Copying //UE4/Dev-Framework to //UE4/Dev-Main (Source: //UE4/Dev-Framework @ 3038004)

diff --git a/Engine/Source/Runtime/Engine/Private/ActiveSound.cpp b/Engine/Source/Runtime/Engine/Private/ActiveSound.cpp
--- a/Engine/Source/Runtime/Engine/Private/ActiveSound.cpp
+++ b/Engine/Source/Runtime/Engine/Private/ActiveSound.cpp
@@ -7,12 +7,15 @@
 #include "Sound/SoundWave.h"
 #include "Sound/SoundNodeAttenuation.h"
 #include "SubtitleManager.h"
+#include "AudioThread.h"
 
 
 FActiveSound::FActiveSound()
-	: Sound(nullptr)
-	, World(nullptr)
-	, AudioComponent(nullptr)
+	: World(nullptr)
+	, WorldID(0)
+	, Sound(nullptr)
+	, AudioComponentID(0)
+	, OwnerID(0)
 	, AudioDevice(nullptr)
 	, ConcurrencyGroupID(0)
 	, ConcurrencyGeneration(0)
@@ -39,14 +42,16 @@ FActiveSound::FActiveSound()
 	, bCenterChannelOnly(false)
 	, bIsPreviewSound(false)
 	, bGotInteriorSettings(false)
+	, bApplyInteriorVolumes(false)
 #if !(NO_LOGGING || UE_BUILD_SHIPPING || UE_BUILD_TEST)
 	, bWarnedAboutOrphanedLooping(false)
 #endif
 	, bEnableLowPassFilter(false)
 	, bOcclusionAsyncTrace(true)
+	, bIsAudible(true)
+	, UserIndex(0)
 	, bIsOccluded(false)
 	, bAsyncOcclusionPending(false)
-	, UserIndex(0)
 	, PlaybackTime(0.f)
 	, RequestedStartTime(0.f)
 	, CurrentAdjustVolumeMultiplier(1.f)
@@ -68,13 +73,12 @@ FActiveSound::FActiveSound()
 	, LastOcclusionCheckTime(0.f)
 	, MaxDistance(WORLD_MAX)
 	, LastLocation(FVector::ZeroVector)
-	, LastAudioVolume(nullptr)
+	, AudioVolumeID(0)
 	, LastUpdateTime(0.f)
 	, SourceInteriorVolume(1.f)
 	, SourceInteriorLPF(MAX_FILTER_FREQUENCY)
 	, CurrentInteriorVolume(1.f)
 	, CurrentInteriorLPF(MAX_FILTER_FREQUENCY)
-	, bIsAudible(true)
 	, ClosestListenerPtr(nullptr)
 {
 }
@@ -90,7 +94,6 @@ FArchive& operator<<( FArchive& Ar, FActiveSound* ActiveSound )
 	if( !Ar.IsLoading() && !Ar.IsSaving() )
 	{
 		Ar << ActiveSound->Sound;
-		Ar << ActiveSound->LastAudioVolume;
 		Ar << ActiveSound->WaveInstances;
 		Ar << ActiveSound->SoundNodeOffsetMap;
 	}
@@ -110,42 +113,84 @@ void FActiveSound::AddReferencedObjects( FReferenceCollector& Collector)
 	}
 
 	Collector.AddReferencedObject(Sound);
-	Collector.AddReferencedObject(LastAudioVolume);
 	Collector.AddReferencedObject(SoundClassOverride);
 }
 
-USoundClass* FActiveSound::GetSoundClass() const
+void FActiveSound::SetWorld(UWorld* InWorld)
 {
-	if (SoundClassOverride)
+	check(IsInGameThread());
+
+	World = InWorld;
+	WorldID = (InWorld ? InWorld->GetUniqueID() : 0);
+}
+
+void FActiveSound::SetSound(USoundBase* InSound)
+{
+	check(IsInGameThread());
+
+	Sound = InSound;
+	bApplyInteriorVolumes = (SoundClassOverride && SoundClassOverride->Properties.bApplyAmbientVolumes)
+							|| (Sound && Sound->ShouldApplyInteriorVolumes());
+}
+
+void FActiveSound::SetSoundClass(USoundClass* SoundClass)
+{
+	check(IsInGameThread());
+
+	SoundClassOverride = SoundClass;
+	bApplyInteriorVolumes = (SoundClassOverride && SoundClassOverride->Properties.bApplyAmbientVolumes)
+							|| (Sound && Sound->ShouldApplyInteriorVolumes());
+}
+
+void FActiveSound::SetAudioComponent(UAudioComponent* Component)
+{
+	check(IsInGameThread());
+
+	AActor* Owner = Component->GetOwner();
+
+	AudioComponentID = Component->GetAudioComponentID();
+	AudioComponentName = Component->GetFName();
+
+	if (Owner)
 	{
-		return SoundClassOverride;
+		OwnerID = Owner->GetUniqueID();
+		OwnerName = Owner->GetFName();
 	}
-	else if (Sound)
+	else
 	{
-		return Sound->GetSoundClass();
+		OwnerID = 0;
+		OwnerName = NAME_None;
 	}
 
-	return NULL;
+	}
+
+FString FActiveSound::GetAudioComponentName() const
+{
+	return (AudioComponentID > 0 ? AudioComponentName.ToString() : TEXT("NO COMPONENT"));
 }
 
-int32 FActiveSound::FindClosestListener( const TArray<FListener>& InListeners ) const
+FString FActiveSound::GetOwnerName() const
 {
-	return FAudioDevice::FindClosestListenerIndex(Transform, InListeners);
+	return (OwnerID > 0 ? OwnerName.ToString() : TEXT("None"));
 }
 
-uint32 FActiveSound::TryGetOwnerID() const
+USoundClass* FActiveSound::GetSoundClass() const
 {
-	// Only have an owner of the active sound if the audio component is valid and if the audio component has an owner.
-	if (UAudioComponent* AudioComponentPtr = AudioComponent.Get())
+	if (SoundClassOverride)
 	{
-		AActor* Owner = AudioComponentPtr->GetOwner();
-		if (Owner)
+		return SoundClassOverride;
+	}
+	else if (Sound)
 		{
-			return Owner->GetUniqueID();
+		return Sound->GetSoundClass();
 		}
+
+	return nullptr;
 	}
 
-	return 0;
+int32 FActiveSound::FindClosestListener( const TArray<FListener>& InListeners ) const
+{
+	return FAudioDevice::FindClosestListenerIndex(Transform, InListeners);
 }
 
 const FSoundConcurrencySettings* FActiveSound::GetSoundConcurrencySettingsToApply() const
@@ -194,20 +239,22 @@ void FActiveSound::UpdateWaveInstances( TArray<FWaveInstance*> &InWaveInstances,
 
 	int32 ClosestListenerIndex = 0;
 
-	if (AudioDevice->Listeners.Num() > 1)
+	const TArray<FListener>& Listeners = AudioDevice->GetListeners();
+
+	if (Listeners.Num() > 1)
 	{
 		SCOPE_CYCLE_COUNTER( STAT_AudioFindNearestLocation );
-		ClosestListenerIndex = FindClosestListener(AudioDevice->Listeners);
+		ClosestListenerIndex = FindClosestListener(Listeners);
 	}
 
 	// Cache the closest listener ptr 
-	ClosestListenerPtr = &AudioDevice->Listeners[ClosestListenerIndex];
+	ClosestListenerPtr = &Listeners[ClosestListenerIndex];
 
 	// The apparent max distance factors the actual max distance of the sound scaled with the distance scale due to focus effects
 	float ApparentMaxDistance = MaxDistance * FocusDistanceScale;
 
 	// Update whether or not his sound is out of range
-	bIsAudible = AudioDevice->LocationIsAudible(Transform.GetTranslation(), *ClosestListenerPtr, ApparentMaxDistance);
+	bIsAudible = AudioDevice->LocationIsAudible(Transform.GetTranslation(), ClosestListenerPtr->Transform, ApparentMaxDistance);
 
 	FSoundParseParameters ParseParams;
 	ParseParams.Transform = Transform;
@@ -220,7 +267,7 @@ void FActiveSound::UpdateWaveInstances( TArray<FWaveInstance*> &InWaveInstances,
 	UpdateAdjustVolumeMultiplier(DeltaTime);
 
 	// If the sound is a preview sound, then ignore the transient master volume and application volume
-	float MasterVolume = AudioDevice->TransientMasterVolume; 
+	float MasterVolume = AudioDevice->GetTransientMasterVolume(); 
 	float ApplicationVolume = FApp::GetVolumeMultiplier();
 	if (bIsPreviewSound)
 	{
@@ -236,7 +283,7 @@ void FActiveSound::UpdateWaveInstances( TArray<FWaveInstance*> &InWaveInstances,
 	ParseParams.LowPassFilterFrequency = LowPassFilterFrequency;
 	ParseParams.SoundClass = GetSoundClass();
 
-	if (ParseParams.SoundClass && ParseParams.SoundClass->Properties.bApplyAmbientVolumes)
+	if (bApplyInteriorVolumes)
 	{
 		// Additional inside/outside processing for ambient sounds
 		// If we aren't in a world there is no interior volumes to be handled.
@@ -264,7 +311,7 @@ void FActiveSound::UpdateWaveInstances( TArray<FWaveInstance*> &InWaveInstances,
 		// if the closest listener is not the primary one, transform the sound transform so it's panned relative to primary listener position
 		if (ClosestListenerIndex != 0)
 		{
-			const FListener& Listener = AudioDevice->Listeners[0];
+			const FListener& Listener = Listeners[0];
 			ParseParams.Transform = ParseParams.Transform * ClosestListenerPtr->Transform.Inverse() * Listener.Transform;
 		}
 
@@ -312,20 +359,24 @@ void FActiveSound::Stop()
 		FWaveInstance*& WaveInstance = WaveInstanceIt.Value();
 
 		// Stop the owning sound source
-		FSoundSource* Source = AudioDevice->WaveInstanceSourceMap.FindRef( WaveInstance );
+		FSoundSource* Source = AudioDevice->GetSoundSource(WaveInstance);
 		if( Source )
 		{
 			Source->Stop();
 		}
 
-		// TODO - Audio Threading. This call would be a task back to game thread
-		// Dequeue subtitles for this sounds
-		FSubtitleManager::GetSubtitleManager()->KillSubtitles( ( PTRINT )WaveInstance );
+		// Dequeue subtitles for this sounds on the game thread
+		DECLARE_CYCLE_STAT(TEXT("FGameThreadAudioTask.KillSubtitles"), STAT_AudioKillSubtitles, STATGROUP_TaskGraphTasks);
+		const PTRINT WaveInstanceID = (PTRINT)WaveInstance;
+		FAudioThread::RunCommandOnGameThread([WaveInstanceID]()
+		{
+			FSubtitleManager::GetSubtitleManager()->KillSubtitles(WaveInstanceID);
+		}, GET_STATID(STAT_AudioKillSubtitles));
 
 		delete WaveInstance;
 
 		// Null the entry out temporarily as later Stop calls could try to access this structure
-		WaveInstance = NULL;
+		WaveInstance = nullptr;
 	}
 	WaveInstances.Empty();
 
@@ -335,7 +386,7 @@ void FActiveSound::Stop()
 FWaveInstance* FActiveSound::FindWaveInstance( const UPTRINT WaveInstanceHash )
 {
 	FWaveInstance** WaveInstance = WaveInstances.Find(WaveInstanceHash);
-	return (WaveInstance ? *WaveInstance : NULL);
+	return (WaveInstance ? *WaveInstance : nullptr);
 }
 
 void FActiveSound::UpdateAdjustVolumeMultiplier(const float DeltaTime)
@@ -408,15 +459,12 @@ void FActiveSound::CheckOcclusion(const FVector ListenerLocation, const FVector
 		LastOcclusionCheckTime = WorldTime;
 		static FName NAME_SoundOcclusion = FName(TEXT("SoundOcclusion"));
 
-		// Make sure we ignore the sound's actor geometry
-		AActor* IgnoreActor = nullptr;
-		if (UAudioComponent* AudioComponentPtr = GetAudioComponent())
+		FCollisionQueryParams Params(NAME_SoundOcclusion, AttenuationSettingsPtr->bUseComplexCollisionForOcclusion);
+		if (OwnerID > 0)
 		{
-			IgnoreActor = AudioComponentPtr->GetOwner();
+			Params.AddIgnoredActor(OwnerID);
 		}
 
-		FCollisionQueryParams Params(NAME_SoundOcclusion, AttenuationSettingsPtr->bUseComplexCollisionForOcclusion, IgnoreActor);
-
 		if (bOcclusionAsyncTrace)
 		{
 			ECollisionChannel OcclusionTraceChannel = AttenuationSettingsPtr->OcclusionTraceChannel;
@@ -432,7 +480,7 @@ void FActiveSound::CheckOcclusion(const FVector ListenerLocation, const FVector
 				// Only do async occlusion trace if we've already made one. The first trace must be synchronous to avoid issues with sounds starting playing as occluded
 				bIsOccluded = WorldPtr->LineTraceTestByChannel(SoundLocation, ListenerLocation, OcclusionTraceChannel, Params);
 			}
-			// don't need to do anyother async trace if we've already got one pending
+			// don't need to do another async trace if we've already got one pending
 			else if (!bAsyncOcclusionPending)
 			{
 				bAsyncOcclusionPending = true;
@@ -453,29 +501,16 @@ const TCHAR* GetAWaveName(TMap<UPTRINT, struct FWaveInstance*> WaveInstances)
 
 void FActiveSound::HandleInteriorVolumes( const FListener& Listener, FSoundParseParameters& ParseParams )
 {
-	UWorld* WorldPtr = World.Get();
-	if (WorldPtr == nullptr)
-	{
-		return;
-	}
-
 	// Get the settings of the ambient sound
-	FInteriorSettings Ambient;
-	class AAudioVolume* AudioVolume;
-
-	if (!bGotInteriorSettings || (GIsEditor && !WorldPtr->IsGameWorld()) || (ParseParams.Transform.GetTranslation() - LastLocation).SizeSquared() > KINDA_SMALL_NUMBER)
+	if (!bGotInteriorSettings || (ParseParams.Transform.GetTranslation() - LastLocation).SizeSquared() > KINDA_SMALL_NUMBER)
 	{
-		AudioVolume = WorldPtr->GetAudioSettings(ParseParams.Transform.GetTranslation(), NULL, &Ambient);
-		LastInteriorSettings = Ambient;
-		LastAudioVolume = AudioVolume;
+		FAudioDevice::FAudioVolumeSettings AudioVolumeSettings;
+		AudioDevice->GetAudioVolumeSettings(WorldID, ParseParams.Transform.GetTranslation(), AudioVolumeSettings);
+
+		InteriorSettings = AudioVolumeSettings.InteriorSettings;
+		AudioVolumeID = AudioVolumeSettings.AudioVolumeID;
 		bGotInteriorSettings = true;
 	}
-	else
-	{
-		// use previous settings as we haven't moved
-		Ambient = LastInteriorSettings;
-		AudioVolume = LastAudioVolume;
-	}
 
 	// Check to see if we've moved to a new audio volume
 	if (LastUpdateTime < Listener.InteriorStartTime)
@@ -485,11 +520,11 @@ void FActiveSound::HandleInteriorVolumes( const FListener& Listener, FSoundParse
 		LastUpdateTime = FApp::GetCurrentTime();
 	}
 
-	if (Listener.Volume == AudioVolume || !bAllowSpatialization)
+	if (Listener.AudioVolumeID == AudioVolumeID || !bAllowSpatialization)
 	{
 		// Ambient and listener in same ambient zone
 		CurrentInteriorVolume = FMath::Lerp(SourceInteriorVolume, 1.0f, Listener.InteriorVolumeInterp);
-		ParseParams.VolumeMultiplier *= CurrentInteriorVolume;
+		ParseParams.InteriorVolumeMultiplier = CurrentInteriorVolume;
 
 		CurrentInteriorLPF = FMath::Lerp(SourceInteriorLPF, MAX_FILTER_FREQUENCY, Listener.InteriorLPFInterp);
 		ParseParams.AmbientZoneFilterFrequency = CurrentInteriorLPF;
@@ -497,26 +532,23 @@ void FActiveSound::HandleInteriorVolumes( const FListener& Listener, FSoundParse
 	else
 	{
 		// Ambient and listener in different ambient zone
-		if( Ambient.bIsWorldSettings )
+		if( InteriorSettings.bIsWorldSettings )
 		{
 			// The ambient sound is 'outside' - use the listener's exterior volume
 			CurrentInteriorVolume = FMath::Lerp(SourceInteriorVolume, Listener.InteriorSettings.ExteriorVolume, Listener.ExteriorVolumeInterp);
-			ParseParams.VolumeMultiplier *= CurrentInteriorVolume;
+			ParseParams.InteriorVolumeMultiplier = CurrentInteriorVolume;
 
 			CurrentInteriorLPF = FMath::Lerp(SourceInteriorLPF, Listener.InteriorSettings.ExteriorLPF, Listener.ExteriorLPFInterp);
 			ParseParams.AmbientZoneFilterFrequency = CurrentInteriorLPF;
-
-			UE_LOG(LogAudio, Verbose, TEXT( "Ambient in diff volume, ambient outside. Volume *= %g LPF *= %g (%s)" ),
-				CurrentInteriorVolume, CurrentInteriorLPF, ( WaveInstances.Num() > 0 ) ? GetAWaveName(WaveInstances) : TEXT( "NULL" ) );
 		}
 		else
 		{
 			// The ambient sound is 'inside' - use the ambient sound's interior volume multiplied with the listeners exterior volume
-			CurrentInteriorVolume = FMath::Lerp(SourceInteriorVolume, Ambient.InteriorVolume, Listener.InteriorVolumeInterp);
+			CurrentInteriorVolume = FMath::Lerp(SourceInteriorVolume, InteriorSettings.InteriorVolume, Listener.InteriorVolumeInterp);
 			CurrentInteriorVolume *= FMath::Lerp(SourceInteriorVolume, Listener.InteriorSettings.ExteriorVolume, Listener.ExteriorVolumeInterp);
-			ParseParams.VolumeMultiplier *= CurrentInteriorVolume;
+			ParseParams.InteriorVolumeMultiplier = CurrentInteriorVolume;
 
-			float AmbientLPFValue = FMath::Lerp(SourceInteriorLPF, Ambient.InteriorLPF, Listener.InteriorLPFInterp);
+			float AmbientLPFValue = FMath::Lerp(SourceInteriorLPF, InteriorSettings.InteriorLPF, Listener.InteriorLPFInterp);
 			float ListenerLPFValue = FMath::Lerp(SourceInteriorLPF, Listener.InteriorSettings.ExteriorLPF, Listener.ExteriorLPFInterp);
 
 			// The current interior LPF value is the less of the LPF due to ambient zone and LPF due to listener settings
@@ -530,9 +562,6 @@ void FActiveSound::HandleInteriorVolumes( const FListener& Listener, FSoundParse
 				CurrentInteriorLPF = ListenerLPFValue;
 				ParseParams.AmbientZoneFilterFrequency = ListenerLPFValue;
 			}
-
-			UE_LOG(LogAudio, Verbose, TEXT( "Ambient in diff volume, ambient inside. Volume *= %g LPF *= %g (%s)" ),
-				CurrentInteriorVolume, CurrentInteriorLPF, ( WaveInstances.Num() > 0 ) ? GetAWaveName(WaveInstances) : TEXT( "NULL" ) );
 		}
 	}
 }
@@ -782,10 +811,8 @@ void FActiveSound::ApplyAttenuation(FSoundParseParameters& ParseParams, const FL
 
 	FAttenuationListenerData ListenerData;
 
-	ListenerData.Listener = &Listener;
-
 	// Get the current focus factor
-	const float FocusFactor = AudioDevice->GetFocusFactor(ListenerData, Sound, SoundTransform, *Settings, &Listener);
+	const float FocusFactor = AudioDevice->GetFocusFactor(ListenerData, Sound, SoundTransform, *Settings, &Listener.Transform);
 
 	// Reset distance and priority scale to 1.0 in case changed in editor
 	FocusDistanceScale = 1.0f;
@@ -795,16 +822,18 @@ void FActiveSound::ApplyAttenuation(FSoundParseParameters& ParseParams, const FL
 	check(Sound);
 	if (Settings->bSpatialize && Settings->bEnableListenerFocus && !Sound->bIgnoreFocus)
 	{
+		const FGlobalFocusSettings& FocusSettings = AudioDevice->GetGlobalFocusSettings();
+
 		// Get the volume scale to apply the volume calculation based on the focus factor
-		const float FocusVolumeAttenuation = Settings->GetFocusAttenuation(AudioDevice->GlobalFocusSettings, FocusFactor);
+		const float FocusVolumeAttenuation = Settings->GetFocusAttenuation(FocusSettings, FocusFactor);
 		Volume *= FocusVolumeAttenuation;
 
 		// Scale the volume-weighted priority scale value we use for sorting this sound for voice-stealing
-		FocusPriorityScale = Settings->GetFocusPriorityScale(AudioDevice->GlobalFocusSettings, FocusFactor);
+		FocusPriorityScale = Settings->GetFocusPriorityScale(FocusSettings, FocusFactor);
 		ParseParams.Priority *= FocusPriorityScale;
 
 		// Get the distance scale to use when computing distance-calculations for 3d attenuation
-		FocusDistanceScale = Settings->GetFocusDistanceScale(AudioDevice->GlobalFocusSettings, FocusFactor);
+		FocusDistanceScale = Settings->GetFocusDistanceScale(FocusSettings, FocusFactor);
 	}
 
 	// Attenuate the volume based on the model
@@ -815,7 +844,7 @@ void FActiveSound::ApplyAttenuation(FSoundParseParameters& ParseParams, const FL
 			case EAttenuationShape::Sphere:
 			{
 				// Update attenuation data in-case it hasn't been updated
-				AudioDevice->GetAttenuationListenerData(ListenerData, SoundTransform, *Settings, &Listener);
+				AudioDevice->GetAttenuationListenerData(ListenerData, SoundTransform, *Settings, &Listener.Transform);
 				Volume *= Settings->AttenuationEval(ListenerData.AttenuationDistance, Settings->FalloffDistance, FocusDistanceScale);
 				break;
 			}
@@ -837,7 +866,7 @@ void FActiveSound::ApplyAttenuation(FSoundParseParameters& ParseParams, const FL
 		}
 	}
 
-	// Only due occlusion traces if the sound is audible
+	// Only do occlusion traces if the sound is audible
 	if (Settings->bEnableOcclusion && Volume > 0.0f && !AudioDevice->IsAudioDeviceMuted())
 	{
 		check(ClosestListenerPtr);
@@ -853,7 +882,7 @@ void FActiveSound::ApplyAttenuation(FSoundParseParameters& ParseParams, const FL
 	// Attenuate with the low pass filter if necessary
 	if (Settings->bAttenuateWithLPF)
 	{
-		AudioDevice->GetAttenuationListenerData(ListenerData, SoundTransform, *Settings, &Listener);
+		AudioDevice->GetAttenuationListenerData(ListenerData, SoundTransform, *Settings, &Listener.Transform);
 
 		// Attenuate with the low pass filter if necessary
 		FVector2D InputRange(Settings->LPFRadiusMin, Settings->LPFRadiusMax);
diff --git a/Engine/Source/Runtime/Engine/Private/Actor.cpp b/Engine/Source/Runtime/Engine/Private/Actor.cpp
--- a/Engine/Source/Runtime/Engine/Private/Actor.cpp
+++ b/Engine/Source/Runtime/Engine/Private/Actor.cpp
@@ -2588,29 +2588,31 @@ void AActor::PostEditImport()
 /** Util that sets up the actor's component hierarchy (when users forget to do so, in their native ctor) */
 static USceneComponent* FixupNativeActorComponents(AActor* Actor)
 {
-	TInlineComponentArray<USceneComponent*> SceneComponents;
-	Actor->GetComponents(SceneComponents);
-
 	USceneComponent* SceneRootComponent = Actor->GetRootComponent();
-	if ((SceneRootComponent == nullptr) && (SceneComponents.Num() > 0))
+	if (SceneRootComponent == nullptr)
 	{
-		UE_LOG(LogActor, Warning, TEXT("%s has natively added scene component(s), but none of them were set as the actor's RootComponent - picking one arbitrarily"), *Actor->GetFullName());
+		TInlineComponentArray<USceneComponent*> SceneComponents;
+		Actor->GetComponents(SceneComponents);
+		if (SceneComponents.Num() > 0)
+		{
+			UE_LOG(LogActor, Warning, TEXT("%s has natively added scene component(s), but none of them were set as the actor's RootComponent - picking one arbitrarily"), *Actor->GetFullName());
 	
-		// if the user forgot to set one of their native components as the root, 
-		// we arbitrarily pick one for them (otherwise the SCS could attempt to 
-		// create its own root, and nest native components under it)
-		for (USceneComponent* Component : SceneComponents)
-		{
-			if ((Component == nullptr) ||
-				(Component->GetAttachParent() != nullptr) ||
-				(Component->CreationMethod != EComponentCreationMethod::Native))
+			// if the user forgot to set one of their native components as the root, 
+			// we arbitrarily pick one for them (otherwise the SCS could attempt to 
+			// create its own root, and nest native components under it)
+			for (USceneComponent* Component : SceneComponents)
 			{
-				continue;
-			}
+				if ((Component == nullptr) ||
+					(Component->GetAttachParent() != nullptr) ||
+					(Component->CreationMethod != EComponentCreationMethod::Native))
+				{
+					continue;
+				}
 
-			SceneRootComponent = Component;
-			Actor->SetRootComponent(Component);
-			break;
+				SceneRootComponent = Component;
+				Actor->SetRootComponent(Component);
+				break;
+			}
 		}
 	}
 
@@ -2626,7 +2628,7 @@ static void ValidateDeferredTransformCache()
 	// could happen if an actor is destroyed before FinishSpawning is called
 	for (auto It = GSpawnActorDeferredTransformCache.CreateIterator(); It; ++It)
 	{
-		TWeakObjectPtr<AActor> ActorRef = It.Key();
+		const TWeakObjectPtr<AActor>& ActorRef = It.Key();
 		if (ActorRef.IsValid() == false)
 		{
 			It.RemoveCurrent();
@@ -3780,8 +3782,8 @@ void AActor::DispatchPhysicsCollisionHit(const FRigidBodyCollisionInfo& MyInfo,
 	Result.PhysMaterial = ContactInfo.PhysMaterial[1];
 	Result.Actor = OtherInfo.Actor;
 	Result.Component = OtherInfo.Component;
-	Result.Item = MyInfo.BodyIndex;
-	Result.BoneName = MyInfo.BoneName;
+	Result.Item = OtherInfo.BodyIndex;
+	Result.BoneName = OtherInfo.BoneName;
 	Result.bBlockingHit = true;
 
 	NotifyHit(MyInfo.Component.Get(), OtherInfo.Actor.Get(), OtherInfo.Component.Get(), true, Result.Location, Result.Normal, RigidCollisionData.TotalNormalImpulse, Result);
diff --git a/Engine/Source/Runtime/Engine/Private/Animation/AnimCompressionDerivedData.h b/Engine/Source/Runtime/Engine/Private/Animation/AnimCompressionDerivedData.h
--- a/Engine/Source/Runtime/Engine/Private/Animation/AnimCompressionDerivedData.h
+++ b/Engine/Source/Runtime/Engine/Private/Animation/AnimCompressionDerivedData.h
@@ -43,7 +43,7 @@ public:
 		// This is a version string that mimics the old versioning scheme. If you
 		// want to bump this version, generate a new guid using VS->Tools->Create GUID and
 		// return it here. Ex.
-		return TEXT("D6830EA6B283485FA15ACC36BBCFA79A");
+		return TEXT("F0CAEDBE5C164D9A8DC1079DC4954CFB");
 	}
 
 	virtual FString GetPluginSpecificCacheKeySuffix() const override;
diff --git a/Engine/Source/Runtime/Engine/Private/Animation/AnimInstance.cpp b/Engine/Source/Runtime/Engine/Private/Animation/AnimInstance.cpp
--- a/Engine/Source/Runtime/Engine/Private/Animation/AnimInstance.cpp
+++ b/Engine/Source/Runtime/Engine/Private/Animation/AnimInstance.cpp
@@ -252,7 +252,7 @@ void UAnimInstance::UninitializeAnimation()
 		TArray<FName> ParamsToClearCopy = MaterialParamatersToClear;
 		for(int i = 0; i < ParamsToClearCopy.Num(); ++i)
 		{
-			AddCurveValue(ParamsToClearCopy[i], 0.0f, ACF_DrivesMaterial);
+			AddCurveValue(ParamsToClearCopy[i], 0.0f, ACF_DriveMaterial);
 		}
 	}
 
@@ -935,24 +935,14 @@ void UAnimInstance::DisplayDebug(class UCanvas* Canvas, const FDebugDisplayInfo&
 
 			DisplayDebugManager.SetLinearDrawColor(TextYellow);
 
-			Heading = FString::Printf(TEXT("Event Curves: %i"), AnimationCurves[(uint8)EAnimCurveType::EventCurve].Num());
+			Heading = FString::Printf(TEXT("Event Curves: %i"), AnimationCurves[(uint8)EAnimCurveType::AttributeCurve].Num());
 			DisplayDebugManager.DrawString(Heading, Indent);
 
 			DisplayDebugManager.SetLinearDrawColor(TextWhite);
 
 			{
 				FIndenter EventCurveIndent(Indent);
-				OutputCurveMap(AnimationCurves[(uint8)EAnimCurveType::EventCurve], Canvas, DisplayDebugManager, Indent);
-			}
-
-			Heading = FString::Printf(TEXT("Pose Curves: %i"), AnimationCurves[(uint8)EAnimCurveType::PoseCurve].Num());
-			DisplayDebugManager.DrawString(Heading, Indent);
-
-			DisplayDebugManager.SetLinearDrawColor(TextWhite);
-
-			{
-				FIndenter EventCurveIndent(Indent);
-				OutputCurveMap(AnimationCurves[(uint8)EAnimCurveType::PoseCurve], Canvas, DisplayDebugManager, Indent);
+				OutputCurveMap(AnimationCurves[(uint8)EAnimCurveType::AttributeCurve], Canvas, DisplayDebugManager, Indent);
 			}
 		}
 	}
@@ -1042,9 +1032,9 @@ void UAnimInstance::AddCurveValue(const FName& CurveName, float Value, int32 Cur
 
 	// save curve value, it will overwrite if same exists, 
 	//CurveValues.Add(CurveName, Value);
-	if (CurveTypeFlags & ACF_TriggerEvent)
+	if (CurveTypeFlags & ACF_DriveAttribute)
 	{
-		float *CurveValPtr = AnimationCurves[(uint8)EAnimCurveType::EventCurve].Find(CurveName);
+		float *CurveValPtr = AnimationCurves[(uint8)EAnimCurveType::AttributeCurve].Find(CurveName);
 		if ( CurveValPtr )
 		{
 			// sum up, in the future we might normalize, but for now this just sums up
@@ -1053,11 +1043,11 @@ void UAnimInstance::AddCurveValue(const FName& CurveName, float Value, int32 Cur
 		}
 		else
 		{
-			AnimationCurves[(uint8)EAnimCurveType::EventCurve].Add(CurveName, Value);
+			AnimationCurves[(uint8)EAnimCurveType::AttributeCurve].Add(CurveName, Value);
 		}
 	}
 
-	if (CurveTypeFlags & ACF_DrivesMorphTarget)
+	if (CurveTypeFlags & ACF_DriveMorphTarget)
 	{
 		float *CurveValPtr = AnimationCurves[(uint8)EAnimCurveType::MorphTargetCurve].Find(CurveName);
 		if ( CurveValPtr )
@@ -1072,7 +1062,7 @@ void UAnimInstance::AddCurveValue(const FName& CurveName, float Value, int32 Cur
 		}
 	}
 
-	if (CurveTypeFlags & ACF_DrivesMaterial)
+	if (CurveTypeFlags & ACF_DriveMaterial)
 	{
 		MaterialParamatersToClear.RemoveSwap(CurveName);
 		float* CurveValPtr = AnimationCurves[(uint8)EAnimCurveType::MaterialCurve].Find(CurveName);
@@ -1085,19 +1075,6 @@ void UAnimInstance::AddCurveValue(const FName& CurveName, float Value, int32 Cur
 			AnimationCurves[(uint8)EAnimCurveType::MaterialCurve].Add(CurveName, Value);
 		}
 	}
-
-	if (CurveTypeFlags & ACF_DrivesPose)
-	{
-		float* CurveValPtr = AnimationCurves[(uint8)EAnimCurveType::PoseCurve].Find(CurveName);
-		if (CurveValPtr)
-		{
-			*CurveValPtr = Value;
-		}
-		else
-		{
-			AnimationCurves[(uint8)EAnimCurveType::PoseCurve].Add(CurveName, Value);
-		}
-	}
 }
 
 void UAnimInstance::AddCurveValue(const USkeleton::AnimCurveUID Uid, float Value, int32 CurveTypeFlags)
@@ -1140,25 +1117,20 @@ void UAnimInstance::GetAnimationCurveList(int32 CurveFlags, TMap<FName, float>&
 {
 	OutCurveList.Reset();
 
-	if (CurveFlags & ACF_TriggerEvent)
+	if (CurveFlags & ACF_DriveAttribute)
 	{
-		OutCurveList.Append(AnimationCurves[(uint8)EAnimCurveType::EventCurve]);
+		OutCurveList.Append(AnimationCurves[(uint8)EAnimCurveType::AttributeCurve]);
 	}
 
-	if (CurveFlags & ACF_DrivesMorphTarget)
+	if (CurveFlags & ACF_DriveMorphTarget)
 	{
 		OutCurveList.Append(AnimationCurves[(uint8)EAnimCurveType::MorphTargetCurve]);
 	}
 
-	if (CurveFlags & ACF_DrivesMaterial)
+	if (CurveFlags & ACF_DriveMaterial)
 	{
 		OutCurveList.Append(AnimationCurves[(uint8)EAnimCurveType::MaterialCurve]);
 	}
-
-	if (CurveFlags & ACF_DrivesPose)
-	{
-		OutCurveList.Append(AnimationCurves[(uint8)EAnimCurveType::PoseCurve]);
-	}
 }
 
 void UAnimInstance::RefreshCurves(USkeletalMeshComponent* Component)
@@ -1207,20 +1179,25 @@ void UAnimInstance::UpdateCurves(const FBlendedHeapCurve& InCurve)
 	TArray<FName> ParamsToClearCopy = MaterialParamatersToClear;
 	for(int i = 0; i < ParamsToClearCopy.Num(); ++i)
 	{
-		AddCurveValue(ParamsToClearCopy[i], 0.0f, ACF_DrivesMaterial);
-	}
-
-#if WITH_EDITOR
-	// if we're supporting this in-game, this code has to change to work with UID
-	for (auto& AddAnimCurveDelegate : OnAddAnimationCurves)
-	{
-		if (AddAnimCurveDelegate.IsBound())
-		{
-			AddAnimCurveDelegate.Execute(this);
-		}
-	}
-
-#endif
+		AddCurveValue(ParamsToClearCopy[i], 0.0f, ACF_DriveMaterial);
+	}
+
+	// @todo: delete me later when james g's change goes in
+	// this won't work well because pose needs to be handled in evaluate
+	// the question is that if we'd like to support preview in anim graph
+	// that will need better handling of the curves - currently UI curves are inserted to 
+	// SignleNodeInstance->PreviewOverride
+// #if WITH_EDITOR
+// 	// if we're supporting this in-game, this code has to change to work with UID
+// 	for (auto& AddAnimCurveDelegate : OnAddAnimationCurves)
+// 	{
+// 		if (AddAnimCurveDelegate.IsBound())
+// 		{
+// 			AddAnimCurveDelegate.Execute(this);
+// 		}
+// 	}
+// 
+// #endif
 	// update curves to component
 	UpdateCurvesToComponents(GetOwningComponent());
 }
@@ -1380,7 +1357,7 @@ float UAnimInstance::GetSlotMontageGlobalWeight(FName SlotNodeName) const
 
 float UAnimInstance::GetCurveValue(FName CurveName)
 {
-	float* Value = AnimationCurves[(uint8)EAnimCurveType::EventCurve].Find(CurveName);
+	float* Value = AnimationCurves[(uint8)EAnimCurveType::AttributeCurve].Find(CurveName);
 	if (Value)
 	{
 		return *Value;
@@ -1791,13 +1768,13 @@ void UAnimInstance::StopSlotAnimation(float InBlendOutTime, FName SlotNodeName)
 	}
 }
 
-bool UAnimInstance::IsPlayingSlotAnimation(UAnimSequenceBase* Asset, FName SlotNodeName)
+bool UAnimInstance::IsPlayingSlotAnimation(const UAnimSequenceBase* Asset, FName SlotNodeName) const
 {
 	UAnimMontage* Montage = nullptr;
 	return IsPlayingSlotAnimation(Asset, SlotNodeName, Montage);
 }
 
-bool UAnimInstance::IsPlayingSlotAnimation(UAnimSequenceBase* Asset, FName SlotNodeName, UAnimMontage*& OutMontage)
+bool UAnimInstance::IsPlayingSlotAnimation(const UAnimSequenceBase* Asset, FName SlotNodeName, UAnimMontage*& OutMontage) const
 {
 	for (int32 InstanceIndex = 0; InstanceIndex < MontageInstances.Num(); InstanceIndex++)
 	{
@@ -1873,11 +1850,11 @@ float UAnimInstance::Montage_Play(UAnimMontage* MontageToPlay, float InPlayRate/
 	return 0.f;
 }
 
-void UAnimInstance::Montage_Stop(float InBlendOutTime, UAnimMontage* Montage)
+void UAnimInstance::Montage_Stop(float InBlendOutTime, const UAnimMontage* Montage)
 {
 	if (Montage)
 	{
-		FAnimMontageInstance* MontageInstance = GetActiveInstanceForMontage(*Montage);
+		FAnimMontageInstance* MontageInstance = GetActiveInstanceForMontage(Montage);
 		if (MontageInstance)
 		{
 			MontageInstance->Stop(FAlphaBlend(Montage->BlendOut, InBlendOutTime));
@@ -1897,11 +1874,11 @@ void UAnimInstance::Montage_Stop(float InBlendOutTime, UAnimMontage* Montage)
 	}
 }
 
-void UAnimInstance::Montage_Pause(UAnimMontage* Montage)
+void UAnimInstance::Montage_Pause(const UAnimMontage* Montage)
 {
 	if (Montage)
 	{
-		FAnimMontageInstance* MontageInstance = GetActiveInstanceForMontage(*Montage);
+		FAnimMontageInstance* MontageInstance = GetActiveInstanceForMontage(Montage);
 		if (MontageInstance)
 		{
 			MontageInstance->Pause();
@@ -1921,11 +1898,11 @@ void UAnimInstance::Montage_Pause(UAnimMontage* Montage)
 	}
 }
 
-void UAnimInstance::Montage_Resume(UAnimMontage* Montage)
+void UAnimInstance::Montage_Resume(const UAnimMontage* Montage)
 {
 	if (Montage)
 	{
-		FAnimMontageInstance* MontageInstance = GetActiveInstanceForMontage(*Montage);
+		FAnimMontageInstance* MontageInstance = GetActiveInstanceForMontage(Montage);
 		if (MontageInstance && !MontageInstance->IsPlaying())
 		{
 			MontageInstance->SetPlaying(true);
@@ -1945,11 +1922,11 @@ void UAnimInstance::Montage_Resume(UAnimMontage* Montage)
 	}
 }
 
-void UAnimInstance::Montage_JumpToSection(FName SectionName, UAnimMontage* Montage)
+void UAnimInstance::Montage_JumpToSection(FName SectionName, const UAnimMontage* Montage)
 {
 	if (Montage)
 	{
-		FAnimMontageInstance* MontageInstance = GetActiveInstanceForMontage(*Montage);
+		FAnimMontageInstance* MontageInstance = GetActiveInstanceForMontage(Montage);
 		if (MontageInstance)
 		{
 			bool const bEndOfSection = (MontageInstance->GetPlayRate() < 0.f);
@@ -1971,11 +1948,11 @@ void UAnimInstance::Montage_JumpToSection(FName SectionName, UAnimMontage* Monta
 	}
 }
 
-void UAnimInstance::Montage_JumpToSectionsEnd(FName SectionName, UAnimMontage* Montage)
+void UAnimInstance::Montage_JumpToSectionsEnd(FName SectionName, const UAnimMontage* Montage)
 {
 	if (Montage)
 	{
-		FAnimMontageInstance* MontageInstance = GetActiveInstanceForMontage(*Montage);
+		FAnimMontageInstance* MontageInstance = GetActiveInstanceForMontage(Montage);
 		if (MontageInstance)
 		{
 			bool const bEndOfSection = (MontageInstance->GetPlayRate() >= 0.f);
@@ -1997,11 +1974,11 @@ void UAnimInstance::Montage_JumpToSectionsEnd(FName SectionName, UAnimMontage* M
 	}
 }
 
-void UAnimInstance::Montage_SetNextSection(FName SectionNameToChange, FName NextSection, UAnimMontage* Montage)
+void UAnimInstance::Montage_SetNextSection(FName SectionNameToChange, FName NextSection, const UAnimMontage* Montage)
 {
 	if (Montage)
 	{
-		FAnimMontageInstance* MontageInstance = GetActiveInstanceForMontage(*Montage);
+		FAnimMontageInstance* MontageInstance = GetActiveInstanceForMontage(Montage);
 		if (MontageInstance)
 		{
 			MontageInstance->SetNextSectionName(SectionNameToChange, NextSection);
@@ -2029,11 +2006,11 @@ void UAnimInstance::Montage_SetNextSection(FName SectionNameToChange, FName Next
 	}
 }
 
-void UAnimInstance::Montage_SetPlayRate(UAnimMontage* Montage, float NewPlayRate)
+void UAnimInstance::Montage_SetPlayRate(const UAnimMontage* Montage, float NewPlayRate)
 {
 	if (Montage)
 	{
-		FAnimMontageInstance* MontageInstance = GetActiveInstanceForMontage(*Montage);
+		FAnimMontageInstance* MontageInstance = GetActiveInstanceForMontage(Montage);
 		if (MontageInstance)
 		{
 			MontageInstance->SetPlayRate(NewPlayRate);
@@ -2053,11 +2030,11 @@ void UAnimInstance::Montage_SetPlayRate(UAnimMontage* Montage, float NewPlayRate
 	}
 }
 
-bool UAnimInstance::Montage_IsActive(UAnimMontage* Montage)
+bool UAnimInstance::Montage_IsActive(const UAnimMontage* Montage) const
 {
 	if (Montage)
 	{
-		FAnimMontageInstance* MontageInstance = GetActiveInstanceForMontage(*Montage);
+		FAnimMontageInstance* MontageInstance = GetActiveInstanceForMontage(Montage);
 		if (MontageInstance)
 		{
 			return true;
@@ -2079,11 +2056,11 @@ bool UAnimInstance::Montage_IsActive(UAnimMontage* Montage)
 	return false;
 }
 
-bool UAnimInstance::Montage_IsPlaying(UAnimMontage* Montage)
+bool UAnimInstance::Montage_IsPlaying(const UAnimMontage* Montage) const
 {
 	if (Montage)
 	{
-		FAnimMontageInstance* MontageInstance = GetActiveInstanceForMontage(*Montage);
+		FAnimMontageInstance* MontageInstance = GetActiveInstanceForMontage(Montage);
 		if (MontageInstance)
 		{
 			return MontageInstance->IsPlaying();
@@ -2105,11 +2082,11 @@ bool UAnimInstance::Montage_IsPlaying(UAnimMontage* Montage)
 	return false;
 }
 
-FName UAnimInstance::Montage_GetCurrentSection(UAnimMontage* Montage)
-{
+FName UAnimInstance::Montage_GetCurrentSection(const UAnimMontage* Montage) const
+{ 
 	if (Montage)
 	{
-		FAnimMontageInstance* MontageInstance = GetActiveInstanceForMontage(*Montage);
+		FAnimMontageInstance* MontageInstance = GetActiveInstanceForMontage(Montage);
 		if (MontageInstance)
 		{
 			return MontageInstance->GetCurrentSection();
@@ -2135,7 +2112,7 @@ void UAnimInstance::Montage_SetEndDelegate(FOnMontageEnded& InOnMontageEnded, UA
 {
 	if (Montage)
 	{
-		FAnimMontageInstance* MontageInstance = GetActiveInstanceForMontage(*Montage);
+		FAnimMontageInstance* MontageInstance = GetActiveInstanceForMontage(Montage);
 		if (MontageInstance)
 		{
 			MontageInstance->OnMontageEnded = InOnMontageEnded;
@@ -2159,7 +2136,7 @@ void UAnimInstance::Montage_SetBlendingOutDelegate(FOnMontageBlendingOutStarted&
 {
 	if (Montage)
 	{
-		FAnimMontageInstance* MontageInstance = GetActiveInstanceForMontage(*Montage);
+		FAnimMontageInstance* MontageInstance = GetActiveInstanceForMontage(Montage);
 		if (MontageInstance)
 		{
 			MontageInstance->OnMontageBlendingOutStarted = InOnMontageBlendingOut;
@@ -2183,7 +2160,7 @@ FOnMontageBlendingOutStarted* UAnimInstance::Montage_GetBlendingOutDelegate(UAni
 {
 	if (Montage)
 	{
-		FAnimMontageInstance* MontageInstance = GetActiveInstanceForMontage(*Montage);
+		FAnimMontageInstance* MontageInstance = GetActiveInstanceForMontage(Montage);
 		if (MontageInstance)
 		{
 			return &MontageInstance->OnMontageBlendingOutStarted;
@@ -2205,11 +2182,11 @@ FOnMontageBlendingOutStarted* UAnimInstance::Montage_GetBlendingOutDelegate(UAni
 	return NULL;
 }
 
-void UAnimInstance::Montage_SetPosition(UAnimMontage* Montage, float NewPosition)
+void UAnimInstance::Montage_SetPosition(const UAnimMontage* Montage, float NewPosition)
 {
 	if (Montage)
 	{
-		FAnimMontageInstance* MontageInstance = GetActiveInstanceForMontage(*Montage);
+		FAnimMontageInstance* MontageInstance = GetActiveInstanceForMontage(Montage);
 		if (MontageInstance)
 		{
 			MontageInstance->SetPosition(NewPosition);
@@ -2229,11 +2206,11 @@ void UAnimInstance::Montage_SetPosition(UAnimMontage* Montage, float NewPosition
 	}
 }
 
-float UAnimInstance::Montage_GetPosition(UAnimMontage* Montage)
+float UAnimInstance::Montage_GetPosition(const UAnimMontage* Montage)
 {
 	if (Montage)
 	{
-		FAnimMontageInstance* MontageInstance = GetActiveInstanceForMontage(*Montage);
+		FAnimMontageInstance* MontageInstance = GetActiveInstanceForMontage(Montage);
 		if (MontageInstance)
 		{
 			return MontageInstance->GetPosition();
@@ -2255,21 +2232,21 @@ float UAnimInstance::Montage_GetPosition(UAnimMontage* Montage)
 	return 0.f;
 }
 
-bool UAnimInstance::Montage_GetIsStopped(UAnimMontage* Montage)
+bool UAnimInstance::Montage_GetIsStopped(const UAnimMontage* Montage)
 {
 	if (Montage)
 	{
-		FAnimMontageInstance* MontageInstance = GetActiveInstanceForMontage(*Montage);
+		FAnimMontageInstance* MontageInstance = GetActiveInstanceForMontage(Montage);
 		return (!MontageInstance); // Not active == Stopped.
 	}
 	return true;
 }
 
-float UAnimInstance::Montage_GetBlendTime(UAnimMontage* Montage)
+float UAnimInstance::Montage_GetBlendTime(const UAnimMontage* Montage)
 {
 	if (Montage)
 	{
-		FAnimMontageInstance* MontageInstance = GetActiveInstanceForMontage(*Montage);
+		FAnimMontageInstance* MontageInstance = GetActiveInstanceForMontage(Montage);
 		if (MontageInstance)
 		{
 			return MontageInstance->GetBlendTime();
@@ -2291,11 +2268,11 @@ float UAnimInstance::Montage_GetBlendTime(UAnimMontage* Montage)
 	return 0.f;
 }
 
-float UAnimInstance::Montage_GetPlayRate(UAnimMontage* Montage)
+float UAnimInstance::Montage_GetPlayRate(const UAnimMontage* Montage)
 {
 	if (Montage)
 	{
-		FAnimMontageInstance* MontageInstance = GetActiveInstanceForMontage(*Montage);
+		FAnimMontageInstance* MontageInstance = GetActiveInstanceForMontage(Montage);
 		if (MontageInstance)
 		{
 			return MontageInstance->GetPlayRate();
@@ -2321,7 +2298,7 @@ int32 UAnimInstance::Montage_GetNextSectionID(UAnimMontage const* const Montage,
 {
 	if (Montage)
 	{
-		FAnimMontageInstance* MontageInstance = GetActiveInstanceForMontage(*Montage);
+		FAnimMontageInstance* MontageInstance = GetActiveInstanceForMontage(Montage);
 		if (MontageInstance)
 		{
 			return MontageInstance->GetNextSectionID(CurrentSectionID);
@@ -2439,7 +2416,12 @@ void UAnimInstance::ClearMontageInstanceReferences(FAnimMontageInstance& InMonta
 
 FAnimMontageInstance* UAnimInstance::GetActiveInstanceForMontage(UAnimMontage const& Montage) const
 {
-	FAnimMontageInstance* const* FoundInstancePtr = ActiveMontagesMap.Find(&Montage);
+	return GetActiveInstanceForMontage(&Montage);
+}
+
+FAnimMontageInstance* UAnimInstance::GetActiveInstanceForMontage(const UAnimMontage* Montage) const
+{
+	FAnimMontageInstance* const* FoundInstancePtr = ActiveMontagesMap.Find(Montage);
 	return FoundInstancePtr ? *FoundInstancePtr : nullptr;
 }
 
diff --git a/Engine/Source/Runtime/Engine/Private/Animation/AnimInstanceProxy.cpp b/Engine/Source/Runtime/Engine/Private/Animation/AnimInstanceProxy.cpp
--- a/Engine/Source/Runtime/Engine/Private/Animation/AnimInstanceProxy.cpp
+++ b/Engine/Source/Runtime/Engine/Private/Animation/AnimInstanceProxy.cpp
@@ -12,6 +12,7 @@
 #include "Animation/AnimNode_StateMachine.h"
 #include "Animation/AnimNode_TransitionResult.h"
 #include "Animation/AnimNode_SaveCachedPose.h"
+#include "Animation/AnimNode_SubInput.h"
 
 #define DO_ANIMSTAT_PROCESSING(StatName) DEFINE_STAT(STAT_ ## StatName)
 #include "AnimMTStats.h"
@@ -140,6 +141,8 @@ void FAnimInstanceProxy::InitializeRootNode()
 				FAnimNode_Base* AnimNode = Property->ContainerPtrToValuePtr<FAnimNode_Base>(AnimInstanceObject);
 				if (AnimNode)
 				{
+					AnimNode->RootInitialize(this);
+
 					// Force our functions to be re-evaluated - this reinitialization may have been a 
 					// consequence of our class being recompiled and functions will be invalid in that
 					// case.
@@ -161,6 +164,12 @@ void FAnimInstanceProxy::InitializeRootNode()
 						FAnimNode_StateMachine* StateMachine = static_cast<FAnimNode_StateMachine*>(AnimNode);
 						StateMachine->CacheMachineDescription(AnimClassInterface);
 					}
+
+					if(Property->Struct->IsChildOf(FAnimNode_SubInput::StaticStruct()))
+					{
+						check(!SubInstanceInputNode); // Should only ever have one
+						SubInstanceInputNode = static_cast<FAnimNode_SubInput*>(AnimNode);
+					}
 				}
 			}
 		}
@@ -173,6 +182,7 @@ void FAnimInstanceProxy::InitializeRootNode()
 
 void FAnimInstanceProxy::Uninitialize(UAnimInstance* InAnimInstance)
 {
+	SubInstanceInputNode = nullptr;
 }
 
 void FAnimInstanceProxy::PreUpdate(UAnimInstance* InAnimInstance, float DeltaSeconds)
diff --git a/Engine/Source/Runtime/Engine/Private/Animation/AnimMontage.cpp b/Engine/Source/Runtime/Engine/Private/Animation/AnimMontage.cpp
--- a/Engine/Source/Runtime/Engine/Private/Animation/AnimMontage.cpp
+++ b/Engine/Source/Runtime/Engine/Private/Animation/AnimMontage.cpp
@@ -2094,7 +2094,7 @@ void FAnimMontageInstance::SetMatineeAnimPositionInner(FName SlotName, USkeletal
 				if(CurrentlyPlayingMontage.IsValid())
 				{
 					// set it's weight to 0
-					struct FAnimMontageInstance* PrevAnimMontageInst = AnimInst->GetActiveInstanceForMontage(*CurrentlyPlayingMontage);
+					struct FAnimMontageInstance* PrevAnimMontageInst = AnimInst->GetActiveInstanceForMontage(CurrentlyPlayingMontage.Get());
 					if(PrevAnimMontageInst)
 					{
 						// set weight to be 0
@@ -2112,7 +2112,7 @@ void FAnimMontageInstance::SetMatineeAnimPositionInner(FName SlotName, USkeletal
 				ensure(CurrentlyPlayingMontage.IsValid());
 			}
 
-			struct FAnimMontageInstance* AnimMontageInst = AnimInst->GetActiveInstanceForMontage(*CurrentlyPlayingMontage);
+			struct FAnimMontageInstance* AnimMontageInst = AnimInst->GetActiveInstanceForMontage(CurrentlyPlayingMontage.Get());
 			if(AnimMontageInst)
 			{
 				// need to set weight to be 1
@@ -2168,7 +2168,7 @@ void FAnimMontageInstance::PreviewMatineeSetAnimPositionInner(FName SlotName, US
 			if(CurrentlyPlayingMontage.IsValid())
 			{
 				// set it's weight to 0
-				struct FAnimMontageInstance* PrevAnimMontageInst = AnimInst->GetActiveInstanceForMontage(*CurrentlyPlayingMontage);
+				struct FAnimMontageInstance* PrevAnimMontageInst = AnimInst->GetActiveInstanceForMontage(CurrentlyPlayingMontage.Get());
 				if(PrevAnimMontageInst)
 				{
 					// set weight to be 0
@@ -2186,7 +2186,7 @@ void FAnimMontageInstance::PreviewMatineeSetAnimPositionInner(FName SlotName, US
 			ensure(CurrentlyPlayingMontage.IsValid());
 		}
 
-		struct FAnimMontageInstance* AnimMontageInst = AnimInst->GetActiveInstanceForMontage(*CurrentlyPlayingMontage);
+		struct FAnimMontageInstance* AnimMontageInst = AnimInst->GetActiveInstanceForMontage(CurrentlyPlayingMontage.Get());
 		if(AnimMontageInst)
 		{
 			// set weight to be 1
diff --git a/Engine/Source/Runtime/Engine/Private/Animation/AnimNodeBase.cpp b/Engine/Source/Runtime/Engine/Private/Animation/AnimNodeBase.cpp
--- a/Engine/Source/Runtime/Engine/Private/Animation/AnimNodeBase.cpp
+++ b/Engine/Source/Runtime/Engine/Private/Animation/AnimNodeBase.cpp
@@ -226,7 +226,7 @@ void FPoseLink::Evaluate(FPoseContext& Output)
 	if (LinkedNode != NULL)
 	{
 #if ENABLE_ANIMNODE_POSE_DEBUG
-		CurrentPose.ResetToIdentity();
+		CurrentPose.ResetToAdditiveIdentity();
 #endif
 		LinkedNode->Evaluate(Output);
 #if ENABLE_ANIMNODE_POSE_DEBUG
@@ -427,7 +427,7 @@ void FExposedValueHandler::Initialize(FAnimNode_Base* AnimNode, UObject* AnimIns
 	}
 
 	// initialize copy records
-	for(auto& CopyRecord : CopyRecords)
+	for(FExposedValueCopyRecord& CopyRecord : CopyRecords)
 	{
 		UProperty* SourceProperty = AnimInstanceObject->GetClass()->FindPropertyByName(CopyRecord.SourcePropertyName);
 		check(SourceProperty);
@@ -467,13 +467,31 @@ void FExposedValueHandler::Initialize(FAnimNode_Base* AnimNode, UObject* AnimIns
 			check(ArrayHelper.IsValidIndex(CopyRecord.DestArrayIndex));
 			CopyRecord.Dest = ArrayHelper.GetRawPtr(CopyRecord.DestArrayIndex);
 			CopyRecord.CachedBoolDestProperty = Cast<UBoolProperty>(CopyRecord.DestProperty);
-			CopyRecord.CachedDestContainer = AnimNode;
+
+			if(CopyRecord.bInstanceIsTarget)
+			{
+				CopyRecord.CachedDestContainer = AnimInstanceObject;
+			}
+			else
+			{
+				CopyRecord.CachedDestContainer = AnimNode;
+			}
 		}
 		else
 		{
 			CopyRecord.Dest = CopyRecord.DestProperty->ContainerPtrToValuePtr<uint8>(AnimNode, CopyRecord.DestArrayIndex);
+			
+			if(CopyRecord.bInstanceIsTarget)
+			{
+				CopyRecord.CachedDestContainer = AnimInstanceObject;
+				CopyRecord.Dest = CopyRecord.DestProperty->ContainerPtrToValuePtr<uint8>(AnimInstanceObject, CopyRecord.DestArrayIndex);
+			}
+			else
+			{
+				CopyRecord.CachedDestContainer = AnimNode;
+			}
+
 			CopyRecord.CachedBoolDestProperty = Cast<UBoolProperty>(CopyRecord.DestProperty);
-			CopyRecord.CachedDestContainer = AnimNode;
 		}
 	}
 
diff --git a/Engine/Source/Runtime/Engine/Private/Animation/AnimNode_SubInput.cpp b/Engine/Source/Runtime/Engine/Private/Animation/AnimNode_SubInput.cpp
new file mode 100644
--- /dev/null
+++ b/Engine/Source/Runtime/Engine/Private/Animation/AnimNode_SubInput.cpp
@@ -0,0 +1,37 @@
+// Copyright 1998-2016 Epic Games, Inc. All Rights Reserved.
+
+#include "EnginePrivate.h"
+#include "Animation/AnimNode_SubInput.h"
+
+void FAnimNode_SubInput::Initialize(const FAnimationInitializeContext& Context)
+{
+	FAnimNode_Base::Initialize(Context);
+}
+
+void FAnimNode_SubInput::CacheBones(const FAnimationCacheBonesContext& Context)
+{
+}
+
+void FAnimNode_SubInput::Update(const FAnimationUpdateContext& Context)
+{
+	EvaluateGraphExposedInputs.Execute(Context);
+}
+
+void FAnimNode_SubInput::Evaluate(FPoseContext& Output)
+{
+	if(InputPose.IsValid() && InputCurve.IsValid())
+	{
+		Output.Pose.CopyBonesFrom(InputPose);
+		Output.Curve.CopyFrom(InputCurve);
+	}
+	else
+	{
+		Output.ResetToRefPose();
+	}
+}
+
+void FAnimNode_SubInput::GatherDebugData(FNodeDebugData& DebugData)
+{
+	FString DebugLine = DebugData.GetNodeName(this);
+	DebugData.AddDebugItem(DebugLine);
+}
diff --git a/Engine/Source/Runtime/Engine/Private/Animation/AnimNode_SubInstance.cpp b/Engine/Source/Runtime/Engine/Private/Animation/AnimNode_SubInstance.cpp
new file mode 100644
--- /dev/null
+++ b/Engine/Source/Runtime/Engine/Private/Animation/AnimNode_SubInstance.cpp
@@ -0,0 +1,200 @@
+// Copyright 1998-2016 Epic Games, Inc. All Rights Reserved.
+
+#include "EnginePrivate.h"
+#include "Components/SkeletalMeshComponent.h"
+#include "AnimInstanceProxy.h"
+#include "Animation/AnimNode_SubInstance.h"
+#include "Animation/AnimNode_SubInput.h"
+
+FAnimNode_SubInstance::FAnimNode_SubInstance()
+	: InstanceClass(nullptr)
+	, InstanceToRun(nullptr)
+{
+
+}
+
+void FAnimNode_SubInstance::Initialize(const FAnimationInitializeContext& Context)
+{
+	FAnimNode_Base::Initialize(Context);
+
+	InPose.Initialize(Context);
+
+	if(InstanceToRun)
+	{
+		InstanceToRun->InitializeAnimation();
+	}
+}
+
+void FAnimNode_SubInstance::ValidateInstance()
+{
+	if(!InstanceToRun)
+	{
+		// Need an instance to run
+		if(*InstanceClass)
+		{
+			InstanceToRun = NewObject<UAnimInstance>(GetTransientPackage(), InstanceClass);
+		}
+		else
+		{
+			UE_LOG(LogAnimation, Warning, TEXT("FAnimNode_SubInstance: Attempting to use a subgraph node without setting an instance class."));
+		}
+	}
+}
+
+void FAnimNode_SubInstance::CacheBones(const FAnimationCacheBonesContext& Context)
+{
+
+}
+
+void FAnimNode_SubInstance::Update(const FAnimationUpdateContext& Context)
+{
+	InPose.Update(Context);
+	EvaluateGraphExposedInputs.Execute(Context);
+
+	if(InstanceToRun)
+	{
+		FAnimInstanceProxy& Proxy = InstanceToRun->GetProxyOnAnyThread<FAnimInstanceProxy>();
+
+		// Only update if we've not had a single-threaded update already
+		if(InstanceToRun->bNeedsUpdate)
+		{
+			Proxy.UpdateAnimation();
+		}
+
+		check(InstanceProperties.Num() == SubInstanceProperties.Num());
+		for(int32 PropIdx = 0; PropIdx < InstanceProperties.Num(); ++PropIdx)
+		{
+			UProperty* CallerProperty = InstanceProperties[PropIdx];
+			UProperty* SubProperty = SubInstanceProperties[PropIdx];
+
+			check(CallerProperty && SubProperty);
+
+			uint8* SrcPtr = CallerProperty->ContainerPtrToValuePtr<uint8>(Context.AnimInstanceProxy->GetAnimInstanceObject());
+			uint8* DestPtr = SubProperty->ContainerPtrToValuePtr<uint8>(InstanceToRun);
+
+			CallerProperty->CopyCompleteValue(DestPtr, SrcPtr);
+		}
+	}
+}
+
+void FAnimNode_SubInstance::Evaluate(FPoseContext& Output)
+{
+	if(InstanceToRun)
+	{
+		InPose.Evaluate(Output);
+
+		FAnimInstanceProxy& Proxy = InstanceToRun->GetProxyOnAnyThread<FAnimInstanceProxy>();
+		FAnimNode_SubInput* InputNode = Proxy.SubInstanceInputNode;
+
+		if(InputNode)
+		{
+			InputNode->InputPose.CopyBonesFrom(Output.Pose);
+			InputNode->InputCurve.CopyFrom(Output.Curve);
+		}
+
+		InstanceToRun->ParallelEvaluateAnimation(false, nullptr, BoneTransforms, BlendedCurve);
+
+		for(const FCompactPoseBoneIndex BoneIndex : Output.Pose.ForEachBoneIndex())
+		{
+			FMeshPoseBoneIndex MeshBoneIndex = Output.Pose.GetBoneContainer().MakeMeshPoseIndex(BoneIndex);
+			Output.Pose[BoneIndex] = BoneTransforms[MeshBoneIndex.GetInt()];
+		}
+
+		Output.Curve.CopyFrom(BlendedCurve);
+	}
+	else
+	{
+		Output.ResetToRefPose();
+	}
+
+}
+
+void FAnimNode_SubInstance::GatherDebugData(FNodeDebugData& DebugData)
+{
+	// Add our entry
+	FString DebugLine = DebugData.GetNodeName(this);
+
+	// Gather data from the sub instance
+	if(InstanceToRun)
+	{
+		FAnimInstanceProxy& Proxy = InstanceToRun->GetProxyOnAnyThread<FAnimInstanceProxy>();
+		Proxy.GatherDebugData(DebugData);
+	}
+
+	// Pass to next
+	InPose.GatherDebugData(DebugData);
+}
+
+bool FAnimNode_SubInstance::HasPreUpdate() const
+{
+	return true;
+}
+
+void FAnimNode_SubInstance::PreUpdate(const UAnimInstance* InAnimInstance)
+{
+	if(USkeletalMeshComponent* SkelComp = InAnimInstance->GetSkelMeshComponent())
+	{
+		USkeletalMesh* SkelMesh = SkelComp->SkeletalMesh;
+
+		const int32 NumTransforms = SkelComp->GetSpaceBases().Num();
+		BoneTransforms.Empty(NumTransforms);
+		BoneTransforms.AddZeroed(NumTransforms);
+	}
+}
+
+void FAnimNode_SubInstance::RootInitialize(const FAnimInstanceProxy* InProxy)
+{
+	if(*InstanceClass)
+	{
+		USkeletalMeshComponent* MeshComp = InProxy->GetSkelMeshComponent();
+		check(MeshComp);
+
+		// Full reinit, kill old instances
+		if(InstanceToRun)
+		{
+			MeshComp->SubInstances.Remove(InstanceToRun);
+			InstanceToRun->MarkPendingKill();
+			InstanceToRun = nullptr;
+		}
+
+		// Need an instance to run
+		InstanceToRun = NewObject<UAnimInstance>(MeshComp, InstanceClass);
+
+		MeshComp->SubInstances.Add(InstanceToRun);
+
+		// Build property lists
+		InstanceProperties.Empty();
+		SubInstanceProperties.Empty();
+
+		check(SourcePropertyNames.Num() == DestPropertyNames.Num());
+
+		for(int32 Idx = 0; Idx < SourcePropertyNames.Num(); ++Idx)
+		{
+			FName& SourceName = SourcePropertyNames[Idx];
+			FName& DestName = DestPropertyNames[Idx];
+
+			UClass* SourceClass = IAnimClassInterface::GetActualAnimClass(InProxy->GetAnimClassInterface());
+
+			UProperty* SourceProperty = FindField<UProperty>(SourceClass, SourceName);
+			UProperty* DestProperty = FindField<UProperty>(*InstanceClass, DestName);
+
+			check(SourceProperty && DestProperty);
+
+			InstanceProperties.Add(SourceProperty);
+			SubInstanceProperties.Add(DestProperty);
+		}
+		
+	}
+	else if(InstanceToRun)
+	{
+		// We have an instance but no instance class
+		TeardownInstance();
+	}
+}
+
+void FAnimNode_SubInstance::TeardownInstance()
+{
+	InstanceToRun->UninitializeAnimation();
+	InstanceToRun = nullptr;
+}
+
diff --git a/Engine/Source/Runtime/Engine/Private/Animation/AnimSequence.cpp b/Engine/Source/Runtime/Engine/Private/Animation/AnimSequence.cpp
--- a/Engine/Source/Runtime/Engine/Private/Animation/AnimSequence.cpp
+++ b/Engine/Source/Runtime/Engine/Private/Animation/AnimSequence.cpp
@@ -287,7 +287,7 @@ void UAnimSequence::Serialize(FArchive& Ar)
 bool UAnimSequence::IsValidToPlay() const
 {
 	// make sure sequence length is valid and raw animation data exists, and compressed
-	return ( SequenceLength > 0.f && RawAnimationData.Num() > 0 && CompressedTrackOffsets.Num() > 0 );
+	return ( SequenceLength > 0.f);
 }
 #endif
 
@@ -1042,7 +1042,7 @@ void UAnimSequence::GetBonePose(FCompactPose& OutPose, FBlendedCurve& OutCurve,
 	{
 		if (bIsBakedAdditive)
 		{
-			OutPose.ResetToIdentity(); 
+			OutPose.ResetToAdditiveIdentity(); 
 		}
 		else
 		{
@@ -1057,7 +1057,7 @@ void UAnimSequence::GetBonePose(FCompactPose& OutPose, FBlendedCurve& OutCurve,
 	if (bIsBakedAdditive)
 	{
 		//When using baked additive ref pose is identity
-		OutPose.ResetToIdentity();
+		OutPose.ResetToAdditiveIdentity();
 	}
 	else
 	{
@@ -1290,7 +1290,7 @@ void UAnimSequence::GetBonePose_Additive(FCompactPose& OutPose, FBlendedCurve& O
 {
 	if (!IsValidAdditive())
 	{
-		OutPose.ResetToIdentity();
+		OutPose.ResetToAdditiveIdentity();
 		return;
 	}
 
@@ -1354,7 +1354,7 @@ void UAnimSequence::GetBonePose_AdditiveMeshRotationOnly(FCompactPose& OutPose,
 	if (!IsValidAdditive())
 	{
 		// since this is additive, need to initialize to identity
-		OutPose.ResetToIdentity();
+		OutPose.ResetToAdditiveIdentity();
 		return;
 	}
 
diff --git a/Engine/Source/Runtime/Engine/Private/Animation/AnimSingleNodeInstance.cpp b/Engine/Source/Runtime/Engine/Private/Animation/AnimSingleNodeInstance.cpp
--- a/Engine/Source/Runtime/Engine/Private/Animation/AnimSingleNodeInstance.cpp
+++ b/Engine/Source/Runtime/Engine/Private/Animation/AnimSingleNodeInstance.cpp
@@ -91,9 +91,9 @@ void UAnimSingleNodeInstance::SetAnimationAsset(class UAnimationAsset* NewAsset,
 	}
 }
 
-void UAnimSingleNodeInstance::SetPreviewPoseOverride(const FName& PoseName, float Value)
+void UAnimSingleNodeInstance::SetPreviewCurveOverride(const FName& PoseName, float Value, bool bRemoveIfZero)
 {
-	GetProxyOnGameThread<FAnimSingleNodeInstanceProxy>().SetPreviewPoseOverride(PoseName, Value);
+	GetProxyOnGameThread<FAnimSingleNodeInstanceProxy>().SetPreviewCurveOverride(PoseName, Value, bRemoveIfZero);
 }
 
 void UAnimSingleNodeInstance::SetMontageLoop(UAnimMontage* Montage, bool bIsLooping, FName StartingSection)
diff --git a/Engine/Source/Runtime/Engine/Private/Animation/AnimSingleNodeInstanceProxy.cpp b/Engine/Source/Runtime/Engine/Private/Animation/AnimSingleNodeInstanceProxy.cpp
--- a/Engine/Source/Runtime/Engine/Private/Animation/AnimSingleNodeInstanceProxy.cpp
+++ b/Engine/Source/Runtime/Engine/Private/Animation/AnimSingleNodeInstanceProxy.cpp
@@ -142,40 +142,113 @@ bool FAnimSingleNodeInstanceProxy::Evaluate(FPoseContext& Output)
 		}
 		else if (UPoseAsset* PoseAsset = Cast<UPoseAsset>(CurrentAsset))
 		{
-#if WITH_EDITORONLY_DATA
+			const TArray<FSmartName>& PoseNames = PoseAsset->GetPoseNames();
+
+			int32 TotalPoses = PoseNames.Num();
+			FAnimExtractContext ExtractContext;
+			ExtractContext.PoseCurves.AddZeroed(TotalPoses);
+
 			USkeleton* MySkeleton = PoseAsset->GetSkeleton();
- 			for (auto Iter=PreviewPoseOverride.CreateConstIterator(); Iter; ++Iter)
+ 			for (auto Iter=PreviewCurveOverride.CreateConstIterator(); Iter; ++Iter)
  			{
  				const FName& Name = Iter.Key();
 				const float Value = Iter.Value();
 				
- 				FSmartNameMapping::UID CurveUID = MySkeleton->GetUIDByName(USkeleton::AnimCurveMappingName, Name);
- 
- 				if (CurveUID != FSmartNameMapping::MaxUID)
- 				{
-					Output.Curve.Set(CurveUID, Value, ACF_DrivesPose);
- 				}
+				FSmartName PoseName;
+				
+				if (MySkeleton->GetSmartNameByName(USkeleton::AnimCurveMappingName, Name, PoseName))
+				{
+					int32 PoseIndex = PoseNames.Find(PoseName);
+					if (PoseIndex != INDEX_NONE)
+					{
+						ExtractContext.PoseCurves[PoseIndex] = Value;
+					}
+				}
  			}
-#endif 
-			if (PoseAsset->IsValidAdditive())
-			{
-				FAnimExtractContext ExtractContext;
 
-				PoseAsset->GetBaseAnimationPose(Output.Pose, Output.Curve, ExtractContext);
-				PoseAsset->GetAnimationPose(Output.Pose, Output.Curve, ExtractContext);
-			}
-			else
+ 			if (PoseAsset->IsValidAdditive())
+ 			{
+ 				PoseAsset->GetBaseAnimationPose(Output.Pose, Output.Curve, FAnimExtractContext());
+
+				FCompactPose AdditivePose;
+				FBlendedCurve AdditiveCurve;
+				AdditivePose.SetBoneContainer(&Output.Pose.GetBoneContainer());
+				AdditiveCurve.InitFrom(Output.Curve);
+				float Weight = PoseAsset->GetAnimationPose(AdditivePose, AdditiveCurve, ExtractContext);
+				FAnimationRuntime::AccumulateAdditivePose(Output.Pose, AdditivePose, Output.Curve, AdditiveCurve, 1.f, EAdditiveAnimationType::AAT_LocalSpaceBase);
+ 			}
+ 			else
+ 			{
+ 				PoseAsset->GetAnimationPose(Output.Pose, Output.Curve, ExtractContext);
+ 			}
+		}
+
+#if WITH_EDITORONLY_DATA
+		USkeleton* MySkeleton = CurrentAsset->GetSkeleton();
+		for (auto Iter = PreviewCurveOverride.CreateConstIterator(); Iter; ++Iter)
+		{
+			const FName& Name = Iter.Key();
+			const float Value = Iter.Value();
+
+			FSmartName PreviewCurveName;
+
+			// @fixme : this is temporary
+			if (MySkeleton->GetSmartNameByName(USkeleton::AnimCurveMappingName, Name, PreviewCurveName))
 			{
-				// @todo : what do we do for ingame? I think in game, we should just allow to load anim sequence with curves
-				PoseAsset->GetAnimationPose(Output.Pose, Output.Curve, FAnimExtractContext());
+				Output.Curve.Set(PreviewCurveName.UID, Value, ACF_EditorPreviewCurves);
+
 			}
 		}
 
-#if WITH_EDITORONLY_DATA
 		// if it has preview pose asset, we have to handle that after we do all animation
 		if (const UPoseAsset* PoseAsset = CurrentAsset->PreviewPoseAsset)
 		{
-			PoseAsset->GetAnimationPose(Output.Pose, Output.Curve, FAnimExtractContext());
+			// if skeleton doesn't match it won't work
+			if (PoseAsset->GetSkeleton() == MySkeleton)
+			{
+				const TArray<FSmartName>& PoseNames = PoseAsset->GetPoseNames();
+
+				int32 TotalPoses = PoseNames.Num();
+				FAnimExtractContext ExtractContext;
+				ExtractContext.PoseCurves.AddZeroed(TotalPoses);
+
+				for (const auto& PoseName : PoseNames)
+				{
+					if (PoseName.UID != FSmartNameMapping::MaxUID)
+					{
+						int32 PoseIndex = PoseNames.Find(PoseName);
+						if (PoseIndex != INDEX_NONE)
+						{
+							ExtractContext.PoseCurves[PoseIndex] = Output.Curve.Get(PoseName.UID);
+						}
+					}
+				}
+
+				if (PoseAsset->IsValidAdditive())
+				{
+					FCompactPose AdditivePose;
+					FBlendedCurve AdditiveCurve;
+					AdditivePose.SetBoneContainer(&Output.Pose.GetBoneContainer());
+					AdditiveCurve.InitFrom(Output.Curve);
+					float Weight = PoseAsset->GetAnimationPose(AdditivePose, AdditiveCurve, ExtractContext);
+					FAnimationRuntime::AccumulateAdditivePose(Output.Pose, AdditivePose, Output.Curve, AdditiveCurve, 1.f, EAdditiveAnimationType::AAT_LocalSpaceBase);
+				}
+				else
+				{
+					FPoseContext CurrentPose(Output);
+					FPoseContext SourcePose(Output);
+
+					SourcePose = Output;
+
+					if (PoseAsset->GetAnimationPose(CurrentPose.Pose, CurrentPose.Curve, ExtractContext))
+					{
+						TArray<float> BoneWeights;
+						BoneWeights.AddZeroed(CurrentPose.Pose.GetNumBones());
+						// once we get it, we have to blend by weight
+						FAnimationRuntime::BlendTwoPosesTogetherPerBone(CurrentPose.Pose, SourcePose.Pose, CurrentPose.Curve, SourcePose.Curve, BoneWeights, Output.Pose, Output.Curve);
+					}
+				}
+			}
 		}
 #endif // WITH_EDITORONLY_DATA
 	}
@@ -300,15 +373,15 @@ void FAnimSingleNodeInstanceProxy::PreUpdate(UAnimInstance* InAnimInstance, floa
 	FAnimInstanceProxy::PreUpdate(InAnimInstance, DeltaSeconds);
 #if WITH_EDITOR
 	// @fixme only do this in pose asset
-	// copy data to PreviewPoseOverride
-	TMap<FName, float> PoseCurveList;
-
-	InAnimInstance->GetAnimationCurveList(ACF_DrivesPose, PoseCurveList);
-
-	if (PoseCurveList.Num() > 0)
-	{
-		PreviewPoseOverride.Append(PoseCurveList);
-	}
+// 	// copy data to PreviewPoseOverride
+// 	TMap<FName, float> PoseCurveList;
+// 
+// 	InAnimInstance->GetAnimationCurveList(ACF_DrivesPose, PoseCurveList);
+// 
+// 	if (PoseCurveList.Num() > 0)
+// 	{
+// 		PreviewPoseOverride.Append(PoseCurveList);
+// 	}
 #endif // WITH_EDITOR
 }
 void FAnimSingleNodeInstanceProxy::InitializeObjects(UAnimInstance* InAnimInstance)
@@ -326,10 +399,10 @@ void FAnimSingleNodeInstanceProxy::ClearObjects()
 	CurrentAsset = nullptr;
 }
 
-void FAnimSingleNodeInstanceProxy::SetPreviewPoseOverride(const FName& PoseName, float Value)
+void FAnimSingleNodeInstanceProxy::SetPreviewCurveOverride(const FName& PoseName, float Value, bool bRemoveIfZero)
 {
-	float *CurveValPtr = PreviewPoseOverride.Find(PoseName);
-	bool bShouldAddToList = FPlatformMath::Abs(Value) > ZERO_ANIMWEIGHT_THRESH;
+	float *CurveValPtr = PreviewCurveOverride.Find(PoseName);
+	bool bShouldAddToList = bRemoveIfZero == false || FPlatformMath::Abs(Value) > ZERO_ANIMWEIGHT_THRESH;
 	if (bShouldAddToList)
 	{
 		if (CurveValPtr)
@@ -340,15 +413,15 @@ void FAnimSingleNodeInstanceProxy::SetPreviewPoseOverride(const FName& PoseName,
 		}
 		else
 		{
-			PreviewPoseOverride.Add(PoseName, Value);
+			PreviewCurveOverride.Add(PoseName, Value);
 		}
 	}
 	// if less than ZERO_ANIMWEIGHT_THRESH
 	// no reason to keep them on the list
-	else
+	else 
 	{
 		// remove if found
-		PreviewPoseOverride.Remove(PoseName);
+		PreviewCurveOverride.Remove(PoseName);
 	}
 }
 
@@ -399,7 +472,7 @@ void FAnimSingleNodeInstanceProxy::SetAnimationAsset(class UAnimationAsset* NewA
 
 #if WITH_EDITORONLY_DATA
 	PreviewPoseCurrentTime = 0.0f;
-	PreviewPoseOverride.Reset();
+	PreviewCurveOverride.Reset();
 #endif
 
 	
diff --git a/Engine/Source/Runtime/Engine/Private/Animation/AnimationRuntime.cpp b/Engine/Source/Runtime/Engine/Private/Animation/AnimationRuntime.cpp
--- a/Engine/Source/Runtime/Engine/Private/Animation/AnimationRuntime.cpp
+++ b/Engine/Source/Runtime/Engine/Private/Animation/AnimationRuntime.cpp
@@ -243,7 +243,42 @@ void FAnimationRuntime::BlendTwoPosesTogether(
 
 	// Ensure that all of the resulting rotations are normalized
 	ResultPose.NormalizeRotations();
-	ResultCurve.Blend(SourceCurve1, SourceCurve2, 1.f - WeightOfPose1);
+	ResultCurve.Lerp(SourceCurve1, SourceCurve2, 1.f - WeightOfPose1);
+}
+
+void FAnimationRuntime::BlendTwoPosesTogetherPerBone(
+	const FCompactPose& SourcePose1,
+	const FCompactPose& SourcePose2,
+	const FBlendedCurve& SourceCurve1,
+	const FBlendedCurve& SourceCurve2,
+	const TArray<float> WeightsOfSource2,
+	/*out*/ FCompactPose& ResultPose,
+	/*out*/ FBlendedCurve& ResultCurve)
+{
+	
+	for (FCompactPoseBoneIndex BoneIndex : ResultPose.ForEachBoneIndex())
+	{
+		const float BlendWeight = WeightsOfSource2[BoneIndex.GetInt()];
+		if (FAnimationRuntime::IsFullWeight(BlendWeight))
+		{
+			ResultPose[BoneIndex] = SourcePose2[BoneIndex];
+		}
+		// if it doens't have weight, take source pose 1
+		else if (FAnimationRuntime::HasWeight(BlendWeight))
+		{
+			BlendTransform<ETransformBlendMode::Overwrite>(SourcePose1[BoneIndex], ResultPose[BoneIndex], 1.f - BlendWeight);
+			BlendTransform<ETransformBlendMode::Accumulate>(SourcePose2[BoneIndex], ResultPose[BoneIndex], BlendWeight);
+		}
+		else
+		{
+			ResultPose[BoneIndex] = SourcePose1[BoneIndex];
+		}
+	}
+
+	// Ensure that all of the resulting rotations are normalized
+	// @fixme: this has to be fixed with name mapping to joint in the future
+	ResultPose.NormalizeRotations();
+	ResultCurve.Lerp(SourceCurve1, SourceCurve2, 1.f);
 }
 
 template <int32 TRANSFORM_BLEND_MODE>
@@ -446,7 +481,11 @@ void FAnimationRuntime::ConvertTransformToAdditive(FTransform& TargetTransform,
 {
 	TargetTransform.SetRotation(TargetTransform.GetRotation() * BaseTransform.GetRotation().Inverse());
 	TargetTransform.SetTranslation(TargetTransform.GetTranslation() - BaseTransform.GetTranslation());
-	TargetTransform.SetScale3D(TargetTransform.GetScale3D() * BaseTransform.GetSafeScaleReciprocal(BaseTransform.GetScale3D()));
+	// additive scale considers how much it grow or lower
+	// in order to support blending between different additive scale, we save [(target scale)/(source scale) - 1.f], and this can blend with 
+	// other delta scale value
+	// when we apply to the another scale, we apply scale * (1 + [additive scale])
+	TargetTransform.SetScale3D(TargetTransform.GetScale3D() * BaseTransform.GetSafeScaleReciprocal(BaseTransform.GetScale3D()) - 1.f);
 	TargetTransform.NormalizeRotation();
 }
 
@@ -512,18 +551,18 @@ void FAnimationRuntime::AccumulateLocalSpaceAdditivePose(FCompactPose& BasePose,
 {
 	if (Weight > ZERO_ANIMWEIGHT_THRESH)
 	{
+		const ScalarRegister VBlendWeight(Weight);
 		if (Weight >= (1.f - ZERO_ANIMWEIGHT_THRESH))
 		{
 			// fast path, no need to weight additive.
 			for (FCompactPoseBoneIndex BoneIndex : BasePose.ForEachBoneIndex())
 			{
-				BasePose[BoneIndex].Accumulate(AdditivePose[BoneIndex]);
+				BasePose[BoneIndex].AccumulateWithAdditiveScale(AdditivePose[BoneIndex], VBlendWeight);
 			}
 		}
 		else
 		{
 			// Slower path w/ weighting
-			const ScalarRegister VBlendWeight(Weight);
 			for (FCompactPoseBoneIndex BoneIndex : BasePose.ForEachBoneIndex())
 			{
 				// copy additive, because BlendFromIdentityAndAccumulate modifies it.
diff --git a/Engine/Source/Runtime/Engine/Private/Animation/PoseAsset.cpp b/Engine/Source/Runtime/Engine/Private/Animation/PoseAsset.cpp
--- a/Engine/Source/Runtime/Engine/Private/Animation/PoseAsset.cpp
+++ b/Engine/Source/Runtime/Engine/Private/Animation/PoseAsset.cpp
@@ -4,7 +4,29 @@
 #include "AnimationUtils.h"
 #include "AnimationRuntime.h"
 #include "Animation/PoseAsset.h"
+#include "FrameworkObjectVersion.h"
 
+// utility function 
+FSmartName GetUniquePoseName(USkeleton* Skeleton)
+{
+	check(Skeleton);
+	int32 NameIndex = 0;
+
+	FSmartNameMapping::UID NewUID;
+	FName NewName;
+
+	do
+	{
+		NewName = FName(*FString::Printf(TEXT("Pose_%d"), NameIndex++));
+		NewUID = Skeleton->GetUIDByName(USkeleton::AnimCurveMappingName, NewName);
+	} while (NewUID != FSmartNameMapping::MaxUID);
+
+	// if found, 
+	FSmartName NewPoseName;
+	Skeleton->AddSmartNameAndModify(USkeleton::AnimCurveMappingName, NewName, NewPoseName);
+
+	return NewPoseName;
+}
 
 //////////////////////////////////////////////////////////////////////////////////////////////////////////////////
 //
@@ -42,7 +64,7 @@ bool FPoseDataContainer::InsertTrack(const FName& InTrackName, USkeleton* InSkel
 	check(InSkeleton);
 
 	// make sure the transform is correct size
-	if (Tracks.Contains(InTrackName))
+	if (Tracks.Contains(InTrackName) == false)
 	{
 		int32 SkeletonIndex = InSkeleton->GetReferenceSkeleton().FindBoneIndex(InTrackName);
 		int32 TrackIndex = INDEX_NONE;
@@ -93,6 +115,51 @@ void FPoseDataContainer::GetPoseCurve(const FPoseData* PoseData, FBlendedCurve&
 // we don't want to save transform per track (good for compression) because that will slow down blending
 void FPoseDataContainer::Shrink(USkeleton* InSkeleton, FName& InRetargetSourceName)
 {
+	for (auto& Pose : Poses)
+	{
+		check(Pose.LocalSpacePose.Num() == Tracks.Num());
+
+		Pose.LocalSpacePoseMask.SetNumZeroed(Tracks.Num(), true);
+		for (int32 TrackIndex = 0; TrackIndex < Tracks.Num(); ++TrackIndex)
+		{
+			FTransform DefaultTransform = GetDefaultTransform(Tracks[TrackIndex], InSkeleton, InRetargetSourceName);
+			if (!Pose.LocalSpacePose[TrackIndex].Equals(DefaultTransform, KINDA_SMALL_NUMBER))
+			{
+				Pose.LocalSpacePoseMask[TrackIndex] = true;
+			}
+		}
+	}
+
+	for (int32 TrackIndex = 0; TrackIndex < Tracks.Num(); ++TrackIndex)
+	{
+		bool bDeleteTrack = true;
+		for (auto& Pose : Poses)
+		{
+			if (Pose.LocalSpacePoseMask[TrackIndex])
+			{
+				bDeleteTrack = false;
+				break;
+			}
+		}
+
+		// if nothing contains any data, just delete this track
+		if (bDeleteTrack)
+		{
+			DeleteTrack(TrackIndex);
+			--TrackIndex;
+		}
+	}
+}
+
+void FPoseDataContainer::DeleteTrack(int32 TrackIndex)
+{
+	TrackMap.Remove(Tracks[TrackIndex]);
+	Tracks.RemoveAt(TrackIndex);
+	for (auto& Pose : Poses)
+	{
+		Pose.LocalSpacePoseMask.RemoveAt(TrackIndex);
+		Pose.LocalSpacePose.RemoveAt(TrackIndex);
+	}
 }
 
 FTransform FPoseDataContainer::GetDefaultTransform(const FName& InTrackName, USkeleton* InSkeleton, const FName& InRetargetSourceName) const
@@ -230,12 +297,13 @@ void FPoseDataContainer::ConvertToFullPose(int32 InBasePoseIndex, const TArray<F
 			TArray<FTransform> AdditivePose = Poses[PoseIndex].LocalSpacePose;
 			FPoseData& PoseData = Poses[PoseIndex];
 
+			const ScalarRegister AdditiveWeight(1.f);
+
 			check(AdditivePose.Num() == InBasePose.Num());
 			for (int32 BoneIndex = 0; BoneIndex < AdditivePose.Num(); ++BoneIndex)
 			{
 				PoseData.LocalSpacePose[BoneIndex] = InBasePose[BoneIndex];
-				PoseData.LocalSpacePose[BoneIndex].Accumulate(AdditivePose[BoneIndex]);
-				PoseData.LocalSpacePose[BoneIndex].NormalizeRotation();
+				PoseData.LocalSpacePose[BoneIndex].AccumulateWithAdditiveScale(AdditivePose[BoneIndex], AdditiveWeight);
 			}
 
 			int32 CurveNum = Curves.Num();
@@ -261,7 +329,6 @@ void FPoseDataContainer::ConvertToAdditivePose(int32 InBasePoseIndex, const TArr
 			for (int32 BoneIndex = 0; BoneIndex < InBasePose.Num(); ++BoneIndex)
 			{
 				FAnimationRuntime::ConvertTransformToAdditive(PoseData.LocalSpacePose[BoneIndex], InBasePose[BoneIndex]);
-				PoseData.LocalSpacePose[BoneIndex].NormalizeRotation();
 			}
 
 			int32 CurveNum = Curves.Num();
@@ -279,22 +346,20 @@ void FPoseDataContainer::ConvertToAdditivePose(int32 InBasePoseIndex, const TArr
 /////////////////////////////////////////////////////
 UPoseAsset::UPoseAsset(const FObjectInitializer& ObjectInitializer)
 	: Super(ObjectInitializer)
-	, NumPoses(0)
-	, bNormalizeWeight(true)
 	, bAdditivePose(false)
 	, BasePoseIndex(-1)
 {
 }
 
-void UPoseAsset::GetBaseAnimationPose(struct FCompactPose& OutPose, FBlendedCurve& InOutCurve, const FAnimExtractContext& ExtractionContext) const
+void UPoseAsset::GetBaseAnimationPose(struct FCompactPose& OutPose, FBlendedCurve& OutCurve, const FAnimExtractContext& ExtractionContext) const
 {
-	OutPose.ResetToRefPose();
-
 	if (bAdditivePose && PoseContainer.Poses.IsValidIndex(BasePoseIndex))
 	{
 		const FBoneContainer& RequiredBones = OutPose.GetBoneContainer();
 		USkeleton* MySkeleton = GetSkeleton();
 
+		OutPose.ResetToRefPose();
+
 		// this contains compact bone pose list that this pose cares
 		TArray<FCompactPoseBoneIndex> CompactBonePoseList;
 
@@ -318,11 +383,15 @@ void UPoseAsset::GetBaseAnimationPose(struct FCompactPose& OutPose, FBlendedCurv
 			}
 		}
 
-		PoseContainer.GetPoseCurve(&PoseContainer.Poses[BasePoseIndex], InOutCurve);
+		PoseContainer.GetPoseCurve(&PoseContainer.Poses[BasePoseIndex], OutCurve);
+	}
+	else
+	{
+		OutPose.ResetToRefPose();
 	}
 }
 
-void UPoseAsset::GetAnimationPose(struct FCompactPose& OutPose, FBlendedCurve& OutCurve, const FAnimExtractContext& ExtractionContext) const
+bool UPoseAsset::GetAnimationPose(struct FCompactPose& OutPose, FBlendedCurve& OutCurve, const FAnimExtractContext& ExtractionContext) const
 {
 	const FBoneContainer& RequiredBones = OutPose.GetBoneContainer();
 	USkeleton* MySkeleton = GetSkeleton();
@@ -340,114 +409,136 @@ void UPoseAsset::GetAnimationPose(struct FCompactPose& OutPose, FBlendedCurve& O
 		CompactBonePoseList.Add(PoseBoneIndex);
 	}
 
-	TMap<const FPoseData*, float> IndexToWeightMap;
-
-	float TotalWeight = 0.f;
-
 	check(PoseContainer.IsValid());
 
-	for (int32 PoseIndex = 0; PoseIndex < PoseContainer.Poses.Num(); ++PoseIndex)
-	{
-		const FSmartName& PoseName = PoseContainer.PoseNames[PoseIndex];
-		const FPoseData& PoseData = PoseContainer.Poses[PoseIndex];
-
-		float Value = OutCurve.Get(PoseName.UID);
-
-		// we only add to the list if it's not additive Or if it's additive, we don't want to add base pose index
-		// and has weight
-		if ((!bAdditivePose || PoseIndex != BasePoseIndex) && FAnimationRuntime::HasWeight(Value))
-		{
-			IndexToWeightMap.Add(&PoseData, Value);
-			TotalWeight += Value;
-		}
-	}
-
-	// need a copy of pose since we don't know if this can be additive
-	FCompactPose CurrentPose;
-	FBlendedCurve CurrentCurve;
-	CurrentPose.SetBoneContainer(&OutPose.GetBoneContainer());
-	CurrentCurve.InitFrom(OutCurve);
-
 	if (bAdditivePose)
 	{
-		CurrentPose.ResetToIdentity();
+		OutPose.ResetToAdditiveIdentity();
 	}
 	else
 	{
-		CurrentPose.ResetToRefPose();
+		OutPose.ResetToRefPose();
 	}
 
-	const int32 TotalNumberOfValidPoses = IndexToWeightMap.Num();
-	if (TotalNumberOfValidPoses > 0)
-	{
-		TArray<FTransform> BlendedBoneTransform;
-		bool bFirstPose = true;
-		TArray<FBlendedCurve> PoseCurves;
-		TArray<float>	CurveWeights;
+	bool bNormalizeWeight = bAdditivePose == false;
 
-		//if full pose, we'll have to normalize by weight
-		if (bNormalizeWeight && TotalWeight != 1.f)
+	if (ensureAlways(ExtractionContext.PoseCurves.Num() == PoseContainer.Poses.Num()))
+	{
+		TMap<const FPoseData*, float> IndexToWeightMap;
+		float TotalWeight = 0.f;
+		for (int32 PoseIndex = 0; PoseIndex < PoseContainer.Poses.Num(); ++PoseIndex)
 		{
-			for (TPair<const FPoseData*, float>& WeightPair : IndexToWeightMap)
+			const FSmartName& PoseName = PoseContainer.PoseNames[PoseIndex];
+			const FPoseData& PoseData = PoseContainer.Poses[PoseIndex];
+
+			float Value = ExtractionContext.PoseCurves[PoseIndex];
+
+			// we only add to the list if it's not additive Or if it's additive, we don't want to add base pose index
+			// and has weight
+			if ((!bAdditivePose || PoseIndex != BasePoseIndex) && FAnimationRuntime::HasWeight(Value))
 			{
-				WeightPair.Value /= TotalWeight;
+				IndexToWeightMap.Add(&PoseData, Value);
+				TotalWeight += Value;
 			}
 		}
 
-		PoseCurves.AddDefaulted(TotalNumberOfValidPoses);
-		CurveWeights.AddUninitialized(TotalNumberOfValidPoses);
-		int32 PoseIndex = 0;
-		for (const TPair<const FPoseData*, float>& ActivePosePair : IndexToWeightMap)
+		const int32 TotalNumberOfValidPoses = IndexToWeightMap.Num();
+		if (TotalNumberOfValidPoses > 0)
 		{
-			const FPoseData* Pose = ActivePosePair.Key;
-			const float Weight = ActivePosePair.Value;
+			TArray<FTransform> BlendedBoneTransform;
+			bool bFirstPose = true;
+			TArray<FBlendedCurve> PoseCurves;
+			TArray<float>	CurveWeights;
 
-			if (bFirstPose)
+			//if full pose, we'll have to normalize by weight
+			if (bNormalizeWeight && TotalWeight != 1.f)
 			{
-				BlendedBoneTransform.AddUninitialized(Pose->LocalSpacePose.Num());
-				for (int32 TrackIndex = 0; TrackIndex < Pose->LocalSpacePose.Num(); ++TrackIndex)
+				for (TPair<const FPoseData*, float>& WeightPair : IndexToWeightMap)
 				{
-					BlendTransform<ETransformBlendMode::Overwrite>(Pose->LocalSpacePose[TrackIndex], BlendedBoneTransform[TrackIndex], Weight);
+					WeightPair.Value /= TotalWeight;
 				}
-				bFirstPose = false;
 			}
-			else
+
+			BlendedBoneTransform.AddUninitialized(TrackNum);
+			for (int32 TrackIndex = 0; TrackIndex < TrackNum; ++TrackIndex)
 			{
-				for (int32 TrackIndex = 0; TrackIndex < Pose->LocalSpacePose.Num(); ++TrackIndex)
+				bool bValidTransform = false;
+
+				TArray<FTransform> BlendingTransform;
+				TArray<float> BlendingWeights;
+				float TotalLocalWeight = 0.f;
+				for (const TPair<const FPoseData*, float>& ActivePosePair : IndexToWeightMap)
 				{
-					BlendTransform<ETransformBlendMode::Accumulate>(Pose->LocalSpacePose[TrackIndex], BlendedBoneTransform[TrackIndex], Weight);
+					const FPoseData* Pose = ActivePosePair.Key;
+					const float Weight = ActivePosePair.Value;
+
+					if (Pose->LocalSpacePoseMask[TrackIndex])
+					{
+						BlendingTransform.Add(Pose->LocalSpacePose[TrackIndex]);
+						BlendingWeights.Add(Weight);
+						TotalLocalWeight += Weight;
+					}
+				}
+
+				if (BlendingTransform.Num() == 0)
+				{
+					// copy from out default pose
+					if (CompactBonePoseList[TrackIndex] != INDEX_NONE)
+					{
+						BlendedBoneTransform[TrackIndex] = OutPose[CompactBonePoseList[TrackIndex]];
+					}
+				}
+				else 
+				{
+					if (bAdditivePose)
+					{
+						const  FTransform AdditiveIdentity(FQuat::Identity, FVector::ZeroVector, FVector::ZeroVector);
+						BlendedBoneTransform[TrackIndex].Blend(AdditiveIdentity, BlendingTransform[0], BlendingWeights[0]);
+					}
+					else
+					{
+						BlendedBoneTransform[TrackIndex] = BlendingTransform[0] * ScalarRegister(BlendingWeights[0]);
+					}
+				}
+
+				for (int32 BlendIndex = 1; BlendIndex < BlendingTransform.Num(); ++BlendIndex)
+				{
+					BlendedBoneTransform[TrackIndex].AccumulateWithShortestRotation(BlendingTransform[BlendIndex], ScalarRegister(BlendingWeights[BlendIndex]));
 				}
 			}
 
-			CurveWeights[PoseIndex] = Weight;
-			PoseCurves[PoseIndex].InitFrom(OutCurve);
-			PoseContainer.GetPoseCurve(Pose, PoseCurves[PoseIndex]);
-			++PoseIndex;
-		}
+			// collect curves
+			PoseCurves.AddDefaulted(TotalNumberOfValidPoses);
+			CurveWeights.AddUninitialized(TotalNumberOfValidPoses);
+			int32 PoseIndex = 0;
+			for (const TPair<const FPoseData*, float>& ActivePosePair : IndexToWeightMap)
+			{
+				const FPoseData* Pose = ActivePosePair.Key;
+				const float Weight = ActivePosePair.Value;
 
-		// blend curves
-		BlendCurves(PoseCurves, CurveWeights, CurrentCurve);
+				CurveWeights[PoseIndex] = Weight;
+				PoseCurves[PoseIndex].InitFrom(OutCurve);
+				PoseContainer.GetPoseCurve(Pose, PoseCurves[PoseIndex]);
+				++PoseIndex;
+			}
 
-		for (int32 TrackIndex = 0; TrackIndex < TrackNum; ++TrackIndex)
-		{
-			if (CompactBonePoseList[TrackIndex] != INDEX_NONE)
+			// blend curves
+			BlendCurves(PoseCurves, CurveWeights, OutCurve);
+
+			for (int32 TrackIndex = 0; TrackIndex < TrackNum; ++TrackIndex)
 			{
-				CurrentPose[CompactBonePoseList[TrackIndex]] = BlendedBoneTransform[TrackIndex];
-				CurrentPose.NormalizeRotations();
+				if (CompactBonePoseList[TrackIndex] != INDEX_NONE)
+				{
+					OutPose[CompactBonePoseList[TrackIndex]] = BlendedBoneTransform[TrackIndex];
+					OutPose.NormalizeRotations();
+				}
 			}
-		}
 
-		if (bAdditivePose)
-		{
-			FAnimationRuntime::AccumulateAdditivePose(OutPose, CurrentPose, OutCurve, CurrentCurve, TotalWeight, AAT_LocalSpaceBase);
-			OutPose.NormalizeRotations();
-		}
-		else
-		{
-			OutPose = CurrentPose;
-			OutCurve = CurrentCurve;
+			return true;
 		}
 	}
+
+	return false;
 }
 
 void UPoseAsset::PostLoad()
@@ -472,7 +563,37 @@ void UPoseAsset::PostLoad()
 
 void UPoseAsset::Serialize(FArchive& Ar)
 {
+	Ar.UsingCustomVersion(FFrameworkObjectVersion::GUID);;
+
 	Super::Serialize(Ar);
+
+	if (Ar.CustomVer(FFrameworkObjectVersion::GUID) < FFrameworkObjectVersion::PoseAssetSupportPerBoneMask)
+	{
+		// fix curve names
+		USkeleton* MySkeleton = GetSkeleton();
+		if (MySkeleton)
+		{
+			PoseContainer.Shrink(MySkeleton, RetargetSource);
+		}
+	}
+}
+
+void UPoseAsset::GetAssetRegistryTags(TArray<FAssetRegistryTag>& OutTags) const
+{
+	Super::GetAssetRegistryTags(OutTags);
+
+	// Number of poses
+	OutTags.Add(FAssetRegistryTag("Poses", FString::FromInt(GetNumPoses()), FAssetRegistryTag::TT_Numerical));
+}
+
+int32 UPoseAsset::GetNumPoses() const
+{ 
+	return PoseContainer.GetNumPoses();
+}
+
+int32 UPoseAsset::GetNumCurves() const
+{
+	return PoseContainer.Curves.Num();
 }
 
 const TArray<FSmartName> UPoseAsset::GetPoseNames() const
@@ -480,6 +601,11 @@ const TArray<FSmartName> UPoseAsset::GetPoseNames() const
 	return PoseContainer.PoseNames;
 }
 
+const TArray<FName>	UPoseAsset::GetTrackNames() const
+{
+	return PoseContainer.Tracks;
+}
+
 const TArray<FSmartName> UPoseAsset::GetCurveNames() const
 {
 	TArray<FSmartName> CurveNames;
@@ -491,6 +617,70 @@ const TArray<FSmartName> UPoseAsset::GetCurveNames() const
 	return CurveNames;
 }
 
+const TArray<FAnimCurveBase> UPoseAsset::GetCurveData() const
+{
+	return PoseContainer.Curves;
+}
+
+const TArray<float> UPoseAsset::GetCurveValues(const int32 PoseIndex) const
+{
+	TArray<float> ResultCurveData;
+
+	if (PoseContainer.Poses.IsValidIndex(PoseIndex))
+	{
+		ResultCurveData = PoseContainer.Poses[PoseIndex].CurveData;
+	}
+
+	return ResultCurveData;
+}
+
+bool UPoseAsset::GetLocalPoseForTrack(const int32 PoseIndex, const int32 TrackIndex, FTransform& OutTransform) const
+{
+	bool bSuccess = false;
+
+	if (PoseContainer.Poses.IsValidIndex(PoseIndex))
+	{
+		const FPoseData& PoseData = PoseContainer.Poses[PoseIndex];
+		if (PoseData.LocalSpacePose.IsValidIndex(TrackIndex))
+		{
+			OutTransform = PoseData.LocalSpacePose[TrackIndex];
+			bSuccess = true;
+		}
+	}
+
+	return bSuccess;
+}
+
+bool UPoseAsset::GetCurveValue(const int32 PoseIndex, const int32 CurveIndex, float& OutValue) const
+{
+	bool bSuccess = false;
+
+	if (PoseContainer.Poses.IsValidIndex(PoseIndex))
+	{
+		const FPoseData& PoseData = PoseContainer.Poses[PoseIndex];
+		if (PoseData.CurveData.IsValidIndex(CurveIndex))
+		{
+			OutValue = PoseData.CurveData[CurveIndex];
+			bSuccess = true;
+		}
+	}
+
+	return bSuccess;
+}
+
+const int32 UPoseAsset::GetTrackIndexByName(const FName& InTrackName) const
+{
+	int32 ResultTrackIndex = INDEX_NONE;
+
+	// Only search if valid name passed in
+	if (InTrackName != NAME_None)
+	{
+		ResultTrackIndex = PoseContainer.Tracks.Find(InTrackName);
+	}
+
+	return ResultTrackIndex;
+}
+
 
 bool UPoseAsset::ContainsPose(const FName& InPoseName) const
 {
@@ -505,6 +695,11 @@ bool UPoseAsset::ContainsPose(const FName& InPoseName) const
 	return false;
 }
 
+void UPoseAsset::AddOrUpdatePoseWithUniqueName(USkeletalMeshComponent* MeshComponent)
+{
+	AddOrUpdatePose(GetUniquePoseName(GetSkeleton()), MeshComponent);
+}
+
 void UPoseAsset::AddOrUpdatePose(const FSmartName& PoseName, USkeletalMeshComponent* MeshComponent)
 {
 	USkeleton* MySkeleton = GetSkeleton();
@@ -530,6 +725,7 @@ void UPoseAsset::AddOrUpdatePose(const FSmartName& PoseName, USkeletalMeshCompon
 		}
 
 		AddOrUpdatePose(PoseName, TrackNames, BoneTransform);
+		PoseContainer.Shrink(MySkeleton, RetargetSource);
 	}
 }
 
@@ -544,6 +740,10 @@ void UPoseAsset::AddOrUpdatePose(const FSmartName& PoseName, TArray<FName> Track
 		FPoseData* PoseData = PoseContainer.FindOrAddPoseData(PoseName);
 		// now copy all transform back to it. 
 		check(PoseData);
+		int32 TotalTracks = TrackNames.Num();
+		PoseData->LocalSpacePose.Reset(TotalTracks);
+		PoseData->LocalSpacePose.AddUninitialized(TotalTracks);
+		PoseData->LocalSpacePoseMask.SetNumZeroed(TotalTracks, true);
 		PoseContainer.FillUpDefaultPose(PoseData, MySkeleton, RetargetSource);
 
 		const FReferenceSkeleton& RefSkeleton = MySkeleton->GetReferenceSkeleton();
@@ -554,9 +754,9 @@ void UPoseAsset::AddOrUpdatePose(const FSmartName& PoseName, TArray<FName> Track
 			int32 SkeletonIndex = RefSkeleton.FindBoneIndex(TrackName);
 			int32 InternalTrackIndex = PoseContainer.Tracks.Find(TrackName);
 			PoseData->LocalSpacePose[InternalTrackIndex] = LocalTransform[Index];
+			// shrink should take care of this but if somebody just calls this function, it should set it to dirty
+			PoseData->LocalSpacePoseMask[InternalTrackIndex] = true;
 		}
-
-		NumPoses = PoseContainer.GetNumPoses();
 	}
 }
 
@@ -578,27 +778,16 @@ void UPoseAsset::CombineTracks(const TArray<FName>& NewTracks)
 	}
 }
 
-FName GetUniquePoseName(USkeleton* Skeleton)
+#if WITH_EDITOR
+void UPoseAsset::Reinitialize()
 {
-	check(Skeleton);
-	int32 NameIndex = 0;
-
-	FSmartNameMapping::UID NewUID;
-	FName NewName;
-
-	do
-	{
-		NewName = FName(*FString::Printf(TEXT("Pose_%d"), NameIndex++));
-		NewUID = Skeleton->GetUIDByName(USkeleton::AnimCurveMappingName, NewName);
-	} while (NewUID != FSmartNameMapping::MaxUID);
+	PoseContainer.Reset();
 
-	// if found, 
-
-	return NewName;
+	bAdditivePose = false;
+	BasePoseIndex = INDEX_NONE;
 }
 
-#if WITH_EDITOR
-void UPoseAsset::CreatePoseFromAnimation(class UAnimSequence* AnimSequence)
+void UPoseAsset::CreatePoseFromAnimation(class UAnimSequence* AnimSequence, const TArray<FSmartName>* InPoseNames/*== nullptr*/)
 {
 	if (AnimSequence)
 	{
@@ -606,10 +795,12 @@ void UPoseAsset::CreatePoseFromAnimation(class UAnimSequence* AnimSequence)
 		if (TargetSkeleton)
 		{
 			SetSkeleton(TargetSkeleton);
+			SourceAnimation = AnimSequence;
 
-			PoseContainer.Reset();
+			// reinitialize, now we're making new pose from this animation
+			Reinitialize();
 
-			NumPoses = AnimSequence->GetNumberOfFrames();
+			const int32 NumPoses = AnimSequence->GetNumberOfFrames();
 
 			// make sure we have more than one pose
 			if (NumPoses > 0)
@@ -655,7 +846,7 @@ void UPoseAsset::CreatePoseFromAnimation(class UAnimSequence* AnimSequence)
 				// add to skeleton UID, so that it knows the curve data
 				for (int32 PoseIndex = 0; PoseIndex < NumPoses; ++PoseIndex)
 				{
-					FName PoseName = GetUniquePoseName(TargetSkeleton);
+					FSmartName NewPoseName = (InPoseNames && InPoseNames->IsValidIndex(PoseIndex))? (*InPoseNames)[PoseIndex] : GetUniquePoseName(TargetSkeleton);
 					// now get rawanimationdata, and each key is converted to new pose
 					for (int32 TrackIndex = 0; TrackIndex < NumTracks; ++TrackIndex)
 					{
@@ -663,10 +854,6 @@ void UPoseAsset::CreatePoseFromAnimation(class UAnimSequence* AnimSequence)
 						AnimSequence->ExtractBoneTransform(RawTrack, NewPose[TrackIndex], PoseIndex);
 					}
 
-					// add to skeleton UID, so that it knows the curve data
-					FSmartName NewPoseName;
-					TargetSkeleton->AddSmartNameAndModify(USkeleton::AnimCurveMappingName, PoseName, NewPoseName);
-
 					if (TotalFloatCurveCount > 0)
 					{
 						// get curve data
@@ -697,6 +884,34 @@ void UPoseAsset::CreatePoseFromAnimation(class UAnimSequence* AnimSequence)
 	}
 }
 
+void UPoseAsset::UpdatePoseFromAnimation(class UAnimSequence* AnimSequence)
+{
+	if (AnimSequence)
+	{
+		// when you update pose, right now, it just only keeps pose names
+		// in the future we might want to make it more flexible
+		// back up old pose names
+		const TArray<FSmartName> OldPoseNames = PoseContainer.PoseNames;
+		const bool bOldAdditive = bAdditivePose;
+		int32 OldBasePoseIndex = BasePoseIndex;
+		CreatePoseFromAnimation(AnimSequence, &OldPoseNames);
+
+		// fix up additive info if it's additive
+		if (bOldAdditive)
+		{
+			if (PoseContainer.Poses.IsValidIndex(OldBasePoseIndex) == false)
+			{
+				// if it's pointing at invalid index, just reset to ref pose
+				OldBasePoseIndex = INDEX_NONE;
+			}
+
+			// Convert to additive again
+			ConvertToAdditivePose(OldBasePoseIndex);
+		}
+
+		OnAssetModifiedNotifier.Broadcast();
+	}
+}
 #endif // WITH_EDITOR
 
 bool UPoseAsset::ModifyPoseName(FName OldPoseName, FName NewPoseName, const FSmartNameMapping::UID* NewUID)
@@ -754,7 +969,6 @@ int32 UPoseAsset::DeletePoses(TArray<FName> PoseNamesToDelete)
 	}
 
 	PoseContainer.Shrink(GetSkeleton(), RetargetSource);
-
 	OnPoseListChanged.Broadcast();
 
 	return ItemsDeleted;
@@ -856,6 +1070,20 @@ const int32 UPoseAsset::GetPoseIndexByName(const FName& InBasePoseName) const
 	return INDEX_NONE;
 }
 
+const int32 UPoseAsset::GetCurveIndexByName(const FName& InCurveName) const
+{
+	for (int32 TestIdx = 0; TestIdx < PoseContainer.Curves.Num(); TestIdx++)
+	{
+		const FAnimCurveBase& Curve = PoseContainer.Curves[TestIdx];
+		if (Curve.Name.DisplayName == InCurveName)
+		{
+			return TestIdx;
+		}
+	}
+	return INDEX_NONE;
+}
+
+
 void UPoseAsset::RecacheTrackmap()
 {
 	USkeleton* MySkeleton = GetSkeleton();
@@ -882,6 +1110,38 @@ void UPoseAsset::RemapTracksToNewSkeleton(USkeleton* NewSkeleton, bool bConvertS
 	// @todo: should allow deleted tracks and so on
 	RecacheTrackmap();
 }
+
+bool UPoseAsset::GetAllAnimationSequencesReferred(TArray<UAnimationAsset*>& AnimationAssets)
+{
+	if (SourceAnimation)
+	{
+		SourceAnimation->HandleAnimReferenceCollection(AnimationAssets);
+	}
+
+	return AnimationAssets.Num() > 0;
+}
+
+void UPoseAsset::ReplaceReferredAnimations(const TMap<UAnimationAsset*, UAnimationAsset*>& ReplacementMap)
+{
+	if (SourceAnimation)
+	{
+		UAnimSequence* const* ReplacementAsset = (UAnimSequence*const*)ReplacementMap.Find(SourceAnimation);
+		if (ReplacementAsset)
+		{
+			SourceAnimation = *ReplacementAsset;
+		}
+	}
+}
+
+void UPoseAsset::RegisterOnAssetModifiedNotifier(const FAssetModifiedNotifier& Delegate)
+{
+	OnAssetModifiedNotifier.Add(Delegate);
+}
+void UPoseAsset::UnregisterOnAssetModifiedNotifier(void* Unregister)
+{
+	OnAssetModifiedNotifier.RemoveAll(Unregister);
+}
+
 #endif // WITH_EDITOR
 
 bool UPoseAsset::GetBasePoseTransform(TArray<FTransform>& OutBasePose, TArray<float>& OutCurve) const
@@ -912,4 +1172,5 @@ bool UPoseAsset::GetBasePoseTransform(TArray<FTransform>& OutBasePose, TArray<fl
 	}
 
 	return false;
-}
\ No newline at end of file
+}
+
diff --git a/Engine/Source/Runtime/Engine/Private/Animation/Skeleton.cpp b/Engine/Source/Runtime/Engine/Private/Animation/Skeleton.cpp
--- a/Engine/Source/Runtime/Engine/Private/Animation/Skeleton.cpp
+++ b/Engine/Source/Runtime/Engine/Private/Animation/Skeleton.cpp
@@ -1135,7 +1135,7 @@ if (bModified)
 	}
 }
 
-bool USkeleton::GetSmartNameByUID(FName ContainerName, FSmartNameMapping::UID UID, FSmartName& OutSmartName)
+bool USkeleton::GetSmartNameByUID(const FName& ContainerName, FSmartNameMapping::UID UID, FSmartName& OutSmartName)
 {
 	const FSmartNameMapping* RequestedMapping = SmartNames.GetContainerInternal(ContainerName);
 	if (RequestedMapping)
@@ -1146,7 +1146,7 @@ bool USkeleton::GetSmartNameByUID(FName ContainerName, FSmartNameMapping::UID UI
 	return false;
 }
 
-bool USkeleton::GetSmartNameByName(FName ContainerName, FName InName, FSmartName& OutSmartName)
+bool USkeleton::GetSmartNameByName(const FName& ContainerName, const FName& InName, FSmartName& OutSmartName)
 {
 	const FSmartNameMapping* RequestedMapping = SmartNames.GetContainerInternal(ContainerName);
 	if (RequestedMapping)
@@ -1157,7 +1157,7 @@ bool USkeleton::GetSmartNameByName(FName ContainerName, FName InName, FSmartName
 	return false;
 }
 
-FSmartNameMapping::UID USkeleton::GetUIDByName(FName ContainerName, FName Name)
+FSmartNameMapping::UID USkeleton::GetUIDByName(const FName& ContainerName, const FName& Name)
 {
 	const FSmartNameMapping* RequestedMapping = SmartNames.GetContainerInternal(ContainerName);
 	if (RequestedMapping)
diff --git a/Engine/Source/Runtime/Engine/Private/Audio.cpp b/Engine/Source/Runtime/Engine/Private/Audio.cpp
--- a/Engine/Source/Runtime/Engine/Private/Audio.cpp
+++ b/Engine/Source/Runtime/Engine/Private/Audio.cpp
@@ -8,6 +8,7 @@
 #include "ActiveSound.h"
 #include "Audio.h"
 #include "AudioDevice.h"
+#include "AudioThread.h"
 #include "Sound/SoundCue.h"
 #include "Sound/SoundWave.h"
 #include "Sound/SoundNodeWavePlayer.h"
@@ -52,7 +53,6 @@ DEFINE_STAT(STAT_AudioSourceCreateTime);
 DEFINE_STAT(STAT_AudioSubmitBuffersTime);
 DEFINE_STAT(STAT_AudioStartSources);
 DEFINE_STAT(STAT_AudioGatherWaveInstances);
-DEFINE_STAT(STAT_AudioUpdateTime);
 DEFINE_STAT(STAT_AudioFindNearestLocation);
 
 
@@ -172,20 +172,10 @@ FString FSoundBuffer::Describe(bool bUseLongName)
 
 FString FSoundSource::Describe(bool bUseLongName)
 {
-	// look for a component and its owner
-	AActor* SoundOwner = NULL;
-	
-	// TODO - Audio Threading. This won't work cross thread.
-	UAudioComponent* AudioComponent = (WaveInstance->ActiveSound ? WaveInstance->ActiveSound->GetAudioComponent() : nullptr);
-	if (AudioComponent)
-	{
-		SoundOwner = AudioComponent->GetOwner();
-	}
-
 	return FString::Printf(TEXT("Wave: %s, Volume: %6.2f, Owner: %s"), 
 		bUseLongName ? *WaveInstance->WaveData->GetPathName() : *WaveInstance->WaveData->GetName(),
 		WaveInstance->GetActualVolume(), 
-		SoundOwner ? *SoundOwner->GetName() : TEXT("None"));
+		WaveInstance->ActiveSound ? *WaveInstance->ActiveSound->GetOwnerName() : TEXT("None"));
 }
 
 void FSoundSource::Stop( void )
@@ -327,31 +317,45 @@ void FSoundSource::UpdateStereoEmitterPositions()
 void FSoundSource::DrawDebugInfo()
 {
 	// Draw 3d Debug information about this source, if enabled
-	// TODO - Audio Threading. This won't work cross thread.
 	FAudioDeviceManager* DeviceManager = GEngine->GetAudioDeviceManager();
 
 	if (DeviceManager && DeviceManager->IsVisualizeDebug3dEnabled())
 	{
-		UAudioComponent* AudioComponent = (WaveInstance->ActiveSound ? WaveInstance->ActiveSound->GetAudioComponent() : nullptr);
-		if (AudioComponent)
+		const uint32 AudioComponentID = WaveInstance->ActiveSound->GetAudioComponentID();
+
+		if (AudioComponentID > 0)
 		{
-			UWorld* SoundWorld = AudioComponent->GetWorld();
-			if (SoundWorld)
-			{
-				FRotator SoundRotation = AudioComponent->GetComponentRotation();
-				DrawDebugCrosshairs(SoundWorld, WaveInstance->Location, SoundRotation, 20.0f, FColor::White, false, -1.0f, SDPG_Foreground);
+			DECLARE_CYCLE_STAT(TEXT("FAudioThreadTask.DrawSourceDebugInfo"), STAT_AudioDrawSourceDebugInfo, STATGROUP_TaskGraphTasks);
+
+			USoundBase* Sound = WaveInstance->ActiveSound->GetSound();
+			const FVector Location = WaveInstance->Location;
 
-				FString Name;
-				WaveInstance->ActiveSound->Sound->GetName(Name);
+			const bool bSpatialized = Buffer->NumChannels == 2 && WaveInstance->bUseSpatialization;
+			const FVector LeftChannelSourceLoc = LeftChannelSourceLocation;
+			const FVector RightChannelSourceLoc = RightChannelSourceLocation;
 
-				if (Buffer->NumChannels == 2 && WaveInstance->bUseSpatialization)
+			FAudioThread::RunCommandOnGameThread([AudioComponentID, Sound, bSpatialized, Location, LeftChannelSourceLoc, RightChannelSourceLoc]()
+			{
+				UAudioComponent* AudioComponent = UAudioComponent::GetAudioComponentFromID(AudioComponentID);
+				if (AudioComponent)
 				{
-					DrawDebugCrosshairs(SoundWorld, LeftChannelSourceLocation, SoundRotation, 20.0f, FColor::Red, false, -1.0f, SDPG_Foreground);
-					DrawDebugCrosshairs(SoundWorld, RightChannelSourceLocation, SoundRotation, 20.0f, FColor::Green, false, -1.0f, SDPG_Foreground);
-				}
+					UWorld* SoundWorld = AudioComponent->GetWorld();
+					if (SoundWorld)
+					{
+						FRotator SoundRotation = AudioComponent->GetComponentRotation();
+						DrawDebugCrosshairs(SoundWorld, Location, SoundRotation, 20.0f, FColor::White, false, -1.0f, SDPG_Foreground);
 
-				DrawDebugString(SoundWorld, AudioComponent->GetComponentLocation() + FVector(0, 0, 32), *Name, nullptr, FColor::White, 0.033, false);
-			}
+						if (bSpatialized)
+						{
+							DrawDebugCrosshairs(SoundWorld, LeftChannelSourceLoc, SoundRotation, 20.0f, FColor::Red, false, -1.0f, SDPG_Foreground);
+							DrawDebugCrosshairs(SoundWorld, RightChannelSourceLoc, SoundRotation, 20.0f, FColor::Green, false, -1.0f, SDPG_Foreground);
+						}
+
+						const FString Name = Sound->GetName();
+						DrawDebugString(SoundWorld, AudioComponent->GetComponentLocation() + FVector(0, 0, 32), *Name, nullptr, FColor::White, 0.033, false);
+					}
+				}
+			}, GET_STATID(STAT_AudioDrawSourceDebugInfo));
 		}
 	}
 }
@@ -360,7 +364,7 @@ FSpatializationParams FSoundSource::GetSpatializationParams()
 {
 	FSpatializationParams Params;
 
-		// Calculate direction from listener to sound, where the sound is at the origin if unspatialised.
+		// Calculate direction from listener to sound, where the sound is at the origin if unspatialized.
 	if (WaveInstance->bUseSpatialization)
 	{
 		FVector EmitterPosition = AudioDevice->GetListenerTransformedDirection(WaveInstance->Location, &Params.Distance);
