Copying //UE4/Dev-Framework to //UE4/Dev-Main (Source: //UE4/Dev-Framework @ 3038004)

diff --git a/Engine/Source/Runtime/Engine/Private/AudioDevice.cpp b/Engine/Source/Runtime/Engine/Private/AudioDevice.cpp
--- a/Engine/Source/Runtime/Engine/Private/AudioDevice.cpp
+++ b/Engine/Source/Runtime/Engine/Private/AudioDevice.cpp
@@ -3,9 +3,10 @@
 #include "EnginePrivate.h"
 #include "ActiveSound.h"
 #include "Audio.h"
+#include "AudioDecompress.h"
 #include "AudioDevice.h"
 #include "AudioEffect.h"
-#include "AudioDecompress.h"
+#include "AudioThread.h"
 #include "Sound/AudioSettings.h"
 #include "Sound/SoundCue.h"
 #include "Sound/SoundNodeWavePlayer.h"
@@ -76,11 +77,14 @@ void FDynamicParameter::Update(float DeltaTime)
 -----------------------------------------------------------------------------*/
 
 FAudioDevice::FAudioDevice()
-	: CommonAudioPool(nullptr)
+	: MaxChannels(0)
+	, SampleRate(0)
+	, CommonAudioPoolSize(0)
+	, CommonAudioPool(nullptr)
 	, CommonAudioPoolFreeBytes(0)
-	, bGameWasTicking(true)
-	, bDisableAudioCaching(false)
-	, bStartupSoundsPreCached(false)
+	, DeviceHandle(INDEX_NONE)
+	, SpatializationPlugin(nullptr)
+	, SpatializeProcessor(nullptr)
 	, CurrentTick(0)
 	, TestAudioComponent(nullptr)
 	, DebugState(DEBUGSTATE_None)
@@ -91,22 +95,26 @@ FAudioDevice::FAudioDevice()
 	, BaseSoundMix(nullptr)
 	, DefaultBaseSoundMix(nullptr)
 	, Effects(nullptr)
-	, CurrentAudioVolume(nullptr)
+	, CurrentReverbEffect(nullptr)
 	, PlatformAudioHeadroom(1.0f)
-	, HighestPriorityReverb(nullptr)
-	, SpatializationPlugin(nullptr)
-	, SpatializeProcessor(nullptr)
+	, bHRTFEnabledForAll_OnGameThread(false)
+	, bGameWasTicking(true)
+	, bDisableAudioCaching(false)
+	, bStartupSoundsPreCached(false)
 	, bSpatializationExtensionEnabled(false)
 	, bHRTFEnabledForAll(false)
 	, bIsDeviceMuted(false)
 	, bIsInitialized(false)
+	, bHasActivatedReverb(false)
+#if !UE_BUILD_SHIPPING
+	, RequestedAudioStats(0)
+#endif
 	, ConcurrencyManager(this)
 {
 }
 
 FAudioEffectsManager* FAudioDevice::CreateEffectsManager()
 {
-	// use a nop effect manager by default
 	return new FAudioEffectsManager(this);
 }
 
@@ -176,7 +184,8 @@ bool FAudioDevice::Init(int32 InMaxChannels)
 	InitSoundSources();
 	
 	// Make sure the Listeners array has at least one entry, so we don't have to check for Listeners.Num() == 0 all the time
-	Listeners.Add(FListener());
+	Listeners.AddDefaulted();
+	ListenerTransforms.AddDefaulted();
 
 	if (!bDeferStartupPrecache)
 	{
@@ -283,29 +292,33 @@ void FAudioDevice::CountBytes(FArchive& Ar)
 void FAudioDevice::AddReferencedObjects(FReferenceCollector& Collector)
 {	
 	Collector.AddReferencedObject(DefaultBaseSoundMix);
+	Collector.AddReferencedObjects(PrevPassiveSoundMixModifiers);
 
-	for (TMap< USoundMix*, FSoundMixState >::TIterator It(SoundMixModifiers); It; ++It)
+	for (TPair<FName, FActivatedReverb>& ActivatedReverbPair : ActivatedReverbs)
 	{
-		Collector.AddReferencedObject(It.Key());
+		Collector.AddReferencedObject(ActivatedReverbPair.Value.ReverbSettings.ReverbEffect);
 	}
 
-	if (Effects)
+	for (TPair<USoundMix*, FSoundMixState>& SoundMixModifierPair : SoundMixModifiers)
 	{
-		Effects->AddReferencedObjects(Collector);
+		Collector.AddReferencedObject(SoundMixModifierPair.Key);
 	}
 
-	for (int32 i = 0; i < ActiveSounds.Num(); ++i)
+	if (Effects)
 	{
-		ActiveSounds[i]->AddReferencedObjects(Collector);
+		Effects->AddReferencedObjects(Collector);
 	}
-	for (int32 i = 0; i < PrevPassiveSoundMixModifiers.Num(); ++i)
+
+	for (FActiveSound* ActiveSound : ActiveSounds)
 	{
-		Collector.AddReferencedObject(PrevPassiveSoundMixModifiers[i]);
+		ActiveSound->AddReferencedObjects(Collector);
 	}
 }
 
 void FAudioDevice::ResetInterpolation()
 {
+	check(IsInAudioThread());
+
 	for (FListener& Listener : Listeners)
 	{
 		Listener.InteriorStartTime = 0.f;
@@ -342,24 +355,28 @@ void FAudioDevice::EnableRadioEffect(bool bEnable)
 {
 	if (bEnable)
 	{
-		DebugState = DEBUGSTATE_None;
+		SetMixDebugState(DEBUGSTATE_None);
 	}
 	else
 	{
 		UE_LOG(LogAudio, Log, TEXT("Radio disabled for all sources"));
-		DebugState = DEBUGSTATE_DisableRadio;
+		SetMixDebugState(DEBUGSTATE_DisableRadio);
 	}
 }
 
 #if !UE_BUILD_SHIPPING
 bool FAudioDevice::HandleShowSoundClassHierarchyCommand(const TCHAR* Cmd, FOutputDevice& Ar)
 {
+	FAudioThreadSuspendContext AudioThreadSuspend;
+
 	ShowSoundClassHierarchy(Ar);
 	return true;
 }
 
 bool FAudioDevice::HandleListWavesCommand(const TCHAR* Cmd, FOutputDevice& Ar)
 {
+	FAudioThreadSuspendContext AudioThreadSuspend;
+
 	TArray<FWaveInstance*> WaveInstances;
 	int32 FirstActiveIndex = GetSortedActiveWaveInstances(WaveInstances, ESortedActiveWaveGetType::QueryOnly);
 
@@ -367,12 +384,13 @@ bool FAudioDevice::HandleListWavesCommand(const TCHAR* Cmd, FOutputDevice& Ar)
 	{
 		FWaveInstance* WaveInstance = WaveInstances[ InstanceIndex ];
 		FSoundSource* Source = WaveInstanceSourceMap.FindRef(WaveInstance);
-		UAudioComponent* AudioComponent = WaveInstance->ActiveSound->GetAudioComponent();
+		UAudioComponent* AudioComponent = UAudioComponent::GetAudioComponentFromID(WaveInstance->ActiveSound->GetAudioComponentID());
 		AActor* SoundOwner = AudioComponent ? AudioComponent->GetOwner() : nullptr;
 		Ar.Logf(TEXT("%4i.    %s %6.2f %6.2f  %s   %s"), InstanceIndex, Source ? TEXT("Yes") : TEXT(" No"), WaveInstance->ActiveSound->PlaybackTime, WaveInstance->GetActualVolume(), *WaveInstance->WaveData->GetPathName(), SoundOwner ? *SoundOwner->GetName() : TEXT("None"));
 	}
 
 	Ar.Logf(TEXT("Total: %i"), WaveInstances.Num()-FirstActiveIndex);
+
 	return true;
 }
 
@@ -543,6 +561,8 @@ TMap<int32, int32> NativeSampleRateCount;
 
 bool FAudioDevice::HandleDumpSoundInfoCommand(const TCHAR* Cmd, FOutputDevice& Ar)
 {
+	FAudioThreadSuspendContext AudioThreadSuspend;
+
 	Ar.Logf(TEXT("Native Count: %d\nRealtime Count: %d\n"), PrecachedNative, PrecachedRealtime);
 	float AverageSize = 0.0f;
 	if (PrecachedNative != 0)
@@ -566,6 +586,8 @@ bool FAudioDevice::HandleDumpSoundInfoCommand(const TCHAR* Cmd, FOutputDevice& A
 
 bool FAudioDevice::HandleListSoundClassVolumesCommand(const TCHAR* Cmd, FOutputDevice& Ar)
 {
+	FAudioThreadSuspendContext AudioThreadSuspend;
+
 	Ar.Logf(TEXT("SoundClass Volumes: (Volume, Pitch)"));
 
 	for (TMap<USoundClass*, FSoundClassProperties>::TIterator It(SoundClasses); It; ++It)
@@ -584,6 +606,8 @@ bool FAudioDevice::HandleListSoundClassVolumesCommand(const TCHAR* Cmd, FOutputD
 
 bool FAudioDevice::HandleListAudioComponentsCommand(const TCHAR* Cmd, FOutputDevice& Ar)
 {
+	FAudioThreadSuspendContext AudioThreadSuspend;
+
 	int32 Count = 0;
 	Ar.Logf(TEXT("AudioComponent Dump"));
 	for (TObjectIterator<UAudioComponent> It; It; ++It)
@@ -610,7 +634,7 @@ bool FAudioDevice::HandleListAudioComponentsCommand(const TCHAR* Cmd, FOutputDev
 	for (int32 ASIndex = 0; ASIndex < ActiveSounds.Num(); ASIndex++)
 	{
 		const FActiveSound* ActiveSound = ActiveSounds[ASIndex];
-		UAudioComponent* AComp = ActiveSound->GetAudioComponent();
+		UAudioComponent* AComp = UAudioComponent::GetAudioComponentFromID(ActiveSound->GetAudioComponentID());
 		if (AComp)
 		{
 			Ar.Logf(TEXT("    0x%p: %4d - %s, %s, %s, %s"),
@@ -712,8 +736,8 @@ bool FAudioDevice::HandlePlaySoundWaveCommand(const TCHAR* Cmd, FOutputDevice& A
 bool FAudioDevice::HandleSetBaseSoundMixCommand(const TCHAR* Cmd, FOutputDevice& Ar)
 {
 	// Ar.Logf(TEXT("Setting base sound mix '%s'"), Cmd);
-	FName NewMix = FName(Cmd);
-	USoundMix* SoundMix = NULL;
+	const FName NewMix = FName(Cmd);
+	USoundMix* SoundMix = nullptr;
 
 	for (TObjectIterator<USoundMix> It; It; ++It)
 	{
@@ -738,42 +762,42 @@ bool FAudioDevice::HandleSetBaseSoundMixCommand(const TCHAR* Cmd, FOutputDevice&
 bool FAudioDevice::HandleIsolateDryAudioCommand(const TCHAR* Cmd, FOutputDevice& Ar)
 {
 	Ar.Logf(TEXT("Dry audio isolated"));
-	DebugState = DEBUGSTATE_IsolateDryAudio;
+	SetMixDebugState(DEBUGSTATE_IsolateDryAudio);
 	return true;
 }
 
 bool FAudioDevice::HandleIsolateReverbCommand(const TCHAR* Cmd, FOutputDevice& Ar)
 {
 	Ar.Logf(TEXT("Reverb audio isolated"));
-	DebugState = DEBUGSTATE_IsolateReverb;
+	SetMixDebugState(DEBUGSTATE_IsolateReverb);
 	return true;
 }
 
 bool FAudioDevice::HandleTestLPFCommand(const TCHAR* Cmd, FOutputDevice& Ar)
 {
 	Ar.Logf(TEXT("LPF set to max for all sources"));
-	DebugState =  DEBUGSTATE_TestLPF;
+	SetMixDebugState(DEBUGSTATE_TestLPF);
 	return true;
 }
 
 bool FAudioDevice::HandleTestStereoBleedCommand(const TCHAR* Cmd, FOutputDevice& Ar)
 {
 	Ar.Logf(TEXT("StereoBleed set to max for all sources"));
-	DebugState =  DEBUGSTATE_TestStereoBleed;
+	SetMixDebugState(DEBUGSTATE_TestStereoBleed);
 	return true;
 }
 
 bool FAudioDevice::HandleTestLFEBleedCommand(const TCHAR* Cmd, FOutputDevice& Ar)
 {
 	Ar.Logf(TEXT("LFEBleed set to max for all sources"));
-	DebugState = DEBUGSTATE_TestLFEBleed;
+	SetMixDebugState(DEBUGSTATE_TestLFEBleed);
 	return true;
 }
 
 bool FAudioDevice::HandleDisableLPFCommand(const TCHAR* Cmd, FOutputDevice& Ar)
 {
 	Ar.Logf(TEXT("LPF disabled for all sources"));
-	DebugState = DEBUGSTATE_DisableLPF;
+	SetMixDebugState(DEBUGSTATE_DisableLPF);
 	return true;
 }
 
@@ -792,22 +816,48 @@ bool FAudioDevice::HandleEnableRadioCommand(const TCHAR* Cmd, FOutputDevice& Ar)
 bool FAudioDevice::HandleResetSoundStateCommand(const TCHAR* Cmd, FOutputDevice& Ar)
 {
 	Ar.Logf(TEXT("All volumes reset to their defaults; all test filters removed"));
-	DebugState = DEBUGSTATE_None;
+	SetMixDebugState(DEBUGSTATE_None);
 	return true;
 }
 
 bool FAudioDevice::HandleToggleSpatializationExtensionCommand(const TCHAR* Cmd, FOutputDevice& Ar)
 {
+	FAudioThread::SuspendAudioThread();
+
 	bSpatializationExtensionEnabled = !bSpatializationExtensionEnabled;
+
+	FAudioThread::ResumeAudioThread();
+
 	return true;
 }
 
 bool FAudioDevice::HandleEnableHRTFForAllCommand(const TCHAR* Cmd, FOutputDevice& Ar)
 {
-	bHRTFEnabledForAll = !bHRTFEnabledForAll;
+	const bool bNewHRTFEnabledForAll = !bHRTFEnabledForAll_OnGameThread;
+
+	bHRTFEnabledForAll_OnGameThread = bNewHRTFEnabledForAll;
+
+	FAudioDevice* AudioDevice = this;
+	FAudioThread::RunCommandOnAudioThread([AudioDevice, bNewHRTFEnabledForAll]()
+	{
+		AudioDevice->bHRTFEnabledForAll = bNewHRTFEnabledForAll;
+
+	});
+
 	return true;
 }
 
+bool FAudioDevice::IsHRTFEnabledForAll() const
+{
+	if (IsInAudioThread())
+	{
+		return bHRTFEnabledForAll && IsSpatializationPluginEnabled();
+	}
+
+	check(IsInGameThread());
+	return bHRTFEnabledForAll_OnGameThread && IsSpatializationPluginEnabled();
+}
+
 bool FAudioDevice::HandleSoloCommand(const TCHAR* Cmd, FOutputDevice& Ar)
 {
 	// Apply the solo to the given device
@@ -1079,9 +1129,23 @@ bool FAudioDevice::HandleAudioMemoryInfo(const TCHAR* Cmd, FOutputDevice& Ar)
 
 #endif // !UE_BUILD_SHIPPING
 
-EDebugState FAudioDevice::GetMixDebugState()
+void FAudioDevice::SetMixDebugState(EDebugState InDebugState)
 {
-	return((EDebugState)DebugState);
+	if (!IsInAudioThread())
+	{
+		DECLARE_CYCLE_STAT(TEXT("FAudioThreadTask.SetMixDebugState"), STAT_AudioSetMixDebugState, STATGROUP_AudioThreadCommands);
+
+		FAudioDevice* AudioDevice = this;
+		FAudioThread::RunCommandOnAudioThread([AudioDevice, InDebugState]()
+		{
+			AudioDevice->SetMixDebugState(InDebugState);
+
+		}, GET_STATID(STAT_AudioSetMixDebugState));
+
+		return;
+	}
+
+	DebugState = InDebugState;
 }
 
 bool FAudioDevice::Exec(UWorld* InWorld, const TCHAR* Cmd, FOutputDevice& Ar)
@@ -1231,33 +1295,49 @@ void FAudioDevice::InitSoundSources()
 
 void FAudioDevice::SetDefaultBaseSoundMix(USoundMix* SoundMix)
 {
-	if (SoundMix == NULL)
+	if (IsInGameThread() && SoundMix == nullptr)
 	{
 		const FStringAssetReference DefaultBaseSoundMixName = GetDefault<UAudioSettings>()->DefaultBaseSoundMix;
 		if (DefaultBaseSoundMixName.IsValid())
 		{			
-			SoundMix = LoadObject<USoundMix>(NULL, *DefaultBaseSoundMixName.ToString());
+			SoundMix = LoadObject<USoundMix>(nullptr, *DefaultBaseSoundMixName.ToString());
 		}
 	}
 
+	if (!IsInAudioThread())
+	{
+		DECLARE_CYCLE_STAT(TEXT("FAudioThreadTask.SetDefaultBaseSoundMix"), STAT_AudioSetDefaultBaseSoundMix, STATGROUP_AudioThreadCommands);
+
+		FAudioDevice* AudioDevice = this;
+		FAudioThread::RunCommandOnAudioThread([AudioDevice, SoundMix]()
+		{
+			AudioDevice->SetDefaultBaseSoundMix(SoundMix);
+
+		}, GET_STATID(STAT_AudioSetDefaultBaseSoundMix));
+
+		return;
+	}
+
 	DefaultBaseSoundMix = SoundMix;
 	SetBaseSoundMix(SoundMix);
 }
 
 void FAudioDevice::RemoveSoundMix(USoundMix* SoundMix)
 {
+	check(IsInAudioThread());
+
 	if (SoundMix)
 	{
 		// Not sure if we will ever destroy the default base SoundMix
 		if (SoundMix == DefaultBaseSoundMix)
 		{
-			DefaultBaseSoundMix = NULL;
+			DefaultBaseSoundMix = nullptr;
 		}
 
 		ClearSoundMix(SoundMix);
 
 		// Try setting to global default if base SoundMix has been cleared
-		if (BaseSoundMix == NULL)
+		if (BaseSoundMix == nullptr)
 		{
 			SetBaseSoundMix(DefaultBaseSoundMix);
 		}
@@ -1300,13 +1380,29 @@ void FAudioDevice::RecurseIntoSoundClasses(USoundClass* CurrentClass, FSoundClas
 
 void FAudioDevice::UpdateHighestPriorityReverb()
 {
-	HighestPriorityReverb = nullptr;
-	for (const TPair<FName, FActivatedReverb>& ActivatedReverb : ActivatedReverbs)
+	check(IsInGameThread());
+
+	DECLARE_CYCLE_STAT(TEXT("FAudioThreadTask.UpdateHighestPriorityReverb"), STAT_AudioUpdateHighestPriorityReverb, STATGROUP_AudioThreadCommands);
+
+	FAudioDevice* AudioDevice = this;
+
+	if (ActivatedReverbs.Num() > 0)
 	{
-		if (HighestPriorityReverb == nullptr || ActivatedReverb.Value.Priority > HighestPriorityReverb->Priority)
+		ActivatedReverbs.ValueSort([](const FActivatedReverb& A, const FActivatedReverb& B) { return A.Priority > B.Priority; });
+
+		const FActivatedReverb& NewActiveReverbRef = ActivatedReverbs.CreateConstIterator().Value();
+		FAudioThread::RunCommandOnAudioThread([AudioDevice, NewActiveReverbRef]()
 		{
-			HighestPriorityReverb = &ActivatedReverb.Value;
-		}
+			AudioDevice->bHasActivatedReverb = true;
+			AudioDevice->HighestPriorityActivatedReverb = NewActiveReverbRef;
+		}, GET_STATID(STAT_AudioUpdateHighestPriorityReverb));
+	}
+	else
+	{
+		FAudioThread::RunCommandOnAudioThread([AudioDevice]()
+		{
+			AudioDevice->bHasActivatedReverb = false;
+		}, GET_STATID(STAT_AudioUpdateHighestPriorityReverb));
 	}
 }
 
@@ -1886,15 +1982,17 @@ void FListener::UpdateCurrentInteriorSettings()
 
 void FAudioDevice::InvalidateCachedInteriorVolumes() const
 {
+	check(IsInAudioThread());
+
 	for (FActiveSound* ActiveSound : ActiveSounds)
 	{
 		ActiveSound->bGotInteriorSettings = false;
 	}
 }
 
-void FListener::ApplyInteriorSettings(class AAudioVolume* InVolume, const FInteriorSettings& Settings)
+void FListener::ApplyInteriorSettings(const uint32 InAudioVolumeID, const FInteriorSettings& Settings)
 {
-	if (InVolume != Volume || Settings != InteriorSettings)
+	if (InAudioVolumeID != AudioVolumeID || Settings != InteriorSettings)
 	{
 		// Use previous/ current interpolation time if we're transitioning to the default worldsettings zone.
 		InteriorStartTime = FApp::GetCurrentTime();
@@ -1903,47 +2001,126 @@ void FListener::ApplyInteriorSettings(class AAudioVolume* InVolume, const FInter
 		InteriorLPFEndTime = InteriorStartTime + (Settings.bIsWorldSettings ? InteriorSettings.InteriorLPFTime : Settings.InteriorLPFTime);
 		ExteriorLPFEndTime = InteriorStartTime + (Settings.bIsWorldSettings ? InteriorSettings.ExteriorLPFTime : Settings.ExteriorLPFTime);
 
-		Volume = InVolume;
+		AudioVolumeID = InAudioVolumeID;
 		InteriorSettings = Settings;
 	}
 }
 
 
-void FAudioDevice::SetListener(const int32 InViewportIndex, const FTransform& InListenerTransform, const float InDeltaSeconds, class AAudioVolume* Volume, const FInteriorSettings& InteriorSettings)
+void FAudioDevice::SetListener(UWorld* World, const int32 InViewportIndex, const FTransform& ListenerTransform, const float InDeltaSeconds)
 {
-	FTransform ListenerTransform = InListenerTransform;
-	
-	if (!ensureMsgf(ListenerTransform.IsValid(), TEXT("Invalid listener transform provided to AudioDevice")))
+	check(IsInGameThread());
+
+	// The copy is done because FTransform doesn't work to pass by value on Win32
+	FTransform ListenerTransformCopy = ListenerTransform;
+
+	if (!ensureMsgf(ListenerTransformCopy.IsValid(), TEXT("Invalid listener transform provided to AudioDevice")))
 	{
 		// If we have a bad transform give it something functional if totally wrong
-		ListenerTransform = FTransform::Identity;
+		ListenerTransformCopy = FTransform::Identity;
 	}
 
-	if (InViewportIndex >= Listeners.Num())
+	if (InViewportIndex >= ListenerTransforms.Num())
 	{
-		UE_LOG(LogAudio, Log, TEXT("Resizing Listeners array: %d -> %d"), Listeners.Num(), InViewportIndex);
-		Listeners.AddZeroed(InViewportIndex - Listeners.Num() + 1);
+		ListenerTransforms.AddDefaulted(InViewportIndex - ListenerTransforms.Num() + 1);
 	}
 
-	Listeners[ InViewportIndex ].Velocity = InDeltaSeconds > 0.f ? 
-											(ListenerTransform.GetTranslation() - Listeners[ InViewportIndex ].Transform.GetTranslation()) / InDeltaSeconds
-											: FVector::ZeroVector;
+	ListenerTransforms[InViewportIndex] = ListenerTransformCopy;
 
-#if ENABLE_NAN_DIAGNOSTIC
-	if (Listeners[InViewportIndex].Velocity.ContainsNaN())
+	DECLARE_CYCLE_STAT(TEXT("FAudioThreadTask.SetListener"), STAT_AudioSetListener, STATGROUP_AudioThreadCommands);
+
+	FAudioDevice* AudioDevice = this;
+	const uint32 WorldID = World->GetUniqueID();
+	FAudioThread::RunCommandOnAudioThread([AudioDevice, WorldID, InViewportIndex, ListenerTransformCopy, InDeltaSeconds]()
 	{
-		logOrEnsureNanError(TEXT("FAudioDevice::SetListener has detected a NaN in Listener Velocity"));
-	}
+		TArray<FListener>& AudioThreadListeners = AudioDevice->Listeners;
+		if (InViewportIndex >= AudioThreadListeners.Num())
+		{
+			AudioThreadListeners.AddDefaulted(InViewportIndex - AudioThreadListeners.Num() + 1);
+		}
+
+		FListener& Listener = AudioThreadListeners[InViewportIndex];
+		Listener.Velocity = InDeltaSeconds > 0.f ? 
+			(ListenerTransformCopy.GetTranslation() - Listener.Transform.GetTranslation()) / InDeltaSeconds
+			: FVector::ZeroVector;
+
+#if ENABLE_NAN_DIAGNOSTIC
+		if (Listener.Velocity.ContainsNaN())
+		{
+			logOrEnsureNanError(TEXT("FAudioDevice::SetListener has detected a NaN in Listener Velocity"));
+		}
 #endif
 
-	Listeners[ InViewportIndex ].Transform = ListenerTransform;
+		Listener.WorldID = WorldID;
+		Listener.Transform = ListenerTransformCopy;
+
+	}, GET_STATID(STAT_AudioSetListener));
+}
+
+void FAudioDevice::SetDefaultAudioSettings(UWorld* World, const FReverbSettings& DefaultReverbSettings, const FInteriorSettings& DefaultInteriorSettings)
+{
+	check(IsInGameThread());
+
+	DECLARE_CYCLE_STAT(TEXT("FAudioThreadTask.SetDefaultAudioSettings"), STAT_AudioSetDefaultAudioSettings, STATGROUP_AudioThreadCommands);
+
+	FAudioDevice* AudioDevice = this;
+	const uint32 WorldID = World->GetUniqueID();
+	FAudioThread::RunCommandOnAudioThread([AudioDevice, WorldID, DefaultReverbSettings, DefaultInteriorSettings]()
+	{
+		AudioDevice->WorldIDToDefaultAudioVolumeSettingsMap.Add(WorldID, TPair<FReverbSettings,FInteriorSettings>(TPairInitializer<FReverbSettings,FInteriorSettings>(DefaultReverbSettings,DefaultInteriorSettings)));
 
-	Listeners[ InViewportIndex ].ApplyInteriorSettings(Volume, InteriorSettings);
+	}, GET_STATID(STAT_AudioSetDefaultAudioSettings));
 }
 
+void FAudioDevice::GetAudioVolumeSettings(const uint32 WorldID, const FVector& Location, FAudioVolumeSettings& OutSettings) const
+{
+	check(IsInAudioThread());
+
+	for (const TPair<uint32,FAudioVolumeProxy>& AudioVolumePair : AudioVolumeProxies)
+	{
+		const FAudioVolumeProxy& Proxy = AudioVolumePair.Value;
+		if (Proxy.WorldID == WorldID)
+		{
+			FVector Dummy;
+			float DistanceSqr = 0.f;
+			if (Proxy.BodyInstance->GetSquaredDistanceToBody(Location, DistanceSqr, Dummy) && DistanceSqr == 0.f)
+			{
+				OutSettings.AudioVolumeID = Proxy.AudioVolumeID;
+				OutSettings.Priority = Proxy.Priority;
+				OutSettings.ReverbSettings = Proxy.ReverbSettings;
+				OutSettings.InteriorSettings = Proxy.InteriorSettings;
+				return;
+			}
+		}
+	}
+
+	OutSettings.AudioVolumeID = 0;
 
-bool FAudioDevice::SetBaseSoundMix(USoundMix* NewMix)
+	const TPair<FReverbSettings, FInteriorSettings>* DefaultAudioVolumeSettings = WorldIDToDefaultAudioVolumeSettingsMap.Find(WorldID);
+
+	if (DefaultAudioVolumeSettings)
+	{
+		OutSettings.ReverbSettings = DefaultAudioVolumeSettings->Key;
+		OutSettings.InteriorSettings = DefaultAudioVolumeSettings->Value;
+	}
+}
+
+void FAudioDevice::SetBaseSoundMix(USoundMix* NewMix)
 {
+	if (!IsInAudioThread())
+	{
+		DECLARE_CYCLE_STAT(TEXT("FAudioThreadTask.SetBaseSoundMix"), STAT_AudioSetBaseSoundMix, STATGROUP_AudioThreadCommands);
+
+		FAudioDevice* AudioDevice = this;
+		FAudioThread::RunCommandOnAudioThread([AudioDevice, NewMix]()
+		{
+			AudioDevice->SetBaseSoundMix(NewMix);
+
+		}, GET_STATID(STAT_AudioSetBaseSoundMix));
+
+		return;
+	}
+
 	if (NewMix && NewMix != BaseSoundMix)
 	{
 		USoundMix* OldBaseSoundMix = BaseSoundMix;
@@ -1972,17 +2149,27 @@ bool FAudioDevice::SetBaseSoundMix(USoundMix* NewMix)
 		}
 
 		ExistingState->IsBaseSoundMix = true;
-
-		return true;
 	}
-
-	return false;
 }
 
 void FAudioDevice::PushSoundMixModifier(USoundMix* SoundMix, bool bIsPassive)
 {
 	if (SoundMix)
 	{
+		if (!IsInAudioThread())
+		{
+			DECLARE_CYCLE_STAT(TEXT("FAudioThreadTask.PushSoundMixModifier"), STAT_AudioPushSoundMixModifier, STATGROUP_AudioThreadCommands);
+
+			FAudioDevice* AudioDevice = this;
+			FAudioThread::RunCommandOnAudioThread([AudioDevice, SoundMix, bIsPassive]()
+			{
+				AudioDevice->PushSoundMixModifier(SoundMix, bIsPassive);
+
+			}, GET_STATID(STAT_AudioPushSoundMixModifier));
+
+			return;
+		}
+
 		FSoundMixState* SoundMixState = SoundMixModifiers.Find(SoundMix);
 
 		if (!SoundMixState)
@@ -2026,6 +2213,20 @@ void FAudioDevice::SetSoundMixClassOverride(USoundMix* InSoundMix, USoundClass*
 		return;
 	}
 
+	if (!IsInAudioThread())
+	{
+		DECLARE_CYCLE_STAT(TEXT("FAudioThreadTask.SetSoundMixClassOverride"), STAT_AudioSetSoundMixClassOverride, STATGROUP_AudioThreadCommands);
+
+		FAudioDevice* AudioDevice = this;
+		FAudioThread::RunCommandOnAudioThread([AudioDevice, InSoundMix, InSoundClass, Volume, Pitch, FadeInTime, bApplyToChildren]()
+		{
+			AudioDevice->SetSoundMixClassOverride(InSoundMix, InSoundClass, Volume, Pitch, FadeInTime, bApplyToChildren);
+
+		}, GET_STATID(STAT_AudioSetSoundMixClassOverride));
+
+		return;
+	}
+
 	FSoundMixClassOverrideMap& SoundMixClassOverrideMap = SoundMixClassEffectOverrides.FindOrAdd(InSoundMix);
 
 	// Check if we've already added this sound class override
@@ -2064,6 +2265,20 @@ void FAudioDevice::ClearSoundMixClassOverride(USoundMix* InSoundMix, USoundClass
 		return;
 	}
 
+	if (!IsInAudioThread())
+	{
+		DECLARE_CYCLE_STAT(TEXT("FAudioThreadTask.ClearSoundMixClassOverride"), STAT_AudioClearSoundMixClassOverride, STATGROUP_AudioThreadCommands);
+
+		FAudioDevice* AudioDevice = this;
+		FAudioThread::RunCommandOnAudioThread([AudioDevice, InSoundMix, InSoundClass, FadeOutTime]()
+		{
+			AudioDevice->ClearSoundMixClassOverride(InSoundMix, InSoundClass, FadeOutTime);
+
+		}, GET_STATID(STAT_AudioClearSoundMixClassOverride));
+
+		return;
+	}
+
 	// Get the sound mix class override map for the sound mix. If this doesn't exist, then nobody overrode the sound mix
 	FSoundMixClassOverrideMap* SoundMixClassOverrideMap = SoundMixClassEffectOverrides.Find(InSoundMix);
 	if (!SoundMixClassOverrideMap)
@@ -2122,6 +2337,20 @@ void FAudioDevice::PopSoundMixModifier(USoundMix* SoundMix, bool bIsPassive)
 {
 	if (SoundMix)
 	{
+		if (!IsInAudioThread())
+		{
+			DECLARE_CYCLE_STAT(TEXT("FAudioThreadTask.PopSoundMixModifier"), STAT_AudioPopSoundMixModifier, STATGROUP_AudioThreadCommands);
+
+			FAudioDevice* AudioDevice = this;
+			FAudioThread::RunCommandOnAudioThread([AudioDevice, SoundMix, bIsPassive]()
+			{
+				AudioDevice->PopSoundMixModifier(SoundMix, bIsPassive);
+
+			}, GET_STATID(STAT_AudioPopSoundMixModifier));
+
+			return;
+		}
+
 		FSoundMixState* SoundMixState = SoundMixModifiers.Find(SoundMix);
 
 		if (SoundMixState)
@@ -2153,6 +2382,20 @@ void FAudioDevice::ClearSoundMixModifier(USoundMix* SoundMix)
 {
 	if (SoundMix)
 	{
+		if (!IsInAudioThread())
+		{
+			DECLARE_CYCLE_STAT(TEXT("FAudioThreadTask.ClearSoundMixModifier"), STAT_AudioClearSoundMixModifier, STATGROUP_AudioThreadCommands);
+
+			FAudioDevice* AudioDevice = this;
+			FAudioThread::RunCommandOnAudioThread([AudioDevice, SoundMix]()
+			{
+				AudioDevice->ClearSoundMixModifier(SoundMix);
+
+			}, GET_STATID(STAT_AudioClearSoundMixModifier));
+
+			return;
+		}
+
 		FSoundMixState* SoundMixState = SoundMixModifiers.Find(SoundMix);
 
 		if (SoundMixState)
@@ -2166,6 +2409,20 @@ void FAudioDevice::ClearSoundMixModifier(USoundMix* SoundMix)
 
 void FAudioDevice::ClearSoundMixModifiers()
 {
+	if (!IsInAudioThread())
+	{
+		DECLARE_CYCLE_STAT(TEXT("FAudioThreadTask.ClearSoundMixModifiers"), STAT_AudioClearSoundMixModifiers, STATGROUP_AudioThreadCommands);
+
+		FAudioDevice* AudioDevice = this;
+		FAudioThread::RunCommandOnAudioThread([AudioDevice]()
+		{
+			AudioDevice->ClearSoundMixModifiers();
+
+		}, GET_STATID(STAT_AudioClearSoundMixModifiers));
+
+		return;
+	}
+
 	// Clear all sound mix modifiers
 	for (TMap< USoundMix*, FSoundMixState >::TIterator It(SoundMixModifiers); It; ++It)
 	{
@@ -2175,6 +2432,8 @@ void FAudioDevice::ClearSoundMixModifiers()
 
 void FAudioDevice::ActivateReverbEffect(UReverbEffect* ReverbEffect, FName TagName, float Priority, float Volume, float FadeTime)
 {
+	check(IsInGameThread());
+
 	FActivatedReverb& ActivatedReverb = ActivatedReverbs.FindOrAdd(TagName);
 
 	ActivatedReverb.ReverbSettings.ReverbEffect = ReverbEffect;
@@ -2187,30 +2446,17 @@ void FAudioDevice::ActivateReverbEffect(UReverbEffect* ReverbEffect, FName TagNa
 
 void FAudioDevice::DeactivateReverbEffect(FName TagName)
 {
+	check(IsInGameThread());
+
 	if (ActivatedReverbs.Remove(TagName) > 0)
 	{
 		UpdateHighestPriorityReverb();
 	}
 }
 
-void FAudioDevice::SetReverbSettings(class AAudioVolume* Volume, const FReverbSettings& ReverbSettings)
-{
-	const FReverbSettings* ActivatedReverb = &ReverbSettings;
-	if (HighestPriorityReverb && (!Volume || HighestPriorityReverb->Priority > Volume->Priority))
-	{
-		ActivatedReverb = &HighestPriorityReverb->ReverbSettings;
-	}
-
-	Effects->SetReverbSettings(*ActivatedReverb);
-}
-
-UReverbEffect* FAudioDevice::GetCurrentReverbEffect()
-{
-	return Effects->GetCurrentReverbEffect();
-}
-
 void* FAudioDevice::InitEffect(FSoundSource* Source)
 {
+	check(IsInAudioThread());
 	return(Effects->InitEffect(Source));
 }
 
@@ -2219,12 +2465,14 @@ void* FAudioDevice::UpdateEffect(FSoundSource* Source)
 {
 	SCOPE_CYCLE_COUNTER(STAT_AudioUpdateEffects);
 
+	check(IsInAudioThread());
 	return(Effects->UpdateEffect(Source));
 }
 
 
 void FAudioDevice::DestroyEffect(FSoundSource* Source)
 {
+	check(IsInAudioThread());
 	return(Effects->DestroyEffect(Source));
 }
 
@@ -2262,6 +2510,8 @@ void FAudioDevice::HandlePause(bool bGameTicking, bool bGlobalPause)
 
 int32 FAudioDevice::GetSortedActiveWaveInstances(TArray<FWaveInstance*>& WaveInstances, const ESortedActiveWaveGetType::Type GetType)
 {
+	check(IsInAudioThread());
+
 	SCOPE_CYCLE_COUNTER(STAT_AudioGatherWaveInstances);
 
 	// Tick all the active audio components.  Use a copy as some operations may remove elements from the list, but we want
@@ -2297,7 +2547,7 @@ int32 FAudioDevice::GetSortedActiveWaveInstances(TArray<FWaveInstance*>& WaveIns
 						ActiveSound->PlaybackTime,
 						Duration,
 						*ActiveSound->Sound->GetName(),
-						(ActiveSound->GetAudioComponent() ? *ActiveSound->GetAudioComponent()->GetName() : TEXT("NO COMPONENT")));
+						*ActiveSound->GetAudioComponentName());
 					AddSoundToStop(ActiveSound);
 				}
 				else
@@ -2422,6 +2672,8 @@ void FAudioDevice::StopSources(TArray<FWaveInstance*>& WaveInstances, int32 Firs
 
 void FAudioDevice::StartSources(TArray<FWaveInstance*>& WaveInstances, int32 FirstActiveIndex, bool bGameTicking)
 {
+	check(IsInAudioThread());
+
 	SCOPE_CYCLE_COUNTER(STAT_AudioStartSources);
 
 	// Start sources as needed.
@@ -2495,10 +2747,23 @@ void FAudioDevice::StartSources(TArray<FWaveInstance*>& WaveInstances, int32 Fir
 				// Otherwise, we need still need to initialize
 				else if (Source->IsPreparedToInit())
 				{
-					Source->Init(WaveInstance);
-					check(Source->IsInitialized());
-					Source->Play();
-					Source->Update();
+					// Try to initialize the source. This may fail if something is wrong with the source.
+					if (Source->Init(WaveInstance))
+					{
+						// Note: if we succeeded in starting to prepare to init, we already added the wave instance map to the source so don't need to add here.
+						check(Source->IsInitialized());
+						Source->Play();
+						Source->Update();
+					}
+					else
+					{
+						// Make sure init cleaned up the buffer when it failed
+						check(Source->Buffer == nullptr);
+
+						// If were ready to call init but failed, then we need to add the source and stop with notification
+						WaveInstance->StopWithoutNotification();
+						FreeSources.Add(Source);
+					}
 				}
 			}
 			else if (!Source)
@@ -2513,7 +2778,20 @@ void FAudioDevice::StartSources(TArray<FWaveInstance*>& WaveInstances, int32 Fir
 
 void FAudioDevice::Update(bool bGameTicking)
 {
-	SCOPE_CYCLE_COUNTER(STAT_AudioUpdateTime);
+	if (!IsInAudioThread())
+	{
+
+		FAudioDevice* AudioDevice = this;
+		FAudioThread::RunCommandOnAudioThread([AudioDevice, bGameTicking]()
+		{
+			AudioDevice->Update(bGameTicking);
+		});
+
+		return;
+	}
+
+	DECLARE_CYCLE_STAT(TEXT("FAudioThreadTask.AudioUpdateTime"), STAT_AudioUpdateTime, STATGROUP_AudioThreadCommands);
+	FScopeCycleCounter(GET_STATID(STAT_AudioUpdateTime));
 
 	double CurrTime = FPlatformTime::Seconds();
 	double DeltaTime = CurrTime - LastUpdateTime;
@@ -2530,20 +2808,47 @@ void FAudioDevice::Update(bool bGameTicking)
 	// Handle pause/unpause for the game and editor.
 	HandlePause(bGameTicking);
 
-	// Update the audio effects - reverb, EQ etc
-	Effects->Update();
-
-	// Gets the current state of the sound classes accounting for sound mix
-	UpdateSoundClassProperties(DeltaTime);
+	bool bHasVolumeSettings = false;
+	float AudioVolumePriority = 0.f;
+	FReverbSettings ReverbSettings;
 
 	// Gets the current state of the interior settings
 	for (FListener& Listener : Listeners)
 	{
+		FAudioVolumeSettings PlayerAudioVolumeSettings;
+		GetAudioVolumeSettings(Listener.WorldID, Listener.Transform.GetLocation(), PlayerAudioVolumeSettings);
+
+		Listener.ApplyInteriorSettings(PlayerAudioVolumeSettings.AudioVolumeID, PlayerAudioVolumeSettings.InteriorSettings);
 		Listener.UpdateCurrentInteriorSettings();
+
+		if (!bHasVolumeSettings || (PlayerAudioVolumeSettings.AudioVolumeID > 0 && PlayerAudioVolumeSettings.Priority > AudioVolumePriority))
+		{
+			bHasVolumeSettings = true;
+			AudioVolumePriority = PlayerAudioVolumeSettings.Priority;
+			ReverbSettings = PlayerAudioVolumeSettings.ReverbSettings;
+		}
 	}
 
+	if (bHasActivatedReverb)
+	{
+		if (HighestPriorityActivatedReverb.Priority > AudioVolumePriority)
+		{
+			ReverbSettings = HighestPriorityActivatedReverb.ReverbSettings;
+		}
+	}
+
+	Effects->SetReverbSettings(ReverbSettings);
+
+	// Update the audio effects - reverb, EQ etc
+	Effects->Update();
+
+	// Gets the current state of the sound classes accounting for sound mix
+	UpdateSoundClassProperties(DeltaTime);
+
 	ProcessingPendingActiveSoundStops();
 
+	int32 FirstActiveIndex = INDEX_NONE;
+
 	if (Sources.Num())
 	{
 		// Kill any sources that have finished
@@ -2558,7 +2863,7 @@ void FAudioDevice::Update(bool bGameTicking)
 
 		// Poll audio components for active wave instances (== paths in node tree that end in a USoundWave)
 		ActiveWaveInstances.Reset();
-		int32 FirstActiveIndex = GetSortedActiveWaveInstances(ActiveWaveInstances, (bGameTicking ? ESortedActiveWaveGetType::FullUpdate : ESortedActiveWaveGetType::PausedUpdate));
+		FirstActiveIndex = GetSortedActiveWaveInstances(ActiveWaveInstances, (bGameTicking ? ESortedActiveWaveGetType::FullUpdate : ESortedActiveWaveGetType::PausedUpdate));
 
 		// Stop sources that need to be stopped, and touch the ones that need to be kept alive
 		StopSources(ActiveWaveInstances, FirstActiveIndex);
@@ -2575,11 +2880,22 @@ void FAudioDevice::Update(bool bGameTicking)
 		INC_DWORD_STAT_BY(STAT_ActiveSounds, ActiveSounds.Num());
 	}
 
-	UpdateListenerTransform();
+	// Update listener transform
+	if (Listeners.Num() > 0)
+	{
+		// Caches the matrix used to transform a sounds position into local space so we can just look
+		// at the Y component after normalization to determine spatialization.
+		const FVector Up = Listeners[0].GetUp();
+		const FVector Right = Listeners[0].GetFront();
+		InverseListenerTransform = FMatrix(Up, Right, Up ^ Right, Listeners[0].Transform.GetTranslation()).InverseFast();
+	}
 
 	// now let the platform perform anything it needs to handle
 	UpdateHardware();
 
+	// send any needed information back to the game thread
+	SendUpdateResultsToGameThread(FirstActiveIndex);
+
 #if !UE_BUILD_SHIPPING
 	// Print statistics for first non initial load allocation.
 	static bool bFirstTime = true;
@@ -2595,9 +2911,88 @@ void FAudioDevice::Update(bool bGameTicking)
 #endif
 }
 
+void FAudioDevice::SendUpdateResultsToGameThread(const int32 FirstActiveIndex)
+{
+#if !UE_BUILD_SHIPPING
+	TArray<FAudioStats::FStatSoundInfo> StatSoundInfos;
+	const bool bStatsStale = (RequestedAudioStats == 0);
+	if (RequestedAudioStats != 0)
+	{
+		const FVector ListenerPosition = Listeners[0].Transform.GetTranslation();
+
+		TMap<FActiveSound*, int32> ActiveSoundToInfoIndex;
+
+		const bool bDebug = (RequestedAudioStats & ERequestedAudioStats::DebugSounds) != 0;
+
+		for (FActiveSound* ActiveSound : ActiveSounds)
+		{
+			if (ActiveSound->Sound)
+			{
+				if (!bDebug || ActiveSound->GetSound()->bDebug)
+				{
+					ActiveSoundToInfoIndex.Add(ActiveSound, StatSoundInfos.AddDefaulted());
+					FAudioStats::FStatSoundInfo& StatSoundInfo = StatSoundInfos.Last();
+					StatSoundInfo.SoundName = ActiveSound->GetSound()->GetPathName();
+					StatSoundInfo.Distance = (ListenerPosition - ActiveSound->Transform.GetTranslation()).Size();
+					StatSoundInfo.SoundClassName = ActiveSound->GetSoundClass()->GetFName();
+					StatSoundInfo.Transform = ActiveSound->Transform;
+					StatSoundInfo.AudioComponentID = ActiveSound->GetAudioComponentID();
+
+					if (bDebug && ActiveSound->GetSound()->bDebug)
+					{
+						ActiveSound->CollectAttenuationShapesForVisualization(StatSoundInfo.ShapeDetailsMap);
+					}
+				}
+			}
+		}
+
+		// Iterate through all wave instances.
+		for (int32 InstanceIndex = FirstActiveIndex; InstanceIndex < ActiveWaveInstances.Num(); ++InstanceIndex)
+		{
+			FWaveInstance* WaveInstance = ActiveWaveInstances[InstanceIndex];
+			int32* SoundInfoIndex = ActiveSoundToInfoIndex.Find(WaveInstance->ActiveSound);
+			if (SoundInfoIndex)
+			{
+				FAudioStats::FStatWaveInstanceInfo WaveInstanceInfo;
+				FSoundSource* Source = WaveInstanceSourceMap.FindRef(WaveInstance);
+				WaveInstanceInfo.Description = Source ? Source->Describe((RequestedAudioStats & ERequestedAudioStats::LongSoundNames) != 0) : FString(TEXT("No source"));
+				WaveInstanceInfo.ActualVolume = WaveInstance->GetActualVolume();
+				WaveInstanceInfo.InstanceIndex = InstanceIndex;
+				StatSoundInfos[*SoundInfoIndex].WaveInstanceInfos.Add(MoveTemp(WaveInstanceInfo));
+			}
+		}
+	}
+#endif
+
+	DECLARE_CYCLE_STAT(TEXT("FGameThreadAudioTask.AudioSendResults"), STAT_AudioSendResults, STATGROUP_TaskGraphTasks);
+
+	const uint32 AudioDeviceID = DeviceHandle;
+
+	UReverbEffect* ReverbEffect = Effects->GetCurrentReverbEffect();
+	FAudioThread::RunCommandOnGameThread([AudioDeviceID, ReverbEffect
+#if !UE_BUILD_SHIPPING
+											, StatSoundInfos, bStatsStale
+#endif
+													]()
+	{
+		if (FAudioDeviceManager* AudioDeviceManager = GEngine->GetAudioDeviceManager())
+		{
+			if (FAudioDevice* AudioDevice = AudioDeviceManager->GetAudioDevice(AudioDeviceID))
+			{
+				AudioDevice->CurrentReverbEffect = ReverbEffect;
+#if !UE_BUILD_SHIPPING
+				AudioDevice->AudioStats.StatSoundInfos = MoveTemp(StatSoundInfos);
+				AudioDevice->AudioStats.bStale = bStatsStale;
+			}
+		}
+#endif
+	}, GET_STATID(STAT_AudioSendResults));
+}
 
 void FAudioDevice::StopAllSounds(bool bShouldStopUISounds)
 {
+	check(IsInAudioThread());
+
 	for (int32 SoundIndex=ActiveSounds.Num() - 1; SoundIndex >= 0; --SoundIndex)
 	{
 		FActiveSound* ActiveSound = ActiveSounds[SoundIndex];
@@ -2632,6 +3027,19 @@ void FAudioDevice::AddNewActiveSound(const FActiveSound& NewActiveSound)
 {
 	check(NewActiveSound.Sound);
 
+	if (!IsInAudioThread())
+	{
+		DECLARE_CYCLE_STAT(TEXT("FAudioThreadTask.AddNewActiveSound"), STAT_AudioAddNewActiveSound, STATGROUP_AudioThreadCommands);
+
+		FAudioDevice* AudioDevice = this;
+		FAudioThread::RunCommandOnAudioThread([AudioDevice, NewActiveSound]()
+		{
+			AudioDevice->AddNewActiveSound(NewActiveSound);
+		}, GET_STATID(STAT_AudioAddNewActiveSound));
+
+		return;		
+	}
+
 	// Evaluate concurrency. This will create an ActiveSound ptr which is a copy of NewActiveSound if the sound can play.
 	FActiveSound* ActiveSound = nullptr;
 
@@ -2649,10 +3057,8 @@ void FAudioDevice::AddNewActiveSound(const FActiveSound& NewActiveSound)
 
 	++NewActiveSound.Sound->CurrentPlayCount;
 
-	UAudioComponent* AudioComponent = NewActiveSound.GetAudioComponent();
-
 #if !(UE_BUILD_SHIPPING || UE_BUILD_TEST)
-	UE_LOG(LogAudio, VeryVerbose, TEXT("New ActiveSound %s Comp: %s Loc: %s"), *NewActiveSound.Sound->GetName(), (AudioComponent ? *AudioComponent->GetFullName() : TEXT("No AudioComponent")), *NewActiveSound.Transform.GetTranslation().ToString());
+	UE_LOG(LogAudio, VeryVerbose, TEXT("New ActiveSound %s Comp: %s Loc: %s"), *NewActiveSound.Sound->GetName(), *NewActiveSound.GetAudioComponentName(), *NewActiveSound.Transform.GetTranslation().ToString());
 #endif // !(UE_BUILD_SHIPPING || UE_BUILD_TEST)
 
 	check(ActiveSound);
@@ -2662,7 +3068,7 @@ void FAudioDevice::AddNewActiveSound(const FActiveSound& NewActiveSound)
 	{
 		if (!ensureMsgf(ActiveSound->Sound->GetFName() != NAME_None, TEXT("AddNewActiveSound with DESTROYED sound %s. AudioComponent=%s. IsPendingKill=%d. BeginDestroy=%d"),
 			*ActiveSound->Sound->GetPathName(),
-			AudioComponent ? *AudioComponent->GetPathName() : TEXT("NO COMPONENT"),
+			*ActiveSound->GetAudioComponentName(),
 			(int32)ActiveSound->Sound->IsPendingKill(),
 			(int32)ActiveSound->Sound->HasAnyFlags(RF_BeginDestroyed)))
 		{
@@ -2677,9 +3083,9 @@ void FAudioDevice::AddNewActiveSound(const FActiveSound& NewActiveSound)
 #endif
 
 	ActiveSounds.Add(ActiveSound);
-	if (AudioComponent)
+	if (ActiveSound->GetAudioComponentID() > 0)
 	{
-		AudioComponentToActiveSoundMap.Add((UPTRINT)AudioComponent, ActiveSound);
+		AudioComponentIDToActiveSoundMap.Add(ActiveSound->GetAudioComponentID(), ActiveSound);
 	}
 
 }
@@ -2723,10 +3129,12 @@ void FAudioDevice::ProcessingPendingActiveSoundStops(bool bForceDelete)
 
 void FAudioDevice::AddSoundToStop(FActiveSound* SoundToStop)
 {
-	const UPTRINT AudioComponentIndex = SoundToStop->GetAudioComponentIndex();
-	if (AudioComponentIndex > 0)
+	check(IsInAudioThread());
+
+	const uint64 AudioComponentID = SoundToStop->GetAudioComponentID();
+	if (AudioComponentID > 0)
 	{
-		AudioComponentToActiveSoundMap.Remove(AudioComponentIndex);
+		AudioComponentIDToActiveSoundMap.Remove(AudioComponentID);
 	}
 
 	check(SoundToStop);
@@ -2738,11 +3146,9 @@ void FAudioDevice::AddSoundToStop(FActiveSound* SoundToStop)
 	}
 }
 
-void FAudioDevice::StopActiveSound(UAudioComponent* AudioComponent)
+void FAudioDevice::StopActiveSound(const uint64 AudioComponentID)
 {
-	check(AudioComponent);
-
-	FActiveSound* ActiveSound = FindActiveSound(AudioComponent);
+	FActiveSound* ActiveSound = FindActiveSound(AudioComponentID);
 	if (ActiveSound)
 	{
 		StopActiveSound(ActiveSound);
@@ -2754,10 +3160,12 @@ void FAudioDevice::StopActiveSound(FActiveSound* ActiveSound)
 	AddSoundToStop(ActiveSound);
 }
 
-FActiveSound* FAudioDevice::FindActiveSound(UAudioComponent* AudioComponent)
+FActiveSound* FAudioDevice::FindActiveSound(const uint64 AudioComponentID)
 {
+	check(IsInAudioThread());
+
 	// find the active sound corresponding to this audio component
-	if (FActiveSound** ActiveSoundPtr = AudioComponentToActiveSoundMap.Find((UPTRINT)AudioComponent))
+	if (FActiveSound** ActiveSoundPtr = AudioComponentIDToActiveSoundMap.Find(AudioComponentID))
 	{
 		return *ActiveSoundPtr;
 	}
@@ -2767,37 +3175,57 @@ FActiveSound* FAudioDevice::FindActiveSound(UAudioComponent* AudioComponent)
 
 void FAudioDevice::RemoveActiveSound(FActiveSound* ActiveSound)
 {
+	check(IsInAudioThread());
+
 	ConcurrencyManager.RemoveActiveSound(ActiveSound);
 
 	// Perform the notification
-	if (UAudioComponent* AudioComponent = ActiveSound->GetAudioComponent())
+	if (ActiveSound->GetAudioComponentID() > 0)
 	{
-		AudioComponent->PlaybackCompleted(false);
+		UAudioComponent::PlaybackCompleted(ActiveSound->GetAudioComponentID(), false);
 	}
 
 	int32 NumRemoved = ActiveSounds.Remove(ActiveSound);
 	check(NumRemoved == 1);
 }
 
-bool FAudioDevice::LocationIsAudible(const FVector& Location, const float MaxDistance)
+bool FAudioDevice::LocationIsAudible(const FVector& Location, const float MaxDistance) const
 {
 	if (MaxDistance >= WORLD_MAX)
 	{
 		return true;
 	}
 
-	for (const FListener& Listener : Listeners)
+	const bool bInAudioThread = IsInAudioThread();
+	const bool bInGameThread = IsInGameThread();
+
+	check(bInAudioThread || bInGameThread);
+
+	if (bInAudioThread)
 	{
-		if (LocationIsAudible(Location, Listener, MaxDistance))
+		for (const FListener& Listener : Listeners)
 		{
-			return true;
+			if (LocationIsAudible(Location, Listener.Transform, MaxDistance))
+			{
+				return true;
+			}
+		}
+	}
+	else //if (bInGameThread)
+	{
+		for (const FTransform& ListenerTransform : ListenerTransforms)
+		{
+			if (LocationIsAudible(Location, ListenerTransform, MaxDistance))
+			{
+				return true;
+			}
 		}
 	}
 
 	return false;
 }
 
-bool FAudioDevice::LocationIsAudible(const FVector& Location, const FListener& Listener, float MaxDistance)
+bool FAudioDevice::LocationIsAudible(const FVector& Location, const FTransform& ListenerTransform, float MaxDistance)
 {
 	if (MaxDistance >= WORLD_MAX)
 	{
@@ -2805,14 +3233,13 @@ bool FAudioDevice::LocationIsAudible(const FVector& Location, const FListener& L
 	}
 
 	const float MaxDistanceSquared = MaxDistance * MaxDistance;
-	return (Listener.Transform.GetTranslation() - Location).SizeSquared() < MaxDistanceSquared;
+	return (ListenerTransform.GetTranslation() - Location).SizeSquared() < MaxDistanceSquared;
 }
 
-void FAudioDevice::GetMaxDistanceAndFocusFactor(USoundBase* Sound, const UWorld* World, const FVector& Location, const FAttenuationSettings* AttenuationSettingsToApply, float *OutMaxDistance, float* OutFocusFactor)
+void FAudioDevice::GetMaxDistanceAndFocusFactor(USoundBase* Sound, const UWorld* World, const FVector& Location, const FAttenuationSettings* AttenuationSettingsToApply, float& OutMaxDistance, float& OutFocusFactor)
 {
+	check(IsInGameThread());
 	check(Sound);
-	check(OutMaxDistance);
-	check(OutFocusFactor);
 
 	bool bIsInGameWorld = World ? World->IsGameWorld() : true;
 	bool bHasAttenuationSettings = (bIsInGameWorld && AttenuationSettingsToApply);
@@ -2822,28 +3249,30 @@ void FAudioDevice::GetMaxDistanceAndFocusFactor(USoundBase* Sound, const UWorld*
 		FTransform SoundTransform;
 		SoundTransform.SetTranslation(Location);
 
-		*OutMaxDistance = AttenuationSettingsToApply->GetMaxDimension();
+		OutMaxDistance = AttenuationSettingsToApply->GetMaxDimension();
 
 		// Now scale the max distance based on the focus settings in the attenuation settings
 		FAttenuationListenerData ListenerData;
-		*OutFocusFactor = GetFocusFactor(ListenerData, Sound, SoundTransform, *AttenuationSettingsToApply);
+		OutFocusFactor = GetFocusFactor(ListenerData, Sound, SoundTransform, *AttenuationSettingsToApply);
 	}
 	else
 	{
 		// No need to scale the distance by focus factor since we're not using any attenuation settings
-		*OutMaxDistance = Sound->GetMaxAudibleDistance();
-		*OutFocusFactor = 1.0f;
+		OutMaxDistance = Sound->GetMaxAudibleDistance();
+		OutFocusFactor = 1.0f;
 	}
 }
 
 bool FAudioDevice::SoundIsAudible(USoundBase* Sound, const UWorld* World, const FVector& Location, const FAttenuationSettings* AttenuationSettingsToApply, float MaxDistance, float FocusFactor)
 {
-	bool bIsInGameWorld = World ? World->IsGameWorld() : true;
-	bool bHasAttenuationSettings = (bIsInGameWorld && AttenuationSettingsToApply);
+	check(IsInGameThread());
+
+	const bool bIsInGameWorld = World ? World->IsGameWorld() : true;
+	const bool bHasAttenuationSettings = (bIsInGameWorld && AttenuationSettingsToApply);
 	float DistanceScale = 1.0f;
 	if (bHasAttenuationSettings)
 	{
-		DistanceScale = AttenuationSettingsToApply->GetFocusDistanceScale(GlobalFocusSettings, FocusFactor);
+		DistanceScale = AttenuationSettingsToApply->GetFocusDistanceScale(GetGlobalFocusSettings(), FocusFactor);
 	}
 
 	DistanceScale = FMath::Max(DistanceScale, 0.0001f);
@@ -2873,25 +3302,60 @@ int32 FAudioDevice::FindClosestListenerIndex(const FTransform& SoundTransform, c
 
 int32 FAudioDevice::FindClosestListenerIndex(const FTransform& SoundTransform) const
 {
-	return FindClosestListenerIndex(SoundTransform, Listeners);
+	if (IsInAudioThread())
+	{
+		return FindClosestListenerIndex(SoundTransform, Listeners);
+	}
+	else if (IsInGameThread())
+	{
+		int32 ClosestListenerIndex = 0;
+		if (ListenerTransforms.Num() > 0)
+		{
+			float ClosestDistSq = FVector::DistSquared(SoundTransform.GetTranslation(), ListenerTransforms[0].GetTranslation());
+
+			for (int32 i = 1; i < ListenerTransforms.Num(); i++)
+			{
+				const float DistSq = FVector::DistSquared(SoundTransform.GetTranslation(), ListenerTransforms[i].GetTranslation());
+				if (DistSq < ClosestDistSq)
+				{
+					ClosestListenerIndex = i;
+					ClosestDistSq = DistSq;
+				}
+			}
+		}
+
+		return ClosestListenerIndex;
+	}
+
+	return INDEX_NONE;
 }
 
-void FAudioDevice::GetAttenuationListenerData(FAttenuationListenerData& OutListenerData, const FTransform& SoundTransform, const FAttenuationSettings& AttenuationSettings, const FListener* InListener) const
+void FAudioDevice::GetAttenuationListenerData(FAttenuationListenerData& OutListenerData, const FTransform& SoundTransform, const FAttenuationSettings& AttenuationSettings, const FTransform* InListenerTransform) const
 {
 	// Only compute various components of the listener of it hasn't been computed yet
 	if (!OutListenerData.bDataComputed)
 	{
-		// Use the optional input listener param
-		OutListenerData.Listener = InListener;
 
+		// Use the optional input listener param
+		if (InListenerTransform)
+		{
+			OutListenerData.ListenerTransform = *InListenerTransform;
+		}
 		// If not set, then we need to find the closest listener
-		if (!OutListenerData.Listener)
+		else
 		{
 			const int32 ClosestListenerIndex = FindClosestListenerIndex(SoundTransform);
-			OutListenerData.Listener = &Listeners[ClosestListenerIndex];
+			if (IsInAudioThread())
+			{
+				OutListenerData.ListenerTransform = Listeners[ClosestListenerIndex].Transform;
+			}
+			else if (IsInGameThread())
+			{
+				OutListenerData.ListenerTransform = ListenerTransforms[ClosestListenerIndex];
+			}
 		}
 
-		const FVector& ListenerLocation = OutListenerData.Listener->Transform.GetTranslation();
+		const FVector ListenerLocation = OutListenerData.ListenerTransform.GetTranslation();
 		FVector ListenerToSound = SoundTransform.GetTranslation() - ListenerLocation;
 		ListenerToSound.ToDirectionAndLength(OutListenerData.ListenerToSoundDir, OutListenerData.ListenerToSoundDistance);
 
@@ -2906,7 +3370,7 @@ void FAudioDevice::GetAttenuationListenerData(FAttenuationListenerData& OutListe
 	}
 }
 
-float FAudioDevice::GetFocusFactor(FAttenuationListenerData& OutListenerData, const USoundBase* Sound, const FTransform& SoundTransform, const FAttenuationSettings& AttenuationSettings, const FListener* InListener) const
+float FAudioDevice::GetFocusFactor(FAttenuationListenerData& OutListenerData, const USoundBase* Sound, const FTransform& SoundTransform, const FAttenuationSettings& AttenuationSettings, const FTransform* InListenerTransform) const
 {
 	check(Sound);
 
@@ -2916,18 +3380,18 @@ float FAudioDevice::GetFocusFactor(FAttenuationListenerData& OutListenerData, co
 	if (AttenuationSettings.bSpatialize && AttenuationSettings.bEnableListenerFocus && !Sound->bIgnoreFocus)
 	{
 		// Compute the focus factor based on listener emitter geometry (SoundTransform)
-		GetAttenuationListenerData(OutListenerData, SoundTransform, AttenuationSettings, InListener);
-
-		check(OutListenerData.Listener);
+		GetAttenuationListenerData(OutListenerData, SoundTransform, AttenuationSettings, InListenerTransform);
 
-		const FVector& ListenerForwardDir = OutListenerData.Listener->Transform.GetUnitAxis(EAxis::X);
+		const FVector& ListenerForwardDir = OutListenerData.ListenerTransform.GetUnitAxis(EAxis::X);
 
 		const float FocusDotProduct = FVector::DotProduct(ListenerForwardDir, OutListenerData.ListenerToSoundDir);
 		const float FocusAngleRadians = FMath::Acos(FocusDotProduct);
 		const float FocusAngle = FMath::RadiansToDegrees(FocusAngleRadians);
 
-		const float FocusAzimuth = FMath::Clamp(GlobalFocusSettings.FocusAzimuthScale * AttenuationSettings.FocusAzimuth, 0.0f, 180.0f);
-		const float NonFocusAzimuth = FMath::Clamp(GlobalFocusSettings.NonFocusAzimuthScale * AttenuationSettings.NonFocusAzimuth, 0.0f, 180.0f);
+		const FGlobalFocusSettings& FocusSettings = GetGlobalFocusSettings();
+
+		const float FocusAzimuth = FMath::Clamp(FocusSettings.FocusAzimuthScale * AttenuationSettings.FocusAzimuth, 0.0f, 180.0f);
+		const float NonFocusAzimuth = FMath::Clamp(FocusSettings.NonFocusAzimuthScale * AttenuationSettings.NonFocusAzimuth, 0.0f, 180.0f);
 
 		if (FocusAzimuth != NonFocusAzimuth)
 		{
@@ -2945,6 +3409,8 @@ float FAudioDevice::GetFocusFactor(FAttenuationListenerData& OutListenerData, co
 
 UAudioComponent* FAudioDevice::CreateComponent(USoundBase* Sound, UWorld* World, AActor* Actor, bool bPlay, bool bStopWhenOwnerDestroyed, const FVector* Location, USoundAttenuation* AttenuationSettings, USoundConcurrency* ConcurrencySettings)
 {
+	check(IsInGameThread());
+
 	UAudioComponent* AudioComponent = nullptr;
 
 	if (Sound && GEngine && GEngine->UseSound())
@@ -2982,7 +3448,7 @@ UAudioComponent* FAudioDevice::CreateComponent(USoundBase* Sound, UWorld* World,
 			{
 				float MaxDistance = 0.0f;
 				float FocusFactor = 0.0f;
-				AudioDevice->GetMaxDistanceAndFocusFactor(Sound, World, *Location, AttenuationSettingsToApply, &MaxDistance, &FocusFactor);
+				AudioDevice->GetMaxDistanceAndFocusFactor(Sound, World, *Location, AttenuationSettingsToApply, MaxDistance, FocusFactor);
 				bIsAudible = AudioDevice->SoundIsAudible(Sound, World, *Location, AttenuationSettingsToApply, MaxDistance, FocusFactor);
 			}
 
@@ -3042,6 +3508,8 @@ UAudioComponent* FAudioDevice::CreateComponent(USoundBase* Sound, UWorld* World,
 
 void FAudioDevice::PlaySoundAtLocation(USoundBase* Sound, UWorld* World, float VolumeMultiplier, float PitchMultiplier, float StartTime, const FVector& Location, const FRotator& Rotation, USoundAttenuation* AttenuationSettings, USoundConcurrency* ConcurrencySettings, const TArray<FAudioComponentParam>* Params)
 {
+	check(IsInGameThread());
+
 	if (!Sound || !World)
 	{
 		return;
@@ -3051,7 +3519,7 @@ void FAudioDevice::PlaySoundAtLocation(USoundBase* Sound, UWorld* World, float V
 	float MaxDistance = 0.0f;
 	float FocusFactor = 0.0f;
 
-	GetMaxDistanceAndFocusFactor(Sound, World, Location, AttenuationSettingsToApply, &MaxDistance, &FocusFactor);
+	GetMaxDistanceAndFocusFactor(Sound, World, Location, AttenuationSettingsToApply, MaxDistance, FocusFactor);
 
 	if (Sound->GetDuration() > 1.0f || SoundIsAudible(Sound, World, Location, AttenuationSettingsToApply, MaxDistance, FocusFactor))
 	{
@@ -3073,9 +3541,11 @@ void FAudioDevice::PlaySoundAtLocation(USoundBase* Sound, UWorld* World, float V
 		NewActiveSound.bHasAttenuationSettings = (bIsInGameWorld && AttenuationSettingsToApply);
 		if (NewActiveSound.bHasAttenuationSettings)
 		{
+			const FGlobalFocusSettings& FocusSettings = GetGlobalFocusSettings();
+
 			NewActiveSound.AttenuationSettings = *AttenuationSettingsToApply;
-			NewActiveSound.FocusPriorityScale = AttenuationSettingsToApply->GetFocusPriorityScale(GlobalFocusSettings, FocusFactor);
-			NewActiveSound.FocusDistanceScale = AttenuationSettingsToApply->GetFocusDistanceScale(GlobalFocusSettings, FocusFactor);
+			NewActiveSound.FocusPriorityScale = AttenuationSettingsToApply->GetFocusPriorityScale(FocusSettings, FocusFactor);
+			NewActiveSound.FocusDistanceScale = AttenuationSettingsToApply->GetFocusDistanceScale(FocusSettings, FocusFactor);
 		}
 
 		NewActiveSound.MaxDistance = MaxDistance;
@@ -3102,6 +3572,27 @@ void FAudioDevice::PlaySoundAtLocation(USoundBase* Sound, UWorld* World, float V
 
 void FAudioDevice::Flush(UWorld* WorldToFlush, bool bClearActivatedReverb)
 {
+	if (!IsInAudioThread())
+	{
+		DECLARE_CYCLE_STAT(TEXT("FAudioThreadTask.Flush"), STAT_AudioFlush, STATGROUP_AudioThreadCommands);
+
+		FAudioDevice* AudioDevice = this;
+		FAudioThread::RunCommandOnAudioThread([AudioDevice, WorldToFlush]()
+		{
+			AudioDevice->Flush(WorldToFlush);
+		}, GET_STATID(STAT_AudioFlush));
+
+		FAudioCommandFence AudioFence;
+		AudioFence.BeginFence();
+		AudioFence.Wait();
+
+		// Clear the GameThread cache of the listner
+		ListenerTransforms.Reset();
+		ListenerTransforms.AddDefaulted();
+
+		return;
+	}
+
 	// Stop all audio components attached to the scene
 	bool bFoundIgnoredComponent = false;
 	for (int32 Index = ActiveSounds.Num() - 1; Index >= 0; --Index)
@@ -3134,13 +3625,13 @@ void FAudioDevice::Flush(UWorld* WorldToFlush, bool bClearActivatedReverb)
 
 	// Anytime we flush, make sure to clear all the listeners.  We'll get the right ones soon enough.
 	Listeners.Reset();
-	Listeners.Add(FListener());
+	Listeners.AddDefaulted();
 
 	// Clear all the activated reverb effects
 	if (bClearActivatedReverb)
 	{
 		ActivatedReverbs.Reset();
-		HighestPriorityReverb = nullptr;
+		bHasActivatedReverb = false;
 	}
 
 	if (WorldToFlush == nullptr)
@@ -3287,8 +3778,10 @@ void FAudioDevice::Precache(USoundWave* SoundWave, bool bSynchronous, bool bTrac
 	}
 }
 
-void FAudioDevice::StopSourcesUsingBuffer(FSoundBuffer * SoundBuffer)
+void FAudioDevice::StopSourcesUsingBuffer(FSoundBuffer* SoundBuffer)
 {
+	check(IsInAudioThread());
+
 	if (SoundBuffer)
 	{
 		for (int32 SrcIndex = 0; SrcIndex < Sources.Num(); SrcIndex++)
@@ -3336,32 +3829,16 @@ FSoundClassProperties* FAudioDevice::GetSoundClassCurrentProperties(USoundClass*
 
 
 #if !UE_BUILD_SHIPPING
-// sort memory usage from large to small 
-struct FSortSoundBuffersBySize
-{
-	FORCEINLINE bool operator()(FSoundBuffer& A, FSoundBuffer& B) const
-	{
-		return B.GetSize() < A.GetSize();
-	}
-};
-
-// sort memory usage by name
-struct FSortSoundBuffersByName
-{
-	FORCEINLINE bool operator()(FSoundBuffer& A, FSoundBuffer& B) const
-	{
-		return A.ResourceName < B.ResourceName;
-	}
-};
-
 /** 
  * Displays debug information about the loaded sounds
  */
 bool FAudioDevice::HandleListSoundsCommand(const TCHAR* Cmd, FOutputDevice& Ar)
 {
+	FAudioThreadSuspendContext AudioThreadSuspend;
+
 	// does the user want to sort by name instead of size?
-	bool bAlphaSort = FParse::Param(Cmd, TEXT("ALPHASORT"));
-	bool bUseLongNames = FParse::Param(Cmd, TEXT("LONGNAMES"));
+	const bool bAlphaSort = FParse::Param(Cmd, TEXT("ALPHASORT"));
+	const bool bUseLongNames = FParse::Param(Cmd, TEXT("LONGNAMES"));
 
 	int32 TotalResident = 0;
 	int32 ResidentCount = 0;
@@ -3381,11 +3858,12 @@ bool FAudioDevice::HandleListSoundsCommand(const TCHAR* Cmd, FOutputDevice& Ar)
 	// sort by name or size, depending on flag
 	if (bAlphaSort)
 	{
-		AllSounds.Sort(FSortSoundBuffersByName());
+		AllSounds.Sort([](FSoundBuffer& A, FSoundBuffer& B) { return A.ResourceName < B.ResourceName; });
 	}
 	else
 	{
-		AllSounds.Sort(FSortSoundBuffersBySize());
+		// sort memory usage from large to small 
+		AllSounds.Sort([](FSoundBuffer& A, FSoundBuffer& B) { return B.GetSize() < A.GetSize(); });
 	}
 
 	// now list the sorted sounds
@@ -3406,22 +3884,45 @@ bool FAudioDevice::HandleListSoundsCommand(const TCHAR* Cmd, FOutputDevice& Ar)
 }
 #endif // !UE_BUILD_SHIPPING
 
-void FAudioDevice::StopSoundsUsingResource(USoundWave* SoundWave, TArray<UAudioComponent*>& StoppedComponents)
+void FAudioDevice::StopSoundsUsingResource(USoundWave* SoundWave, TArray<UAudioComponent*>* StoppedComponents)
 {
+	if (StoppedComponents == nullptr && !IsInAudioThread())
+	{
+		DECLARE_CYCLE_STAT(TEXT("FAudioThreadTask.StopSoundsUsingResource"), STAT_AudioStopSoundsUsingResource, STATGROUP_AudioThreadCommands);
+
+		FAudioDevice* AudioDevice = this;
+		FAudioThread::RunCommandOnAudioThread([AudioDevice, SoundWave]()
+		{
+			AudioDevice->StopSoundsUsingResource(SoundWave);
+		}, GET_STATID(STAT_AudioStopSoundsUsingResource));
+
+		return;
+	}
+	else if (StoppedComponents)
+	{
+		check(IsInGameThread());
+		FAudioCommandFence AudioFence;
+		AudioFence.BeginFence();
+		AudioFence.Wait();
+	}
+
 	bool bStoppedSounds = false;
 
 	for (int32 ActiveSoundIndex = ActiveSounds.Num() - 1; ActiveSoundIndex >= 0; --ActiveSoundIndex)
 	{
 		FActiveSound* ActiveSound = ActiveSounds[ActiveSoundIndex];
-		for (auto WaveInstanceIt(ActiveSound->WaveInstances.CreateConstIterator()); WaveInstanceIt; ++WaveInstanceIt)
+		for (const TPair<UPTRINT, FWaveInstance*>& WaveInstancePair : ActiveSound->WaveInstances)
 		{
 			// If anything the ActiveSound uses the wave then we stop the sound
-			FWaveInstance* WaveInstance = WaveInstanceIt.Value();
+			FWaveInstance* WaveInstance = WaveInstancePair.Value;
 			if (WaveInstance->WaveData == SoundWave)
 			{
-				if (UAudioComponent* AudioComponent = ActiveSound->GetAudioComponent())
+				if (StoppedComponents)
 				{
-					StoppedComponents.Add(AudioComponent);
+					if (UAudioComponent* AudioComponent = UAudioComponent::GetAudioComponentFromID(ActiveSound->GetAudioComponentID()))
+					{
+						StoppedComponents->Add(AudioComponent);
+					}
 				}
 				AddSoundToStop(ActiveSound);
 				bStoppedSounds = true;
@@ -3430,6 +3931,9 @@ void FAudioDevice::StopSoundsUsingResource(USoundWave* SoundWave, TArray<UAudioC
 		}
 	}
 
+	// Immediately stop all pending active sounds
+	ProcessingPendingActiveSoundStops();
+
 	if (!GIsEditor && bStoppedSounds)
 	{
 		UE_LOG(LogAudio, Warning, TEXT("All Sounds using SoundWave '%s' have been stopped"), *SoundWave->GetName());
@@ -3438,6 +3942,8 @@ void FAudioDevice::StopSoundsUsingResource(USoundWave* SoundWave, TArray<UAudioC
 
 bool FAudioDevice::IsAudioDeviceMuted() const
 {
+	check(IsInAudioThread());
+
 	// First check to see if the device manager has "bPlayAllPIEAudio" enabled
 	FAudioDeviceManager* DeviceManager = GEngine->GetAudioDeviceManager();
 	if (DeviceManager && DeviceManager->IsPlayAllDeviceAudio())
@@ -3448,20 +3954,26 @@ bool FAudioDevice::IsAudioDeviceMuted() const
 	return bIsDeviceMuted;
 }
 
-void FAudioDevice::UpdateListenerTransform()
+void FAudioDevice::SetDeviceMuted(const bool bMuted)
 {
-	if (Listeners.Num() > 0)
+	if (!IsInAudioThread())
 	{
-		// Caches the matrix used to transform a sounds position into local space so we can just look
-		// at the Y component after normalization to determine spatialization.
-		const FVector Up = Listeners[0].GetUp();
-		const FVector Right = Listeners[0].GetFront();
-		InverseListenerTransform = FMatrix(Up, Right, Up ^ Right, Listeners[0].Transform.GetTranslation()).InverseFast();
+		FAudioDevice* AudioDevice = this;
+		FAudioThread::RunCommandOnAudioThread([AudioDevice, bMuted]()
+		{
+			AudioDevice->SetDeviceMuted(bMuted);
+
+		});
+
+		return;
 	}
+
+	bIsDeviceMuted = bMuted;
 }
 
 FVector FAudioDevice::GetListenerTransformedDirection(const FVector& Position, float* OutDistance)
 {
+	check(IsInAudioThread());
 	FVector UnnormalizedDirection = InverseListenerTransform.TransformPosition(Position);
 	if (OutDistance)
 	{
@@ -3501,3 +4013,84 @@ bool FAudioDevice::CanUseVRAudioDevice()
 #endif
 	return true;
 }
+
+void FAudioDevice::SetTransientMasterVolume(const float InTransientMasterVolume)
+{
+	if (!IsInAudioThread())
+	{
+		FAudioDevice* AudioDevice = this;
+		FAudioThread::RunCommandOnAudioThread([AudioDevice, InTransientMasterVolume]()
+		{
+			AudioDevice->SetTransientMasterVolume(InTransientMasterVolume);
+		});
+
+		return;
+	}
+
+	TransientMasterVolume = InTransientMasterVolume;
+}
+
+FSoundSource* FAudioDevice::GetSoundSource(FWaveInstance* WaveInstance) const
+{
+	check(IsInAudioThread());
+	return WaveInstanceSourceMap.FindRef( WaveInstance );
+}
+
+const FGlobalFocusSettings& FAudioDevice::GetGlobalFocusSettings() const
+{
+	if (IsInAudioThread())
+	{
+		return GlobalFocusSettings;
+	}
+
+	check(IsInGameThread());
+	return GlobalFocusSettings_GameThread;
+}
+
+void FAudioDevice::SetGlobalFocusSettings(const FGlobalFocusSettings& NewFocusSettings)
+{
+	check(IsInGameThread());
+
+	GlobalFocusSettings_GameThread = NewFocusSettings;
+
+	DECLARE_CYCLE_STAT(TEXT("FAudioThreadTask.SetGlobalListenerFocusParameters"), STAT_AudioSetGlobalListenerFocusParameters, STATGROUP_TaskGraphTasks);
+	FAudioDevice* AudioDevice = this;
+	FAudioThread::RunCommandOnAudioThread([AudioDevice, NewFocusSettings]()
+	{
+		AudioDevice->GlobalFocusSettings = NewFocusSettings;
+	}, GET_STATID(STAT_AudioSetGlobalListenerFocusParameters));
+}
+
+void FAudioDevice::SetGlobalPitchModulation(float PitchModulation, float TimeSec)
+{
+	if (!IsInAudioThread())
+	{
+		DECLARE_CYCLE_STAT(TEXT("FAudioThreadTask.SetGlobalPitchModulation"), STAT_AudioSetGlobalPitchModulation, STATGROUP_TaskGraphTasks);
+
+		FAudioDevice* AudioDevice = this;
+		FAudioThread::RunCommandOnAudioThread([AudioDevice, PitchModulation, TimeSec]()
+		{
+			AudioDevice->SetGlobalPitchModulation(PitchModulation, TimeSec);
+		}, GET_STATID(STAT_AudioSetGlobalPitchModulation));
+
+		return;
+	}
+
+	GlobalPitchScale.Set(PitchModulation, TimeSec);
+}
+
+void FAudioDevice::SetPlatformAudioHeadroom(const float InPlatformHeadRoom)
+{
+	if (!IsInAudioThread())
+	{
+		FAudioDevice* AudioDevice = this;
+		FAudioThread::RunCommandOnAudioThread([AudioDevice, InPlatformHeadRoom]()
+		{
+			AudioDevice->SetPlatformAudioHeadroom(InPlatformHeadRoom);
+		});
+
+		return;
+	}
+
+	PlatformAudioHeadroom = InPlatformHeadRoom;
+}
\ No newline at end of file
diff --git a/Engine/Source/Runtime/Engine/Private/AudioDeviceManager.cpp b/Engine/Source/Runtime/Engine/Private/AudioDeviceManager.cpp
--- a/Engine/Source/Runtime/Engine/Private/AudioDeviceManager.cpp
+++ b/Engine/Source/Runtime/Engine/Private/AudioDeviceManager.cpp
@@ -4,6 +4,7 @@
 #include "Audio.h"
 #include "AudioDevice.h"
 #include "AudioDeviceManager.h"
+#include "AudioThread.h"
 #include "Sound/SoundWave.h"
 #include "Sound/AudioSettings.h"
 #include "GameFramework/GameUserSettings.h"
@@ -25,6 +26,13 @@ static const uint32 AUDIO_DEVICE_DEFAULT_ALLOWED_DEVICE_COUNT = 2;
 // The max number of audio devices allowed
 static const uint32 AUDIO_DEVICE_MAX_DEVICE_COUNT = 8;
 
+FAudioDeviceManager::FCreateAudioDeviceResults::FCreateAudioDeviceResults()
+	: Handle(AUDIO_DEVICE_HANDLE_INVALID)
+	, bNewDevice(false)
+	, AudioDevice(nullptr)
+{
+}
+
 /*-----------------------------------------------------------------------------
 FAudioDeviceManager implementation.
 -----------------------------------------------------------------------------*/
@@ -61,13 +69,14 @@ void FAudioDeviceManager::RegisterAudioDeviceModule(IAudioDeviceModule* AudioDev
 	AudioDeviceModule = AudioDeviceModuleInput;
 }
 
-FAudioDevice* FAudioDeviceManager::CreateAudioDevice(uint32& HandleOut, bool bCreateNewDevice)
+bool FAudioDeviceManager::CreateAudioDevice(bool bCreateNewDevice, FCreateAudioDeviceResults& OutResults)
 {
+	OutResults = FCreateAudioDeviceResults();
+
 	// If we don't have an audio device module, then we can't create new audio devices.
 	if (AudioDeviceModule == nullptr)
 	{
-		HandleOut = AUDIO_DEVICE_HANDLE_INVALID;
-		return nullptr;
+		return false;
 	}
 
 	// If we are running without the editor, we only need one audio device.
@@ -78,23 +87,21 @@ FAudioDevice* FAudioDeviceManager::CreateAudioDevice(uint32& HandleOut, bool bCr
 			FAudioDevice* MainAudioDevice = GEngine->GetMainAudioDevice();
 			if (MainAudioDevice)
 			{
-				HandleOut = MainAudioDevice->DeviceHandle;
-				return MainAudioDevice;
+				OutResults.Handle = MainAudioDevice->DeviceHandle;
+				OutResults.AudioDevice = MainAudioDevice;
+				return true;
 			}
-			return nullptr;
+			return false;
 		}
 	}
 
-	FAudioDevice* NewAudioDevice = nullptr;
-
 	if (NumActiveAudioDevices < AUDIO_DEVICE_DEFAULT_ALLOWED_DEVICE_COUNT || (bCreateNewDevice && NumActiveAudioDevices < AUDIO_DEVICE_MAX_DEVICE_COUNT))
 	{
 		// Create the new audio device and make sure it succeeded
-		NewAudioDevice = AudioDeviceModule->CreateAudioDevice();
-		if (NewAudioDevice == nullptr)
+		OutResults.AudioDevice = AudioDeviceModule->CreateAudioDevice();
+		if (OutResults.AudioDevice == nullptr)
 		{
-			HandleOut = AUDIO_DEVICE_HANDLE_INVALID;
-			return nullptr;
+			return false;
 		}
 
 		// Now generation a new audio device handle for the device and store the
@@ -110,7 +117,7 @@ FAudioDevice* FAudioDeviceManager::CreateAudioDevice(uint32& HandleOut, bool bCr
 			--FreeIndicesSize;
 			check(int32(AudioDeviceIndex) < Devices.Num());
 			check(Devices[AudioDeviceIndex] == nullptr);
-			Devices[AudioDeviceIndex] = NewAudioDevice;
+			Devices[AudioDeviceIndex] = OutResults.AudioDevice;
 		}
 		else
 		{
@@ -120,13 +127,14 @@ FAudioDevice* FAudioDeviceManager::CreateAudioDevice(uint32& HandleOut, bool bCr
 			Generations.Add(0);
 			AudioDeviceIndex = Generations.Num() - 1;
 			check(AudioDeviceIndex < (1 << AUDIO_DEVICE_HANDLE_INDEX_BITS));
-			Devices.Add(NewAudioDevice);
+			Devices.Add(OutResults.AudioDevice);
 		}
 
-		HandleOut = CreateHandle(AudioDeviceIndex, Generations[AudioDeviceIndex]);
+		OutResults.bNewDevice = true;
+		OutResults.Handle = CreateHandle(AudioDeviceIndex, Generations[AudioDeviceIndex]);
 
 		// Store the handle on the audio device itself
-		NewAudioDevice->DeviceHandle = HandleOut;
+		OutResults.AudioDevice->DeviceHandle = OutResults.Handle;
 	}
 	else
 	{
@@ -134,22 +142,21 @@ FAudioDevice* FAudioDeviceManager::CreateAudioDevice(uint32& HandleOut, bool bCr
 		FAudioDevice* MainAudioDevice = GEngine->GetMainAudioDevice();
 		if (MainAudioDevice)
 		{
-			HandleOut = MainAudioDevice->DeviceHandle;
-			NewAudioDevice = MainAudioDevice;
+			OutResults.Handle = MainAudioDevice->DeviceHandle;
+			OutResults.AudioDevice = MainAudioDevice;
 		}
 	}
 
 	++NumActiveAudioDevices;
 
 	const UAudioSettings* AudioSettings = GetDefault<UAudioSettings>();
-	if (!NewAudioDevice->Init(AudioSettings->GetQualityLevelSettings(GEngine->GetGameUserSettings()->GetAudioQualityLevel()).MaxChannels))
+	if (!OutResults.AudioDevice->Init(AudioSettings->GetQualityLevelSettings(GEngine->GetGameUserSettings()->GetAudioQualityLevel()).MaxChannels))
 	{
-		ShutdownAudioDevice(HandleOut);
-		HandleOut = AUDIO_DEVICE_HANDLE_INVALID;
-		NewAudioDevice = nullptr;
+		ShutdownAudioDevice(OutResults.Handle);
+		OutResults = FCreateAudioDeviceResults();
 	}
 
-	return NewAudioDevice;
+	return (OutResults.AudioDevice != nullptr);
 }
 
 bool FAudioDeviceManager::IsValidAudioDeviceHandle(uint32 Handle) const
@@ -292,19 +299,7 @@ void FAudioDeviceManager::AddReferencedObjects(FReferenceCollector& Collector)
 	}
 }
 
-void FAudioDeviceManager::StopSoundsUsingWave(class USoundWave* InSoundWave)
-{
-	for (FAudioDevice* AudioDevice : Devices)
-	{
-		if (AudioDevice)
-		{
-			TArray<UAudioComponent*> StoppedComponents;
-			AudioDevice->StopSoundsUsingResource(InSoundWave, StoppedComponents);
-		}
-	}
-}
-
-void FAudioDeviceManager::StopSoundsUsingResource(class USoundWave* InSoundWave, TArray<UAudioComponent*>& StoppedComponents)
+void FAudioDeviceManager::StopSoundsUsingResource(class USoundWave* InSoundWave, TArray<UAudioComponent*>* StoppedComponents)
 {
 	for (FAudioDevice* AudioDevice : Devices)
 	{
@@ -360,11 +355,11 @@ void FAudioDeviceManager::SetActiveDevice(uint32 InAudioDeviceHandle)
 				if (AudioDevice->DeviceHandle == InAudioDeviceHandle)
 				{
 					ActiveAudioDeviceHandle = InAudioDeviceHandle;
-					AudioDevice->bIsDeviceMuted = false;
+					AudioDevice->SetDeviceMuted(false);
 				}
 				else
 				{
-					AudioDevice->bIsDeviceMuted = true;
+					AudioDevice->SetDeviceMuted(true);
 				}
 			}
 		}
@@ -384,11 +379,11 @@ void FAudioDeviceManager::SetSoloDevice(uint32 InAudioDeviceHandle)
 				if (AudioDevice->DeviceHandle == InAudioDeviceHandle)
 				{
 					ActiveAudioDeviceHandle = InAudioDeviceHandle;
-					AudioDevice->bIsDeviceMuted = false;
+					AudioDevice->SetDeviceMuted(false);
 				}
 				else
 				{
-					AudioDevice->bIsDeviceMuted = true;
+					AudioDevice->SetDeviceMuted(true);
 				}
 			}
 		}
@@ -485,6 +480,20 @@ void FAudioDeviceManager::RemoveSoundBufferForResourceID(uint32 ResourceID)
 
 void FAudioDeviceManager::RemoveSoundMix(USoundMix* SoundMix)
 {
+	if (!IsInAudioThread())
+	{
+		DECLARE_CYCLE_STAT(TEXT("FAudioThreadTask.RemoveSoundMix"), STAT_AudioRemoveSoundMix, STATGROUP_AudioThreadCommands);
+
+		FAudioDeviceManager* AudioDeviceManager = this;
+		FAudioThread::RunCommandOnAudioThread([AudioDeviceManager, SoundMix]()
+		{
+			AudioDeviceManager->RemoveSoundMix(SoundMix);
+
+		}, GET_STATID(STAT_AudioRemoveSoundMix));
+
+		return;
+	}
+
 	for (FAudioDevice* AudioDevice : Devices)
 	{
 		if (AudioDevice)
@@ -496,21 +505,35 @@ void FAudioDeviceManager::RemoveSoundMix(USoundMix* SoundMix)
 
 void FAudioDeviceManager::TogglePlayAllDeviceAudio()
 {
-	bPlayAllDeviceAudio = !bPlayAllDeviceAudio;
-}
+	if (!IsInAudioThread())
+	{
+		FAudioDeviceManager* AudioDeviceManager = this;
+		FAudioThread::RunCommandOnAudioThread([AudioDeviceManager]()
+		{
+			AudioDeviceManager->TogglePlayAllDeviceAudio();
 
-bool FAudioDeviceManager::IsPlayAllDeviceAudio() const
-{
-	return bPlayAllDeviceAudio;
-}
+		});
 
-bool FAudioDeviceManager::IsVisualizeDebug3dEnabled() const
-{
-	return bVisualize3dDebug;
+		return;
+	}
+	
+	bPlayAllDeviceAudio = !bPlayAllDeviceAudio;
 }
 
 void FAudioDeviceManager::ToggleVisualize3dDebug()
 {
+	if (!IsInAudioThread())
+	{
+		FAudioDeviceManager* AudioDeviceManager = this;
+		FAudioThread::RunCommandOnAudioThread([AudioDeviceManager]()
+		{
+			AudioDeviceManager->ToggleVisualize3dDebug();
+
+		});
+
+		return;
+	}
+
 	bVisualize3dDebug = !bVisualize3dDebug;
 }
 
diff --git a/Engine/Source/Runtime/Engine/Private/AudioThread.cpp b/Engine/Source/Runtime/Engine/Private/AudioThread.cpp
new file mode 100644
--- /dev/null
+++ b/Engine/Source/Runtime/Engine/Private/AudioThread.cpp
@@ -0,0 +1,450 @@
+// Copyright 1998-2016 Epic Games, Inc. All Rights Reserved.
+
+/*=============================================================================
+	AudioThread.cpp: Audio thread implementation.
+=============================================================================*/
+
+#include "EnginePrivate.h"
+#include "AudioThread.h"
+#include "ExceptionHandling.h"
+#include "TaskGraphInterfaces.h"
+
+//
+// Globals
+//
+
+int32 GCVarSuspendAudioThread = 0;
+TAutoConsoleVariable<int32> CVarSuspendAudioThread(TEXT("AudioThread.SuspendAudioThread"), GCVarSuspendAudioThread, TEXT("0=Resume, 1=Suspend"), ECVF_Cheat);
+
+struct FAudioThreadInteractor
+{
+	static void UseAudioThreadCVarSinkFunction()
+	{
+		static bool bLastSuspendAudioThread = false;
+		const bool bSuspendAudioThread = (CVarSuspendAudioThread.GetValueOnGameThread() != 0);
+
+		if (bLastSuspendAudioThread != bSuspendAudioThread)
+		{
+			bLastSuspendAudioThread = bSuspendAudioThread;
+			if (GAudioThread)
+			{
+				if (bSuspendAudioThread)
+				{
+					FAudioThread::SuspendAudioThread();
+				}
+				else
+				{
+					FAudioThread::ResumeAudioThread();
+				}
+			}
+			else if (GIsEditor)
+			{
+				UE_LOG(LogAudio, Warning, TEXT("Audio threading is disabled in the editor."));
+			}
+			else
+			{
+				UE_LOG(LogAudio, Warning, TEXT("Cannot manipulate audio thread when disabled by platform or ini."));
+			}
+		}
+	}
+};
+
+static FAutoConsoleVariableSink CVarUseAudioThreadSink(FConsoleCommandDelegate::CreateStatic(&FAudioThreadInteractor::UseAudioThreadCVarSinkFunction));
+
+bool FAudioThread::bIsAudioThreadRunning = false;
+bool FAudioThread::bUseThreadedAudio = false;
+uint32 FAudioThread::SuspensionCount = 0;
+uint32 FAudioThread::CachedAudioThreadId = 0;
+FRunnable* FAudioThread::AudioThreadRunnable = nullptr;
+FString FAudioThread::AudioThreadError;
+volatile bool FAudioThread::bIsAudioThreadHealthy = true;
+
+
+/** The audio thread main loop */
+void AudioThreadMain( FEvent* TaskGraphBoundSyncEvent )
+{
+	FTaskGraphInterface::Get().AttachToThread(ENamedThreads::AudioThread);
+	FPlatformMisc::MemoryBarrier();
+
+	// Inform main thread that the audio thread has been attached to the taskgraph and is ready to receive tasks
+	if( TaskGraphBoundSyncEvent != nullptr )
+	{
+		TaskGraphBoundSyncEvent->Trigger();
+	}
+
+	// set the thread back to real time mode
+	FPlatformProcess::SetRealTimeMode();
+
+#if STATS
+	if (FThreadStats::WillEverCollectData())
+	{
+		FThreadStats::ExplicitFlush(); // flush the stats and set update the scope so we don't flush again until a frame update, this helps prevent fragmentation
+	}
+#endif
+
+	FTaskGraphInterface::Get().ProcessThreadUntilRequestReturn(ENamedThreads::AudioThread);
+	FPlatformMisc::MemoryBarrier();
+	
+#if STATS
+	if (FThreadStats::WillEverCollectData())
+	{
+		FThreadStats::ExplicitFlush(); // Another explicit flush to clean up the ScopeCount established above for any stats lingering since the last frame
+	}
+#endif
+}
+
+FAudioThread::FAudioThread()
+{
+	TaskGraphBoundSyncEvent	= FPlatformProcess::GetSynchEventFromPool(true);
+
+	FCoreUObjectDelegates::PreGarbageCollect.AddRaw(this, &FAudioThread::OnPreGarbageCollect);
+	FCoreUObjectDelegates::PostGarbageCollect.AddRaw(this, &FAudioThread::OnPostGarbageCollect);
+}
+
+FAudioThread::~FAudioThread()
+{
+	FCoreUObjectDelegates::PreGarbageCollect.RemoveAll(this);
+	FCoreUObjectDelegates::PostGarbageCollect.RemoveAll(this);
+
+	FPlatformProcess::ReturnSynchEventToPool(TaskGraphBoundSyncEvent);
+	TaskGraphBoundSyncEvent = nullptr;
+}
+
+void FAudioThread::SuspendAudioThread()
+{
+	check(IsInGameThread());
+	if (IsAudioThreadRunning())
+	{
+		++SuspensionCount;
+
+		// Make GC wait on the audio thread finishing processing
+		FAudioCommandFence AudioFence;
+		AudioFence.BeginFence();
+		AudioFence.Wait();
+
+		CachedAudioThreadId = GAudioThreadId;
+		GAudioThreadId = 0; // While we are suspended we will pretend we have no audio thread
+		bIsAudioThreadRunning = false;
+
+		FPlatformMisc::MemoryBarrier();
+	}
+}
+
+void FAudioThread::ResumeAudioThread()
+{
+	check(IsInGameThread());
+	if (SuspensionCount > 0)
+	{
+		--SuspensionCount;
+
+		if (SuspensionCount == 0)
+		{
+			GAudioThreadId = CachedAudioThreadId;
+			bIsAudioThreadRunning = true;
+
+			FPlatformMisc::MemoryBarrier();
+		}
+	}
+}
+
+void FAudioThread::OnPreGarbageCollect()
+{
+	SuspendAudioThread();
+}
+
+void FAudioThread::OnPostGarbageCollect()
+{
+	ResumeAudioThread();
+}
+
+bool FAudioThread::Init()
+{ 
+	GAudioThreadId = FPlatformTLS::GetCurrentThreadId();
+	return true; 
+}
+
+void FAudioThread::Exit()
+{
+	GAudioThreadId = 0;
+}
+
+uint32 FAudioThread::Run()
+{
+	FMemory::SetupTLSCachesOnCurrentThread();
+	FPlatformProcess::SetupAudioThread();
+
+#if PLATFORM_WINDOWS
+	if ( !FPlatformMisc::IsDebuggerPresent() || GAlwaysReportCrash )
+	{
+#if !PLATFORM_SEH_EXCEPTIONS_DISABLED
+		__try
+#endif
+		{
+			AudioThreadMain( TaskGraphBoundSyncEvent );
+		}
+#if !PLATFORM_SEH_EXCEPTIONS_DISABLED
+		__except(ReportCrash(GetExceptionInformation()))
+		{
+			AudioThreadError = GErrorHist;
+
+			// Use a memory barrier to ensure that the game thread sees the write to AudioThreadError before
+			// the write to GIsAudioThreadHealthy.
+			FPlatformMisc::MemoryBarrier();
+
+			bIsAudioThreadHealthy = false;
+		}
+#endif
+	}
+	else
+#endif // PLATFORM_WINDOWS
+	{
+		AudioThreadMain( TaskGraphBoundSyncEvent );
+	}
+	FMemory::ClearAndDisableTLSCachesOnCurrentThread();
+	return 0;
+}
+
+void FAudioThread::SetUseThreadedAudio(const bool bInUseThreadedAudio)
+{
+	if (bIsAudioThreadRunning && !bInUseThreadedAudio)
+	{
+		UE_LOG(LogAudio, Error, TEXT("You cannot disable using threaded audio once the thread has already begun running."));
+	}
+	else
+	{
+		bUseThreadedAudio = bInUseThreadedAudio;
+	}
+}
+
+void FAudioThread::RunCommandOnAudioThread(TFunction<void()> InFunction, const TStatId InStatId)
+{
+	if (bIsAudioThreadRunning)
+	{
+		FFunctionGraphTask::CreateAndDispatchWhenReady(MoveTemp(InFunction), InStatId, nullptr, ENamedThreads::AudioThread);
+	}
+	else
+	{
+		FScopeCycleCounter ScopeCycleCounter(InStatId);
+		InFunction();
+	}
+}
+
+void FAudioThread::RunCommandOnGameThread(TFunction<void()> InFunction, const TStatId InStatId)
+{
+	if (bIsAudioThreadRunning)
+	{
+		FFunctionGraphTask::CreateAndDispatchWhenReady(MoveTemp(InFunction), InStatId, nullptr, ENamedThreads::GameThread);
+	}
+	else
+	{
+		FScopeCycleCounter ScopeCycleCounter(InStatId);
+		InFunction();
+	}
+}
+
+static FString BuildAudioThreadName( uint32 ThreadIndex )
+{
+	return FString::Printf( TEXT( "%s %u" ), *FName( NAME_AudioThread ).GetPlainNameString(), ThreadIndex );
+}
+
+void FAudioThread::StartAudioThread()
+{
+	if (bUseThreadedAudio)
+	{
+		check(GAudioThread == nullptr);
+
+		static uint32 ThreadCount = 0;
+
+		bIsAudioThreadRunning = true;
+
+		// Create the audio thread.
+		AudioThreadRunnable = new FAudioThread();
+
+		GAudioThread = FRunnableThread::Create(AudioThreadRunnable, *BuildAudioThreadName(ThreadCount), 0, TPri_Normal, FPlatformAffinity::GetAudioThreadMask());
+
+		// Wait for audio thread to have taskgraph bound before we dispatch any tasks for it.
+		((FAudioThread*)AudioThreadRunnable)->TaskGraphBoundSyncEvent->Wait();
+
+		// ensure the thread has actually started and is idling
+		FAudioCommandFence Fence;
+		Fence.BeginFence();
+		Fence.Wait();
+
+		ThreadCount++;
+	}
+}
+
+void FAudioThread::StopAudioThread()
+{
+	// This function is not thread-safe. Ensure it is only called by the main game thread.
+	check( IsInGameThread() );
+
+	if (!bIsAudioThreadRunning)
+	{
+		return;
+	}
+
+	// unregister
+	IConsoleManager::Get().RegisterThreadPropagation();
+
+	FGraphEventRef QuitTask = TGraphTask<FReturnGraphTask>::CreateTask(nullptr, ENamedThreads::GameThread).ConstructAndDispatchWhenReady(ENamedThreads::AudioThread);
+
+	// Busy wait while BP debugging, to avoid opportunistic execution of game thread tasks
+	// If the game thread is already executing tasks, then we have no choice but to spin
+	if (GIntraFrameDebuggingGameThread || FTaskGraphInterface::Get().IsThreadProcessingTasks(ENamedThreads::GameThread) ) 
+	{
+		while ((QuitTask.GetReference() != nullptr) && !QuitTask->IsComplete())
+		{
+			FPlatformProcess::Sleep(0.0f);
+		}
+	}
+	else
+	{
+		QUICK_SCOPE_CYCLE_COUNTER(STAT_StopAudioThread);
+		FTaskGraphInterface::Get().WaitUntilTaskCompletes(QuitTask, ENamedThreads::GameThread_Local);
+	}
+
+	// Wait for the audio thread to return.
+	GAudioThread->WaitForCompletion();
+
+	bIsAudioThreadRunning = false;
+
+	// Destroy the audio thread objects.
+	delete GAudioThread;
+	GAudioThread = nullptr;
+			
+	delete AudioThreadRunnable;
+	AudioThreadRunnable = nullptr;
+}
+
+void FAudioThread::CheckAudioThreadHealth()
+{
+	if(!bIsAudioThreadHealthy)
+	{
+		GErrorHist[0] = 0;
+		GIsCriticalError = false;
+		UE_LOG(LogAudio, Fatal, TEXT("Audio thread exception:\r\n%s"),*AudioThreadError);
+	}
+
+	if (IsInGameThread())
+	{
+		GLog->FlushThreadedLogs();
+		SCOPE_CYCLE_COUNTER(STAT_PumpMessages);
+		FPlatformMisc::PumpMessages(false);
+	}
+}
+
+bool FAudioThread::IsAudioThreadHealthy()
+{
+	return bIsAudioThreadHealthy;
+}
+
+void FAudioCommandFence::BeginFence()
+{
+	if (FAudioThread::IsAudioThreadRunning())
+	{
+		DECLARE_CYCLE_STAT(TEXT("FNullGraphTask.FenceAudioCommand"),
+			STAT_FNullGraphTask_FenceAudioCommand,
+			STATGROUP_TaskGraphTasks);
+
+		CompletionEvent = TGraphTask<FNullGraphTask>::CreateTask(NULL, ENamedThreads::GameThread).ConstructAndDispatchWhenReady(
+			GET_STATID(STAT_FNullGraphTask_FenceAudioCommand), ENamedThreads::AudioThread);
+	}
+}
+
+bool FAudioCommandFence::IsFenceComplete() const
+{
+	if (!FAudioThread::IsAudioThreadRunning())
+	{
+		return true;
+	}
+	check(IsInGameThread() || IsInAsyncLoadingThread());
+	FAudioThread::CheckAudioThreadHealth();
+	if (!CompletionEvent.GetReference() || CompletionEvent->IsComplete())
+	{
+		CompletionEvent = NULL; // this frees the handle for other uses, the NULL state is considered completed
+		return true;
+	}
+	return false;
+}
+
+static int32 GTimeToBlockOnAudioFence = 1;
+static FAutoConsoleVariableRef CVarTimeToBlockOnAudioFence(
+	TEXT("g.TimeToBlockOnAudioFence"),
+	GTimeToBlockOnAudioFence,
+	TEXT("Number of milliseconds the game thread should block when waiting on a audio thread fence.")
+	);
+
+/**
+ * Block the game thread waiting for a task to finish on the audio thread.
+ */
+static void GameThreadWaitForTask(const FGraphEventRef& Task, bool bEmptyGameThreadTasks = false)
+{
+	check(IsInGameThread());
+	check(IsValidRef(Task));
+
+	if (!Task->IsComplete())
+	{
+		SCOPE_CYCLE_COUNTER(STAT_GameIdleTime);
+		{
+			static int32 NumRecursiveCalls = 0;
+		
+			// Check for recursion. It's not completely safe but because we pump messages while 
+			// blocked it is expected.
+			NumRecursiveCalls++;
+			if (NumRecursiveCalls > 1)
+			{
+				UE_LOG(LogAudio,Warning,TEXT("Wait on Audio Thread called recursively! %d calls on the stack."), NumRecursiveCalls);
+			}
+			if (NumRecursiveCalls > 1 || FTaskGraphInterface::Get().IsThreadProcessingTasks(ENamedThreads::GameThread))
+			{
+				bEmptyGameThreadTasks = false; // we don't do this on recursive calls or if we are at a blueprint breakpoint
+			}
+
+			// Grab an event from the pool and fire off a task to trigger it.
+			FEvent* Event = FPlatformProcess::GetSynchEventFromPool();
+			FTaskGraphInterface::Get().TriggerEventWhenTaskCompletes(Event, Task, ENamedThreads::GameThread);
+
+			// Check audio thread health needs to be called from time to
+			// time in order to make sure the audio thread hasn't crashed :)
+			bool bDone;
+			uint32 WaitTime = FMath::Clamp<uint32>(GTimeToBlockOnAudioFence, 0, 33);
+			do
+			{
+				FAudioThread::CheckAudioThreadHealth();
+				if (bEmptyGameThreadTasks)
+				{
+					// process gamethread tasks if there are any
+					FTaskGraphInterface::Get().ProcessThreadUntilIdle(ENamedThreads::GameThread);
+				}
+				bDone = Event->Wait(WaitTime);
+			}
+			while (!bDone);
+
+			// Return the event to the pool and decrement the recursion counter.
+			FPlatformProcess::ReturnSynchEventToPool(Event);
+			Event = nullptr;
+			NumRecursiveCalls--;
+		}
+	}
+}
+
+/**
+ * Waits for pending fence commands to retire.
+ */
+void FAudioCommandFence::Wait(bool bProcessGameThreadTasks) const
+{
+	if (!IsFenceComplete())
+	{
+#if 0
+		// on most platforms this is a better solution because it doesn't spin
+		// windows needs to pump messages
+		if (bProcessGameThreadTasks)
+		{
+			QUICK_SCOPE_CYCLE_COUNTER(STAT_FAudioCommandFence_Wait);
+			FTaskGraphInterface::Get().WaitUntilTaskCompletes(CompletionEvent, ENamedThreads::GameThread);
+		}
+#endif
+		GameThreadWaitForTask(CompletionEvent, bProcessGameThreadTasks);
+	}
+}
diff --git a/Engine/Source/Runtime/Engine/Private/AudioVolume.cpp b/Engine/Source/Runtime/Engine/Private/AudioVolume.cpp
--- a/Engine/Source/Runtime/Engine/Private/AudioVolume.cpp
+++ b/Engine/Source/Runtime/Engine/Private/AudioVolume.cpp
@@ -6,10 +6,19 @@
 
 #include "EnginePrivate.h"
 #include "AudioDevice.h"
+#include "AudioThread.h"
 #include "Components/BrushComponent.h"
 #include "Net/UnrealNetwork.h"
 #include "Sound/ReverbEffect.h"
 
+bool FReverbSettings::operator==(const FReverbSettings& Other) const
+{
+	return (bApplyReverb == Other.bApplyReverb
+			&& ReverbEffect == Other.ReverbEffect
+			&& Volume == Other.Volume
+			&& FadeTime == Other.FadeTime);
+}
+
 void FReverbSettings::PostSerialize(const FArchive& Ar)
 {
 	if( Ar.UE4Ver() < VER_UE4_REVERB_EFFECT_ASSET_TYPE )
@@ -152,114 +161,204 @@ void AAudioVolume::GetLifetimeReplicatedProps(TArray< FLifetimeProperty > & OutL
 	DOREPLIFETIME(AAudioVolume, bEnabled);
 }
 
-void AAudioVolume::PostUnregisterAllComponents( void )
+FAudioVolumeProxy::FAudioVolumeProxy(const AAudioVolume* AudioVolume)
+	: AudioVolumeID(AudioVolume->GetUniqueID())
+	, WorldID(AudioVolume->GetWorld()->GetUniqueID())
+	, Priority(AudioVolume->GetPriority())
+	, ReverbSettings(AudioVolume->GetReverbSettings())
+	, InteriorSettings(AudioVolume->GetInteriorSettings())
+	, BodyInstance(AudioVolume->GetBrushComponent()->GetBodyInstance())
 {
-	// Route clear to super first.
-	Super::PostUnregisterAllComponents();
+}
+
+void AAudioVolume::AddProxy() const
+{
+	UWorld* World = GetWorld();
 
+	if (FAudioDevice* AudioDevice = World->GetAudioDevice())
+	{
+		DECLARE_CYCLE_STAT(TEXT("FAudioThreadTask.AddAudioVolumeProxy"), STAT_AudioAddAudioVolumeProxy, STATGROUP_TaskGraphTasks);
+
+		FAudioVolumeProxy Proxy(this);
+
+		FAudioThread::RunCommandOnAudioThread([AudioDevice, Proxy]()
+		{
+			AudioDevice->AddAudioVolumeProxy(Proxy);
+		}, GET_STATID(STAT_AudioAddAudioVolumeProxy));
+	}
+}
+
+void FAudioDevice::AddAudioVolumeProxy(const FAudioVolumeProxy& Proxy)
+{
+	check(IsInAudioThread());
+
+	AudioVolumeProxies.Add(Proxy.AudioVolumeID, Proxy);
+	AudioVolumeProxies.ValueSort([](const FAudioVolumeProxy& A, const FAudioVolumeProxy& B) { return A.Priority > B.Priority; });
+
+	InvalidateCachedInteriorVolumes();
+}
+
+
+void AAudioVolume::RemoveProxy() const
+{
 	// World will be NULL during exit purge.
 	UWorld* World = GetWorld();
 	if (World)
 	{
-		AAudioVolume* CurrentVolume = World->HighestPriorityAudioVolume;
-		AAudioVolume* PreviousVolume = NULL;
-
-		// Iterate over linked list, removing this volume if found.
-		while (CurrentVolume)
+		if (FAudioDevice* AudioDevice = World->GetAudioDevice())
 		{
-			// Found.
-			if (CurrentVolume == this)
+			DECLARE_CYCLE_STAT(TEXT("FAudioThreadTask.RemoveAudioVolumeProxy"), STAT_AudioRemoveAudioVolumeProxy, STATGROUP_TaskGraphTasks);
+
+			const uint32 AudioVolumeID = GetUniqueID();
+			FAudioThread::RunCommandOnAudioThread([AudioDevice, AudioVolumeID]()
 			{
-				// Remove from linked list.
-				if (PreviousVolume)
-				{
-					PreviousVolume->NextLowerPriorityVolume = NextLowerPriorityVolume;
-				}
-				else
-				{
-					// Special case removal from first entry.
-					World->HighestPriorityAudioVolume = NextLowerPriorityVolume;
-				}
+				AudioDevice->RemoveAudioVolumeProxy(AudioVolumeID);
+			}, GET_STATID(STAT_AudioRemoveAudioVolumeProxy));
+		}
+	}
+}
 
-				break;
-			}
+void FAudioDevice::RemoveAudioVolumeProxy(const uint32 AudioVolumeID)
+{
+	check(IsInAudioThread());
 
-			// List traversal.
-			PreviousVolume = CurrentVolume;
-			CurrentVolume = CurrentVolume->NextLowerPriorityVolume;
-		}
+	AudioVolumeProxies.Remove(AudioVolumeID);
+
+	InvalidateCachedInteriorVolumes();
+}
 
-		// Reset next pointer to avoid dangling end bits and also for GC.
-		NextLowerPriorityVolume = nullptr;
+void AAudioVolume::UpdateProxy() const
+{
+	UWorld* World = GetWorld();
 
-		if (FAudioDevice* AudioDevice = World->GetAudioDevice())
+	if (FAudioDevice* AudioDevice = World->GetAudioDevice())
+	{
+		DECLARE_CYCLE_STAT(TEXT("FAudioThreadTask.UpdateAudioVolumeProxy"), STAT_AudioUpdateAudioVolumeProxy, STATGROUP_TaskGraphTasks);
+
+		FAudioVolumeProxy Proxy(this);
+
+		FAudioThread::RunCommandOnAudioThread([AudioDevice, Proxy]()
 		{
-			AudioDevice->InvalidateCachedInteriorVolumes();
+			AudioDevice->UpdateAudioVolumeProxy(Proxy);
+		}, GET_STATID(STAT_AudioUpdateAudioVolumeProxy));
+	}
+}
+
+void FAudioDevice::UpdateAudioVolumeProxy(const FAudioVolumeProxy& NewProxy)
+{
+	check(IsInAudioThread());
+
+	if (FAudioVolumeProxy* CurrentProxy = AudioVolumeProxies.Find(NewProxy.AudioVolumeID))
+	{
+		const float CurrentPriority = CurrentProxy->Priority;
+
+		*CurrentProxy = NewProxy;
+
+		if (CurrentPriority != NewProxy.Priority)
+		{
+			AudioVolumeProxies.ValueSort([](const FAudioVolumeProxy& A, const FAudioVolumeProxy& B) { return A.Priority > B.Priority; });
 		}
 	}
 }
 
+void AAudioVolume::PostUnregisterAllComponents()
+{
+	// Route clear to super first.
+	Super::PostUnregisterAllComponents();
+
+	GetRootComponent()->TransformUpdated.RemoveAll(this);
+	RemoveProxy();
+
+	if (UWorld* World = GetWorld())
+	{
+		World->AudioVolumes.Remove(this);
+	}
+}
+
 void AAudioVolume::PostRegisterAllComponents()
 {
 	// Route update to super first.
 	Super::PostRegisterAllComponents();
 
+	GetRootComponent()->TransformUpdated.AddUObject(this, &AAudioVolume::TransformUpdated);
+	AddProxy();
+
 	UWorld* World = GetWorld();
-	AAudioVolume* CurrentVolume = World->HighestPriorityAudioVolume;
-	AAudioVolume* PreviousVolume = nullptr;
+	World->AudioVolumes.Add(this);
+	World->AudioVolumes.Sort([](const AAudioVolume& A, const AAudioVolume& B) { return (A.GetPriority() > B.GetPriority()); });
+}
+
+void AAudioVolume::TransformUpdated(USceneComponent* InRootComponent, EUpdateTransformFlags UpdateTransformFlags, ETeleportType Teleport)
+{
+	UpdateProxy();
+}
 
-	// Find where to insert in sorted linked list.
-	if (CurrentVolume)
+void AAudioVolume::SetEnabled(const bool bNewEnabled)
+{
+	if (bNewEnabled != bEnabled)
 	{
-		// Avoid double insertion!
-		while (CurrentVolume && CurrentVolume != this)
+		bEnabled = bNewEnabled;
+		if (bEnabled)
 		{
-			// We use > instead of >= to be sure that we are not inserting twice in the case of multiple volumes having
-			// the same priority and the current one already having being inserted after one with the same priority.
-			if (Priority > CurrentVolume->Priority)
-			{
-				if (PreviousVolume)
-				{
-					// Insert before current node by fixing up previous to point to current.
-					PreviousVolume->NextLowerPriorityVolume = this;
-				}
-				else
-				{
-					// Special case for insertion at the beginning.
-					World->HighestPriorityAudioVolume = this;
-				}
-
-				// Point to current volume, finalizing insertion.
-				NextLowerPriorityVolume = CurrentVolume;
-				return;
-			}
-
-			// List traversal.
-			PreviousVolume = CurrentVolume;
-			CurrentVolume = CurrentVolume->NextLowerPriorityVolume;
+			AddProxy();
 		}
-
-		// We're the lowest priority volume, insert at the end.
-		if (!CurrentVolume)
+		else
 		{
-			checkSlow(PreviousVolume);
-			PreviousVolume->NextLowerPriorityVolume = this;
-			NextLowerPriorityVolume = nullptr;
+			RemoveProxy();
 		}
 	}
+}
+
+void AAudioVolume::OnRep_bEnabled()
+{
+	if (bEnabled)
+	{
+		AddProxy();
+	}
 	else
 	{
-		// First volume in the world info.
-		World->HighestPriorityAudioVolume = this;
-		NextLowerPriorityVolume	= nullptr;
+		RemoveProxy();
 	}
+}
 
-	if (FAudioDevice* AudioDevice = World->GetAudioDevice())
+void AAudioVolume::SetPriority(const float NewPriority)
+{
+	if (NewPriority != Priority)
+	{
+		Priority = NewPriority;
+		GetWorld()->AudioVolumes.Sort([](const AAudioVolume& A, const AAudioVolume& B) { return (A.GetPriority() > B.GetPriority()); });
+		if (bEnabled)
+		{
+			UpdateProxy();
+		}
+	}
+}
+
+void AAudioVolume::SetInteriorSettings(const FInteriorSettings& NewInteriorSettings)
+{
+	if (NewInteriorSettings != AmbientZoneSettings)
+	{
+		AmbientZoneSettings = NewInteriorSettings;
+		if (bEnabled)
+		{
+			UpdateProxy();
+		}
+	}
+}
+
+void AAudioVolume::SetReverbSettings(const FReverbSettings& NewReverbSettings)
+{
+	if (NewReverbSettings != Settings)
 	{
-		AudioDevice->InvalidateCachedInteriorVolumes();
+		Settings = NewReverbSettings;
+		if (bEnabled)
+		{
+			UpdateProxy();
+		}
 	}
 }
 
+
 #if WITH_EDITOR
 void AAudioVolume::PostEditChangeProperty(FPropertyChangedEvent& PropertyChangedEvent)
 {
@@ -270,5 +369,27 @@ void AAudioVolume::PostEditChangeProperty(FPropertyChangedEvent& PropertyChanged
 	AmbientZoneSettings.InteriorLPFTime = FMath::Max<float>( 0.01f, AmbientZoneSettings.InteriorLPFTime );
 	AmbientZoneSettings.ExteriorTime = FMath::Max<float>( 0.01f, AmbientZoneSettings.ExteriorTime );
 	AmbientZoneSettings.ExteriorLPFTime = FMath::Max<float>( 0.01f, AmbientZoneSettings.ExteriorLPFTime );
+
+	if (PropertyChangedEvent.Property && PropertyChangedEvent.Property->GetFName() == GET_MEMBER_NAME_CHECKED(AAudioVolume, Priority))
+	{
+		GetWorld()->AudioVolumes.Sort([](const AAudioVolume& A, const AAudioVolume& B) { return (A.GetPriority() > B.GetPriority()); });
+	}
+
+	if (PropertyChangedEvent.Property && PropertyChangedEvent.Property->GetFName() == GET_MEMBER_NAME_CHECKED(AAudioVolume, bEnabled))
+	{
+		if (bEnabled)
+		{
+			AddProxy();
+		}
+		else
+		{
+			RemoveProxy();
+		}
+	}
+	else if (bEnabled)
+	{
+		UpdateProxy();
+	}
+
 }
 #endif // WITH_EDITOR
diff --git a/Engine/Source/Runtime/Engine/Private/Character.cpp b/Engine/Source/Runtime/Engine/Private/Character.cpp
--- a/Engine/Source/Runtime/Engine/Private/Character.cpp
+++ b/Engine/Source/Runtime/Engine/Private/Character.cpp
@@ -353,14 +353,15 @@ void ACharacter::OnEndCrouch( float HeightAdjust, float ScaledHeightAdjust )
 {
 	RecalculateBaseEyeHeight();
 
-	if (Mesh)
+	const ACharacter* DefaultChar = GetDefault<ACharacter>(GetClass());
+	if (Mesh && DefaultChar->Mesh)
 	{
-		Mesh->RelativeLocation.Z = GetDefault<ACharacter>(GetClass())->Mesh->RelativeLocation.Z;
+		Mesh->RelativeLocation.Z = DefaultChar->Mesh->RelativeLocation.Z;
 		BaseTranslationOffset.Z = Mesh->RelativeLocation.Z;
 	}
 	else
 	{
-		BaseTranslationOffset.Z = GetDefault<ACharacter>(GetClass())->BaseTranslationOffset.Z;
+		BaseTranslationOffset.Z = DefaultChar->BaseTranslationOffset.Z;
 	}
 
 	K2_OnEndCrouch(HeightAdjust, ScaledHeightAdjust);
@@ -370,9 +371,10 @@ void ACharacter::OnStartCrouch( float HeightAdjust, float ScaledHeightAdjust )
 {
 	RecalculateBaseEyeHeight();
 
-	if (Mesh)
+	const ACharacter* DefaultChar = GetDefault<ACharacter>(GetClass());
+	if (Mesh && DefaultChar->Mesh)
 	{
-		Mesh->RelativeLocation.Z = GetDefault<ACharacter>(GetClass())->Mesh->RelativeLocation.Z + HeightAdjust;
+		Mesh->RelativeLocation.Z = DefaultChar->Mesh->RelativeLocation.Z + HeightAdjust;
 		BaseTranslationOffset.Z = Mesh->RelativeLocation.Z;
 	}
 	else
diff --git a/Engine/Source/Runtime/Engine/Private/Collision/Collision.cpp b/Engine/Source/Runtime/Engine/Private/Collision/Collision.cpp
--- a/Engine/Source/Runtime/Engine/Private/Collision/Collision.cpp
+++ b/Engine/Source/Runtime/Engine/Private/Collision/Collision.cpp
@@ -223,6 +223,11 @@ void FCollisionQueryParams::AddIgnoredActor(const AActor* InIgnoreActor)
 	}
 }
 
+void FCollisionQueryParams::AddIgnoredActor(const uint32 InIgnoreActorID)
+{
+	IgnoreActors.Add(InIgnoreActorID);
+}
+
 void FCollisionQueryParams::AddIgnoredActors(const TArray<AActor*>& InIgnoreActors)
 {
 	for (int32 Idx = 0; Idx < InIgnoreActors.Num(); ++Idx)
diff --git a/Engine/Source/Runtime/Engine/Private/Collision/CollisionProfile.cpp b/Engine/Source/Runtime/Engine/Private/Collision/CollisionProfile.cpp
--- a/Engine/Source/Runtime/Engine/Private/Collision/CollisionProfile.cpp
+++ b/Engine/Source/Runtime/Engine/Private/Collision/CollisionProfile.cpp
@@ -190,11 +190,7 @@ bool UCollisionProfile::ReadConfig(FName ProfileName, FBodyInstance& BodyInstanc
 		BodyInstance.CollisionResponses.SetCollisionResponseContainer(Template.ResponseToChannels);
 		BodyInstance.ResponseToChannels_DEPRECATED = Template.ResponseToChannels;
 
-		// if valid instance, make sure to update physics filter data
-		if (BodyInstance.IsValidBodyInstance())
-		{
-			BodyInstance.UpdatePhysicsFilterData();
-		}
+		BodyInstance.UpdatePhysicsFilterData();
 		return true;
 	}
 
diff --git a/Engine/Source/Runtime/Engine/Private/ComponentInstanceDataCache.cpp b/Engine/Source/Runtime/Engine/Private/ComponentInstanceDataCache.cpp
--- a/Engine/Source/Runtime/Engine/Private/ComponentInstanceDataCache.cpp
+++ b/Engine/Source/Runtime/Engine/Private/ComponentInstanceDataCache.cpp
@@ -56,10 +56,13 @@ FActorComponentInstanceData::FActorComponentInstanceData(const UActorComponent*
 			FComponentPropertyWriter(const UActorComponent* Component, TArray<uint8>& InBytes)
 				: FObjectWriter(InBytes)
 			{
-				Component->GetUCSModifiedProperties(PropertiesToSkip);
+				if (Component)
+				{
+					Component->GetUCSModifiedProperties(PropertiesToSkip);
 
-				UClass* ComponentClass = Component->GetClass();
-				ComponentClass->SerializeTaggedProperties(*this, (uint8*)Component, ComponentClass, (uint8*)Component->GetArchetype());
+					UClass* ComponentClass = Component->GetClass();
+					ComponentClass->SerializeTaggedProperties(*this, (uint8*)Component, ComponentClass, (uint8*)Component->GetArchetype());
+				}
 			}
 
 			virtual bool ShouldSkipProperty(const UProperty* InProperty) const override
@@ -72,7 +75,31 @@ FActorComponentInstanceData::FActorComponentInstanceData(const UActorComponent*
 		private:
 			TSet<const UProperty*> PropertiesToSkip;
 
-		} ComponentPropertyWriter(SourceComponent, SavedProperties);
+		};
+
+		FComponentPropertyWriter ComponentPropertyWriter(SourceComponent, SavedProperties);
+
+		// Cache off the length of an array that will come from SerializeTaggedProperties that had no properties saved in to it.
+		auto GetSizeOfEmptyArchive = [](const UActorComponent* DummyComponent) -> int32
+		{
+			TArray<uint8> NoWrittenPropertyReference;
+			FComponentPropertyWriter NullWriter(nullptr, NoWrittenPropertyReference);
+			UClass* ComponentClass = DummyComponent->GetClass();
+			
+			// By serializing the component with itself as its defaults we guarantee that no properties will be written out
+			ComponentClass->SerializeTaggedProperties(NullWriter, (uint8*)DummyComponent, ComponentClass, (uint8*)DummyComponent);
+
+			return NoWrittenPropertyReference.Num();
+		};
+
+		static const int32 SizeOfEmptyArchive = GetSizeOfEmptyArchive(SourceComponent);
+
+		// SerializeTaggedProperties will always put a sentinel NAME_None at the end of the Archive. 
+		// If that is the only thing in the buffer then empty it because we want to know that we haven't stored anything.
+		if (SavedProperties.Num() == SizeOfEmptyArchive)
+		{
+			SavedProperties.Empty();
+		}
 	}
 }
 
