Copying //UE4/Dev-Framework to //UE4/Dev-Main (Source: //UE4/Dev-Framework @ 3038004)

diff --git a/Engine/Source/Runtime/Engine/Private/Components/AudioComponent.cpp b/Engine/Source/Runtime/Engine/Private/Components/AudioComponent.cpp
--- a/Engine/Source/Runtime/Engine/Private/Components/AudioComponent.cpp
+++ b/Engine/Source/Runtime/Engine/Private/Components/AudioComponent.cpp
@@ -6,10 +6,15 @@
 #include "Sound/SoundCue.h"
 #include "SubtitleManager.h"
 #include "Audio.h"
+#include "AudioThread.h"
+#include "TaskGraphInterfaces.h"
 
 /*-----------------------------------------------------------------------------
 UAudioComponent implementation.
 -----------------------------------------------------------------------------*/
+uint64 UAudioComponent::AudioComponentIDCounter = 0;
+TMap<uint64, UAudioComponent*> UAudioComponent::AudioIDToComponentMap;
+
 UAudioComponent::UAudioComponent(const FObjectInitializer& ObjectInitializer)
 	: Super(ObjectInitializer)
 {
@@ -36,6 +41,24 @@ UAudioComponent::UAudioComponent(const FObjectInitializer& ObjectInitializer)
 	LowPassFilterFrequency = MAX_FILTER_FREQUENCY;
 	OcclusionCheckInterval = 0.1f;
 	ActiveCount = 0;
+
+	AudioComponentID = ++AudioComponentIDCounter;
+
+	// TODO: Consider only putting played/active components in to the map
+	AudioIDToComponentMap.Add(AudioComponentID, this);
+}
+
+UAudioComponent* UAudioComponent::GetAudioComponentFromID(uint64 AudioComponentID)
+{
+	check(IsInGameThread());
+	return AudioIDToComponentMap.FindRef(AudioComponentID);
+}
+
+void UAudioComponent::BeginDestroy()
+{
+	Super::BeginDestroy();
+
+	AudioIDToComponentMap.Remove(AudioComponentID);
 }
 
 FString UAudioComponent::GetDetailedInfoInternal( void ) const
@@ -131,14 +154,22 @@ void UAudioComponent::OnUpdateTransform(EUpdateTransformFlags UpdateTransformFla
 
 	if (bIsActive && !bPreviewComponent)
 	{
-		if (FAudioDevice * AudioDevice = GetAudioDevice())
+		if (FAudioDevice* AudioDevice = GetAudioDevice())
 		{
-			FActiveSound* ActiveSound = AudioDevice->FindActiveSound(this);
-			if (ActiveSound)
+			DECLARE_CYCLE_STAT(TEXT("FAudioThreadTask.UpdateTransform"), STAT_AudioUpdateTransform, STATGROUP_AudioThreadCommands);
+
+			const uint64 MyAudioComponentID = AudioComponentID;
+			const FTransform& MyTransform = ComponentToWorld;
+
+			FAudioThread::RunCommandOnAudioThread([AudioDevice, MyAudioComponentID, MyTransform]()
 			{
-				FScopeCycleCounterUObject ComponentScope(ActiveSound->Sound);
-				ActiveSound->Transform = ComponentToWorld;
-			}
+				FActiveSound* ActiveSound = AudioDevice->FindActiveSound(MyAudioComponentID);
+				if (ActiveSound)
+				{
+					ActiveSound->Transform = MyTransform;
+
+				}
+			}, GET_STATID(STAT_AudioUpdateTransform));
 		}
 	}
 };
@@ -174,13 +205,13 @@ void UAudioComponent::PlayInternal(const float StartTime, const float FadeInDura
 			float FocusFactor = 0.0f;
 			FVector Location = ComponentToWorld.GetLocation();
 
-			AudioDevice->GetMaxDistanceAndFocusFactor(Sound, World, Location, AttenuationSettingsToApply, &MaxDistance, &FocusFactor);
+			AudioDevice->GetMaxDistanceAndFocusFactor(Sound, World, Location, AttenuationSettingsToApply, MaxDistance, FocusFactor);
 
 			FActiveSound NewActiveSound;
 			NewActiveSound.SetAudioComponent(this);
-			NewActiveSound.World = GetWorld();
-			NewActiveSound.Sound = Sound;
-			NewActiveSound.SoundClassOverride = SoundClassOverride;
+			NewActiveSound.SetWorld(GetWorld());
+			NewActiveSound.SetSound(Sound);
+			NewActiveSound.SetSoundClass(SoundClassOverride);
 			NewActiveSound.ConcurrencySettings = ConcurrencySettings;
 
 			NewActiveSound.VolumeMultiplier = (VolumeModulationMax + ((VolumeModulationMin - VolumeModulationMax) * FMath::SRand())) * VolumeMultiplier;
@@ -221,7 +252,7 @@ void UAudioComponent::PlayInternal(const float StartTime, const float FadeInDura
 			if (NewActiveSound.bHasAttenuationSettings)
 			{
 				NewActiveSound.AttenuationSettings = *AttenuationSettingsToApply;
-				NewActiveSound.FocusPriorityScale = AttenuationSettingsToApply->GetFocusPriorityScale(AudioDevice->GlobalFocusSettings, FocusFactor);
+				NewActiveSound.FocusPriorityScale = AttenuationSettingsToApply->GetFocusPriorityScale(AudioDevice->GetGlobalFocusSettings(), FocusFactor);
 			}
 
 			NewActiveSound.MaxDistance = MaxDistance;
@@ -242,7 +273,6 @@ void UAudioComponent::PlayInternal(const float StartTime, const float FadeInDura
 			// Bump ActiveCount... this is used to determine if an audio component is still active after "finishing"
 			++ActiveCount;
 
-			// TODO - Audio Threading. This call would be a task call to dispatch to the audio thread
 			AudioDevice->AddNewActiveSound(NewActiveSound);
 			bIsActive = true;
 		}
@@ -279,16 +309,21 @@ void UAudioComponent::FadeOut( float FadeOutDuration, float FadeVolumeLevel )
 	{
 		if (FadeOutDuration > 0.0f)
 		{
-			// TODO - Audio Threading. This call would be a task
 			if (FAudioDevice* AudioDevice = GetAudioDevice())
 			{
-				FActiveSound* ActiveSound = AudioDevice->FindActiveSound(this);
-				if (ActiveSound)
+				DECLARE_CYCLE_STAT(TEXT("FAudioThreadTask.FadeOut"), STAT_AudioFadeOut, STATGROUP_AudioThreadCommands);
+
+				const uint64 MyAudioComponentID = AudioComponentID;
+				FAudioThread::RunCommandOnAudioThread([AudioDevice, MyAudioComponentID, FadeOutDuration, FadeVolumeLevel]()
 				{
-					ActiveSound->TargetAdjustVolumeMultiplier = FadeVolumeLevel;
-					ActiveSound->TargetAdjustVolumeStopTime = ActiveSound->PlaybackTime + FadeOutDuration;
-					ActiveSound->bFadingOut = true;
-				}
+					FActiveSound* ActiveSound = AudioDevice->FindActiveSound(MyAudioComponentID);
+					if (ActiveSound)
+					{
+						ActiveSound->TargetAdjustVolumeMultiplier = FadeVolumeLevel;
+						ActiveSound->TargetAdjustVolumeStopTime = ActiveSound->PlaybackTime + FadeOutDuration;
+						ActiveSound->bFadingOut = true;
+					}
+				}, GET_STATID(STAT_AudioFadeOut));
 			}
 		}
 		else
@@ -302,25 +337,30 @@ void UAudioComponent::AdjustVolume( float AdjustVolumeDuration, float AdjustVolu
 {
 	if (bIsActive)
 	{
-		// TODO - Audio Threading. This call would be a task
 		if (FAudioDevice* AudioDevice = GetAudioDevice())
 		{
-			FActiveSound* ActiveSound = AudioDevice->FindActiveSound(this);
-			if (ActiveSound)
-			{
-				ActiveSound->bFadingOut = false;
-				ActiveSound->TargetAdjustVolumeMultiplier = AdjustVolumeLevel;
+			DECLARE_CYCLE_STAT(TEXT("FAudioThreadTask.AdjustVolume"), STAT_AudioAdjustVolume, STATGROUP_AudioThreadCommands);
 
-				if( AdjustVolumeDuration > 0.0f )
-				{
-					ActiveSound->TargetAdjustVolumeStopTime = ActiveSound->PlaybackTime + AdjustVolumeDuration;
-				}
-				else
+			const uint64 MyAudioComponentID = AudioComponentID;
+			FAudioThread::RunCommandOnAudioThread([AudioDevice, MyAudioComponentID, AdjustVolumeDuration, AdjustVolumeLevel]()
+			{
+				FActiveSound* ActiveSound = AudioDevice->FindActiveSound(MyAudioComponentID);
+				if (ActiveSound)
 				{
-					ActiveSound->CurrentAdjustVolumeMultiplier = AdjustVolumeLevel;
-					ActiveSound->TargetAdjustVolumeStopTime = -1.0f;
+					ActiveSound->bFadingOut = false;
+					ActiveSound->TargetAdjustVolumeMultiplier = AdjustVolumeLevel;
+
+					if (AdjustVolumeDuration > 0.0f)
+					{
+						ActiveSound->TargetAdjustVolumeStopTime = ActiveSound->PlaybackTime + AdjustVolumeDuration;
+					}
+					else
+					{
+						ActiveSound->CurrentAdjustVolumeMultiplier = AdjustVolumeLevel;
+						ActiveSound->TargetAdjustVolumeStopTime = -1.0f;
+					}
 				}
-			}
+			}, GET_STATID(STAT_AudioAdjustVolume));
 		}
 	}
 }
@@ -334,14 +374,34 @@ void UAudioComponent::Stop()
 
 		UE_LOG(LogAudio, Verbose, TEXT( "%g: Stopping AudioComponent : '%s' with Sound: '%s'" ), GetWorld() ? GetWorld()->GetAudioTimeSeconds() : 0.0f, *GetFullName(), Sound ? *Sound->GetName() : TEXT( "nullptr" ) );
 
-		// TODO - Audio Threading. This call would be a task
 		if (FAudioDevice* AudioDevice = GetAudioDevice())
 		{
-			AudioDevice->StopActiveSound(this);
+			DECLARE_CYCLE_STAT(TEXT("FAudioThreadTask.StopActiveSound"), STAT_AudioStopActiveSound, STATGROUP_AudioThreadCommands);
+
+			const uint64 MyAudioComponentID = AudioComponentID;
+			FAudioThread::RunCommandOnAudioThread([AudioDevice, MyAudioComponentID]()
+			{
+				AudioDevice->StopActiveSound(MyAudioComponentID);
+			}, GET_STATID(STAT_AudioStopActiveSound));
 		}
 	}
 }
 
+void UAudioComponent::PlaybackCompleted(uint64 AudioComponentID, bool bFailedToStart)
+{
+	check(IsInAudioThread());
+
+	DECLARE_CYCLE_STAT(TEXT("FGameThreadAudioTask.PlaybackCompleted"), STAT_AudioPlaybackCompleted, STATGROUP_TaskGraphTasks);
+
+	FAudioThread::RunCommandOnGameThread([AudioComponentID, bFailedToStart]()
+	{
+		if (UAudioComponent* AudioComponent = GetAudioComponentFromID(AudioComponentID))
+		{
+			AudioComponent->PlaybackCompleted(bFailedToStart);
+		}
+	}, GET_STATID(STAT_AudioPlaybackCompleted));
+}
+
 void UAudioComponent::PlaybackCompleted(bool bFailedToStart)
 {
 	check(ActiveCount > 0);
@@ -513,14 +573,19 @@ void UAudioComponent::SetFloatParameter( FName InName, float InFloat )
 		// If we're active we need to push this value to the ActiveSound
 		if (bIsActive)
 		{
-			// TODO - Audio Threading. This call would be a task
 			if (FAudioDevice* AudioDevice = GetAudioDevice())
 			{
-				FActiveSound* ActiveSound = AudioDevice->FindActiveSound(this);
-				if (ActiveSound)
+				DECLARE_CYCLE_STAT(TEXT("FAudioThreadTask.SetFloatParameter"), STAT_AudioSetFloatParameter, STATGROUP_AudioThreadCommands);
+
+				const uint64 MyAudioComponentID = AudioComponentID;
+				FAudioThread::RunCommandOnAudioThread([AudioDevice, MyAudioComponentID, InName, InFloat]()
 				{
-					ActiveSound->SetFloatParameter(InName, InFloat);
-				}
+					FActiveSound* ActiveSound = AudioDevice->FindActiveSound(MyAudioComponentID);
+					if (ActiveSound)
+					{
+						ActiveSound->SetFloatParameter(InName, InFloat);
+					}
+				}, GET_STATID(STAT_AudioSetFloatParameter));
 			}
 		}
 	}
@@ -554,14 +619,19 @@ void UAudioComponent::SetWaveParameter( FName InName, USoundWave* InWave )
 		// If we're active we need to push this value to the ActiveSound
 		if (bIsActive)
 		{
-			// TODO - Audio Threading. This call would be a task
 			if (FAudioDevice* AudioDevice = GetAudioDevice())
 			{
-				FActiveSound* ActiveSound = AudioDevice->FindActiveSound(this);
-				if (ActiveSound)
+				DECLARE_CYCLE_STAT(TEXT("FAudioThreadTask.SetWaveParameter"), STAT_AudioSetWaveParameter, STATGROUP_AudioThreadCommands);
+
+				const uint64 MyAudioComponentID = AudioComponentID;
+				FAudioThread::RunCommandOnAudioThread([AudioDevice, MyAudioComponentID, InName, InWave]()
 				{
-					ActiveSound->SetWaveParameter(InName, InWave);
-				}
+					FActiveSound* ActiveSound = AudioDevice->FindActiveSound(MyAudioComponentID);
+					if (ActiveSound)
+					{
+						ActiveSound->SetWaveParameter(InName, InWave);
+					}
+				}, GET_STATID(STAT_AudioSetWaveParameter));
 			}
 		}
 	}
@@ -595,14 +665,19 @@ void UAudioComponent::SetBoolParameter( FName InName, bool InBool )
 		// If we're active we need to push this value to the ActiveSound
 		if (bIsActive)
 		{
-			// TODO - Audio Threading. This call would be a task
 			if (FAudioDevice* AudioDevice = GetAudioDevice())
 			{
-				FActiveSound* ActiveSound = AudioDevice->FindActiveSound(this);
-				if (ActiveSound)
+				DECLARE_CYCLE_STAT(TEXT("FAudioThreadTask.SetBoolParameter"), STAT_AudioSetBoolParameter, STATGROUP_AudioThreadCommands);
+
+				const uint64 MyAudioComponentID = AudioComponentID;
+				FAudioThread::RunCommandOnAudioThread([AudioDevice, MyAudioComponentID, InName, InBool]()
 				{
-					ActiveSound->SetBoolParameter(InName, InBool);
-				}
+					FActiveSound* ActiveSound = AudioDevice->FindActiveSound(MyAudioComponentID);
+					if (ActiveSound)
+					{
+						ActiveSound->SetBoolParameter(InName, InBool);
+					}
+				}, GET_STATID(STAT_AudioSetBoolParameter));
 			}
 		}
 	}
@@ -637,14 +712,19 @@ void UAudioComponent::SetIntParameter( FName InName, int32 InInt )
 		// If we're active we need to push this value to the ActiveSound
 		if (bIsActive)
 		{
-			// TODO - Audio Threading. This call would be a task
 			if (FAudioDevice* AudioDevice = GetAudioDevice())
 			{
-				FActiveSound* ActiveSound = AudioDevice->FindActiveSound(this);
-				if (ActiveSound)
+				DECLARE_CYCLE_STAT(TEXT("FAudioThreadTask.SetIntParameter"), STAT_AudioSetIntParameter, STATGROUP_AudioThreadCommands);
+
+				const uint64 MyAudioComponentID = AudioComponentID;
+				FAudioThread::RunCommandOnAudioThread([AudioDevice, MyAudioComponentID, InName, InInt]()
 				{
-					ActiveSound->SetIntParameter(InName, InInt);
-				}
+					FActiveSound* ActiveSound = AudioDevice->FindActiveSound(MyAudioComponentID);
+					if (ActiveSound)
+					{
+						ActiveSound->SetIntParameter(InName, InInt);
+					}
+				}, GET_STATID(STAT_AudioSetIntParameter));
 			}
 		}
 	}
@@ -676,14 +756,19 @@ void UAudioComponent::SetSoundParameter(const FAudioComponentParam& Param)
 
 		if (bIsActive)
 		{
-			// TODO - Audio Threading. This call would be a task
 			if (FAudioDevice* AudioDevice = GetAudioDevice())
 			{
-				FActiveSound* ActiveSound = AudioDevice->FindActiveSound(this);
-				if (ActiveSound)
+				DECLARE_CYCLE_STAT(TEXT("FAudioThreadTask.SetSoundParameter"), STAT_AudioSetSoundParameter, STATGROUP_AudioThreadCommands);
+
+				const uint64 MyAudioComponentID = AudioComponentID;
+				FAudioThread::RunCommandOnAudioThread([AudioDevice, MyAudioComponentID, Param]()
 				{
-					ActiveSound->SetSoundParameter(Param);
-				}
+					FActiveSound* ActiveSound = AudioDevice->FindActiveSound(MyAudioComponentID);
+					if (ActiveSound)
+					{
+						ActiveSound->SetSoundParameter(Param);
+					}
+				}, GET_STATID(STAT_AudioSetSoundParameter));
 			}
 		}
 	}
@@ -694,14 +779,19 @@ void UAudioComponent::SetVolumeMultiplier(float NewVolumeMultiplier)
 	VolumeMultiplier = NewVolumeMultiplier;
 	VolumeModulationMin = VolumeModulationMax = 1.f;
 
-	// TODO - Audio Threading. This call would be a task
 	if (FAudioDevice* AudioDevice = GetAudioDevice())
 	{
-		FActiveSound* ActiveSound = AudioDevice->FindActiveSound(this);
-		if (ActiveSound)
+		DECLARE_CYCLE_STAT(TEXT("FAudioThreadTask.SetVolumeMultiplier"), STAT_AudioSetVolumeMultiplier, STATGROUP_AudioThreadCommands);
+
+		const uint64 MyAudioComponentID = AudioComponentID;
+		FAudioThread::RunCommandOnAudioThread([AudioDevice, MyAudioComponentID, NewVolumeMultiplier]()
 		{
-			ActiveSound->VolumeMultiplier = NewVolumeMultiplier;
-		}
+			FActiveSound* ActiveSound = AudioDevice->FindActiveSound(MyAudioComponentID);
+			if (ActiveSound)
+			{
+				ActiveSound->VolumeMultiplier = NewVolumeMultiplier;
+			}
+		}, GET_STATID(STAT_AudioSetVolumeMultiplier));
 	}
 }
 
@@ -710,14 +800,19 @@ void UAudioComponent::SetPitchMultiplier(float NewPitchMultiplier)
 	PitchMultiplier = NewPitchMultiplier;
 	PitchModulationMin = PitchModulationMax = 1.f;
 
-	// TODO - Audio Threading. This call would be a task
 	if (FAudioDevice* AudioDevice = GetAudioDevice())
 	{
-		FActiveSound* ActiveSound = AudioDevice->FindActiveSound(this);
-		if (ActiveSound)
+		DECLARE_CYCLE_STAT(TEXT("FAudioThreadTask.SetPitchMultiplier"), STAT_AudioSetPitchMultiplier, STATGROUP_AudioThreadCommands);
+
+		const uint64 MyAudioComponentID = AudioComponentID;
+		FAudioThread::RunCommandOnAudioThread([AudioDevice, MyAudioComponentID, NewPitchMultiplier]()
 		{
-			ActiveSound->PitchMultiplier = NewPitchMultiplier;
-		}
+			FActiveSound* ActiveSound = AudioDevice->FindActiveSound(MyAudioComponentID);
+			if (ActiveSound)
+			{
+				ActiveSound->VolumeMultiplier = NewPitchMultiplier;
+			}
+		}, GET_STATID(STAT_AudioSetPitchMultiplier));
 	}
 }
 
@@ -725,14 +820,19 @@ void UAudioComponent::SetUISound(bool bInIsUISound)
 {
 	bIsUISound = bInIsUISound;
 
-	// TODO - Audio Threading. This call would be a task
 	if (FAudioDevice* AudioDevice = GetAudioDevice())
 	{
-		FActiveSound* ActiveSound = AudioDevice->FindActiveSound(this);
-		if (ActiveSound)
+		DECLARE_CYCLE_STAT(TEXT("FAudioThreadTask.SetIsUISound"), STAT_AudioSetIsUISound, STATGROUP_AudioThreadCommands);
+
+		const uint64 MyAudioComponentID = AudioComponentID;
+		FAudioThread::RunCommandOnAudioThread([AudioDevice, MyAudioComponentID, bInIsUISound]()
 		{
-			ActiveSound->bIsUISound = bIsUISound;
-		}
+			FActiveSound* ActiveSound = AudioDevice->FindActiveSound(MyAudioComponentID);
+			if (ActiveSound)
+			{
+				ActiveSound->bIsUISound = bInIsUISound;
+			}
+		}, GET_STATID(STAT_AudioSetIsUISound));
 	}
 }
 
@@ -741,13 +841,18 @@ void UAudioComponent::AdjustAttenuation(const FAttenuationSettings& InAttenuatio
 	bOverrideAttenuation = true;
 	AttenuationOverrides = InAttenuationSettings;
 
-	// TODO - Audio Threading. This call would be a task
 	if (FAudioDevice* AudioDevice = GetAudioDevice())
 	{
-		FActiveSound* ActiveSound = AudioDevice->FindActiveSound(this);
-		if (ActiveSound)
+		DECLARE_CYCLE_STAT(TEXT("FAudioThreadTask.AdjustAttenuation"), STAT_AudioAdjustAttenuation, STATGROUP_AudioThreadCommands);
+
+		const uint64 MyAudioComponentID = AudioComponentID;
+		FAudioThread::RunCommandOnAudioThread([AudioDevice, MyAudioComponentID, InAttenuationSettings]()
 		{
-			ActiveSound->AttenuationSettings = AttenuationOverrides;
-		}
+			FActiveSound* ActiveSound = AudioDevice->FindActiveSound(MyAudioComponentID);
+			if (ActiveSound)
+			{
+				ActiveSound->AttenuationSettings = InAttenuationSettings;
+			}
+		}, GET_STATID(STAT_AudioAdjustAttenuation));
 	}
 }
diff --git a/Engine/Source/Runtime/Engine/Private/Components/CharacterMovementComponent.cpp b/Engine/Source/Runtime/Engine/Private/Components/CharacterMovementComponent.cpp
--- a/Engine/Source/Runtime/Engine/Private/Components/CharacterMovementComponent.cpp
+++ b/Engine/Source/Runtime/Engine/Private/Components/CharacterMovementComponent.cpp
@@ -2224,12 +2224,24 @@ void UCharacterMovementComponent::Crouch(bool bClientSimulation)
 	bForceNextFloorCheck = true;
 
 	// OnStartCrouch takes the change from the Default size, not the current one (though they are usually the same).
+	const float MeshAdjust = ScaledHalfHeightAdjust;
 	ACharacter* DefaultCharacter = CharacterOwner->GetClass()->GetDefaultObject<ACharacter>();
 	HalfHeightAdjust = (DefaultCharacter->GetCapsuleComponent()->GetUnscaledCapsuleHalfHeight() - ClampedCrouchedHalfHeight);
 	ScaledHalfHeightAdjust = HalfHeightAdjust * ComponentScale;
 
 	AdjustProxyCapsuleSize();
 	CharacterOwner->OnStartCrouch( HalfHeightAdjust, ScaledHalfHeightAdjust );
+
+	// Don't smooth this change in mesh position
+	if (bClientSimulation && CharacterOwner->Role == ROLE_SimulatedProxy)
+	{
+		FNetworkPredictionData_Client_Character* ClientData = GetPredictionData_Client_Character();
+		if (ClientData && ClientData->MeshTranslationOffset.Z != 0.f)
+		{
+			ClientData->MeshTranslationOffset -= FVector(0.f, 0.f, MeshAdjust);
+			ClientData->OriginalMeshTranslationOffset = ClientData->MeshTranslationOffset;
+		}
+	}
 }
 
 
@@ -2362,8 +2374,20 @@ void UCharacterMovementComponent::UnCrouch(bool bClientSimulation)
 	bUpdateOverlaps = true;
 	CharacterOwner->GetCapsuleComponent()->SetCapsuleSize(DefaultCharacter->GetCapsuleComponent()->GetUnscaledCapsuleRadius(), DefaultCharacter->GetCapsuleComponent()->GetUnscaledCapsuleHalfHeight(), bUpdateOverlaps);
 
+	const float MeshAdjust = ScaledHalfHeightAdjust;
 	AdjustProxyCapsuleSize();
 	CharacterOwner->OnEndCrouch( HalfHeightAdjust, ScaledHalfHeightAdjust );
+
+	// Don't smooth this change in mesh position
+	if (bClientSimulation && CharacterOwner->Role == ROLE_SimulatedProxy)
+	{
+		FNetworkPredictionData_Client_Character* ClientData = GetPredictionData_Client_Character();
+		if (ClientData && ClientData->MeshTranslationOffset.Z != 0.f)
+		{
+			ClientData->MeshTranslationOffset += FVector(0.f, 0.f, MeshAdjust);
+			ClientData->OriginalMeshTranslationOffset = ClientData->MeshTranslationOffset;
+		}
+	}
 }
 
 void UCharacterMovementComponent::StartNewPhysics(float deltaTime, int32 Iterations)
diff --git a/Engine/Source/Runtime/Engine/Private/Components/SceneComponent.cpp b/Engine/Source/Runtime/Engine/Private/Components/SceneComponent.cpp
--- a/Engine/Source/Runtime/Engine/Private/Components/SceneComponent.cpp
+++ b/Engine/Source/Runtime/Engine/Private/Components/SceneComponent.cpp
@@ -503,6 +503,7 @@ void USceneComponent::PropagateTransformUpdate(bool bTransformChanged, EUpdateTr
 				//QUICK_SCOPE_CYCLE_COUNTER(STAT_USceneComponent_PropagateTransformUpdate_OnUpdateTransform);
 				OnUpdateTransform(UpdateTransformFlags, Teleport);
 			}
+			TransformUpdated.Broadcast(this, UpdateTransformFlags, Teleport);
 
 			// Flag render transform as dirty
 			MarkRenderTransformDirty();
diff --git a/Engine/Source/Runtime/Engine/Private/Components/SkeletalMeshComponent.cpp b/Engine/Source/Runtime/Engine/Private/Components/SkeletalMeshComponent.cpp
--- a/Engine/Source/Runtime/Engine/Private/Components/SkeletalMeshComponent.cpp
+++ b/Engine/Source/Runtime/Engine/Private/Components/SkeletalMeshComponent.cpp
@@ -608,6 +608,11 @@ void USkeletalMeshComponent::TickAnimation(float DeltaTime, bool bNeedsValidRoot
 			// Tick the animation
 			AnimScriptInstance->UpdateAnimation(DeltaTime * GlobalAnimRateScale, bNeedsValidRootMotion);
 		}
+
+		for(UAnimInstance* SubInstance : SubInstances)
+		{
+			SubInstance->UpdateAnimation(DeltaTime * GlobalAnimRateScale, false);
+		}
 	}
 }
 
@@ -924,9 +929,9 @@ void USkeletalMeshComponent::RecalcRequiredBones(int32 LODIndex)
 	if(PhysicsAsset)
 	{
 		TArray<FBoneIndexType> PhysAssetBones;
-		for(int32 i=0; i<PhysicsAsset->BodySetup.Num(); i++ )
+		for(int32 i=0; i<PhysicsAsset->SkeletalBodySetups.Num(); i++ )
 		{
-			int32 PhysBoneIndex = SkeletalMesh->RefSkeleton.FindBoneIndex( PhysicsAsset->BodySetup[i]->BoneName );
+			int32 PhysBoneIndex = SkeletalMesh->RefSkeleton.FindBoneIndex( PhysicsAsset->SkeletalBodySetups[i]->BoneName );
 			if(PhysBoneIndex != INDEX_NONE)
 			{
 				PhysAssetBones.Add(PhysBoneIndex);
@@ -1382,7 +1387,7 @@ void USkeletalMeshComponent::PostAnimEvaluation(FAnimationEvaluationContext& Eva
 		FillSpaceBases(SkeletalMesh, LocalAtoms, GetEditableSpaceBases());
 
 		// interpolate curve
-		EvaluationContext.Curve.BlendWith(CachedCurve, Alpha);
+		EvaluationContext.Curve.LerpTo(CachedCurve, Alpha);
 	}
 
 	if(AnimScriptInstance)
@@ -1498,9 +1503,16 @@ void USkeletalMeshComponent::SetSkeletalMesh(USkeletalMesh* InSkelMesh, bool bRe
 		ValidateAnimation();
 #endif
 	
-		if (GetPhysicsAsset() != OldPhysAsset && IsPhysicsStateCreated())
+		if(IsPhysicsStateCreated())
 		{
-			RecreatePhysicsState();
+			if(GetPhysicsAsset() == OldPhysAsset)
+			{
+				UpdateBoneBodyMapping();
+			}
+			else
+			{
+				RecreatePhysicsState();
+			}
 		}
 
 		UpdateHasValidBodies();
@@ -2017,10 +2029,10 @@ float USkeletalMeshComponent::CalculateMass(FName BoneName)
 	}
 	else	//We want to calculate mass before we've initialized body instances - in this case use physics asset setup
 	{
-		TArray<class UBodySetup*> * BodySetups = NULL;
-		if (UPhysicsAsset * PhysicsAsset = GetPhysicsAsset())
+		TArray<USkeletalBodySetup*>* BodySetups = nullptr;
+		if (UPhysicsAsset* PhysicsAsset = GetPhysicsAsset())
 		{
-			BodySetups = &PhysicsAsset->BodySetup;
+			BodySetups = &PhysicsAsset->SkeletalBodySetups;
 		}
 
 		if (BodySetups)
@@ -2410,10 +2422,10 @@ bool USkeletalMeshComponent::DoCustomNavigableGeometryExport(FNavigableGeometryE
 	UPhysicsAsset* PhysicsAsset = GetPhysicsAsset();
 	if (PhysicsAsset && ComponentToWorld.GetScale3D().IsUniform())
 	{
-		const int32 MaxBodies = PhysicsAsset->BodySetup.Num();
+		const int32 MaxBodies = PhysicsAsset->SkeletalBodySetups.Num();
 		for (int32 Idx = 0; Idx < MaxBodies; Idx++)
 		{
-			UBodySetup* const BS = PhysicsAsset->BodySetup[Idx];
+			UBodySetup* const BS = PhysicsAsset->SkeletalBodySetups[Idx];
 			int32 const BoneIndex = BS ? GetBoneIndex(BS->BoneName) : INDEX_NONE;
 
 			if (BoneIndex != INDEX_NONE)
diff --git a/Engine/Source/Runtime/Engine/Private/Components/SkinnedMeshComponent.cpp b/Engine/Source/Runtime/Engine/Private/Components/SkinnedMeshComponent.cpp
--- a/Engine/Source/Runtime/Engine/Private/Components/SkinnedMeshComponent.cpp
+++ b/Engine/Source/Runtime/Engine/Private/Components/SkinnedMeshComponent.cpp
@@ -703,8 +703,8 @@ FBoxSphereBounds USkinnedMeshComponent::CalcMeshBound(const FVector& RootOffset,
 	const bool bDetailModeAllowsRendering = (DetailMode <= GetCachedScalabilityCVars().DetailMode);
 	const bool bIsVisible = ( bDetailModeAllowsRendering && (ShouldRender() || bCastHiddenShadow));
 
-	const bool bHasPhysBodies = PhysicsAsset && PhysicsAsset->BodySetup.Num();
-	const bool bMasterHasPhysBodies = MasterPhysicsAsset && MasterPhysicsAsset->BodySetup.Num();
+	const bool bHasPhysBodies = PhysicsAsset && PhysicsAsset->SkeletalBodySetups.Num();
+	const bool bMasterHasPhysBodies = MasterPhysicsAsset && MasterPhysicsAsset->SkeletalBodySetups.Num();
 
 	// if not visible, or we were told to use fixed bounds, use skelmesh bounds
 	if ( (!bIsVisible || bComponentUseFixedSkelBounds) && SkeletalMesh ) 
diff --git a/Engine/Source/Runtime/Engine/Private/DialogueWave.cpp b/Engine/Source/Runtime/Engine/Private/DialogueWave.cpp
--- a/Engine/Source/Runtime/Engine/Private/DialogueWave.cpp
+++ b/Engine/Source/Runtime/Engine/Private/DialogueWave.cpp
@@ -447,22 +447,16 @@ void UDialogueSoundWaveProxy::Parse(class FAudioDevice* AudioDevice, const UPTRI
 	// Add in the subtitle if they exist
 	if (ActiveSound.bHandleSubtitles && Subtitles.Num() > 0)
 	{
-		// TODO - Audio Threading. This would need to be a call back to the main thread.
-		UAudioComponent* AudioComponent = ActiveSound.GetAudioComponent();
-		if (AudioComponent && AudioComponent->OnQueueSubtitles.IsBound())
-		{
-			// intercept the subtitles if the delegate is set
-			AudioComponent->OnQueueSubtitles.ExecuteIfBound(Subtitles, GetDuration());
-		}
-		else
-		{
-			// otherwise, pass them on to the subtitle manager for display
-			// Subtitles are hashed based on the associated sound (wave instance).
-			if (ActiveSound.World.IsValid())
-			{
-				FSubtitleManager::GetSubtitleManager()->QueueSubtitles((PTRINT)WaveInstance, ActiveSound.SubtitlePriority, false, false, GetDuration(), Subtitles, 0.0f);
-			}
-		}
+		FQueueSubtitleParams QueueSubtitleParams(Subtitles);
+		QueueSubtitleParams.AudioComponentID = ActiveSound.GetAudioComponentID();
+		QueueSubtitleParams.WorldPtr = ActiveSound.GetWeakWorld();
+		QueueSubtitleParams.WaveInstance = (PTRINT)WaveInstance;
+		QueueSubtitleParams.SubtitlePriority = ActiveSound.SubtitlePriority;
+		QueueSubtitleParams.Duration = GetDuration();
+		QueueSubtitleParams.bManualWordWrap = false;
+		QueueSubtitleParams.bSingleLine = false;
+
+		FSubtitleManager::QueueSubtitles(QueueSubtitleParams);
 	}
 }
 
diff --git a/Engine/Source/Runtime/Engine/Private/GameMode.cpp b/Engine/Source/Runtime/Engine/Private/GameMode.cpp
--- a/Engine/Source/Runtime/Engine/Private/GameMode.cpp
+++ b/Engine/Source/Runtime/Engine/Private/GameMode.cpp
@@ -390,11 +390,18 @@ AActor* AGameMode::FindPlayerStart_Implementation( AController* Player, const FS
 	// always pick StartSpot at start of match
 	if ( ShouldSpawnAtStartSpot(Player) )
 	{
-		return Player->StartSpot.Get();
+		if (AActor* PlayerStartSpot = Player->StartSpot.Get())
+		{
+			return PlayerStartSpot;
+		}
+		else
+		{
+			UE_LOG(LogGameMode, Error, TEXT("Error - ShouldSpawnAtStartSpot returned true but the Player StartSpot was null."));
+		}
 	}
 
 	AActor* BestStart = ChoosePlayerStart(Player);
-	if (BestStart == NULL)
+	if (BestStart == nullptr)
 	{
 		// no player start found
 		UE_LOG(LogGameMode, Log, TEXT("Warning - PATHS NOT DEFINED or NO PLAYERSTART with positive rating"));
diff --git a/Engine/Source/Runtime/Engine/Private/GameViewportClient.cpp b/Engine/Source/Runtime/Engine/Private/GameViewportClient.cpp
--- a/Engine/Source/Runtime/Engine/Private/GameViewportClient.cpp
+++ b/Engine/Source/Runtime/Engine/Private/GameViewportClient.cpp
@@ -287,13 +287,20 @@ void UGameViewportClient::Init(struct FWorldContext& WorldContext, UGameInstance
 		FAudioDeviceManager* AudioDeviceManager = GEngine->GetAudioDeviceManager();
 		if (AudioDeviceManager)
 		{
-			FAudioDevice* NewAudioDevice = AudioDeviceManager->CreateAudioDevice(AudioDeviceHandle, bCreateNewAudioDevice);
-			if (NewAudioDevice)
+			FAudioDeviceManager::FCreateAudioDeviceResults NewDeviceResults;
+			if (AudioDeviceManager->CreateAudioDevice(bCreateNewAudioDevice, NewDeviceResults))
 			{
+				AudioDeviceHandle = NewDeviceResults.Handle;
+
+				if (NewDeviceResults.bNewDevice)
+				{
+					NewDeviceResults.AudioDevice->UpdateSoundShowFlags(0, GetSoundShowFlags());
+				}
+
 				// Set the base mix of the new device based on the world settings of the world
 				if (World)
 				{
-					NewAudioDevice->SetDefaultBaseSoundMix(World->GetWorldSettings()->DefaultBaseSoundMix);
+					NewDeviceResults.AudioDevice->SetDefaultBaseSoundMix(World->GetWorldSettings()->DefaultBaseSoundMix);
 
 					// Set the world's audio device handle to use so that sounds which play in that world will use the correct audio device
 					World->SetAudioDeviceHandle(AudioDeviceHandle);
@@ -744,9 +751,9 @@ bool UGameViewportClient::ShouldForceFullscreenViewport() const
 	{
 		bResult = true;
 	}
-	else if ( GetWorld() )
+	else if ( UWorld* MyWorld = GetWorld() )
 	{
-		if ( GetWorld()->bIsDefaultLevel )
+		if ( MyWorld->bIsDefaultLevel )
 		{
 			bResult = true;
 		}
@@ -849,11 +856,13 @@ void UGameViewportClient::Draw(FViewport* InViewport, FCanvas* SceneCanvas)
 
 	bool bUIDisableWorldRendering = false;
 	FGameViewDrawer GameViewDrawer;
-	
+
+	UWorld* MyWorld = GetWorld();
+
 	// create the view family for rendering the world scene to the viewport's render target
 	FSceneViewFamilyContext ViewFamily(FSceneViewFamily::ConstructionValues( 	
 		InViewport,
-		GetWorld()->Scene,
+		MyWorld->Scene,
 		EngineShowFlags)
 		.SetRealtimeUpdate(true));
 
@@ -912,13 +921,9 @@ void UGameViewportClient::Draw(FViewport* InViewport, FCanvas* SceneCanvas)
 
 	TMap<ULocalPlayer*,FSceneView*> PlayerViewMap;
 
-	FAudioDevice* AudioDevice = GetWorld()->GetAudioDevice();
-
-	bool bReverbSettingsFound = false;
-	FReverbSettings ReverbSettings;
-	class AAudioVolume* AudioVolume = nullptr;
+	FAudioDevice* AudioDevice = MyWorld->GetAudioDevice();
 
-	for (FLocalPlayerIterator Iterator(GEngine, GetWorld()); Iterator; ++Iterator)
+	for (FLocalPlayerIterator Iterator(GEngine, MyWorld); Iterator; ++Iterator)
 	{
 		ULocalPlayer* LocalPlayer = *Iterator;
 		if (LocalPlayer)
@@ -1022,27 +1027,15 @@ void UGameViewportClient::Draw(FViewport* InViewport, FCanvas* SceneCanvas)
 								ListenerTransform.SetTranslation(Location);
 								ListenerTransform.NormalizeRotation();
 
-								bReverbSettingsFound = true;
-
-								FReverbSettings PlayerReverbSettings;
-								FInteriorSettings PlayerInteriorSettings;
-								class AAudioVolume* PlayerAudioVolume = GetWorld()->GetAudioSettings(Location, &PlayerReverbSettings, &PlayerInteriorSettings);
-
-								if (AudioVolume == nullptr || (PlayerAudioVolume != nullptr && PlayerAudioVolume->Priority > AudioVolume->Priority))
-								{
-									AudioVolume = PlayerAudioVolume;
-									ReverbSettings = PlayerReverbSettings;
-								}
-
 								uint32 ViewportIndex = PlayerViewMap.Num() - 1;
-								AudioDevice->SetListener(ViewportIndex, ListenerTransform, (View->bCameraCut ? 0.f : GetWorld()->GetDeltaSeconds()), PlayerAudioVolume, PlayerInteriorSettings);
+								AudioDevice->SetListener(MyWorld, ViewportIndex, ListenerTransform, (View->bCameraCut ? 0.f : MyWorld->GetDeltaSeconds()));
 							}
 						}
 					}
 
 					// Add view information for resource streaming.
 					IStreamingManager::Get().AddViewInformation(View->ViewMatrices.ViewOrigin, View->ViewRect.Width(), View->ViewRect.Width() * View->ViewMatrices.ProjMatrix.M[0][0]);
-					GetWorld()->ViewLocationsRenderedLastFrame.Add(View->ViewMatrices.ViewOrigin);
+					MyWorld->ViewLocationsRenderedLastFrame.Add(View->ViewMatrices.ViewOrigin);
 				}
 			}
 		}
@@ -1050,13 +1043,8 @@ void UGameViewportClient::Draw(FViewport* InViewport, FCanvas* SceneCanvas)
 
 	FinalizeViews(&ViewFamily, PlayerViewMap);
 
-	if (bReverbSettingsFound)
-	{
-		AudioDevice->SetReverbSettings( AudioVolume, ReverbSettings );
-	}
-
 	// Update level streaming.
-	GetWorld()->UpdateLevelStreaming();
+	MyWorld->UpdateLevelStreaming();
 
 	// Find largest rectangle bounded by all rendered views.
 	uint32 MinX=InViewport->GetSizeXY().X, MinY=InViewport->GetSizeXY().Y, MaxX=0, MaxY=0;
@@ -1130,20 +1118,20 @@ void UGameViewportClient::Draw(FViewport* InViewport, FCanvas* SceneCanvas)
 	}
 	
 	// Remove temporary debug lines.
-	if (GetWorld()->LineBatcher != NULL)
+	if (MyWorld->LineBatcher != nullptr)
 	{
-		GetWorld()->LineBatcher->Flush();
+		MyWorld->LineBatcher->Flush();
 	}
 
-	if (GetWorld()->ForegroundLineBatcher != NULL)
+	if (MyWorld->ForegroundLineBatcher != nullptr)
 	{
-		GetWorld()->ForegroundLineBatcher->Flush();
+		MyWorld->ForegroundLineBatcher->Flush();
 	}
 
 	// Draw FX debug information.
-	if (GetWorld()->FXSystem)
+	if (MyWorld->FXSystem)
 	{
-		GetWorld()->FXSystem->DrawDebug(SceneCanvas);
+		MyWorld->FXSystem->DrawDebug(SceneCanvas);
 	}
 
 	// Render the UI.
@@ -1152,7 +1140,7 @@ void UGameViewportClient::Draw(FViewport* InViewport, FCanvas* SceneCanvas)
 
 		// render HUD
 		bool bDisplayedSubtitles = false;
-		for( FConstPlayerControllerIterator Iterator = GetWorld()->GetPlayerControllerIterator(); Iterator; ++Iterator )
+		for( FConstPlayerControllerIterator Iterator = MyWorld->GetPlayerControllerIterator(); Iterator; ++Iterator )
 		{
 			APlayerController* PlayerController = *Iterator;
 			if (PlayerController)
@@ -1214,7 +1202,7 @@ void UGameViewportClient::Draw(FViewport* InViewport, FCanvas* SceneCanvas)
 							const uint32 SizeX = SceneCanvas->GetRenderTarget()->GetSizeXY().X;
 							const uint32 SizeY = SceneCanvas->GetRenderTarget()->GetSizeXY().Y;
 							FIntRect SubtitleRegion(FMath::TruncToInt(SizeX * MinPos.X), FMath::TruncToInt(SizeY * MinPos.Y), FMath::TruncToInt(SizeX * MaxPos.X), FMath::TruncToInt(SizeY * MaxPos.Y));
-							FSubtitleManager::GetSubtitleManager()->DisplaySubtitles( SceneCanvas, SubtitleRegion, GetWorld()->GetAudioTimeSeconds() );
+							FSubtitleManager::GetSubtitleManager()->DisplaySubtitles( SceneCanvas, SubtitleRegion, MyWorld->GetAudioTimeSeconds() );
 							bDisplayedSubtitles = true;
 						}
 					}
@@ -1246,13 +1234,13 @@ void UGameViewportClient::Draw(FViewport* InViewport, FCanvas* SceneCanvas)
 	FVector PlayerCameraLocation = FVector::ZeroVector;
 	FRotator PlayerCameraRotation = FRotator::ZeroRotator;
 	{
-		for( FConstPlayerControllerIterator Iterator = GetWorld()->GetPlayerControllerIterator(); Iterator; ++Iterator )
+		for( FConstPlayerControllerIterator Iterator = MyWorld->GetPlayerControllerIterator(); Iterator; ++Iterator )
 		{
 			(*Iterator)->GetPlayerViewPoint( PlayerCameraLocation, PlayerCameraRotation );
 		}
 	}
 
-	DrawStatsHUD( GetWorld(), InViewport, DebugCanvas, DebugCanvasObject, DebugProperties, PlayerCameraLocation, PlayerCameraRotation );
+	DrawStatsHUD( MyWorld, InViewport, DebugCanvas, DebugCanvasObject, DebugProperties, PlayerCameraLocation, PlayerCameraRotation );
 
 	if (GEngine->IsStereoscopic3D(InViewport))
 	{
diff --git a/Engine/Source/Runtime/Engine/Private/GameplayStatics.cpp b/Engine/Source/Runtime/Engine/Private/GameplayStatics.cpp
--- a/Engine/Source/Runtime/Engine/Private/GameplayStatics.cpp
+++ b/Engine/Source/Runtime/Engine/Private/GameplayStatics.cpp
@@ -900,7 +900,7 @@ void UGameplayStatics::SetGlobalPitchModulation(UObject* WorldContextObject, flo
 
 	if (FAudioDevice* AudioDevice = ThisWorld->GetAudioDevice())
 	{
-		AudioDevice->GlobalPitchScale.Set(PitchModulation, TimeSec);
+		AudioDevice->SetGlobalPitchModulation(PitchModulation, TimeSec);
 	}
 }
 
@@ -919,14 +919,17 @@ void UGameplayStatics::SetGlobalListenerFocusParameters(UObject* WorldContextObj
 
 	if (FAudioDevice* AudioDevice = ThisWorld->GetAudioDevice())
 	{
-		AudioDevice->GlobalFocusSettings.FocusAzimuthScale = FMath::Max(FocusAzimuthScale, 0.0f);
-		AudioDevice->GlobalFocusSettings.NonFocusAzimuthScale = FMath::Max(NonFocusAzimuthScale, 0.0f);
-		AudioDevice->GlobalFocusSettings.FocusDistanceScale = FMath::Max(FocusDistanceScale, 0.0f);
-		AudioDevice->GlobalFocusSettings.NonFocusDistanceScale = FMath::Max(NonFocusDistanceScale, 0.0f);
-		AudioDevice->GlobalFocusSettings.FocusVolumeScale = FMath::Max(FocusVolumeScale, 0.0f);
-		AudioDevice->GlobalFocusSettings.NonFocusVolumeScale = FMath::Max(NonFocusVolumeScale, 0.0f);
-		AudioDevice->GlobalFocusSettings.FocusPriorityScale = FMath::Max(FocusPriorityScale, 0.0f);
-		AudioDevice->GlobalFocusSettings.NonFocusPriorityScale = FMath::Max(NonFocusPriorityScale, 0.0f);
+		FGlobalFocusSettings NewFocusSettings;
+		NewFocusSettings.FocusAzimuthScale = FMath::Max(FocusAzimuthScale, 0.0f);
+		NewFocusSettings.NonFocusAzimuthScale = FMath::Max(NonFocusAzimuthScale, 0.0f);
+		NewFocusSettings.FocusDistanceScale = FMath::Max(FocusDistanceScale, 0.0f);
+		NewFocusSettings.NonFocusDistanceScale = FMath::Max(NonFocusDistanceScale, 0.0f);
+		NewFocusSettings.FocusVolumeScale = FMath::Max(FocusVolumeScale, 0.0f);
+		NewFocusSettings.NonFocusVolumeScale = FMath::Max(NonFocusVolumeScale, 0.0f);
+		NewFocusSettings.FocusPriorityScale = FMath::Max(FocusPriorityScale, 0.0f);
+		NewFocusSettings.NonFocusPriorityScale = FMath::Max(NonFocusPriorityScale, 0.0f);
+
+		AudioDevice->SetGlobalFocusSettings(NewFocusSettings);
 	}
 }
 
@@ -943,11 +946,10 @@ void UGameplayStatics::PlaySound2D(UObject* WorldContextObject, class USoundBase
 		return;
 	}
 
-	// TODO - Audio Threading. This call would be a task call to dispatch to the audio thread
 	if (FAudioDevice* AudioDevice = ThisWorld->GetAudioDevice())
 	{
 		FActiveSound NewActiveSound;
-		NewActiveSound.Sound = Sound;
+		NewActiveSound.SetSound(Sound);
 
 		NewActiveSound.VolumeMultiplier = VolumeMultiplier;
 		NewActiveSound.PitchMultiplier = PitchMultiplier;
@@ -1629,8 +1631,7 @@ bool UGameplayStatics::SuggestProjectileVelocity(UObject* WorldContextObject, FV
 
 	UWorld* const World = GEngine->GetWorldFromContextObject(WorldContextObject);
 
-	const float GravityZ = (OverrideGravityZ != 0.f) ? -OverrideGravityZ : -World->GetGravityZ();
-
+	const float GravityZ = FMath::IsNearlyEqual(OverrideGravityZ, 0.0f) ? -OverrideGravityZ : -World->GetGravityZ();
 
 	// v^4 - g*(g*x^2 + 2*y*v^2)
 	const float InsideTheSqrt = FMath::Square(TossSpeedSq) - GravityZ * ( (GravityZ * FMath::Square(DeltaXY)) + (2.f * DeltaZ * TossSpeedSq) );
@@ -1788,6 +1789,132 @@ bool UGameplayStatics::SuggestProjectileVelocity(UObject* WorldContextObject, FV
 	return bFoundAValidSolution;
 }
 
+static const FName NAME_PredictProjectilePath = FName(TEXT("PredictProjectilePath"));
+
+// note: this will automatically fall back to line test if radius is small enough
+bool UGameplayStatics::PredictProjectilePath(UObject* WorldContextObject, FHitResult& OutHit, TArray<FVector>& OutPathPositions, FVector& OutLastTraceDestination, FVector StartPos, FVector LaunchVelocity, bool bTracePath, float ProjectileRadius, const TArray<TEnumAsByte<EObjectTypeQuery> >& ObjectTypes, bool bTraceComplex, const TArray<AActor*>& ActorsToIgnore, EDrawDebugTrace::Type DrawDebugType, float DrawDebugTime, float SimFrequency /*= 30.f*/, float MaxSimTime /*= 2.f*/, float OverrideGravityZ /*= 0*/)
+{
+	OutPathPositions.Empty();
+	bool bBlockingHit = false;
+
+	UWorld const* const World = GEngine->GetWorldFromContextObject(WorldContextObject);
+	if (World && SimFrequency > KINDA_SMALL_NUMBER)
+	{
+		float const SubstepDeltaTime = 1.f / SimFrequency;
+		int32 const StepLimit = FMath::CeilToInt(SimFrequency * MaxSimTime);
+
+		const float GravityZ = FMath::IsNearlyEqual(OverrideGravityZ, 0.0f) ? OverrideGravityZ : World->GetGravityZ();
+
+		OutPathPositions.Add(StartPos);
+
+		FVector CurrentVel = LaunchVelocity;
+		FVector TraceStart = StartPos;
+		FVector TraceEnd = TraceStart + CurrentVel * SubstepDeltaTime;
+
+		FCollisionQueryParams QueryParams(NAME_PredictProjectilePath, true);
+		QueryParams.AddIgnoredActors(ActorsToIgnore);
+		QueryParams.bTraceComplex = bTraceComplex;
+
+		FCollisionObjectQueryParams ObjQueryParams;
+		for (auto Iter = ObjectTypes.CreateConstIterator(); Iter; ++Iter)
+		{
+			const ECollisionChannel& Channel = UCollisionProfile::Get()->ConvertToCollisionChannel(false, *Iter);
+			ObjQueryParams.AddObjectTypesToQuery(Channel);
+		}
+
+		FHitResult TraceHit(0.f);
+		int32 StepCount = 0;
+		while (StepCount < StepLimit)
+		{
+			OutLastTraceDestination = TraceEnd;
+
+			if (bTracePath)
+			{
+				if (World->SweepSingleByObjectType(TraceHit, TraceStart, TraceEnd, FQuat::Identity, ObjQueryParams, FCollisionShape::MakeSphere(ProjectileRadius), QueryParams))
+				{
+					// hit! we are done
+					OutHit = TraceHit;
+					OutPathPositions.Add(OutHit.Location);
+					bBlockingHit = true;
+					break;
+				}
+			}
+
+			OutPathPositions.Add(TraceEnd);
+
+			// integrate and continue
+			TraceStart = TraceEnd;
+			CurrentVel = CurrentVel + FVector(0, 0, GravityZ * SubstepDeltaTime);
+			TraceEnd = TraceStart + CurrentVel * SubstepDeltaTime;
+
+			++StepCount;
+		}
+	}
+
+	if (DrawDebugType != EDrawDebugTrace::None)
+	{
+		bool bPersistent = DrawDebugType == EDrawDebugTrace::Persistent;
+		float LifeTime = (DrawDebugType == EDrawDebugTrace::ForDuration) ? DrawDebugTime : 0.f;
+
+		float const DrawRadius = (ProjectileRadius > 0.f) ? ProjectileRadius : 5.f;
+
+		// draw the path
+		for (FVector PathPt : OutPathPositions)
+		{
+			::DrawDebugSphere(World, PathPt, DrawRadius, 12, FColor::Green, bPersistent, LifeTime);
+		}
+		// draw the impact point
+		if (bBlockingHit)
+		{
+			::DrawDebugSphere(World, OutHit.Location, 15.f, 12, FColor::Red, bPersistent, LifeTime);
+		}
+	}
+
+	return bBlockingHit;
+}
+
+bool UGameplayStatics::SuggestProjectileVelocity_MediumArc(UObject* WorldContextObject, FVector& OutLaunchVelocity, FVector StartPos, FVector EndPos, float OverrideGravityZ /*= 0*/)
+{
+	/* Make sure the start and end aren't the same location */
+	FVector const StartToEnd = EndPos - StartPos;
+	float const StartToEndDist = StartToEnd.Size();
+
+	UWorld const* const World = GEngine->GetWorldFromContextObject(WorldContextObject);
+	if (World && StartToEndDist > KINDA_SMALL_NUMBER)
+	{
+		const float GravityZ = FMath::IsNearlyEqual(OverrideGravityZ, 0.0f) ? OverrideGravityZ : World->GetGravityZ();
+
+		// choose arc halfway between up and direct line
+
+		FVector const StartToEndDir = StartToEnd / StartToEndDist;
+		FVector LaunchDir = (StartToEndDir + FVector::UpVector) * 0.5f;
+
+		// v = sqrt ( g * dx^2 / ( (dx tan(angle) + dz) * 2 * cos(angle))^2 ) )
+
+		FRotator const LaunchRot = LaunchDir.Rotation();
+		float const Angle = FMath::DegreesToRadians(LaunchRot.Pitch);
+
+		float const Dx = StartToEnd.Size2D();
+		float const Dz = StartToEnd.Z;
+		float const NumeratorInsideSqrt = (GravityZ * FMath::Square(Dx) * 0.5f);
+		float const DenominatorInsideSqrt = (Dz - (Dx * FMath::Tan(Angle))) * FMath::Square(FMath::Cos(Angle));
+		float const InsideSqrt = NumeratorInsideSqrt / DenominatorInsideSqrt;
+		if (InsideSqrt >= 0.f)
+		{
+			// there exists a solution
+			float const SpeedY = FMath::Sqrt(InsideSqrt);	// this is the mag of the vertical component
+			float VelMag = SpeedY / FMath::Sin(Angle);			// find speed in direction of the launch dir
+
+			OutLaunchVelocity = LaunchDir * VelMag;
+
+			return true;
+		}
+	}
+
+	OutLaunchVelocity = FVector::ZeroVector;
+	return false;
+}
+
 FIntVector UGameplayStatics::GetWorldOriginLocation(UObject* WorldContextObject)
 {
 	UWorld* World = GEngine->GetWorldFromContextObject(WorldContextObject);
diff --git a/Engine/Source/Runtime/Engine/Private/InstancedStaticMesh.cpp b/Engine/Source/Runtime/Engine/Private/InstancedStaticMesh.cpp
--- a/Engine/Source/Runtime/Engine/Private/InstancedStaticMesh.cpp
+++ b/Engine/Source/Runtime/Engine/Private/InstancedStaticMesh.cpp
@@ -1532,6 +1532,22 @@ bool UInstancedStaticMeshComponent::GetInstanceTransform(int32 InstanceIndex, FT
 	return true;
 }
 
+void UInstancedStaticMeshComponent::OnUpdateTransform(EUpdateTransformFlags UpdateTransformFlags, ETeleportType Teleport)
+{
+	// We are handling the physics move below, so don't handle it at higher levels
+	Super::OnUpdateTransform(UpdateTransformFlags | EUpdateTransformFlags::SkipPhysicsUpdate, Teleport);
+
+	// Always send new transform to physics
+	if (bPhysicsStateCreated && !(EUpdateTransformFlags::SkipPhysicsUpdate & UpdateTransformFlags))
+	{
+		for (int i = 0; i < PerInstanceSMData.Num(); i++)
+		{
+			const FTransform InstanceTransform(PerInstanceSMData[i].Transform);
+			UpdateInstanceTransform(i, InstanceTransform * ComponentToWorld, true);
+		}
+	}
+}
+
 bool UInstancedStaticMeshComponent::UpdateInstanceTransform(int32 InstanceIndex, const FTransform& NewInstanceTransform, bool bWorldSpace, bool bMarkRenderStateDirty)
 {
 	if (!PerInstanceSMData.IsValidIndex(InstanceIndex))
@@ -1544,6 +1560,9 @@ bool UInstancedStaticMeshComponent::UpdateInstanceTransform(int32 InstanceIndex,
 
 	FInstancedStaticMeshInstanceData& InstanceData = PerInstanceSMData[InstanceIndex];
 
+    // TODO: Computing LocalTransform is useless when we're updating the world location for the entire mesh.
+	// Should find some way around this for performance.
+    
 	// Render data uses local transform of the instance
 	FTransform LocalTransform = bWorldSpace ? NewInstanceTransform.GetRelativeTransform(ComponentToWorld) : NewInstanceTransform;
 	InstanceData.Transform = LocalTransform.ToMatrixWithScale();
diff --git a/Engine/Source/Runtime/Engine/Private/InterpolateComponentToAction.h b/Engine/Source/Runtime/Engine/Private/InterpolateComponentToAction.h
--- a/Engine/Source/Runtime/Engine/Private/InterpolateComponentToAction.h
+++ b/Engine/Source/Runtime/Engine/Private/InterpolateComponentToAction.h
@@ -42,7 +42,10 @@ public:
 	/** Should we east out (ie end slowly) during interpolation */
 	bool bEaseOut;
 
-	FInterpolateComponentToAction(float Duration, const FLatentActionInfo& LatentInfo, USceneComponent* Component, bool bInEaseOut, bool bInEaseIn)
+	/** Force use of shortest Path for rotation **/
+	bool bForceShortestRotationPath;
+
+	FInterpolateComponentToAction(float Duration, const FLatentActionInfo& LatentInfo, USceneComponent* Component, bool bInEaseOut, bool bInEaseIn, bool bInForceShortestRotationPath)
 		: TotalTime(Duration)
 		, TimeElapsed(0.f)
 		, bInterpolating(true)
@@ -58,6 +61,7 @@ public:
 		, TargetLocation(FVector::ZeroVector)
 		, bEaseIn(bInEaseIn)
 		, bEaseOut(bInEaseOut)
+		, bForceShortestRotationPath(bInForceShortestRotationPath)
 	{
 	}
 
@@ -111,7 +115,27 @@ public:
 
 			if(bInterpRotation && TargetComponent.IsValid())
 			{
-				FRotator NewRotation = bComplete ? TargetRotation : FMath::Lerp(InitialRotation, TargetRotation, BlendPct);
+				FRotator NewRotation;
+				// If we are done just set the final rotation
+				if (bComplete)
+				{
+					NewRotation = TargetRotation;
+				}
+				else if (bForceShortestRotationPath)
+				{
+					// We want the shortest path 
+					FQuat AQuat(InitialRotation);
+					FQuat BQuat(TargetRotation);
+
+					FQuat Result = FQuat::Slerp(AQuat, BQuat, BlendPct);
+					Result.Normalize();
+					NewRotation = Result.Rotator();
+				}
+				else
+				{
+					// dont care about it being the shortest path - just lerp
+					NewRotation = FMath::Lerp(InitialRotation, TargetRotation, BlendPct);
+				}
 				TargetComponent->SetRelativeRotation(NewRotation, false);
 			}
 		}
diff --git a/Engine/Source/Runtime/Engine/Private/Interpolation.cpp b/Engine/Source/Runtime/Engine/Private/Interpolation.cpp
--- a/Engine/Source/Runtime/Engine/Private/Interpolation.cpp
+++ b/Engine/Source/Runtime/Engine/Private/Interpolation.cpp
@@ -8057,19 +8057,6 @@ void UInterpTrackSound::UpdateTrack(float NewPosition, UInterpTrackInst* TrInst,
 			SoundInst->PlayAudioComp->Stop();
 			bPlaying = false;
 		}
-		else
-		{
-			FAudioDevice* AudioDevice = SoundInst->PlayAudioComp->GetAudioDevice();
-			if (AudioDevice)
-			{
-				FActiveSound* ActiveSound = AudioDevice->FindActiveSound(SoundInst->PlayAudioComp);
-				if (ActiveSound && FMath::Abs(NewPosition - (ActiveSound->RequestedStartTime + ActiveSound->PlaybackTime)) > 0.f)
-				{
-					SoundInst->PlayAudioComp->Stop();
-					bPlaying = false;
-				}
-			}
-		}
 	}
 
 
diff --git a/Engine/Source/Runtime/Engine/Private/KismetSystemLibrary.cpp b/Engine/Source/Runtime/Engine/Private/KismetSystemLibrary.cpp
--- a/Engine/Source/Runtime/Engine/Private/KismetSystemLibrary.cpp
+++ b/Engine/Source/Runtime/Engine/Private/KismetSystemLibrary.cpp
@@ -2949,7 +2949,7 @@ void UKismetSystemLibrary::RetriggerableDelay(UObject* WorldContextObject, float
 	}
 }
 
-void UKismetSystemLibrary::MoveComponentTo(USceneComponent* Component, FVector TargetRelativeLocation, FRotator TargetRelativeRotation, bool bEaseOut, bool bEaseIn, float OverTime, TEnumAsByte<EMoveComponentAction::Type> MoveAction, FLatentActionInfo LatentInfo)
+void UKismetSystemLibrary::MoveComponentTo(USceneComponent* Component, FVector TargetRelativeLocation, FRotator TargetRelativeRotation, bool bEaseOut, bool bEaseIn, float OverTime, bool bForceShortestRotationPath, TEnumAsByte<EMoveComponentAction::Type> MoveAction, FLatentActionInfo LatentInfo)
 {
 	if (UWorld* World = ((Component != NULL) ? Component->GetWorld() : NULL))
 	{
@@ -2965,7 +2965,7 @@ void UKismetSystemLibrary::MoveComponentTo(USceneComponent* Component, FVector T
 			if (MoveAction == EMoveComponentAction::Move)
 			{
 				// Only act on a 'move' input if not running
-				Action = new FInterpolateComponentToAction(OverTime, LatentInfo, Component, bEaseOut, bEaseIn);
+				Action = new FInterpolateComponentToAction(OverTime, LatentInfo, Component, bEaseOut, bEaseIn, bForceShortestRotationPath);
 
 				Action->TargetLocation = TargetRelativeLocation;
 				Action->TargetRotation = TargetRelativeRotation;
diff --git a/Engine/Source/Runtime/Engine/Private/LevelActor.cpp b/Engine/Source/Runtime/Engine/Private/LevelActor.cpp
--- a/Engine/Source/Runtime/Engine/Private/LevelActor.cpp
+++ b/Engine/Source/Runtime/Engine/Private/LevelActor.cpp
@@ -1271,53 +1271,40 @@ void UWorld::IssueEditorLoadWarnings()
 
 AAudioVolume* UWorld::GetAudioSettings( const FVector& ViewLocation, FReverbSettings* OutReverbSettings, FInteriorSettings* OutInteriorSettings )
 {
-	// Find the highest priority volume encompassing the current view location. This is made easier by the linked
-	// list being sorted by priority. @todo: it remains to be seen whether we should trade off sorting for constant
-	// time insertion/ removal time via e.g. TLinkedList.
-	AAudioVolume* Volume = HighestPriorityAudioVolume;
-	while( Volume )
+	// Find the highest priority volume encompassing the current view location.
+	for (AAudioVolume* Volume : AudioVolumes)
 	{
 		// Volume encompasses, break out of loop.
-		if (Volume->bEnabled && Volume->EncompassesPoint(ViewLocation))
+		if (Volume->GetEnabled() && Volume->EncompassesPoint(ViewLocation))
 		{
-			break;
-		}
-		// Volume doesn't encompass location, further traverse linked list.
-		else
-		{
-			Volume = Volume->NextLowerPriorityVolume;
+			if( OutReverbSettings )
+			{
+				*OutReverbSettings = Volume->GetReverbSettings();
+			}
+
+			if( OutInteriorSettings )
+			{
+				*OutInteriorSettings = Volume->GetInteriorSettings();
+			}
+			return Volume;
 		}
 	}
 
-	if( Volume )
-	{
-		if( OutReverbSettings )
-		{
-			*OutReverbSettings = Volume->Settings;
-		}
+	// If first level is a FakePersistentLevel (see CommitMapChange for more info)
+	// then use its world info for reverb settings
+	AWorldSettings* CurrentWorldSettings = GetWorldSettings(true);
 
-		if( OutInteriorSettings )
-		{
-			*OutInteriorSettings = Volume->AmbientZoneSettings;
-		}
-	}
-	else
+	if( OutReverbSettings )
 	{
-		// If first level is a FakePersistentLevel (see CommitMapChange for more info)
-		// then use its world info for reverb settings
-		AWorldSettings* CurrentWorldSettings = GetWorldSettings(true);
-
-		if( OutReverbSettings )
-		{
-			*OutReverbSettings = CurrentWorldSettings->DefaultReverbSettings;
-		}
+		*OutReverbSettings = CurrentWorldSettings->DefaultReverbSettings;
+	}
 
-		if( OutInteriorSettings )
-		{
-			*OutInteriorSettings = CurrentWorldSettings->DefaultAmbientZoneSettings;
-		}
+	if( OutInteriorSettings )
+	{
+		*OutInteriorSettings = CurrentWorldSettings->DefaultAmbientZoneSettings;
 	}
-	return Volume;
+
+	return nullptr;
 }
 
 void UWorld::SetAudioDeviceHandle(const uint32 InAudioDeviceHandle)
diff --git a/Engine/Source/Runtime/Engine/Private/PhysicsEngine/BodyInstance.cpp b/Engine/Source/Runtime/Engine/Private/PhysicsEngine/BodyInstance.cpp
--- a/Engine/Source/Runtime/Engine/Private/PhysicsEngine/BodyInstance.cpp
+++ b/Engine/Source/Runtime/Engine/Private/PhysicsEngine/BodyInstance.cpp
@@ -245,6 +245,9 @@ bool FCollisionResponse::operator==(const FCollisionResponse& Other) const
 }
 ////////////////////////////////////////////////////////////////////////////
 
+FBodyInstanceInit FBodyInstance::InitBodyDelegate;
+FBodyInstanceTerm FBodyInstance::TermBodyDelegate;
+
 FBodyInstance::FBodyInstance()
 	: InstanceBodyIndex(INDEX_NONE)
 	, InstanceBoneIndex(INDEX_NONE)
@@ -706,7 +709,7 @@ void FBodyInstance::CreateDOFLock()
 		DOFConstraint->Pos2 = TM.GetLocation();
 
 		// Create constraint instance based on DOF
-		DOFConstraint->InitConstraint(OwnerComponent.Get(), this, nullptr, 1.f);
+		DOFConstraint->InitConstraint(this, nullptr, 1.f, OwnerComponent.Get());
 	}
 }
 
@@ -1797,6 +1800,8 @@ void FBodyInstance::InitBody(class UBodySetup* Setup, const FTransform& Transfor
 
 	Bodies.Reset();
 	Transforms.Reset();
+
+	InitBodyDelegate.Broadcast(this);
 }
 
 TSharedPtr<TArray<ANSICHAR>> GetDebugDebugName(const UPrimitiveComponent* PrimitiveComp, const UBodySetup* BodySetup, FString& DebugName)
@@ -1990,6 +1995,8 @@ void TermBodyHelper(int16& SceneIndex, PxRigidActor*& PRigidActor, FBodyInstance
  */
 void FBodyInstance::TermBody()
 {
+	TermBodyDelegate.Broadcast(this);
+
 	SCOPE_CYCLE_COUNTER(STAT_TermBody);
 #if WITH_BOX2D
 	if (BodyInstancePtr != NULL)
diff --git a/Engine/Source/Runtime/Engine/Private/PhysicsEngine/ConstraintInstance.cpp b/Engine/Source/Runtime/Engine/Private/PhysicsEngine/ConstraintInstance.cpp
--- a/Engine/Source/Runtime/Engine/Private/PhysicsEngine/ConstraintInstance.cpp
+++ b/Engine/Source/Runtime/Engine/Private/PhysicsEngine/ConstraintInstance.cpp
@@ -7,6 +7,9 @@
 #include "PhysicsEngine/PhysicsConstraintComponent.h"
 #include "UObject/DevObjectVersion.h"
 
+#include "MessageLog.h"
+#include "UObjectToken.h"
+
 #define LOCTEXT_NAMESPACE "ConstraintInstance"
 
 TAutoConsoleVariable<float> CVarConstraintDampingScale(
@@ -141,7 +144,6 @@ void FConstraintInstance::UpdateDriveTarget()
 /** Constructor **/
 FConstraintInstance::FConstraintInstance()
 	: ConstraintIndex(0)
-	, OwnerComponent(NULL)
 #if WITH_PHYSX
 	, ConstraintData(NULL)
 #endif	//WITH_PHYSX
@@ -267,7 +269,7 @@ float ComputeAverageMass_AssumesLocked(const PxRigidActor* PActor1, const PxRigi
 }
 
 /** Finds the common scene and appropriate actors for the passed in body instances. Makes sure to do this without requiring a scene lock*/
-PxScene* GetPScene_LockFree(const FBodyInstance* Body1, const FBodyInstance* Body2)
+PxScene* GetPScene_LockFree(const FBodyInstance* Body1, const FBodyInstance* Body2, UObject* DebugOwner)
 {
 	const int32 SceneIndex1 = Body1 ? Body1->GetSceneIndex() : -1;
 	const int32 SceneIndex2 = Body2 ? Body2->GetSceneIndex() : -1;
@@ -276,11 +278,21 @@ PxScene* GetPScene_LockFree(const FBodyInstance* Body1, const FBodyInstance* Bod
 	//now we check if the two actors are valid
 	if(SceneIndex1 == -1 && SceneIndex2 == -1)
 	{
-		UE_LOG(LogPhysics, Log, TEXT("Attempting to create a joint between two null actors.  No joint created."));
+#if !(UE_BUILD_SHIPPING || UE_BUILD_TEST)
+		FMessageLog("PIE").Warning()
+			->AddToken(FTextToken::Create(LOCTEXT("JointBetweenNullStart", "Constraint")))
+			->AddToken(FUObjectToken::Create(DebugOwner))
+			->AddToken(FTextToken::Create(LOCTEXT("JointBetweenNullEnd", "attempting to create a joint between two null actors.")));
+#endif // !(UE_BUILD_SHIPPING || UE_BUILD_TEST)
 	}
 	else if(SceneIndex1 >= 0 && SceneIndex2 >= 0 && SceneIndex1 != SceneIndex2)
 	{
-		UE_LOG(LogPhysics, Log, TEXT("Attempting to create a joint between two actors in different scenes.  No joint created."));
+#if !(UE_BUILD_SHIPPING || UE_BUILD_TEST)
+		FMessageLog("PIE").Warning()
+			->AddToken(FTextToken::Create(LOCTEXT("JointBetweenScenesStart", "Constraint")))
+			->AddToken(FUObjectToken::Create(DebugOwner))
+			->AddToken(FTextToken::Create(LOCTEXT("JointBetweenScenesEnd", "attempting to create a joint between two actors in different scenes.  No joint created.")));
+#endif // !(UE_BUILD_SHIPPING || UE_BUILD_TEST)
 	}
 	else
 	{
@@ -290,12 +302,19 @@ PxScene* GetPScene_LockFree(const FBodyInstance* Body1, const FBodyInstance* Bod
 	return PScene;
 }
 
-bool CanActorSimulate(const FBodyInstance* BI, const PxRigidActor* PActor)
+bool CanActorSimulate(const FBodyInstance* BI, const PxRigidActor* PActor, UObject* DebugOwner)
 {
 	if (PActor && (PActor->getActorFlags() & PxActorFlag::eDISABLE_SIMULATION))
 	{
+#if !(UE_BUILD_SHIPPING || UE_BUILD_TEST)
 		const UPrimitiveComponent* PrimComp = BI->OwnerComponent.Get();
-		UE_LOG(LogPhysics, Warning, TEXT("Attempting to create a joint on actor (%s) which is not eligible for simulation. Is it marked QueryOnly?"), *PrimComp->GetReadableName());
+		FMessageLog("PIE").Warning()
+			->AddToken(FTextToken::Create(LOCTEXT("InvalidBodyStart", "Attempting to create a joint")))
+			->AddToken(FUObjectToken::Create(DebugOwner))
+			->AddToken(FTextToken::Create(LOCTEXT("InvalidBodyMid", "to body")))
+			->AddToken(FUObjectToken::Create(PrimComp))
+			->AddToken(FTextToken::Create(LOCTEXT("InvalidBodyEnd", "which is not eligible for simulation. Is it marked QueryOnly?")));
+#endif // !(UE_BUILD_SHIPPING || UE_BUILD_TEST)
 
 		return false;
 	}
@@ -304,7 +323,7 @@ bool CanActorSimulate(const FBodyInstance* BI, const PxRigidActor* PActor)
 }
 
 /*various logical checks to find the correct physx actor. Returns true if found valid actors that can be constrained*/
-bool GetPActors_AssumesLocked(const FBodyInstance* Body1, const FBodyInstance* Body2, PxRigidActor** PActor1Out, PxRigidActor** PActor2Out)
+bool GetPActors_AssumesLocked(const FBodyInstance* Body1, const FBodyInstance* Body2, PxRigidActor** PActor1Out, PxRigidActor** PActor2Out, UObject* DebugOwner)
 {
 	PxRigidActor* PActor1 = Body1 ? Body1->GetPxRigidActor_AssumesLocked() : NULL;
 	PxRigidActor* PActor2 = Body2 ? Body2->GetPxRigidActor_AssumesLocked() : NULL;
@@ -313,18 +332,29 @@ bool GetPActors_AssumesLocked(const FBodyInstance* Body1, const FBodyInstance* B
 	// Do not create joint unless one of the actors is dynamic
 	if ((!PActor1 || !PActor1->isRigidBody()) && (!PActor2 || !PActor2->isRigidBody()))
 	{
-		UE_LOG(LogPhysics, Warning, TEXT("Attempting to create a joint between actors that are static.  No joint created."));
+#if !(UE_BUILD_SHIPPING || UE_BUILD_TEST)
+		FMessageLog("PIE").Warning()
+			->AddToken(FTextToken::Create(LOCTEXT("TwoStaticBodiesWarningStart", "Constraint in")))
+			->AddToken(FUObjectToken::Create(DebugOwner))
+			->AddToken(FTextToken::Create(LOCTEXT("TwoStaticBodiesWarningEnd", "attempting to create a joint between objects that are both static.  No joint created.")));
+#endif // !(UE_BUILD_SHIPPING || UE_BUILD_TEST)
 		return false;
 	}
 
 	if(PActor1 == PActor2)
 	{
+#if !(UE_BUILD_SHIPPING || UE_BUILD_TEST)
 		const UPrimitiveComponent* PrimComp = Body1->OwnerComponent.Get();
-		UE_LOG(LogPhysics, Warning, TEXT("Attempting to create a joint between the same actor (%s)"), *PrimComp->GetReadableName());
+		FMessageLog("PIE").Warning()
+			->AddToken(FTextToken::Create(LOCTEXT("SameBodyWarningStart", "Constraint in")))
+			->AddToken(FUObjectToken::Create(DebugOwner))
+			->AddToken(FTextToken::Create(LOCTEXT("SameBodyWarningMid", "attempting to create a joint to the same body")))
+			->AddToken(FUObjectToken::Create(PrimComp));
+#endif // !(UE_BUILD_SHIPPING || UE_BUILD_TEST)
 		return false;
 	}
 
-	if(!CanActorSimulate(Body1, PActor1) || !CanActorSimulate(Body2, PActor2))
+	if(!CanActorSimulate(Body1, PActor1, DebugOwner) || !CanActorSimulate(Body2, PActor2, DebugOwner))
 	{
 		return false;
 	}
@@ -468,42 +498,54 @@ void EnsureSleepingActorsStaySleeping_AssumesLocked(PxRigidActor* PActor1, PxRig
 /** 
  *	Create physics engine constraint.
  */
-void FConstraintInstance::InitConstraint(USceneComponent* Owner, FBodyInstance* Body1, FBodyInstance* Body2, float InScale)
+void FConstraintInstance::InitConstraint(FBodyInstance* Body1, FBodyInstance* Body2, float InScale, UObject* DebugOwner, FOnConstraintBroken InConstraintBrokenDelegate)
 {
-	OwnerComponent = Owner;
-	LastKnownScale = InScale;
-
 #if WITH_PHYSX
-	PhysxUserData = FPhysxUserData(this);
-
-	// if there's already a constraint, get rid of it first
-	if (ConstraintData)
-	{
-		TermConstraint();
-	}
-
 	PxRigidActor* PActor1 = nullptr;
 	PxRigidActor* PActor2 = nullptr;
-	PxScene* PScene = GetPScene_LockFree(Body1, Body2);
+	PxScene* PScene = GetPScene_LockFree(Body1, Body2, DebugOwner);
 	SCOPED_SCENE_WRITE_LOCK(PScene);
 	{
-		const bool bValidConstraintSetup = PScene && GetPActors_AssumesLocked(Body1, Body2, &PActor1, &PActor2) && CreatePxJoint_AssumesLocked(PActor1, PActor2, PScene);
+		const bool bValidConstraintSetup = PScene && GetPActors_AssumesLocked(Body1, Body2, &PActor1, &PActor2, DebugOwner);
 		if (!bValidConstraintSetup)
 		{
 			return;
 		}
 
-		// update mass
-		UpdateAverageMass_AssumesLocked(PActor1, PActor2);
+		InitConstraintPhysX_AssumesLocked(PActor1, PActor2, PScene, InScale, InConstraintBrokenDelegate);
+	}
+	
+#endif // WITH_PHYSX
 
-		ProfileInstance.UpdatePhysX_AssumesLocked(ConstraintData, AverageMass, bScaleLinearLimits ? LastKnownScale : 1.f);
+}
 
-		EnsureSleepingActorsStaySleeping_AssumesLocked(PActor1, PActor2);
+#if WITH_PHYSX
+void FConstraintInstance::InitConstraintPhysX_AssumesLocked(physx::PxRigidActor* PActor1, physx::PxRigidActor* PActor2, physx::PxScene* PScene, float InScale, FOnConstraintBroken InConstraintBrokenDelegate)
+{
+	OnConstraintBrokenDelegate = InConstraintBrokenDelegate;
+	LastKnownScale = InScale;
+
+	PhysxUserData = FPhysxUserData(this);
+
+	// if there's already a constraint, get rid of it first
+	if (ConstraintData)
+	{
+		TermConstraint();
+	}
+
+	if (!CreatePxJoint_AssumesLocked(PActor1, PActor2, PScene))
+	{
+		return;
 	}
 	
-#endif // WITH_PHYSX
+	// update mass
+	UpdateAverageMass_AssumesLocked(PActor1, PActor2);
+
+	ProfileInstance.UpdatePhysX_AssumesLocked(ConstraintData, AverageMass, bScaleLinearLimits ? LastKnownScale : 1.f);
 
+	EnsureSleepingActorsStaySleeping_AssumesLocked(PActor1, PActor2);
 }
+#endif
 
 #if WITH_PHYSX
 void FConstraintProfileProperties::UpdatePhysX_AssumesLocked(PxD6Joint* Joint, float AverageMass, float UseScale) const
@@ -1162,18 +1204,7 @@ void FConstraintInstance::PostSerialize(const FArchive& Ar)
 
 void FConstraintInstance::OnConstraintBroken()
 {
-	UPhysicsConstraintComponent* ConstraintComp = Cast<UPhysicsConstraintComponent>(OwnerComponent);
-	USkeletalMeshComponent* SkelMeshComp = Cast<USkeletalMeshComponent>(OwnerComponent);
-
-	if (ConstraintComp)
-	{
-		ConstraintComp->OnConstraintBroken.Broadcast(ConstraintIndex);
-	}
-	else if (SkelMeshComp)
-	{
-		SkelMeshComp->OnConstraintBroken.Broadcast(ConstraintIndex);
-	}
-
+	OnConstraintBrokenDelegate.ExecuteIfBound(ConstraintIndex);
 }
 
 //Hacks to easily get zeroed memory for special case when we don't use GC
diff --git a/Engine/Source/Runtime/Engine/Private/PhysicsEngine/PhysAnim.cpp b/Engine/Source/Runtime/Engine/Private/PhysicsEngine/PhysAnim.cpp
--- a/Engine/Source/Runtime/Engine/Private/PhysicsEngine/PhysAnim.cpp
+++ b/Engine/Source/Runtime/Engine/Private/PhysicsEngine/PhysAnim.cpp
@@ -212,7 +212,7 @@ void USkeletalMeshComponent::PerformBlendPhysicsBones(const TArray<FBoneIndexTyp
 				if ( PhysicsAsset )
 				{
 					UE_LOG(LogPhysics, Warning, TEXT(" - # of BodySetup (%d), # of Bodies (%d), Invalid BodyIndex(%d)"), 
-						PhysicsAsset->BodySetup.Num(), Bodies.Num(), BodyIndex);
+						PhysicsAsset->SkeletalBodySetups.Num(), Bodies.Num(), BodyIndex);
 				}
 				continue;
 			}
@@ -535,11 +535,11 @@ void USkeletalMeshComponent::UpdateKinematicBonesToAnim(const TArray<FTransform>
 		if (PhysicsAsset && SkeletalMesh && Bodies.Num() > 0)
 		{
 #if !(UE_BUILD_SHIPPING || UE_BUILD_TEST)
-			if (!ensure(PhysicsAsset->BodySetup.Num() == Bodies.Num()))
+			if (!ensure(PhysicsAsset->SkeletalBodySetups.Num() == Bodies.Num()))
 			{
 				// related to TTP 280315
 				UE_LOG(LogPhysics, Warning, TEXT("Mesh (%s) has PhysicsAsset(%s), and BodySetup(%d) and Bodies(%d) don't match"),
-					*SkeletalMesh->GetName(), *PhysicsAsset->GetName(), PhysicsAsset->BodySetup.Num(), Bodies.Num());
+					*SkeletalMesh->GetName(), *PhysicsAsset->GetName(), PhysicsAsset->SkeletalBodySetups.Num(), Bodies.Num());
 				return;
 			}
 #endif
@@ -571,7 +571,7 @@ void USkeletalMeshComponent::UpdateKinematicBonesToAnim(const TArray<FTransform>
 					// If we could not find it - warn.
 					if (BoneIndex == INDEX_NONE || BoneIndex >= GetNumSpaceBases())
 					{
-						const FName BodyName = PhysicsAsset->BodySetup[i]->BoneName;
+						const FName BodyName = PhysicsAsset->SkeletalBodySetups[i]->BoneName;
 						UE_LOG(LogPhysics, Log, TEXT("UpdateRBBones: WARNING: Failed to find bone '%s' need by PhysicsAsset '%s' in SkeletalMesh '%s'."), *BodyName.ToString(), *PhysicsAsset->GetName(), *SkeletalMesh->GetName());
 					}
 					else
@@ -581,7 +581,7 @@ void USkeletalMeshComponent::UpdateKinematicBonesToAnim(const TArray<FTransform>
 						const FTransform BoneTransform = InSpaceBases[BoneIndex] * CurrentLocalToWorld;
 						if(!BoneTransform.IsValid())
 						{
-							const FName BodyName = PhysicsAsset->BodySetup[i]->BoneName;
+							const FName BodyName = PhysicsAsset->SkeletalBodySetups[i]->BoneName;
 							UE_LOG(LogPhysics, Warning, TEXT("UpdateKinematicBonesToAnim: Trying to set transform with bad data %s on PhysicsAsset '%s' in SkeletalMesh '%s' for bone '%s'"), *BoneTransform.ToHumanReadableString(), *PhysicsAsset->GetName(), *SkeletalMesh->GetName(), *BodyName.ToString());
 							BoneTransform.DiagnosticCheck_IsValid();	//In special nan mode we want to actually ensure
 
diff --git a/Engine/Source/Runtime/Engine/Private/PhysicsEngine/PhysDrawing.cpp b/Engine/Source/Runtime/Engine/Private/PhysicsEngine/PhysDrawing.cpp
--- a/Engine/Source/Runtime/Engine/Private/PhysicsEngine/PhysDrawing.cpp
+++ b/Engine/Source/Runtime/Engine/Private/PhysicsEngine/PhysDrawing.cpp
@@ -425,10 +425,22 @@ void FKConvexElem::DrawElemWire(FPrimitiveDrawInterface* PDI, const FTransform&
 #endif // WITH_PHYSX
 }
 
-void FKConvexElem::AddCachedSolidConvexGeom(TArray<FDynamicMeshVertex>& VertexBuffer, TArray<int32>& IndexBuffer, const float Scale, const FColor VertexColor, const bool bIsMirrored) const
+void FKConvexElem::AddCachedSolidConvexGeom(TArray<FDynamicMeshVertex>& VertexBuffer, TArray<int32>& IndexBuffer, const FColor VertexColor) const
 {
 #if WITH_PHYSX
-	const PxConvexMesh* ConvexMeshToUse = bIsMirrored ? ConvexMeshNegX : ConvexMesh;
+	// We always want to generate 'non-mirrored geometry', so if all we have is flipped, we have to un-flip it in this function
+	bool bIsMirrored = false;
+	const PxConvexMesh* ConvexMeshToUse = nullptr; 
+	if (ConvexMesh != nullptr)
+	{
+		ConvexMeshToUse = ConvexMesh;
+	}
+	else if (ConvexMeshNegX != nullptr)
+	{
+		ConvexMeshToUse = ConvexMeshNegX;
+		bIsMirrored = true;
+	}
+
 	if(ConvexMeshToUse)
 	{
 		int32 StartVertOffset = VertexBuffer.Num();
@@ -438,6 +450,8 @@ void FKConvexElem::AddCachedSolidConvexGeom(TArray<FDynamicMeshVertex>& VertexBu
 		const PxU8* PIndexBuffer = ConvexMeshToUse->getIndexBuffer();
 		PxU32 NbPolygons = ConvexMeshToUse->getNbPolygons();
 
+		FVector Scale3D = bIsMirrored ? FVector(-1, 1, 1) : FVector(1, 1, 1);
+
 		for(PxU32 i=0;i<NbPolygons;i++)
 		{
 			PxHullPolygon Data;
@@ -457,7 +471,7 @@ void FKConvexElem::AddCachedSolidConvexGeom(TArray<FDynamicMeshVertex>& VertexBu
 				int32 VertIndex = indices[j];
 
 				FDynamicMeshVertex Vert1;
-				Vert1.Position = Transform.TransformPosition( P2UVector(PVertices[VertIndex]) * Scale); // Apply element transform to get geom in component space
+				Vert1.Position = Transform.TransformPosition( P2UVector(PVertices[VertIndex]) * Scale3D ); // Apply element transform to get geom in component space
 				Vert1.Color = VertexColor;
 				Vert1.SetTangents(
 					TangentX,
@@ -472,8 +486,17 @@ void FKConvexElem::AddCachedSolidConvexGeom(TArray<FDynamicMeshVertex>& VertexBu
 			for(PxU32 j=0;j<nbTris;j++)
 			{
 				IndexBuffer.Add(StartVertOffset+0);
-				IndexBuffer.Add(StartVertOffset+j+2);
-				IndexBuffer.Add(StartVertOffset+j+1);
+
+				if (bIsMirrored)
+				{
+					IndexBuffer.Add(StartVertOffset + j + 1);
+					IndexBuffer.Add(StartVertOffset + j + 2);
+				}
+				else
+				{
+					IndexBuffer.Add(StartVertOffset + j + 2);
+					IndexBuffer.Add(StartVertOffset + j + 1);
+				}
 			}
 
 			StartVertOffset += Data.mNbVerts;
@@ -593,8 +616,6 @@ void FKAggregateGeom::GetAggGeom(const FTransform& Transform, const FColor Color
 	{
 		if(bDrawSolid)
 		{
-			const bool bIsMirrored = (Scale3D.X * Scale3D.Y * Scale3D.Z < 0.0f);
-
 			// Cache collision vertex/index buffer
 			if(!RenderInfo)
 			{
@@ -607,7 +628,7 @@ void FKAggregateGeom::GetAggGeom(const FTransform& Transform, const FColor Color
 				for(int32 i=0; i<ConvexElems.Num(); i++)
 				{
 					// Get vertices/triangles from this hull.
-					ConvexElems[i].AddCachedSolidConvexGeom(ThisGeom.RenderInfo->VertexBuffer->Vertices, ThisGeom.RenderInfo->IndexBuffer->Indices, 1.0f, FColor::White, bIsMirrored);
+					ConvexElems[i].AddCachedSolidConvexGeom(ThisGeom.RenderInfo->VertexBuffer->Vertices, ThisGeom.RenderInfo->IndexBuffer->Indices, FColor::White);
 				}
 
 				// Only continue if we actually got some valid geometry
@@ -626,7 +647,7 @@ void FKAggregateGeom::GetAggGeom(const FTransform& Transform, const FColor Color
 			if(RenderInfo->HasValidGeometry())
 			{
 				// Calculate transform
-				FTransform LocalToWorld = FTransform( FQuat::Identity, FVector::ZeroVector, bIsMirrored ? (Scale3D * FVector(-1, 1, 1)) : Scale3D ) * ParentTM;
+				FTransform LocalToWorld = FTransform(FQuat::Identity, FVector::ZeroVector, Scale3D) * ParentTM;
 
 				// Draw the mesh.
 				FMeshBatch& Mesh = Collector.AllocateMesh();
@@ -720,16 +741,16 @@ FTransform GetSkelBoneTransform(int32 BoneIndex, const TArray<FTransform>& Space
 
 void UPhysicsAsset::GetCollisionMesh(int32 ViewIndex, FMeshElementCollector& Collector, const USkeletalMesh* SkelMesh, const TArray<FTransform>& SpaceBases, const FTransform& LocalToWorld, const FVector& Scale3D)
 {
-	for( int32 i=0; i<BodySetup.Num(); i++)
+	for( int32 i=0; i<SkeletalBodySetups.Num(); i++)
 	{
-		int32 BoneIndex = SkelMesh->RefSkeleton.FindBoneIndex( BodySetup[i]->BoneName );
+		int32 BoneIndex = SkelMesh->RefSkeleton.FindBoneIndex( SkeletalBodySetups[i]->BoneName );
 		
-		FColor* BoneColor = (FColor*)( &BodySetup[i] );
+		FColor* BoneColor = (FColor*)( &SkeletalBodySetups[i] );
 
 		FTransform BoneTransform = GetSkelBoneTransform(BoneIndex, SpaceBases, LocalToWorld);
 		BoneTransform.SetScale3D(Scale3D);
-		BodySetup[i]->CreatePhysicsMeshes();
-		BodySetup[i]->AggGeom.GetAggGeom(BoneTransform, *BoneColor, NULL, false, false, false, ViewIndex, Collector);
+		SkeletalBodySetups[i]->CreatePhysicsMeshes();
+		SkeletalBodySetups[i]->AggGeom.GetAggGeom(BoneTransform, *BoneColor, NULL, false, false, false, ViewIndex, Collector);
 	}
 }
 
diff --git a/Engine/Source/Runtime/Engine/Private/PhysicsEngine/PhysScene.cpp b/Engine/Source/Runtime/Engine/Private/PhysicsEngine/PhysScene.cpp
--- a/Engine/Source/Runtime/Engine/Private/PhysicsEngine/PhysScene.cpp
+++ b/Engine/Source/Runtime/Engine/Private/PhysicsEngine/PhysScene.cpp
@@ -66,6 +66,8 @@ DECLARE_DWORD_COUNTER_STAT(TEXT("(ASync) Shapes"), STAT_NumShapesAsync, STATGROU
 static int16 PhysXSceneCount = 1;
 static const int PhysXSlowRebuildRate = 10;
 
+CreateSimEventCallbackFn* FPhysScene::CreateSimEventCallback;
+
 EPhysicsSceneType FPhysScene::SceneType_AssumesLocked(const FBodyInstance* BodyInstance) const
 {
 #if WITH_PHYSX
@@ -1670,7 +1672,7 @@ void FPhysScene::InitPhysScene(uint32 SceneType)
 	PhysxUserData = FPhysxUserData(this);
 
 	// Create sim event callback
-	SimEventCallback[SceneType] = new FPhysXSimEventCallback(this, SceneType);
+	SimEventCallback[SceneType] = CreateSimEventCallback ? (*CreateSimEventCallback)(this, SceneType) : new FPhysXSimEventCallback(this, SceneType);
 
 	// Include scene descriptor in loop, so that we might vary it with scene type
 	PxSceneDesc PSceneDesc(GPhysXSDK->getTolerancesScale());
diff --git a/Engine/Source/Runtime/Engine/Private/PhysicsEngine/PhysXSupport.h b/Engine/Source/Runtime/Engine/Private/PhysicsEngine/PhysXSupport.h
--- a/Engine/Source/Runtime/Engine/Private/PhysicsEngine/PhysXSupport.h
+++ b/Engine/Source/Runtime/Engine/Private/PhysicsEngine/PhysXSupport.h
@@ -444,7 +444,7 @@ PxFilterFlags PhysXSimFilterShader(	PxFilterObjectAttributes attributes0, PxFilt
 class FPhysScene;
 
 /** Event callback used to notify engine about various collision events */
-class FPhysXSimEventCallback : public PxSimulationEventCallback
+class ENGINE_API FPhysXSimEventCallback : public PxSimulationEventCallback
 {
 public:
 	FPhysXSimEventCallback(FPhysScene* InOwningScene, int32 InSceneType) : OwningScene(InOwningScene), SceneType(InSceneType){}
diff --git a/Engine/Source/Runtime/Engine/Private/PhysicsEngine/PhysicalAnimationComponent.cpp b/Engine/Source/Runtime/Engine/Private/PhysicsEngine/PhysicalAnimationComponent.cpp
new file mode 100644
--- /dev/null
+++ b/Engine/Source/Runtime/Engine/Private/PhysicsEngine/PhysicalAnimationComponent.cpp
@@ -0,0 +1,365 @@
+// Copyright 1998-2016 Epic Games, Inc. All Rights Reserved.
+
+#include "EnginePrivate.h"
+#include "PhysicsEngine/PhysicalAnimationComponent.h"
+#include "PhysicsEngine/PhysicsAsset.h"
+#include "PhysXSupport.h"
+
+const FConstraintProfileProperties UPhysicalAnimationComponent::PhysicalAnimationProfile = []()
+{
+	//Setup the default constraint profile for all joints created by physical animation system
+	FConstraintProfileProperties RetProfile;
+	RetProfile.LinearLimit.XMotion = LCM_Free;
+	RetProfile.LinearLimit.YMotion = LCM_Free;
+	RetProfile.LinearLimit.ZMotion = LCM_Free;
+
+	RetProfile.ConeLimit.Swing1Motion = ACM_Free;
+	RetProfile.ConeLimit.Swing2Motion = ACM_Free;
+	RetProfile.TwistLimit.TwistMotion = ACM_Free;
+
+	RetProfile.LinearDrive.XDrive.bEnablePositionDrive = true;
+	RetProfile.LinearDrive.XDrive.bEnableVelocityDrive = true;
+	RetProfile.LinearDrive.YDrive.bEnablePositionDrive = true;
+	RetProfile.LinearDrive.YDrive.bEnableVelocityDrive = true;
+	RetProfile.LinearDrive.ZDrive.bEnablePositionDrive = true;
+	RetProfile.LinearDrive.ZDrive.bEnableVelocityDrive = true;
+
+	RetProfile.AngularDrive.SlerpDrive.bEnablePositionDrive = true;
+	RetProfile.AngularDrive.SlerpDrive.bEnableVelocityDrive = true;
+	RetProfile.AngularDrive.AngularDriveMode = EAngularDriveMode::SLERP;
+
+	return RetProfile;
+}();
+
+UPhysicalAnimationComponent::UPhysicalAnimationComponent(const FObjectInitializer& ObjectInitializer)
+	: Super(ObjectInitializer)
+{
+	bWantsInitializeComponent = true;
+	PrimaryComponentTick.bCanEverTick = true;
+	PrimaryComponentTick.bTickEvenWhenPaused = true;
+	PrimaryComponentTick.TickGroup = TG_PrePhysics;
+
+	SceneIndex = INDEX_NONE;
+}
+
+void UPhysicalAnimationComponent::InitializeComponent()
+{
+	Super::InitializeComponent();
+	InitComponent();
+}
+
+void UPhysicalAnimationComponent::InitComponent()
+{
+	if (SkeletalMeshComponent)
+	{
+		PrimaryComponentTick.AddPrerequisite(SkeletalMeshComponent, SkeletalMeshComponent->PrimaryComponentTick);
+		UpdatePhysicsEngine();
+	}
+}
+
+void UPhysicalAnimationComponent::BeginDestroy()
+{
+	ReleasePhysicsEngine();
+	Super::BeginDestroy();
+}
+
+void UPhysicalAnimationComponent::SetSkeletalMeshComponent(USkeletalMeshComponent* InSkeletalMeshComponent)
+{
+	SkeletalMeshComponent = InSkeletalMeshComponent;
+	DriveData.Empty();
+	ReleasePhysicsEngine();
+	InitComponent();
+
+}
+
+bool UpdatePhysicalAnimationSettings(FName BodyName, const FPhysicalAnimationData& InPhysicalAnimationData, TArray<FPhysicalAnimationData>& DriveData, const UPhysicsAsset& PhysAsset)
+{
+	int32 BodyIdx = PhysAsset.FindBodyIndex(BodyName);
+	if (BodyIdx != INDEX_NONE)
+	{
+		//This code does a linear search in the insertion. This is by choice so that during tick we get nice tight cache for iteration. We could keep a map of sorts, but expected number of bodies is small
+		FPhysicalAnimationData* UpdateAnimationData = DriveData.FindByPredicate([BodyName](const FPhysicalAnimationData& Elem) { return Elem.BodyName == BodyName; });
+		if (UpdateAnimationData == nullptr)
+		{
+			UpdateAnimationData = &DriveData[DriveData.AddUninitialized()];
+		}
+		*UpdateAnimationData = InPhysicalAnimationData;
+		UpdateAnimationData->BodyName = BodyName;
+
+		return true;
+	}
+
+	return false;
+}
+
+void UPhysicalAnimationComponent::ApplyPhysicalAnimationSettings(FName BodyName, const FPhysicalAnimationData& PhysicalAnimationData)
+{
+	UPhysicsAsset* PhysAsset = SkeletalMeshComponent ? SkeletalMeshComponent->GetPhysicsAsset() : nullptr;
+	if (PhysAsset)
+	{
+		if(UpdatePhysicalAnimationSettings(BodyName, PhysicalAnimationData, DriveData, *PhysAsset))
+		{
+			UpdatePhysicsEngine();
+		}
+	}
+}
+
+void UPhysicalAnimationComponent::ApplyPhysicalAnimationSettingsBelow(FName BodyName, const FPhysicalAnimationData& PhysicalAnimationData, bool bIncludeSelf)
+{
+	UPhysicsAsset* PhysAsset = SkeletalMeshComponent ? SkeletalMeshComponent->GetPhysicsAsset() : nullptr;
+	if (PhysAsset && SkeletalMeshComponent)
+	{
+		TArray<FPhysicalAnimationData>& NewDriveData = DriveData;
+		bool bNeedsUpdating = false;
+		SkeletalMeshComponent->ForEachBodyBelow(BodyName, bIncludeSelf, /*bSkipCustomType=*/false, [PhysAsset, &NewDriveData, PhysicalAnimationData, &bNeedsUpdating](const FBodyInstance* BI)
+		{
+			const FName IterBodyName = PhysAsset->SkeletalBodySetups[BI->InstanceBodyIndex]->BoneName;
+			bNeedsUpdating |= UpdatePhysicalAnimationSettings(IterBodyName, PhysicalAnimationData, NewDriveData, *PhysAsset);
+		});
+
+		if(bNeedsUpdating)
+		{
+			UpdatePhysicsEngine();
+		}
+	}
+}
+
+
+void UPhysicalAnimationComponent::ApplyPhysicalAnimationProfileBelow(FName BodyName, FName ProfileName, bool bIncludeSelf, bool bClearNotFound)
+{
+	UPhysicsAsset* PhysAsset = SkeletalMeshComponent ? SkeletalMeshComponent->GetPhysicsAsset() : nullptr;
+	if (PhysAsset && SkeletalMeshComponent)
+	{
+		TArray<FPhysicalAnimationData>& NewDriveData = DriveData;
+		bool bNeedsUpdating = false;
+		SkeletalMeshComponent->ForEachBodyBelow(BodyName, bIncludeSelf, /*bSkipCustomType=*/false, [bClearNotFound, ProfileName, PhysAsset, &NewDriveData, &bNeedsUpdating](const FBodyInstance* BI)
+		{
+			USkeletalBodySetup* BodySetup = Cast<USkeletalBodySetup>(BI->BodySetup.Get());
+			if(ensure(BodySetup))
+			{
+				const FName IterBodyName = PhysAsset->SkeletalBodySetups[BI->InstanceBodyIndex]->BoneName;
+				if(FPhysicalAnimationProfile* Profile = BodySetup->FindPhysicalAnimationProfile(ProfileName))
+				{
+					bNeedsUpdating |= UpdatePhysicalAnimationSettings(IterBodyName, Profile->PhysicalAnimationData, NewDriveData, *PhysAsset);
+				}
+				else if(bClearNotFound)
+				{
+					bNeedsUpdating |= UpdatePhysicalAnimationSettings(IterBodyName, FPhysicalAnimationData(), NewDriveData, *PhysAsset);
+				}
+			}
+			
+		});
+
+		if (bNeedsUpdating)
+		{
+			UpdatePhysicsEngine();
+		}
+	}
+}
+
+FTransform ComputeWorldSpaceTargetTM(const USkeletalMeshComponent& SkeletalMeshComponent, const TArray<FTransform>& SpaceBases, int32 BoneIndex)
+{
+	return SpaceBases[BoneIndex] * SkeletalMeshComponent.GetComponentToWorld();
+}
+
+FTransform ComputeLocalSpaceTargetTM(const USkeletalMeshComponent& SkeletalMeshComponent, const UPhysicsAsset& PhysAsset, int32 BoneIndex)
+{
+	const FReferenceSkeleton& RefSkeleton = SkeletalMeshComponent.SkeletalMesh->RefSkeleton;
+	FTransform AccumulatedDelta = SkeletalMeshComponent.LocalAtoms[BoneIndex];
+	int32 CurBoneIdx = BoneIndex;
+	while ((CurBoneIdx = RefSkeleton.GetParentIndex(CurBoneIdx)) != INDEX_NONE)
+	{
+		FName BoneName = RefSkeleton.GetBoneName(CurBoneIdx);
+		int32 BodyIndex = PhysAsset.FindBodyIndex(BoneName);
+
+		if (CurBoneIdx == BoneIndex)	//some kind of loop so just stop TODO: warn?
+		{
+			break;
+		}
+
+		if (BodyIndex != INDEX_NONE)
+		{
+
+			FBodyInstance* ParentBody = SkeletalMeshComponent.Bodies[BodyIndex];
+			const FTransform NewWorldTM = AccumulatedDelta * ParentBody->GetUnrealWorldTransform_AssumesLocked();
+			return NewWorldTM;
+		}
+
+		AccumulatedDelta = AccumulatedDelta * SkeletalMeshComponent.LocalAtoms[CurBoneIdx];
+	}
+
+	return FTransform::Identity;
+}
+
+FTransform ComputeTargetTM(const FPhysicalAnimationData& PhysAnimData, const USkeletalMeshComponent& SkeletalMeshComponent, const UPhysicsAsset& PhysAsset, const TArray<FTransform>& SpaceBases, int32 BoneIndex)
+{
+	return PhysAnimData.bIsLocalSimulation ? ComputeLocalSpaceTargetTM(SkeletalMeshComponent, PhysAsset, BoneIndex) : ComputeWorldSpaceTargetTM(SkeletalMeshComponent, SpaceBases, BoneIndex);
+}
+
+void UPhysicalAnimationComponent::TickComponent(float DeltaTime, enum ELevelTick TickType, FActorComponentTickFunction *ThisTickFunction)
+{
+	UPhysicsAsset* PhysAsset = SkeletalMeshComponent ? SkeletalMeshComponent->GetPhysicsAsset() : nullptr;
+	if (PhysAsset && SkeletalMeshComponent->SkeletalMesh)
+	{
+		const FReferenceSkeleton& RefSkeleton = SkeletalMeshComponent->SkeletalMesh->RefSkeleton;
+		const TArray<FTransform>& SpaceBases = SkeletalMeshComponent->GetEditableSpaceBases();
+
+#if WITH_PHYSX
+		if (PxScene* Scene = GetPhysXSceneFromIndex(SceneIndex))
+		{
+			SCOPED_SCENE_WRITE_LOCK(Scene);
+			for (int32 DataIdx = 0; DataIdx < DriveData.Num(); ++DataIdx)
+			{
+				const FPhysicalAnimationData& PhysAnimData = DriveData[DataIdx];
+				FPhysicalAnimationInstanceData& InstanceData = RuntimeInstanceData[DataIdx];
+				if(PxRigidDynamic* TargetActor = InstanceData.TargetActor)
+				{
+					const int32 BoneIdx = RefSkeleton.FindBoneIndex(PhysAnimData.BodyName);
+					check(BoneIdx != INDEX_NONE);	//If we've created the actor the bone name should exist
+
+					const FTransform TargetTM =	ComputeTargetTM(PhysAnimData, *SkeletalMeshComponent, *PhysAsset, SpaceBases, BoneIdx);
+					TargetActor->setKinematicTarget(U2PTransform(TargetTM));	//TODO: this doesn't work with sub-stepping!
+				}
+			}
+		}
+#endif
+	}
+}
+
+//NOTE: Technically skeletal mesh component could have bodies in multiple scenes. This doesn't seem like a legit setup though and we should probably enforce that it's not supported.
+int32 FindSceneIndexForSkeletalMeshComponent(const USkeletalMeshComponent* SkeletalMeshComp)
+{
+	for(FBodyInstance* BI : SkeletalMeshComp->Bodies)
+	{
+		return BI->GetSceneIndex();
+	}
+
+	return INDEX_NONE;
+}
+
+
+void UPhysicalAnimationComponent::UpdatePhysicsEngine()
+{
+	UPhysicsAsset* PhysAsset = SkeletalMeshComponent ? SkeletalMeshComponent->GetPhysicsAsset() : nullptr;
+	if(PhysAsset && SkeletalMeshComponent->SkeletalMesh)
+	{
+		//TODO: This is hacky and assumes constraints can only be added and not removed. True for now, but bad in general!
+		const int32 NumData = DriveData.Num();
+		const int32 NumInstances = RuntimeInstanceData.Num();
+		
+		RuntimeInstanceData.AddZeroed(NumData - NumInstances);
+		SceneIndex = FindSceneIndexForSkeletalMeshComponent(SkeletalMeshComponent);
+
+		const TArray<FTransform>& SpaceBases = SkeletalMeshComponent->GetEditableSpaceBases();
+		const FReferenceSkeleton& RefSkeleton = SkeletalMeshComponent->SkeletalMesh->RefSkeleton;
+#if WITH_PHYSX
+		if(PxScene* Scene = GetPhysXSceneFromIndex(SceneIndex))
+		{
+			SCOPED_SCENE_WRITE_LOCK(Scene);
+			for (int32 DataIdx = 0; DataIdx < DriveData.Num(); ++DataIdx)
+			{
+				bool bNewConstraint = false;
+				const FPhysicalAnimationData& PhysAnimData = DriveData[DataIdx];
+				FPhysicalAnimationInstanceData& InstanceData = RuntimeInstanceData[DataIdx];
+				FConstraintInstance*& ConstraintInstance = InstanceData.ConstraintInstance;
+				if(ConstraintInstance == nullptr)
+				{
+					bNewConstraint = true;
+					ConstraintInstance = new FConstraintInstance();
+					ConstraintInstance->ProfileInstance = PhysicalAnimationProfile;
+				}
+
+				//Apply drive forces
+				ConstraintInstance->SetAngularDriveParams(PhysAnimData.OrientationStrength, PhysAnimData.AngularVelocityStrength, PhysAnimData.MaxAngularForce);
+				if(PhysAnimData.bIsLocalSimulation)	//linear only works for world space simulation
+				{
+					ConstraintInstance->SetLinearDriveParams(0.f, 0.f, 0.f);
+				}
+				else
+				{
+					ConstraintInstance->SetLinearDriveParams(PhysAnimData.PositionStrength, PhysAnimData.VelocityStrength, PhysAnimData.MaxLinearForce);
+				}
+
+				if(bNewConstraint)
+				{
+					//Find body instances
+					int32 ChildBodyIdx = PhysAsset->FindBodyIndex(PhysAnimData.BodyName);
+					if (FBodyInstance* ChildBody = (ChildBodyIdx == INDEX_NONE ? nullptr : SkeletalMeshComponent->Bodies[ChildBodyIdx]))
+					{
+						if (PxRigidActor* PRigidActor = ChildBody->GetPxRigidActor_AssumesLocked())
+						{
+							ConstraintInstance->SetRefFrame(EConstraintFrame::Frame1, FTransform::Identity);
+							ConstraintInstance->SetRefFrame(EConstraintFrame::Frame2, FTransform::Identity);
+
+							const FTransform TargetTM = ComputeTargetTM(PhysAnimData, *SkeletalMeshComponent, *PhysAsset, SpaceBases, ChildBody->InstanceBoneIndex);
+
+							// Create kinematic actor we are going to create joint with. This will be moved around with calls to SetLocation/SetRotation.
+							PxRigidDynamic* KineActor = Scene->getPhysics().createRigidDynamic(U2PTransform(TargetTM));
+							KineActor->setRigidDynamicFlag(PxRigidDynamicFlag::eKINEMATIC, true);
+							KineActor->setMass(1.0f);
+							KineActor->setMassSpaceInertiaTensor(PxVec3(1.0f, 1.0f, 1.0f));
+
+							// No bodyinstance
+							KineActor->userData = NULL;
+
+							// Add to Scene
+							Scene->addActor(*KineActor);
+
+							// Save reference to the kinematic actor.
+							InstanceData.TargetActor = KineActor;
+
+							ConstraintInstance->InitConstraintPhysX_AssumesLocked(PRigidActor, InstanceData.TargetActor, Scene, 1.f);
+						}
+					}
+				}
+			}
+		}
+#endif
+	}
+}
+
+void UPhysicalAnimationComponent::ReleasePhysicsEngine()
+{
+#if WITH_PHYSX
+	PxScene* Scene = GetPhysXSceneFromIndex(SceneIndex);
+	SCOPED_SCENE_WRITE_LOCK(Scene);
+
+	for (FPhysicalAnimationInstanceData& Instance : RuntimeInstanceData)
+	{
+		if(Instance.ConstraintInstance)
+		{
+			Instance.ConstraintInstance->TermConstraint();
+			delete Instance.ConstraintInstance;
+			Instance.ConstraintInstance = nullptr;
+		}
+		
+		if(Instance.TargetActor)
+		{
+			Scene->removeActor(*Instance.TargetActor);
+			Instance.TargetActor->release();
+
+			Instance.TargetActor = nullptr;
+		}
+	}
+
+	RuntimeInstanceData.Reset();
+#endif
+}
+
+#if WITH_EDITOR
+static const FColor	TargetActorColor(255, 0, 0);
+
+void UPhysicalAnimationComponent::DebugDraw(FPrimitiveDrawInterface* PDI) const
+{
+	for (const FPhysicalAnimationInstanceData& PhysAnimData : RuntimeInstanceData)
+	{
+#if WITH_PHYSX
+		if (PhysAnimData.TargetActor)
+		{
+			PDI->DrawPoint(P2UVector(PhysAnimData.TargetActor->getGlobalPose().p), TargetActorColor, 3.f, SDPG_World);
+		}
+#endif
+
+	}
+}
+#endif
\ No newline at end of file
diff --git a/Engine/Source/Runtime/Engine/Private/PhysicsEngine/PhysicsAsset.cpp b/Engine/Source/Runtime/Engine/Private/PhysicsEngine/PhysicsAsset.cpp
--- a/Engine/Source/Runtime/Engine/Private/PhysicsEngine/PhysicsAsset.cpp
+++ b/Engine/Source/Runtime/Engine/Private/PhysicsEngine/PhysicsAsset.cpp
@@ -6,6 +6,7 @@
 
 #include "EnginePrivate.h"
 #include "PhysicsEngine/PhysicsConstraintTemplate.h"
+#include "FrameworkObjectVersion.h"
 
 
 #if WITH_PHYSX
@@ -27,25 +28,83 @@ void UPhysicsAsset::UpdateBoundsBodiesArray()
 {
 	BoundsBodies.Empty();
 
-	for(int32 i=0; i<BodySetup.Num(); i++)
+	for(int32 i=0; i<SkeletalBodySetups.Num(); i++)
 	{
-		check(BodySetup[i]);
-		if(BodySetup[i]->bConsiderForBounds)
+		check(SkeletalBodySetups[i]);
+		if(SkeletalBodySetups[i]->bConsiderForBounds)
 		{
 			BoundsBodies.Add(i);
 		}
 	}
 }
 
+#if WITH_EDITOR
+void USkeletalBodySetup::PostEditChangeProperty(struct FPropertyChangedEvent& PropertyChangedEvent)
+{
+	UPhysicsAsset* OwningPhysAsset = Cast<UPhysicsAsset>(GetOuter());
+
+	if(PropertyChangedEvent.Property == nullptr || !OwningPhysAsset)
+	{
+		return;
+	}
+
+	if (FPhysicalAnimationProfile* PhysProfile = FindPhysicalAnimationProfile(OwningPhysAsset->CurrentPhysicalAnimationProfileName))
+	{
+			//changed any setting so copy dummy UI into profile location
+			PhysProfile->PhysicalAnimationData = CurrentPhysicalAnimationProfile.PhysicalAnimationData;
+	}
+}
+
+FName USkeletalBodySetup::GetCurrentPhysicalAnimationProfileName() const
+{
+	FName CurrentProfileName;
+	if(UPhysicsAsset* OwningPhysAsset = Cast<UPhysicsAsset>(GetOuter()))
+	{
+		CurrentProfileName = OwningPhysAsset->CurrentPhysicalAnimationProfileName;
+	}
+
+	return CurrentProfileName;
+}
+
+void USkeletalBodySetup::AddPhysicalAnimationProfile(FName ProfileName)
+{
+	FPhysicalAnimationProfile* NewProfile = new (PhysicalAnimationData) FPhysicalAnimationProfile();
+	NewProfile->ProfileName = ProfileName;
+}
+
+void USkeletalBodySetup::RemovePhysicalAnimationProfile(FName ProfileName)
+{
+	for(int32 ProfileIdx = 0; ProfileIdx < PhysicalAnimationData.Num(); ++ProfileIdx)
+	{
+		if(PhysicalAnimationData[ProfileIdx].ProfileName == ProfileName)
+		{
+			PhysicalAnimationData.RemoveAtSwap(ProfileIdx--);
+		}
+	}
+}
+
+void USkeletalBodySetup::UpdatePhysicalAnimationProfiles(const TArray<FName>& Profiles)
+{
+	//TODO: rename code. At the moment just deletes
+	for(int32 ProfileIdx = 0; ProfileIdx < PhysicalAnimationData.Num(); ++ProfileIdx)
+	{
+		if(Profiles.Contains(PhysicalAnimationData[ProfileIdx].ProfileName) == false)
+		{
+			PhysicalAnimationData.RemoveAtSwap(ProfileIdx--);
+		}
+	}
+}
+#endif
+
 void UPhysicsAsset::UpdateBodySetupIndexMap()
 {
 	// update BodySetupIndexMap
 	BodySetupIndexMap.Empty();
 
-	for(int32 i=0; i<BodySetup.Num(); i++)
+	for(int32 i=0; i<SkeletalBodySetups.Num(); i++)
 	{
-		check(BodySetup[i]);
-		BodySetupIndexMap.Add(BodySetup[i]->BoneName, i);
+		check(SkeletalBodySetups[i]);
+		BodySetupIndexMap.Add(SkeletalBodySetups[i]->BoneName, i);
 	}
 }
 
@@ -53,13 +112,29 @@ void UPhysicsAsset::PostLoad()
 {
 	Super::PostLoad();
 
+	if(GetLinkerCustomVersion(FFrameworkObjectVersion::GUID) < FFrameworkObjectVersion::PhysAssetUseSkeletalBodySetup)
+	{
+		SkeletalBodySetups.AddUninitialized(BodySetup_DEPRECATED.Num());
+		for(int32 Idx = 0; Idx < BodySetup_DEPRECATED.Num(); ++Idx)
+		{
+			//Used to use BodySetup for physics asset, but moved to more specialized SkeletalBodySetup
+			SkeletalBodySetups[Idx] = NewObject<USkeletalBodySetup>(this, NAME_None);
+			
+			TArray<uint8> OldData;
+			FObjectWriter ObjWriter(BodySetup_DEPRECATED[Idx], OldData);
+			FObjectReader ObjReader(SkeletalBodySetups[Idx], OldData);
+		}
+		
+		BodySetup_DEPRECATED.Empty();
+	}
+
 	// Ensure array of bounds bodies is up to date.
-	if(BoundsBodies.Num() == 0)
+	if(SkeletalBodySetups.Num() == 0)
 	{
 		UpdateBoundsBodiesArray();
 	}
 
-	if (BodySetup.Num() > 0 && BodySetupIndexMap.Num() == 0)
+	if (SkeletalBodySetups.Num() > 0 && BodySetupIndexMap.Num() == 0)
 	{
 		UpdateBodySetupIndexMap();
 	}
@@ -133,8 +208,8 @@ FBox UPhysicsAsset::CalcAABB(const USkinnedMeshComponent* MeshComp, const FTrans
 		// If we want to consider all bodies, make array with all body indices in
 		if(MeshComp->bConsiderAllBodiesForBounds)
 		{
-			AllBodies.AddUninitialized(BodySetup.Num());
-			for(int32 i=0; i<BodySetup.Num();i ++)
+			AllBodies.AddUninitialized(SkeletalBodySetups.Num());
+			for(int32 i=0; i<SkeletalBodySetups.Num();i ++)
 			{
 				AllBodies[i] = i;
 			}
@@ -152,7 +227,7 @@ FBox UPhysicsAsset::CalcAABB(const USkinnedMeshComponent* MeshComp, const FTrans
 		for(int32 i=0; i<BodySetupNum; i++)
 		{
 			const int32 BodyIndex = (*BodyIndexRefs)[i];
-			UBodySetup* bs = BodySetup[BodyIndex];
+			UBodySetup* bs = SkeletalBodySetups[BodyIndex];
 
 			// Check if setup should be considered for bounds, or if all bodies should be considered anyhow
 			if (bs->bConsiderForBounds || MeshComp->bConsiderAllBodiesForBounds)
@@ -160,8 +235,8 @@ FBox UPhysicsAsset::CalcAABB(const USkinnedMeshComponent* MeshComp, const FTrans
 				if (i+1<BodySetupNum)
 				{
 					int32 NextIndex = (*BodyIndexRefs)[i+1];
-					FPlatformMisc::Prefetch(BodySetup[NextIndex]);
-					FPlatformMisc::Prefetch(BodySetup[NextIndex], PLATFORM_CACHE_LINE_SIZE);
+					FPlatformMisc::Prefetch(SkeletalBodySetups[NextIndex]);
+					FPlatformMisc::Prefetch(SkeletalBodySetups[NextIndex], PLATFORM_CACHE_LINE_SIZE);
 				}
 
 				int32 BoneIndex = MeshComp->GetBoneIndex(bs->BoneName);
@@ -305,9 +380,9 @@ void UPhysicsAsset::GetBodyIndicesBelow(TArray<int32>& OutBodyIndices, FName InB
 	int32 BaseIndex = SkelMesh->RefSkeleton.FindBoneIndex(InBoneName);
 
 	// Iterate over all other bodies, looking for 'children' of this one
-	for(int32 i=0; i<BodySetup.Num(); i++)
+	for(int32 i=0; i<SkeletalBodySetups.Num(); i++)
 	{
-		UBodySetup* BS = BodySetup[i];
+		UBodySetup* BS = SkeletalBodySetups[i];
 		FName TestName = BS->BoneName;
 		int32 TestIndex = SkelMesh->RefSkeleton.FindBoneIndex(TestName);
 
@@ -325,7 +400,7 @@ void UPhysicsAsset::GetNearestBodyIndicesBelow(TArray<int32> & OutBodyIndices, F
 
 	//we need to filter all bodies below to first in the chain
 	TArray<bool> Nearest;
-	Nearest.AddUninitialized(BodySetup.Num());
+	Nearest.AddUninitialized(SkeletalBodySetups.Num());
 	for (int32 i = 0; i < Nearest.Num(); ++i)
 	{
 		Nearest[i] = true;
@@ -336,7 +411,7 @@ void UPhysicsAsset::GetNearestBodyIndicesBelow(TArray<int32> & OutBodyIndices, F
 		int32 BodyIndex = AllBodiesBelow[i];
 		if (Nearest[BodyIndex] == false) continue;
 
-		UBodySetup * Body = BodySetup[BodyIndex];
+		UBodySetup * Body = SkeletalBodySetups[BodyIndex];
 		TArray<int32> BodiesBelowMe;
 		GetBodyIndicesBelow(BodiesBelowMe, Body->BoneName, InSkelMesh, false);
 		
@@ -358,9 +433,9 @@ void UPhysicsAsset::GetNearestBodyIndicesBelow(TArray<int32> & OutBodyIndices, F
 
 void UPhysicsAsset::ClearAllPhysicsMeshes()
 {
-	for(int32 i=0; i<BodySetup.Num(); i++)
+	for(int32 i=0; i<SkeletalBodySetups.Num(); i++)
 	{
-		BodySetup[i]->ClearPhysicsMeshes();
+		SkeletalBodySetups[i]->ClearPhysicsMeshes();
 	}
 }
 
@@ -368,9 +443,9 @@ void UPhysicsAsset::ClearAllPhysicsMeshes()
 
 void UPhysicsAsset::InvalidateAllPhysicsMeshes()
 {
-	for(int32 i=0; i<BodySetup.Num(); i++)
+	for(int32 i=0; i<SkeletalBodySetups.Num(); i++)
 	{
-		BodySetup[i]->InvalidatePhysicsData();
+		SkeletalBodySetups[i]->InvalidatePhysicsData();
 	}
 }
 
@@ -383,6 +458,24 @@ void UPhysicsAsset::PostEditUndo()
 	UpdateBoundsBodiesArray();
 }
 
+void UPhysicsAsset::PostEditChangeProperty(struct FPropertyChangedEvent& PropertyChangedEvent)
+{
+	if(UProperty* Property = PropertyChangedEvent.Property)
+	{
+		if(Profiles.Find(CurrentPhysicalAnimationProfileName) == INDEX_NONE)
+		{
+			CurrentPhysicalAnimationProfileName = NAME_None;
+		}
+
+		for(USkeletalBodySetup* BS : SkeletalBodySetups)
+		{
+			BS->UpdatePhysicalAnimationProfiles(Profiles);
+		}
+	}
+
+	Super::PostEditChangeProperty(PropertyChangedEvent);
+}
+
 #endif
 
 //// THUMBNAIL SUPPORT //////
@@ -392,12 +485,12 @@ void UPhysicsAsset::PostEditUndo()
  */
 FString UPhysicsAsset::GetDesc()
 {
-	return FString::Printf( TEXT("%d Bodies, %d Constraints"), BodySetup.Num(), ConstraintSetup.Num() );
+	return FString::Printf( TEXT("%d Bodies, %d Constraints"), SkeletalBodySetups.Num(), ConstraintSetup.Num() );
 }
 
 void UPhysicsAsset::GetAssetRegistryTags(TArray<FAssetRegistryTag>& OutTags) const
 {
-	OutTags.Add( FAssetRegistryTag("Bodies", FString::FromInt(BodySetup.Num()), FAssetRegistryTag::TT_Numerical) );
+	OutTags.Add( FAssetRegistryTag("Bodies", FString::FromInt(SkeletalBodySetups.Num()), FAssetRegistryTag::TT_Numerical) );
 	OutTags.Add( FAssetRegistryTag("Constraints", FString::FromInt(ConstraintSetup.Num()), FAssetRegistryTag::TT_Numerical) );
 
 	Super::GetAssetRegistryTags(OutTags);
@@ -407,7 +500,7 @@ void UPhysicsAsset::GetAssetRegistryTags(TArray<FAssetRegistryTag>& OutTags) con
 void UPhysicsAsset::BodyFindConstraints(int32 BodyIndex, TArray<int32>& Constraints)
 {
 	Constraints.Empty();
-	FName BodyName = BodySetup[BodyIndex]->BoneName;
+	FName BodyName = SkeletalBodySetups[BodyIndex]->BoneName;
 
 	for(int32 ConIdx=0; ConIdx<ConstraintSetup.Num(); ConIdx++)
 	{
@@ -422,7 +515,7 @@ SIZE_T UPhysicsAsset::GetResourceSize(EResourceSizeMode::Type Mode)
 {
 	SIZE_T ResourceSize = Super::GetResourceSize(Mode);
 
-	for (const auto& SingleBody : BodySetup)
+	for (const auto& SingleBody : SkeletalBodySetups)
 	{
 		ResourceSize += SingleBody->GetResourceSize(Mode);
 	}
diff --git a/Engine/Source/Runtime/Engine/Private/PhysicsEngine/PhysicsConstraintComponent.cpp b/Engine/Source/Runtime/Engine/Private/PhysicsEngine/PhysicsConstraintComponent.cpp
--- a/Engine/Source/Runtime/Engine/Private/PhysicsEngine/PhysicsConstraintComponent.cpp
+++ b/Engine/Source/Runtime/Engine/Private/PhysicsEngine/PhysicsConstraintComponent.cpp
@@ -132,7 +132,7 @@ FBox UPhysicsConstraintComponent::GetBodyBoxInternal(EConstraintFrame::Type Fram
 			if(BoneIndex != INDEX_NONE && BodyIndex != INDEX_NONE)
 			{	
 				const FTransform BoneTransform = SkelComp->GetBoneTransform(BoneIndex);
-				ResultBox = PhysicsAsset->BodySetup[BodyIndex]->AggGeom.CalcAABB(BoneTransform);
+				ResultBox = PhysicsAsset->SkeletalBodySetups[BodyIndex]->AggGeom.CalcAABB(BoneTransform);
 			}
 		}
 	}
@@ -187,6 +187,12 @@ FBodyInstance* UPhysicsConstraintComponent::GetBodyInstance(EConstraintFrame::Ty
 }
 
 
+/** Wrapper that calls our constraint broken delegate */
+void UPhysicsConstraintComponent::OnConstraintBrokenWrapper(int32 ConstraintIndex)
+{
+	OnConstraintBroken.Broadcast(ConstraintIndex);
+}
+
 void UPhysicsConstraintComponent::InitComponentConstraint()
 {
 	// First we convert world space position of constraint into local space frames
@@ -195,7 +201,8 @@ void UPhysicsConstraintComponent::InitComponentConstraint()
 	// Then we init the constraint
 	FBodyInstance* Body1 = GetBodyInstance(EConstraintFrame::Frame1);
 	FBodyInstance* Body2 = GetBodyInstance(EConstraintFrame::Frame2);
-	ConstraintInstance.InitConstraint(this, Body1, Body2, GetConstraintScale());
+
+	ConstraintInstance.InitConstraint(Body1, Body2, GetConstraintScale(), this, FOnConstraintBroken::CreateUObject(this, &UPhysicsConstraintComponent::OnConstraintBrokenWrapper));
 }
 
 void UPhysicsConstraintComponent::TermComponentConstraint()
