Copying //UE4/Dev-Framework to //UE4/Dev-Main (Source: //UE4/Dev-Framework @ 3038004)

diff --git a/Engine/Source/Runtime/Engine/Private/PlayerCameraManager.cpp b/Engine/Source/Runtime/Engine/Private/PlayerCameraManager.cpp
--- a/Engine/Source/Runtime/Engine/Private/PlayerCameraManager.cpp
+++ b/Engine/Source/Runtime/Engine/Private/PlayerCameraManager.cpp
@@ -548,7 +548,7 @@ void APlayerCameraManager::ApplyAudioFade()
 		{
 			if (FAudioDevice* AudioDevice = World->GetAudioDevice())
 			{
-				AudioDevice->TransientMasterVolume = 1.0f - FadeAmount;
+				AudioDevice->SetTransientMasterVolume(1.0f - FadeAmount);
 			}
 		}
 	}
@@ -563,7 +563,7 @@ void APlayerCameraManager::StopAudioFade()
 		{
 			if (FAudioDevice* AudioDevice = World->GetAudioDevice())
 			{
-				AudioDevice->TransientMasterVolume = 1.0f;
+				AudioDevice->SetTransientMasterVolume(1.0f);
 			}
 		}
 	}
diff --git a/Engine/Source/Runtime/Engine/Private/PreviewScene.cpp b/Engine/Source/Runtime/Engine/Private/PreviewScene.cpp
--- a/Engine/Source/Runtime/Engine/Private/PreviewScene.cpp
+++ b/Engine/Source/Runtime/Engine/Private/PreviewScene.cpp
@@ -123,10 +123,7 @@ void FPreviewScene::RemoveComponent(UActorComponent* Component)
 
 void FPreviewScene::AddReferencedObjects( FReferenceCollector& Collector )
 {
-	for( int32 Index = 0; Index < Components.Num(); Index++ )
-	{
-		Collector.AddReferencedObject( Components[ Index ] );
-	}
+	Collector.AddReferencedObjects( Components );
 	Collector.AddReferencedObject( DirectionalLight );
 	Collector.AddReferencedObject( PreviewWorld );
 }
diff --git a/Engine/Source/Runtime/Engine/Private/SkeletalMesh.cpp b/Engine/Source/Runtime/Engine/Private/SkeletalMesh.cpp
--- a/Engine/Source/Runtime/Engine/Private/SkeletalMesh.cpp
+++ b/Engine/Source/Runtime/Engine/Private/SkeletalMesh.cpp
@@ -4804,9 +4804,9 @@ FSkeletalMeshSceneProxy::FSkeletalMeshSceneProxy(const USkinnedMeshComponent* Co
 			&& SkeletalMeshComponent->CastShadow
 			&& (SkeletalMeshComponent->bCastCapsuleDirectShadow || SkeletalMeshComponent->bCastCapsuleIndirectShadow))
 		{
-			for (int32 BodyIndex = 0; BodyIndex < ShadowPhysicsAsset->BodySetup.Num(); BodyIndex++)
+			for (int32 BodyIndex = 0; BodyIndex < ShadowPhysicsAsset->SkeletalBodySetups.Num(); BodyIndex++)
 			{
-				UBodySetup* BodySetup = ShadowPhysicsAsset->BodySetup[BodyIndex];
+				UBodySetup* BodySetup = ShadowPhysicsAsset->SkeletalBodySetups[BodyIndex];
 				int32 BoneIndex = SkeletalMeshComponent->GetBoneIndex(BodySetup->BoneName);
 
 				if (BoneIndex != INDEX_NONE)
diff --git a/Engine/Source/Runtime/Engine/Private/SkeletalMeshComponentPhysics.cpp b/Engine/Source/Runtime/Engine/Private/SkeletalMeshComponentPhysics.cpp
--- a/Engine/Source/Runtime/Engine/Private/SkeletalMeshComponentPhysics.cpp
+++ b/Engine/Source/Runtime/Engine/Private/SkeletalMeshComponentPhysics.cpp
@@ -423,7 +423,7 @@ UBodySetup* USkeletalMeshComponent::GetBodySetup()
 				int32 BodyIndex = PhysicsAsset->FindBodyIndex(SkeletalMesh->RefSkeleton.GetBoneName(i));
 				if (BodyIndex != INDEX_NONE)
 				{
-					return PhysicsAsset->BodySetup[BodyIndex];
+					return PhysicsAsset->SkeletalBodySetups[BodyIndex];
 				}
 			}
 		}
@@ -466,7 +466,7 @@ void USkeletalMeshComponent::SetSimulatePhysics(bool bSimulate)
 		{
 			if (FBodyInstance* BodyInst = Bodies[BodyIdx])
 			{
-				if (UBodySetup * PhysAssetBodySetup = PhysAsset->BodySetup[BodyIdx])
+				if (UBodySetup * PhysAssetBodySetup = PhysAsset->SkeletalBodySetups[BodyIdx])
 				{
 					if (PhysAssetBodySetup->PhysicsType == EPhysicsType::PhysType_Default)
 					{
@@ -752,7 +752,7 @@ void USkeletalMeshComponent::SetEnableGravity(bool bGravityEnabled)
 		{
 			if (FBodyInstance* BodyInst = Bodies[BodyIdx])
 			{
-				if (UBodySetup * PhysAssetBodySetup = PhysAsset->BodySetup[BodyIdx])
+				if (UBodySetup * PhysAssetBodySetup = PhysAsset->SkeletalBodySetups[BodyIdx])
 				{
 					bool bUseGravityEnabled = bGravityEnabled;
 					
@@ -774,6 +774,11 @@ bool USkeletalMeshComponent::IsGravityEnabled() const
 	return BodyInstance.bEnableGravity;
 }
 
+void USkeletalMeshComponent::OnConstraintBrokenWrapper(int32 ConstraintIndex)
+{
+	OnConstraintBroken.Broadcast(ConstraintIndex);
+}
+
 DECLARE_CYCLE_STAT(TEXT("Init Articulated"), STAT_InitArticulated, STATGROUP_Physics);
 
 TAutoConsoleVariable<int32> CVarEnableRagdollPhysics(TEXT("p.RagdollPhysics"), 1, TEXT("If 1, ragdoll physics will be used. Otherwise just root body is simulated"));
@@ -821,10 +826,10 @@ void USkeletalMeshComponent::InitArticulated(FPhysScene* PhysScene)
 	uint32 SkelMeshCompID = GetUniqueID();
 	PhysScene->DeferredAddCollisionDisableTable(SkelMeshCompID, &PhysicsAsset->CollisionDisableTable);
 
-	int32 NumBodies = PhysicsAsset->BodySetup.Num();
+	int32 NumBodies = PhysicsAsset->SkeletalBodySetups.Num();
 	if(Aggregate == NULL && NumBodies > RagdollAggregateThreshold && NumBodies <= AggregateMaxSize)
 	{
-		Aggregate = GPhysXSDK->createAggregate(PhysicsAsset->BodySetup.Num(), true);
+		Aggregate = GPhysXSDK->createAggregate(PhysicsAsset->SkeletalBodySetups.Num(), true);
 	}
 	else if(Aggregate && NumBodies > AggregateMaxSize)
 	{
@@ -837,7 +842,7 @@ void USkeletalMeshComponent::InitArticulated(FPhysScene* PhysScene)
 	Bodies.AddZeroed(NumBodies);
 	for(int32 i=0; i<NumBodies; i++)
 	{
-		UBodySetup* PhysicsAssetBodySetup = PhysicsAsset->BodySetup[i];
+		UBodySetup* PhysicsAssetBodySetup = PhysicsAsset->SkeletalBodySetups[i];
 		Bodies[i] = new FBodyInstance;
 		FBodyInstance* BodyInst = Bodies[i];
 		check(BodyInst);
@@ -946,7 +951,7 @@ void USkeletalMeshComponent::InitArticulated(FPhysScene* PhysScene)
 		// If we have 2, joint 'em
 		if(Body1 != NULL && Body2 != NULL)
 		{
-			ConInst->InitConstraint(this, Body1, Body2, Scale);
+			ConInst->InitConstraint(Body1, Body2, Scale, this, FOnConstraintBroken::CreateUObject(this, &USkeletalMeshComponent::OnConstraintBrokenWrapper));
 		}
 	}
 
@@ -1023,7 +1028,7 @@ void USkeletalMeshComponent::TermBodiesBelow(FName ParentBoneName)
 	UPhysicsAsset* const PhysicsAsset = GetPhysicsAsset();
 	if(PhysicsAsset && SkeletalMesh && Bodies.Num() > 0)
 	{
-		check(Bodies.Num() == PhysicsAsset->BodySetup.Num());
+		check(Bodies.Num() == PhysicsAsset->SkeletalBodySetups.Num());
 
 		// Get index of parent bone
 		int32 ParentBoneIndex = GetBoneIndex(ParentBoneName);
@@ -1142,7 +1147,7 @@ void USkeletalMeshComponent::SetAllMotorsAngularPositionDrive(bool bEnableSwingD
 		if( bSkipCustomPhysicsType )
 		{
 			int32 BodyIndex = PhysicsAsset->FindBodyIndex(Constraints[i]->JointName);
-			if( BodyIndex != INDEX_NONE && PhysicsAsset->BodySetup[BodyIndex]->PhysicsType != PhysType_Default)
+			if( BodyIndex != INDEX_NONE && PhysicsAsset->SkeletalBodySetups[BodyIndex]->PhysicsType != PhysType_Default)
 			{
 				continue;
 			}
@@ -1209,7 +1214,7 @@ void USkeletalMeshComponent::SetAllMotorsAngularVelocityDrive(bool bEnableSwingD
 		if( bSkipCustomPhysicsType )
 		{
 			int32 BodyIndex = PhysicsAsset->FindBodyIndex(Constraints[i]->JointName);
-			if( BodyIndex != INDEX_NONE && PhysicsAsset->BodySetup[BodyIndex]->PhysicsType != PhysType_Default )
+			if( BodyIndex != INDEX_NONE && PhysicsAsset->SkeletalBodySetups[BodyIndex]->PhysicsType != PhysType_Default )
 			{
 				continue;
 			}
@@ -1253,7 +1258,7 @@ void USkeletalMeshComponent::SetAllMotorsAngularDriveParams(float InSpring, floa
 		if( bSkipCustomPhysicsType )
 		{
 			int32 BodyIndex = PhysicsAsset->FindBodyIndex(Constraints[i]->JointName);
-			if( BodyIndex != INDEX_NONE && PhysicsAsset->BodySetup[BodyIndex]->PhysicsType != PhysType_Default )
+			if( BodyIndex != INDEX_NONE && PhysicsAsset->SkeletalBodySetups[BodyIndex]->PhysicsType != PhysType_Default )
 			{
 				continue;
 			}
@@ -1520,9 +1525,9 @@ void USkeletalMeshComponent::UpdateMeshForBrokenConstraints()
 			DEBUGBROKENCONSTRAINTUPDATE(UE_LOG(LogSkeletalMesh, Log, TEXT("  Found Broken Constraint: (%d) %s"), JointBoneIndex, *PhysicsAsset->ConstraintSetup(ConstraintInstIndex)->JointName.ToString());)
 
 			// Get child bodies of this joint
-			for(int32 BodySetupIndex = 0; BodySetupIndex < PhysicsAsset->BodySetup.Num(); BodySetupIndex++)
+			for(int32 BodySetupIndex = 0; BodySetupIndex < PhysicsAsset->SkeletalBodySetups.Num(); BodySetupIndex++)
 			{
-				UBodySetup* PhysicsAssetBodySetup = PhysicsAsset->BodySetup[BodySetupIndex];
+				UBodySetup* PhysicsAssetBodySetup = PhysicsAsset->SkeletalBodySetups[BodySetupIndex];
 				int32 BoneIndex = GetBoneIndex(PhysicsAssetBodySetup->BoneName);
 				if( BoneIndex != INDEX_NONE && 
 					(BoneIndex == JointBoneIndex || SkeletalMesh->RefSkeleton.BoneIsChildOf(BoneIndex, JointBoneIndex)) )
@@ -1627,7 +1632,7 @@ void USkeletalMeshComponent::GetWeldedBodies(TArray<FBodyInstance*> & OutWeldedB
 			OutWeldedBodies.Add(&BodyInstance);
 			if (PhysicsAsset)
 			{
-				if (UBodySetup * PhysicsAssetBodySetup = PhysicsAsset->BodySetup[BodyIdx])
+				if (UBodySetup * PhysicsAssetBodySetup = PhysicsAsset->SkeletalBodySetups[BodyIdx])
 				{
 					OutLabels.Add(PhysicsAssetBodySetup->BoneName);
 				}
@@ -1688,7 +1693,7 @@ int32 USkeletalMeshComponent::ForEachBodyBelow(FName BoneName, bool bIncludeSelf
 			FBodyInstance* BI = Bodies[BodyIndices[BodyIdx]];
 			if (bSkipCustomType)
 			{
-				if (UBodySetup* PhysAssetBodySetup = PhysicsAsset->BodySetup[BodyIdx])
+				if (UBodySetup* PhysAssetBodySetup = PhysicsAsset->SkeletalBodySetups[BodyIdx])
 				{
 					if (PhysAssetBodySetup->PhysicsType != EPhysicsType::PhysType_Default)
 					{
@@ -1853,10 +1858,10 @@ void USkeletalMeshComponent::UpdateHasValidBodies()
 	if(PhysicsAsset != NULL)
 	{
 		// For each body in physics asset..
-		for( int32 BodyIndex = 0; BodyIndex < PhysicsAsset->BodySetup.Num(); BodyIndex++ )
+		for( int32 BodyIndex = 0; BodyIndex < PhysicsAsset->SkeletalBodySetups.Num(); BodyIndex++ )
 		{
 			// .. find the matching graphics bone index
-			int32 BoneIndex = GetBoneIndex( PhysicsAsset->BodySetup[ BodyIndex ]->BoneName );
+			int32 BoneIndex = GetBoneIndex( PhysicsAsset->SkeletalBodySetups[ BodyIndex ]->BoneName );
 
 			// If we found a valid graphics bone, set the 'valid' flag
 			if(BoneIndex != INDEX_NONE)
@@ -1868,6 +1873,26 @@ void USkeletalMeshComponent::UpdateHasValidBodies()
 	}
 }
 
+void USkeletalMeshComponent::UpdateBoneBodyMapping()
+{
+	// If we have a physics asset..
+	if (const UPhysicsAsset* const PhysicsAsset = GetPhysicsAsset())
+	{
+		// For each body in physics asset..
+		for (int32 BodyIndex = 0; BodyIndex < PhysicsAsset->SkeletalBodySetups.Num(); BodyIndex++)
+		{
+			// .. find the matching graphics bone index
+			int32 BoneIndex = GetBoneIndex(PhysicsAsset->SkeletalBodySetups[BodyIndex]->BoneName);
+			Bodies[BodyIndex]->InstanceBoneIndex = BoneIndex;
+
+			if (BoneIndex == INDEX_NONE)
+			{
+				//TODO: warn that body was found without corresponding bone
+			}
+		}
+	}
+}
+
 void USkeletalMeshComponent::UpdatePhysicsToRBChannels()
 {
 	// Iterate over each bone/body.
@@ -2016,7 +2041,7 @@ bool USkeletalMeshComponent::GetClosestPointOnPhysicsAsset(const FVector& WorldP
 		int32 CurrentClosestBoneIndex = INDEX_NONE;
 		const UBodySetup* CurrentClosestBodySetup = nullptr;
 
-		for(const UBodySetup* BodySetupInstance : PhysicsAsset->BodySetup)
+		for(const UBodySetup* BodySetupInstance : PhysicsAsset->SkeletalBodySetups)
 		{
 			ClosestPointOnPhysicsAsset.Distance = FLT_MAX;
 			const FName BoneName = BodySetupInstance->BoneName;
diff --git a/Engine/Source/Runtime/Engine/Private/SkeletalMeshMerge.cpp b/Engine/Source/Runtime/Engine/Private/SkeletalMeshMerge.cpp
--- a/Engine/Source/Runtime/Engine/Private/SkeletalMeshMerge.cpp
+++ b/Engine/Source/Runtime/Engine/Private/SkeletalMeshMerge.cpp
@@ -25,12 +25,14 @@ FSkeletalMeshMerge::FSkeletalMeshMerge(USkeletalMesh* InMergeMesh,
 									   const TArray<USkeletalMesh*>& InSrcMeshList, 
 									   const TArray<FSkelMeshMergeSectionMapping>& InForceSectionMapping,
 									   int32 InStripTopLODs,
-                                       EMeshBufferAccess InMeshBufferAccess)
+                                       EMeshBufferAccess InMeshBufferAccess,
+									   FSkelMeshMergeUVTransforms* InSectionUVTransforms)
 :	MergeMesh(InMergeMesh)
 ,	SrcMeshList(InSrcMeshList)
 ,	StripTopLODs(InStripTopLODs)
 ,   MeshBufferAccess(InMeshBufferAccess)
 ,	ForceSectionMapping(InForceSectionMapping)
+,	SectionUVTransforms(InSectionUVTransforms)
 {
 	check(MergeMesh);
 }
@@ -324,10 +326,17 @@ void FSkeletalMeshMerge::GenerateNewSectionArray( TArray<FNewSectionInfo>& NewSe
 						// check to see if the newly merged bonemap is still within the bone limit for GPU skinning
 						if( TempMergedBoneMap.Num() <= MaxGPUSkinBones )
 						{
+							TArray<FTransform> SrcUVTransform;
+							if (SectionUVTransforms != nullptr && MeshIdx < SectionUVTransforms->UVTransformsPerMesh.Num())
+							{
+								SrcUVTransform = SectionUVTransforms->UVTransformsPerMesh[MeshIdx];
+							}
+
 							// add the source section as a new merge entry
 							FMergeSectionInfo& MergeSectionInfo = *new(NewSectionInfo.MergeSections) FMergeSectionInfo(
 								SrcMesh,
-								&SrcLODModel.Sections[SectionIdx]
+								&SrcLODModel.Sections[SectionIdx],
+								SrcUVTransform
 								);
 							// keep track of remapping for the existing chunk's bonemap 
 							// so that the bone matrix indices can be updated for the vertices
@@ -351,10 +360,17 @@ void FSkeletalMeshMerge::GenerateNewSectionArray( TArray<FNewSectionInfo>& NewSe
 					FNewSectionInfo& NewSectionInfo = *new(NewSectionArray) FNewSectionInfo(Material,MaterialId);
 					// initialize the merged bonemap to simply use the original chunk bonemap
 					NewSectionInfo.MergedBoneMap = DestChunkBoneMap;
+
+					TArray<FTransform> SrcUVTransform;
+					if (SectionUVTransforms != nullptr && MeshIdx < SectionUVTransforms->UVTransformsPerMesh.Num())
+					{
+						SrcUVTransform = SectionUVTransforms->UVTransformsPerMesh[MeshIdx];
+					}
 					// add a new merge section entry
 					FMergeSectionInfo& MergeSectionInfo = *new(NewSectionInfo.MergeSections) FMergeSectionInfo(
 						SrcMesh,
-						&SrcLODModel.Sections[SectionIdx]);
+						&SrcLODModel.Sections[SectionIdx],
+						SrcUVTransform);
 					// since merged bonemap == chunk.bonemap then remapping is just pass-through
 					MergeSectionInfo.BoneMapToMergedBoneMap.Empty( DestChunkBoneMap.Num() );
 					for( int32 i=0; i < DestChunkBoneMap.Num(); i++ )
@@ -533,7 +549,13 @@ void FSkeletalMeshMerge::GenerateLODModel( int32 LODIdx )
 				uint32 LODNumTexCoords = SrcLODModel.VertexBufferGPUSkin.GetNumTexCoords();
 				for( uint32 UVIndex = 0; UVIndex < LODNumTexCoords && UVIndex < MAX_TEXCOORDS; ++UVIndex )
 				{
-					DestVert.UVs[UVIndex] = SrcLODModel.VertexBufferGPUSkin.GetVertexUVFast<bExtraBoneInfluencesT>(VertIdx,UVIndex);
+					FVector2D UVs = SrcLODModel.VertexBufferGPUSkin.GetVertexUVFast<bExtraBoneInfluencesT>(VertIdx,UVIndex);
+					if (UVIndex < (uint32)MergeSectionInfo.UVTransforms.Num())
+					{
+						auto Transformed = MergeSectionInfo.UVTransforms[UVIndex].TransformPosition(FVector(UVs, 1.f));
+						UVs = FVector2D(Transformed.X, Transformed.Y);
+					}
+					DestVert.UVs[UVIndex] = UVs;
 				}
 
 				if( TotalNumUVs < LODNumTexCoords )
diff --git a/Engine/Source/Runtime/Engine/Private/Slate/SlateSoundDevice.cpp b/Engine/Source/Runtime/Engine/Private/Slate/SlateSoundDevice.cpp
--- a/Engine/Source/Runtime/Engine/Private/Slate/SlateSoundDevice.cpp
+++ b/Engine/Source/Runtime/Engine/Private/Slate/SlateSoundDevice.cpp
@@ -21,10 +21,10 @@ void FSlateSoundDevice::PlaySound(const FSlateSound& Sound, int32 UserIndex) con
 			if(SoundResource)
 			{
 				FActiveSound NewActiveSound;
-				NewActiveSound.Sound = SoundResource;
+				NewActiveSound.SetSound(SoundResource);
 				NewActiveSound.bIsUISound = true;
 				NewActiveSound.UserIndex = UserIndex;
-				NewActiveSound.Priority = NewActiveSound.Sound->Priority;
+				NewActiveSound.Priority = SoundResource->Priority;
 
 				AudioDevice->AddNewActiveSound(NewActiveSound);
 			}
diff --git a/Engine/Source/Runtime/Engine/Private/SoundBase.cpp b/Engine/Source/Runtime/Engine/Private/SoundBase.cpp
--- a/Engine/Source/Runtime/Engine/Private/SoundBase.cpp
+++ b/Engine/Source/Runtime/Engine/Private/SoundBase.cpp
@@ -81,6 +81,10 @@ bool USoundBase::IsLooping()
 	return (GetDuration() >= INDEFINITELY_LOOPING_DURATION); 
 }
 
+bool USoundBase::ShouldApplyInteriorVolumes() const
+{
+	return (SoundClassObject && SoundClassObject->Properties.bApplyAmbientVolumes);
+}
 
 USoundClass* USoundBase::GetSoundClass() const
 {
diff --git a/Engine/Source/Runtime/Engine/Private/SoundConcurrency.cpp b/Engine/Source/Runtime/Engine/Private/SoundConcurrency.cpp
--- a/Engine/Source/Runtime/Engine/Private/SoundConcurrency.cpp
+++ b/Engine/Source/Runtime/Engine/Private/SoundConcurrency.cpp
@@ -149,7 +149,7 @@ FSoundConcurrencyManager::~FSoundConcurrencyManager()
 
 FActiveSound* FSoundConcurrencyManager::CreateNewActiveSound(const FActiveSound& NewActiveSound)
 {
-	check(NewActiveSound.Sound);
+	check(NewActiveSound.GetSound());
 	
 	// If there are no concurrency settings associated then there is no limit on this sound
 	const FSoundConcurrencySettings* Concurrency = NewActiveSound.GetSoundConcurrencySettingsToApply();
@@ -236,19 +236,16 @@ FActiveSound* FSoundConcurrencyManager::ResolveConcurrency(const FActiveSound& N
 			break;
 
 			case EMaxConcurrentResolutionRule::StopLowestPriority:
+			case EMaxConcurrentResolutionRule::StopLowestPriorityThenPreventNew:
 			{
-				// Find the current lowest priority sound.
-				FActiveSound* Oldest = nullptr;
 				for (FActiveSound* CurrSound : ActiveSounds)
 				{
-					if (SoundToStop == nullptr || CurrSound->GetPriority() < SoundToStop->GetPriority())
+					// This will find oldest and oldest lowest priority sound in the group
+					if (SoundToStop == nullptr 
+						|| (CurrSound->GetPriority() < SoundToStop->GetPriority())
+						|| (CurrSound->GetPriority() == SoundToStop->GetPriority() && CurrSound->PlaybackTime > SoundToStop->PlaybackTime))
 					{
 						SoundToStop = CurrSound;
-						
-					}
-					if (Oldest == nullptr || CurrSound->PlaybackTime > SoundToStop->PlaybackTime)
-					{
-						Oldest = CurrSound;
 					}
 				}
 
@@ -259,12 +256,13 @@ FActiveSound* FSoundConcurrencyManager::ResolveConcurrency(const FActiveSound& N
 					{
 						SoundToStop = nullptr;
 					}
-					// If all sounds are the same priority, then stop the oldest sound
-					else if (SoundToStop->GetPriority() == NewActiveSound.GetPriority())
+					else if (Concurrency->ResolutionRule == EMaxConcurrentResolutionRule::StopLowestPriorityThenPreventNew 
+						&& SoundToStop->GetPriority() == NewActiveSound.GetPriority())
 					{
-						SoundToStop = Oldest;
+						SoundToStop = nullptr;
 					}
 				}
+
 			}
 			break;
 
@@ -286,12 +284,12 @@ FActiveSound* FSoundConcurrencyManager::ResolveConcurrency(const FActiveSound& N
 		if (SoundToStop == nullptr)
 		{
 			// If we didn't find any sound to stop, then we're not going to play this sound, so 
-			// immediatley do the playback complete notify
+			// immediately do the playback complete notify
 			bCanPlay = false;
 
-			if (UAudioComponent* AudioComponent = NewActiveSound.GetAudioComponent())
+			if (NewActiveSound.GetAudioComponentID() > 0)
 			{
-				AudioComponent->PlaybackCompleted(true);
+				UAudioComponent::PlaybackCompleted(NewActiveSound.GetAudioComponentID(), true);
 			}
 		}
 	}
@@ -344,12 +342,12 @@ FActiveSound* FSoundConcurrencyManager::ResolveConcurrency(const FActiveSound& N
 
 FActiveSound* FSoundConcurrencyManager::HandleConcurrencyEvaluationOverride(const FActiveSound& NewActiveSound)
 {
-	check(NewActiveSound.Sound);
+	check(NewActiveSound.GetSound());
 	const FSoundConcurrencySettings* ConcurrencySettings = NewActiveSound.GetSoundConcurrencySettingsToApply();
 	check(ConcurrencySettings);
 
-	uint32 OwnerObjectID = NewActiveSound.TryGetOwnerID();
-	FSoundObjectID SoundObjectID = NewActiveSound.Sound->GetUniqueID();
+	const uint32 OwnerObjectID = NewActiveSound.GetOwnerID();
+	FSoundObjectID SoundObjectID = NewActiveSound.GetSound()->GetUniqueID();
 
 	FActiveSound* ActiveSound = nullptr;
 
@@ -417,11 +415,11 @@ FActiveSound* FSoundConcurrencyManager::HandleConcurrencyEvaluationOverride(cons
 
 FActiveSound* FSoundConcurrencyManager::HandleConcurrencyEvaluation(const FActiveSound& NewActiveSound)
 {
-	check(NewActiveSound.Sound);
+	check(NewActiveSound.GetSound());
 	const FSoundConcurrencySettings* ConcurrencySettings = NewActiveSound.GetSoundConcurrencySettingsToApply();
 	check(ConcurrencySettings);
 
-	uint32 OwnerObjectID = NewActiveSound.TryGetOwnerID();
+	const uint32 OwnerObjectID = NewActiveSound.GetOwnerID();
 	FConcurrencyObjectID ConcurrencyObjectID = NewActiveSound.GetSoundConcurrencyObjectID();
 	check(ConcurrencyObjectID != 0);
 
diff --git a/Engine/Source/Runtime/Engine/Private/SoundCue.cpp b/Engine/Source/Runtime/Engine/Private/SoundCue.cpp
--- a/Engine/Source/Runtime/Engine/Private/SoundCue.cpp
+++ b/Engine/Source/Runtime/Engine/Private/SoundCue.cpp
@@ -6,6 +6,7 @@
 #include "Sound/SoundNodeWavePlayer.h"
 #include "Sound/SoundNodeAttenuation.h"
 #include "Sound/SoundNodeQualityLevel.h"
+#include "Sound/SoundNodeSoundClass.h"
 #include "Sound/SoundWave.h"
 #include "GameFramework/GameUserSettings.h"
 #if WITH_EDITOR
@@ -363,9 +364,34 @@ float USoundCue::GetDuration()
 	return Duration;
 }
 
+bool USoundCue::ShouldApplyInteriorVolumes() const
+{
+	if (Super::ShouldApplyInteriorVolumes())
+	{
+		return true;
+	}
+
+	// TODO: Consider caching this so we only reevaluate in editor
+	TArray<UObject*> Children;
+	GetObjectsWithOuter(this, Children);
+
+	for (UObject* Child : Children)
+	{
+		if (USoundNodeSoundClass* SoundClassNode = Cast<USoundNodeSoundClass>(Child))
+		{
+			if (SoundClassNode->SoundClassOverride && SoundClassNode->SoundClassOverride->Properties.bApplyAmbientVolumes)
+			{
+				return true;
+			}
+		}
+	}
+
+	return false;
+}
+
 bool USoundCue::IsPlayable() const
 {
-	return FirstNode != NULL;
+	return FirstNode != nullptr;
 }
 
 void USoundCue::Parse( FAudioDevice* AudioDevice, const UPTRINT NodeWaveInstanceHash, FActiveSound& ActiveSound, const FSoundParseParameters& ParseParams, TArray<FWaveInstance*>& WaveInstances )
diff --git a/Engine/Source/Runtime/Engine/Private/SoundNodeAttenuation.cpp b/Engine/Source/Runtime/Engine/Private/SoundNodeAttenuation.cpp
--- a/Engine/Source/Runtime/Engine/Private/SoundNodeAttenuation.cpp
+++ b/Engine/Source/Runtime/Engine/Private/SoundNodeAttenuation.cpp
@@ -53,7 +53,7 @@ void USoundNodeAttenuation::ParseNodes( FAudioDevice* AudioDevice, const UPTRINT
 	FSoundParseParameters UpdatedParseParams = ParseParams;
 	if (Settings)
 	{
-		const FListener& Listener = AudioDevice->Listeners[0];
+		const FListener& Listener = AudioDevice->GetListeners()[0];
 
 		// Update this node's attenuation settings overrides
 		ActiveSound.ApplyAttenuation(UpdatedParseParams, Listener, Settings);
diff --git a/Engine/Source/Runtime/Engine/Private/SoundNodeDistanceCrossFade.cpp b/Engine/Source/Runtime/Engine/Private/SoundNodeDistanceCrossFade.cpp
--- a/Engine/Source/Runtime/Engine/Private/SoundNodeDistanceCrossFade.cpp
+++ b/Engine/Source/Runtime/Engine/Private/SoundNodeDistanceCrossFade.cpp
@@ -166,7 +166,7 @@ void USoundNodeDistanceCrossFade::SetChildNodes(TArray<USoundNode*>& InChildNode
 
 float USoundNodeDistanceCrossFade::GetCurrentDistance(FAudioDevice* AudioDevice, FActiveSound& ActiveSound, const FSoundParseParameters& ParseParams) const
 {
-	return ActiveSound.bLocationDefined ? FVector::Dist( ParseParams.Transform.GetTranslation(), AudioDevice->Listeners[0].Transform.GetTranslation() ) : 0.f;
+	return ActiveSound.bLocationDefined ? FVector::Dist( ParseParams.Transform.GetTranslation(), AudioDevice->GetListeners()[0].Transform.GetTranslation() ) : 0.f;
 }
 
 bool USoundNodeDistanceCrossFade::AllowCrossfading(FActiveSound& ActiveSound) const
diff --git a/Engine/Source/Runtime/Engine/Private/SoundNodeDoppler.cpp b/Engine/Source/Runtime/Engine/Private/SoundNodeDoppler.cpp
--- a/Engine/Source/Runtime/Engine/Private/SoundNodeDoppler.cpp
+++ b/Engine/Source/Runtime/Engine/Private/SoundNodeDoppler.cpp
@@ -17,7 +17,7 @@ USoundNodeDoppler::USoundNodeDoppler(const FObjectInitializer& ObjectInitializer
 void USoundNodeDoppler::ParseNodes( FAudioDevice* AudioDevice, const UPTRINT NodeWaveInstanceHash, FActiveSound& ActiveSound, const FSoundParseParameters& ParseParams, TArray<FWaveInstance*>& WaveInstances )
 {
 	FSoundParseParameters UpdatedParams = ParseParams;
-	UpdatedParams.Pitch *= GetDopplerPitchMultiplier(AudioDevice->Listeners[0], ParseParams.Transform.GetTranslation(), ParseParams.Velocity);
+	UpdatedParams.Pitch *= GetDopplerPitchMultiplier(AudioDevice->GetListeners()[0], ParseParams.Transform.GetTranslation(), ParseParams.Velocity);
 
 	Super::ParseNodes( AudioDevice, NodeWaveInstanceHash, ActiveSound, UpdatedParams, WaveInstances );
 }
diff --git a/Engine/Source/Runtime/Engine/Private/SoundNodeLooping.cpp b/Engine/Source/Runtime/Engine/Private/SoundNodeLooping.cpp
--- a/Engine/Source/Runtime/Engine/Private/SoundNodeLooping.cpp
+++ b/Engine/Source/Runtime/Engine/Private/SoundNodeLooping.cpp
@@ -30,9 +30,9 @@ void USoundNodeLooping::ParseNodes( FAudioDevice* AudioDevice, const UPTRINT Nod
 	}
 
 #if !(NO_LOGGING || UE_BUILD_SHIPPING || UE_BUILD_TEST)
-	if (bLoopIndefinitely && !ActiveSound.bWarnedAboutOrphanedLooping && ActiveSound.GetAudioComponent() == nullptr)
+	if (bLoopIndefinitely && !ActiveSound.bWarnedAboutOrphanedLooping && ActiveSound.GetAudioComponentID() == 0)
 	{
-		UE_LOG(LogAudio, Warning, TEXT("Detected orphaned looping sound '%s'."), *ActiveSound.Sound->GetName());
+		UE_LOG(LogAudio, Warning, TEXT("Detected orphaned looping sound '%s'."), *ActiveSound.GetSound()->GetName());
 		ActiveSound.bWarnedAboutOrphanedLooping = true;
 	}
 #endif
@@ -61,13 +61,13 @@ bool USoundNodeLooping::NotifyWaveInstanceFinished( FWaveInstance* InWaveInstanc
 	RETRIEVE_SOUNDNODE_PAYLOAD(sizeof(int32));
 	DECLARE_SOUNDNODE_ELEMENT(int32, CurrentLoopCount);
 
+	const bool bRequiresInitialization = (*RequiresInitialization) != 0;
+
 	// The looping node should be initialized, but we hit this check and need to get information about what sound cue is causing the issue.
-	// Suppressing the "needless deref" static analysis warning here since it is a false positive. * has higher precedence than == and thus it is needed.
-	CA_SUPPRESS(6269)
-	ensureMsgf((*RequiresInitialization == 0), TEXT("Sound looping finished but not initialized. SoundWave: %s, Sound: %s"), *InWaveInstance->GetName(), *ActiveSound.Sound->GetName());
+	ensureMsgf(!bRequiresInitialization, TEXT("Sound looping finished but not initialized. SoundWave: %s, Sound: %s"), *InWaveInstance->GetName(), *ActiveSound.GetSound()->GetName());
 
 	// Avoid the crash and continue...
-	if (!*RequiresInitialization)
+	if (!bRequiresInitialization)
 	{
 		if (bLoopIndefinitely || ++CurrentLoopCount < LoopCount)
 		{
diff --git a/Engine/Source/Runtime/Engine/Private/SoundNodeQualityLevel.cpp b/Engine/Source/Runtime/Engine/Private/SoundNodeQualityLevel.cpp
--- a/Engine/Source/Runtime/Engine/Private/SoundNodeQualityLevel.cpp
+++ b/Engine/Source/Runtime/Engine/Private/SoundNodeQualityLevel.cpp
@@ -68,7 +68,7 @@ void USoundNodeQualityLevel::ParseNodes( FAudioDevice* AudioDevice, const UPTRIN
 
 		if (*RequiresInitialization)
 		{
-			const bool bIsPIESound = ((GEditor->bIsSimulatingInEditor || GEditor->PlayWorld != NULL) && ActiveSound.World != NULL);
+			const bool bIsPIESound = ((GEditor->bIsSimulatingInEditor || GEditor->PlayWorld != nullptr) && ActiveSound.GetWorldID() > 0);
 			if (bIsPIESound)
 			{
 				CachedQualityLevel = GetDefault<ULevelEditorPlaySettings>()->PlayInEditorSoundQualityLevel;
diff --git a/Engine/Source/Runtime/Engine/Private/SoundNodeRandom.cpp b/Engine/Source/Runtime/Engine/Private/SoundNodeRandom.cpp
--- a/Engine/Source/Runtime/Engine/Private/SoundNodeRandom.cpp
+++ b/Engine/Source/Runtime/Engine/Private/SoundNodeRandom.cpp
@@ -71,7 +71,7 @@ int32 USoundNodeRandom::ChooseNodeIndex(FActiveSound& ActiveSound)
 	float WeightSum = 0.0f;
 
 #if WITH_EDITOR
-	bool bIsPIESound = (GEditor != nullptr) && ((GEditor->bIsSimulatingInEditor || GEditor->PlayWorld != NULL) && ActiveSound.World != NULL);
+	bool bIsPIESound = (GEditor != nullptr) && ((GEditor->bIsSimulatingInEditor || GEditor->PlayWorld != nullptr) && ActiveSound.GetWorldID() > 0);
 
 	if (bIsPIESound)
 	{
@@ -139,7 +139,7 @@ void USoundNodeRandom::ParseNodes( FAudioDevice* AudioDevice, const UPTRINT Node
 	}
 
 #if WITH_EDITOR
-	bool bIsPIESound = (GEditor != nullptr) && ((GEditor->bIsSimulatingInEditor || GEditor->PlayWorld != NULL) && ActiveSound.World != NULL);
+	bool bIsPIESound = (GEditor != nullptr) && ((GEditor->bIsSimulatingInEditor || GEditor->PlayWorld != nullptr) && ActiveSound.GetWorldID() > 0);
 #endif //WITH_EDITOR
 
 	// check to see if we have used up our random sounds
diff --git a/Engine/Source/Runtime/Engine/Private/SoundWave.cpp b/Engine/Source/Runtime/Engine/Private/SoundWave.cpp
--- a/Engine/Source/Runtime/Engine/Private/SoundWave.cpp
+++ b/Engine/Source/Runtime/Engine/Private/SoundWave.cpp
@@ -6,6 +6,7 @@
 #include "Audio.h"
 #include "AudioDevice.h"
 #include "AudioDecompress.h"
+#include "AudioThread.h"
 #include "TargetPlatform.h"
 #include "AudioDerivedData.h"
 #include "SubtitleManager.h"
@@ -76,6 +77,7 @@ USoundWave::USoundWave(const FObjectInitializer& ObjectInitializer)
 	Pitch = 1.0;
 	CompressionQuality = 40;
 	SubtitlePriority = DEFAULT_SUBTITLE_PRIORITY;
+	ResourceState = ESoundWaveResourceState::NeedsFree;
 }
 
 SIZE_T USoundWave::GetResourceSize(EResourceSizeMode::Type Mode)
@@ -376,7 +378,7 @@ void USoundWave::PostLoad()
 	if( !GIsEditor && !IsTemplate( RF_ClassDefaultObject ) && GEngine )
 	{
 		FAudioDevice* AudioDevice = GEngine->GetMainAudioDevice();
-		if( AudioDevice && AudioDevice->bStartupSoundsPreCached)
+		if( AudioDevice && AudioDevice->AreStartupSoundsPreCached())
 		{
 			// Upload the data to the hardware, but only if we've precached startup sounds already
 			AudioDevice->Precache( this );
@@ -496,6 +498,8 @@ void USoundWave::PostEditChangeProperty(FPropertyChangedEvent& PropertyChangedEv
 
 void USoundWave::FreeResources()
 {
+	check(IsInAudioThread());
+
 	// Housekeeping of stats
 	DEC_FLOAT_STAT_BY( STAT_AudioBufferTime, Duration );
 	DEC_FLOAT_STAT_BY( STAT_AudioBufferTimeChannels, NumChannels * Duration );
@@ -508,7 +512,7 @@ void USoundWave::FreeResources()
 		FAudioDeviceManager* AudioDeviceManager = GEngine->GetAudioDeviceManager();
 		if (AudioDeviceManager)
 		{
-			AudioDeviceManager->StopSoundsUsingWave(this);
+			AudioDeviceManager->StopSoundsUsingResource(this);
 			AudioDeviceManager->FreeResource(this);
 		}
 	}
@@ -538,9 +542,15 @@ void USoundWave::FreeResources()
 	bDynamicResource = false;
 	DecompressionType = DTYPE_Setup;
 	bDecompressedFromOgg = 0;
+
+	USoundWave* SoundWave = this;
+	FAudioThread::RunCommandOnGameThread([SoundWave]()
+	{
+		SoundWave->ResourceState = ESoundWaveResourceState::Freed;
+	}, TStatId());
 }
 
-FWaveInstance* USoundWave::HandleStart( FActiveSound& ActiveSound, const UPTRINT WaveInstanceHash )
+FWaveInstance* USoundWave::HandleStart( FActiveSound& ActiveSound, const UPTRINT WaveInstanceHash ) const
 {
 	// Create a new wave instance and associate with the ActiveSound
 	FWaveInstance* WaveInstance = new FWaveInstance( &ActiveSound );
@@ -550,21 +560,16 @@ FWaveInstance* USoundWave::HandleStart( FActiveSound& ActiveSound, const UPTRINT
 	// Add in the subtitle if they exist
 	if (ActiveSound.bHandleSubtitles && Subtitles.Num() > 0)
 	{
-		if (UAudioComponent* AudioComponent = ActiveSound.GetAudioComponent())
-		{
-			// TODO - Audio Threading. This would need to be a call back to the main thread.
-			if (AudioComponent->OnQueueSubtitles.IsBound())
-			{
-				// intercept the subtitles if the delegate is set
-				AudioComponent->OnQueueSubtitles.ExecuteIfBound( Subtitles, Duration );
-			}
-			else if( ActiveSound.World.IsValid() )
-			{
-				// otherwise, pass them on to the subtitle manager for display
-				// Subtitles are hashed based on the associated sound (wave instance).
-				FSubtitleManager::GetSubtitleManager()->QueueSubtitles( ( PTRINT )WaveInstance, ActiveSound.SubtitlePriority, bManualWordWrap, bSingleLine, Duration, Subtitles, ActiveSound.World->GetAudioTimeSeconds() );
-			}
-		}
+		FQueueSubtitleParams QueueSubtitleParams(Subtitles);
+		QueueSubtitleParams.AudioComponentID = ActiveSound.GetAudioComponentID();
+		QueueSubtitleParams.WorldPtr = ActiveSound.GetWeakWorld();
+		QueueSubtitleParams.WaveInstance = (PTRINT)WaveInstance;
+		QueueSubtitleParams.SubtitlePriority = ActiveSound.SubtitlePriority;
+		QueueSubtitleParams.Duration = Duration;
+		QueueSubtitleParams.bManualWordWrap = bManualWordWrap;
+		QueueSubtitleParams.bSingleLine = bSingleLine;
+
+		FSubtitleManager::QueueSubtitles(QueueSubtitleParams);
 	}
 
 	return WaveInstance;
@@ -572,17 +577,30 @@ FWaveInstance* USoundWave::HandleStart( FActiveSound& ActiveSound, const UPTRINT
 
 bool USoundWave::IsReadyForFinishDestroy()
 {
-	bool bIsStreamingInProgress = IStreamingManager::Get().GetAudioStreamingManager().IsStreamingInProgress(this);
+	const bool bIsStreamingInProgress = IStreamingManager::Get().GetAudioStreamingManager().IsStreamingInProgress(this);
 
 	// Wait till streaming and decompression finishes before deleting resource.
-	return( !bIsStreamingInProgress && (( AudioDecompressor == NULL ) || AudioDecompressor->IsDone()) );
+	if ( !bIsStreamingInProgress && (( AudioDecompressor == nullptr ) || AudioDecompressor->IsDone()) )
+	{
+		if (ResourceState == ESoundWaveResourceState::NeedsFree)
+		{
+			DECLARE_CYCLE_STAT(TEXT("FAudioThreadTask.FreeResources"), STAT_AudioFreeResources, STATGROUP_AudioThreadCommands);
+
+			USoundWave* SoundWave = this;
+			ResourceState = ESoundWaveResourceState::Freeing;
+			FAudioThread::RunCommandOnAudioThread([SoundWave]()
+			{
+				SoundWave->FreeResources();
+			}, GET_STATID(STAT_AudioFreeResources));
+		}
+	}
+	
+	return ResourceState == ESoundWaveResourceState::Freed;
 }
 
 
 void USoundWave::FinishDestroy()
 {
-	FreeResources();
-
 	Super::FinishDestroy();
 
 	CleanupCachedRunningPlatformData();
@@ -613,9 +631,9 @@ void USoundWave::Parse( FAudioDevice* AudioDevice, const UPTRINT NodeWaveInstanc
 	{
 		WaveInstance->bIsFinished = false;
 #if !(NO_LOGGING || UE_BUILD_SHIPPING || UE_BUILD_TEST)
-		if (!ActiveSound.bWarnedAboutOrphanedLooping && ActiveSound.GetAudioComponent() == nullptr)
+		if (!ActiveSound.bWarnedAboutOrphanedLooping && ActiveSound.GetAudioComponentID() == 0)
 		{
-			UE_LOG(LogAudio, Warning, TEXT("Detected orphaned looping sound '%s'."), *ActiveSound.Sound->GetName());
+			UE_LOG(LogAudio, Warning, TEXT("Detected orphaned looping sound '%s'."), *ActiveSound.GetSound()->GetName());
 			ActiveSound.bWarnedAboutOrphanedLooping = true;
 		}
 #endif
@@ -651,7 +669,6 @@ void USoundWave::Parse( FAudioDevice* AudioDevice, const UPTRINT NodeWaveInstanc
 			WaveInstance->Pitch *= SoundClassProperties->Pitch;
 			//TODO: Add in HighFrequencyGainMultiplier property to sound classes
 
-
 			WaveInstance->VoiceCenterChannelVolume = SoundClassProperties->VoiceCenterChannelVolume;
 			WaveInstance->RadioFilterVolume = SoundClassProperties->RadioFilterVolume * ParseParams.VolumeMultiplier;
 			WaveInstance->RadioFilterVolumeThreshold = SoundClassProperties->RadioFilterVolumeThreshold * ParseParams.VolumeMultiplier;
@@ -665,6 +682,13 @@ void USoundWave::Parse( FAudioDevice* AudioDevice, const UPTRINT NodeWaveInstanc
 			WaveInstance->bReverb = ActiveSound.bReverb || SoundClassProperties->bReverb;
 			WaveInstance->OutputTarget = SoundClassProperties->OutputTarget;
 
+			if (SoundClassProperties->bApplyAmbientVolumes)
+			{
+				WaveInstance->VolumeMultiplier *= ParseParams.InteriorVolumeMultiplier;
+				WaveInstance->RadioFilterVolume *= ParseParams.InteriorVolumeMultiplier;
+				WaveInstance->RadioFilterVolumeThreshold *= ParseParams.InteriorVolumeMultiplier;
+			}
+
 			bAlwaysPlay = ActiveSound.bAlwaysPlay || SoundClassProperties->bAlwaysPlay;
 		}
 		else
@@ -723,21 +747,33 @@ void USoundWave::Parse( FAudioDevice* AudioDevice, const UPTRINT NodeWaveInstanc
 			if (!ReportedSounds.Contains(this))
 			{
 				FString SoundWarningInfo = FString::Printf(TEXT("Spatialisation on stereo and multichannel sounds is not supported. SoundWave: %s"), *GetName());
-				if (ActiveSound.Sound != this)
+				if (ActiveSound.GetSound() != this)
 				{
-					SoundWarningInfo += FString::Printf(TEXT(" SoundCue: %s"), *ActiveSound.Sound->GetName());
+					SoundWarningInfo += FString::Printf(TEXT(" SoundCue: %s"), *ActiveSound.GetSound()->GetName());
 				}
 
-				if (UAudioComponent* AudioComponent = ActiveSound.GetAudioComponent())
+#if !NO_LOGGING
+				const uint64 AudioComponentID = ActiveSound.GetAudioComponentID();
+				if (AudioComponentID > 0)
 				{
-					// TODO - Audio Threading. This log would have to be a task back to game thread
-					AActor* SoundOwner = AudioComponent->GetOwner();
-					UE_LOG(LogAudio, Warning, TEXT( "%s Actor: %s AudioComponent: %s" ), *SoundWarningInfo, (SoundOwner ? *SoundOwner->GetName() : TEXT("None")), *AudioComponent->GetName() );
+					FAudioThread::RunCommandOnGameThread([AudioComponentID, SoundWarningInfo]()
+					{
+						if (UAudioComponent* AudioComponent = UAudioComponent::GetAudioComponentFromID(AudioComponentID))
+						{
+							AActor* SoundOwner = AudioComponent->GetOwner();
+							UE_LOG(LogAudio, Warning, TEXT( "%s Actor: %s AudioComponent: %s" ), *SoundWarningInfo, (SoundOwner ? *SoundOwner->GetName() : TEXT("None")), *AudioComponent->GetName() );
+						}
+						else
+						{
+							UE_LOG(LogAudio, Warning, TEXT("%s"), *SoundWarningInfo );
+						}
+					});
 				}
 				else
 				{
 					UE_LOG(LogAudio, Warning, TEXT("%s"), *SoundWarningInfo );
 				}
+#endif
 
 				ReportedSounds.Add(this);
 			}
diff --git a/Engine/Source/Runtime/Engine/Private/SubtitleManager.cpp b/Engine/Source/Runtime/Engine/Private/SubtitleManager.cpp
--- a/Engine/Source/Runtime/Engine/Private/SubtitleManager.cpp
+++ b/Engine/Source/Runtime/Engine/Private/SubtitleManager.cpp
@@ -5,6 +5,7 @@
 #include "SubtitleManager.h"
 #include "SoundDefinitions.h"
 #include "Engine/Font.h"
+#include "AudioThread.h"
 
 DEFINE_LOG_CATEGORY(LogSubtitle);
 
@@ -28,6 +29,29 @@ void FSubtitleManager::KillSubtitles( PTRINT SubtitleID )
 	ActiveSubtitles.Remove( SubtitleID );
 }
 
+void FSubtitleManager::QueueSubtitles(const FQueueSubtitleParams& QueueSubtitleParams)
+{
+	check(IsInAudioThread());
+
+	DECLARE_CYCLE_STAT(TEXT("FGameThreadAudioTask.QueueSubtitles"), STAT_AudioQueueSubtitles, STATGROUP_TaskGraphTasks);
+
+	FAudioThread::RunCommandOnGameThread([QueueSubtitleParams]()
+	{
+		UAudioComponent* AudioComponent = UAudioComponent::GetAudioComponentFromID(QueueSubtitleParams.AudioComponentID);
+		if (AudioComponent && AudioComponent->OnQueueSubtitles.IsBound())
+		{
+			// intercept the subtitles if the delegate is set
+			AudioComponent->OnQueueSubtitles.ExecuteIfBound(QueueSubtitleParams.Subtitles, QueueSubtitleParams.Duration);
+		}
+		else if (UWorld* World = QueueSubtitleParams.WorldPtr.Get())
+		{
+			// otherwise, pass them on to the subtitle manager for display
+			// Subtitles are hashed based on the associated sound (wave instance).
+			FSubtitleManager::GetSubtitleManager()->QueueSubtitles( QueueSubtitleParams.WaveInstance, QueueSubtitleParams.SubtitlePriority, QueueSubtitleParams.bManualWordWrap, QueueSubtitleParams.bSingleLine, QueueSubtitleParams.Duration, QueueSubtitleParams.Subtitles, World->GetAudioTimeSeconds() );
+		}
+	}, GET_STATID(STAT_AudioQueueSubtitles));
+}
+
 /**
  * Add an array of subtitles to the active list
  * This is called from AudioComponent::Play
@@ -39,9 +63,10 @@ void FSubtitleManager::KillSubtitles( PTRINT SubtitleID )
  * @param  SoundDuration - time in seconds after which the subtitles do not display
  * @param  Subtitles - TArray of lines of subtitle and time offset to play them
  */
-void FSubtitleManager::QueueSubtitles( PTRINT SubtitleID, float Priority, bool bManualWordWrap, bool bSingleLine, float SoundDuration, TArray<FSubtitleCue>& Subtitles, float InStartTime )
+void FSubtitleManager::QueueSubtitles( PTRINT SubtitleID, float Priority, bool bManualWordWrap, bool bSingleLine, float SoundDuration, const TArray<FSubtitleCue>& Subtitles, float InStartTime )
 {
 	check( GEngine );
+	check(IsInGameThread());
 
 	// No subtitles to display
 	if( !Subtitles.Num() )
diff --git a/Engine/Source/Runtime/Engine/Private/TickTaskManager.cpp b/Engine/Source/Runtime/Engine/Private/TickTaskManager.cpp
--- a/Engine/Source/Runtime/Engine/Private/TickTaskManager.cpp
+++ b/Engine/Source/Runtime/Engine/Private/TickTaskManager.cpp
@@ -751,8 +751,8 @@ public:
 				CumulativeCooldown += TickFunction->RelativeTickCooldown;
 
 				TickFunction->TickState = FTickFunction::ETickState::Enabled;
-				// we store a bit in here if this came from an interval queue 
-				AllTickFunctions.Add((FTickFunction*)(UPTRINT(TickFunction) | 1));
+				TickFunction->bWasInterval = true;
+				AllTickFunctions.Add(TickFunction);
 
 				TickFunctionsToReschedule.Add(FTickScheduleDetails(TickFunction, TickFunction->TickInterval - (Context.DeltaSeconds - CumulativeCooldown))); // Give credit for any overrun
 
@@ -1375,13 +1375,10 @@ public:
 			ParallelFor(AllTickFunctions.Num(), 
 				[this](int32 Index)
 				{
-					// we store a bit in here if this came from an interval queue 
 					FTickFunction* TickFunction = AllTickFunctions[Index];
-					bool bWasInterval = !!(UPTRINT(TickFunction) & 1);
-					TickFunction = (FTickFunction*)(UPTRINT(TickFunction) & ~1);
 
 					TArray<FTickFunction*, TInlineAllocator<8> > StackForCycleDetection;
-					TickFunction->QueueTickFunctionParallel(Context, StackForCycleDetection, bWasInterval);
+					TickFunction->QueueTickFunctionParallel(Context, StackForCycleDetection);
 				}
 			);
 		    for( int32 LevelIndex = 0; LevelIndex < LevelList.Num(); LevelIndex++ )
@@ -1587,6 +1584,7 @@ FTickFunction::FTickFunction()
 	, bHighPriority(false)
 	, bRunOnAnyThread(false)
 	, bRegistered(false)
+	, bWasInterval(false)
 	, TickState(ETickState::Enabled)
 	, TickVisitedGFrameCounter(0)
 	, TickQueuedGFrameCounter(0)
@@ -1778,7 +1776,7 @@ void FTickFunction::QueueTickFunction(FTickTaskSequencer& TTS, const struct FTic
 	}
 }
 
-void FTickFunction::QueueTickFunctionParallel(const struct FTickContext& TickContext, TArray<FTickFunction*, TInlineAllocator<8> >& StackForCycleDetection, bool bWasInterval)
+void FTickFunction::QueueTickFunctionParallel(const struct FTickContext& TickContext, TArray<FTickFunction*, TInlineAllocator<8> >& StackForCycleDetection)
 {
 
 	bool bProcessTick;
@@ -1812,7 +1810,7 @@ void FTickFunction::QueueTickFunctionParallel(const struct FTickContext& TickCon
 					else if (Prereq->bRegistered)
 					{
 						// recursive call to make sure my prerequisite is set up so I can use its completion handle
-						Prereq->QueueTickFunctionParallel(TickContext, StackForCycleDetection, false);
+						Prereq->QueueTickFunctionParallel(TickContext, StackForCycleDetection);
 						if (!Prereq->TaskPointer)
 						{
 							//ok UE_LOG(LogTick, Warning, TEXT("While processing prerequisites for %s, could use %s because it is disabled."),*DiagnosticMessage(), *Prereq->DiagnosticMessage());
@@ -1862,6 +1860,8 @@ void FTickFunction::QueueTickFunctionParallel(const struct FTickContext& TickCon
 				}
 			}
 		}
+		bWasInterval = false;
+
 		FPlatformMisc::MemoryBarrier();
 		TickQueuedGFrameCounter = GFrameCounter;
 	}
diff --git a/Engine/Source/Runtime/Engine/Private/UnrealEngine.cpp b/Engine/Source/Runtime/Engine/Private/UnrealEngine.cpp
--- a/Engine/Source/Runtime/Engine/Private/UnrealEngine.cpp
+++ b/Engine/Source/Runtime/Engine/Private/UnrealEngine.cpp
@@ -70,6 +70,7 @@
 #include "Particles/ParticleSystem.h"
 #include "Particles/ParticleSystemComponent.h"
 
+#include "AudioThread.h"
 #include "Sound/ReverbEffect.h"
 #include "Sound/SoundWave.h"
 
@@ -1021,11 +1022,11 @@ void UEngine::Init(IEngineLoop* InEngineLoop)
 #if !UE_BUILD_SHIPPING
 	EngineStats.Add(FEngineStatFuncs(TEXT("STAT_SoundMixes"), TEXT("STATCAT_Engine"), FText::GetEmpty(), &UEngine::RenderStatSoundMixes, NULL));
 	EngineStats.Add(FEngineStatFuncs(TEXT("STAT_Reverb"), TEXT("STATCAT_Engine"), FText::GetEmpty(), &UEngine::RenderStatReverb, NULL));
-	EngineStats.Add(FEngineStatFuncs(TEXT("STAT_SoundWaves"), TEXT("STATCAT_Engine"), FText::GetEmpty(), &UEngine::RenderStatSoundWaves, NULL));
-	EngineStats.Add(FEngineStatFuncs(TEXT("STAT_SoundCues"), TEXT("STATCAT_Engine"), FText::GetEmpty(), &UEngine::RenderStatSoundCues, NULL));
-#endif // !UE_BUILD_SHIPPING
+	EngineStats.Add(FEngineStatFuncs(TEXT("STAT_SoundWaves"), TEXT("STATCAT_Engine"), FText::GetEmpty(), &UEngine::RenderStatSoundWaves, &UEngine::ToggleStatSoundWaves));
+	EngineStats.Add(FEngineStatFuncs(TEXT("STAT_SoundCues"), TEXT("STATCAT_Engine"), FText::GetEmpty(), &UEngine::RenderStatSoundCues, &UEngine::ToggleStatSoundCues));
 	EngineStats.Add(FEngineStatFuncs(TEXT("STAT_Sounds"), TEXT("STATCAT_Engine"), FText::GetEmpty(), &UEngine::RenderStatSounds, &UEngine::ToggleStatSounds));
-/* @todo UE4 physx fix this once we have convexelem drawing again
+#endif // !UE_BUILD_SHIPPING
+	/* @todo UE4 physx fix this once we have convexelem drawing again
 	EngineStats.Add(FEngineStatFuncs(TEXT("STAT_LevelMap"), TEXT("STATCAT_Engine"), FText::GetEmpty(), &UEngine::RenderStatLevelMap, NULL));
 */
 	EngineStats.Add(FEngineStatFuncs(TEXT("STAT_Detailed"), TEXT("STATCAT_Engine"), FText::GetEmpty(), NULL, &UEngine::ToggleStatDetailed));
@@ -1072,7 +1073,14 @@ void UEngine::ShutdownAudioDeviceManager()
 	// Shutdown the main audio device in the UEEngine
 	if (AudioDeviceManager)
 	{
+		FAudioCommandFence Fence;
+		Fence.BeginFence();
+		Fence.Wait();
+
+		FAudioThread::StopAudioThread();
+
 		AudioDeviceManager->ShutdownAllAudioDevices();
+
 		delete AudioDeviceManager;
 		AudioDeviceManager = NULL;
 	}
@@ -1267,7 +1275,8 @@ void UEngine::UpdateTimeAndHandleMaxTickRate()
 		UpdateRunningAverageDeltaTime(DeltaTime);
 
 		// Get max tick rate based on network settings and current delta time.
-		const float MaxTickRate = FABTest::StaticIsActive() ? 0.0f : GetMaxTickRate(DeltaTime);
+		const float GivenMaxTickRate = GetMaxTickRate(DeltaTime);
+		const float MaxTickRate = FABTest::StaticIsActive() ? 0.0f : bUseFixedFrameRate ? FMath::Min(GivenMaxTickRate, FixedFrameRate) : GivenMaxTickRate;
 		float WaitTime		= 0;
 		// Convert from max FPS to wait time.
 		if( MaxTickRate > 0 )
@@ -1307,7 +1316,7 @@ void UEngine::UpdateTimeAndHandleMaxTickRate()
 			}
 			FApp::SetCurrentTime(FPlatformTime::Seconds());
 		}
-		else if(bUseFixedFrameRate)
+		else if(bUseFixedFrameRate && MaxTickRate == FixedFrameRate)
 		{
 			//We are doing fixed framerate and the real delta time is bigger than our desired delta time. In this case we start falling behind real time (and that's ok)
 			const float FrameRate = 1.f / FixedFrameRate;
@@ -1921,14 +1930,17 @@ bool UEngine::InitializeAudioDeviceManager()
 					AudioDeviceManager = new FAudioDeviceManager();
 					AudioDeviceManager->RegisterAudioDeviceModule(AudioDeviceModule);
 
+					FAudioDeviceManager::FCreateAudioDeviceResults NewDeviceResults;
+
 					// Create a new audio device.
-					FAudioDevice* NewAudioDevice = AudioDeviceManager->CreateAudioDevice(MainAudioDeviceHandle, true);
-					if (NewAudioDevice)
+					if (AudioDeviceManager->CreateAudioDevice(true, NewDeviceResults))
+					{
+						MainAudioDeviceHandle = NewDeviceResults.Handle;
+						AudioDeviceManager->SetActiveDevice(MainAudioDeviceHandle);
+						FAudioThread::StartAudioThread();
+					}
+					else
 					{
-							AudioDeviceManager->SetActiveDevice(MainAudioDeviceHandle);
-						}
-						else
-						{
 						ShutdownAudioDeviceManager();
 					}
 				}
@@ -6654,11 +6666,6 @@ float UEngine::GetMaxTickRate(float DeltaTime, bool bAllowFrameRateSmoothing) co
 		}
 	}
 
-	if(bUseFixedFrameRate)
-	{
-		MaxTickRate = FixedFrameRate;
-	}
-
 	if (CVarCauseHitches.GetValueOnGameThread())
 	{
 		static float RunningHitchTimer = 0.f;
@@ -9622,7 +9629,7 @@ bool UEngine::LoadMap( FWorldContext& WorldContext, FURL URL, class UPendingNetG
 		if (FAudioDevice* AudioDevice = WorldContext.World()->GetAudioDevice())
 		{
 			AudioDevice->Flush(WorldContext.World());
-			AudioDevice->TransientMasterVolume = 1.0f;
+			AudioDevice->SetTransientMasterVolume(1.0f);
 		}
 
 		WorldContext.SetCurrentWorld(nullptr);
@@ -12144,77 +12151,88 @@ int32 UEngine::RenderStatReverb(UWorld* World, FViewport* Viewport, FCanvas* Can
 {
 	if (FAudioDevice* AudioDevice = World->GetAudioDevice())
 	{
-		UReverbEffect* ReverbEffect = (AudioDevice->Effects ? AudioDevice->Effects->GetCurrentReverbEffect() : NULL);
-		FString TheString;
-		if (ReverbEffect)
-		{
-			TheString = FString::Printf(TEXT("Active Reverb Effect: %s"), *ReverbEffect->GetName());
-			Canvas->DrawShadowedString(X, Y, *TheString, GetSmallFont(), FLinearColor::White);
-			Y += 12;
+		AudioDevice->RenderStatReverb(World, Viewport, Canvas, X, Y, ViewLocation, ViewRotation);
+	}
+
+	return Y;
+}
+
+void FAudioDevice::RenderStatReverb(UWorld* World, FViewport* Viewport, FCanvas* Canvas, int32 X, int32& Y, const FVector* ViewLocation, const FRotator* ViewRotation) const
+{
+	UReverbEffect* ReverbEffect = GetCurrentReverbEffect();
+	FString TheString;
+	if (ReverbEffect)
+	{
+		TheString = FString::Printf(TEXT("Active Reverb Effect: %s"), *ReverbEffect->GetName());
+		Canvas->DrawShadowedString(X, Y, *TheString, UEngine::GetSmallFont(), FLinearColor::White);
+		Y += 12;
 
-			ULocalPlayer* LocalPlayer = GetFirstGamePlayer(World);
-			if (LocalPlayer)
+		AAudioVolume* CurrentAudioVolume = nullptr;
+		for (const FTransform& Transform : ListenerTransforms)
+		{
+			AAudioVolume* PlayerAudioVolume = World->GetAudioSettings(Transform.GetLocation(), nullptr, nullptr);
+			if (PlayerAudioVolume && ((CurrentAudioVolume == nullptr) || (PlayerAudioVolume->GetPriority() > CurrentAudioVolume->GetPriority())))
 			{
-				const AAudioVolume* AudioVolume = AudioDevice->CurrentAudioVolume;
-				if (AudioVolume && AudioVolume->Settings.ReverbEffect)
-				{
-					TheString = FString::Printf(TEXT("  Audio Volume Reverb Effect: %s (Priority: %g Volume Name: %s)"), *AudioVolume->Settings.ReverbEffect->GetName(), AudioVolume->Priority, *AudioVolume->GetName());
-				}
-				else
-				{
-					TheString = TEXT("  Audio Volume Reverb Effect: None");
-				}
-				Canvas->DrawShadowedString(X, Y, *TheString, GetSmallFont(), FLinearColor::White);
-				Y += 12;
-				if (AudioDevice->ActivatedReverbs.Num() == 0)
-				{
-					TheString = TEXT("  Activated Reverb: None");
-					Canvas->DrawShadowedString(X, Y, *TheString, GetSmallFont(), FLinearColor::White);
-					Y += 12;
-				}
-				else if (AudioDevice->ActivatedReverbs.Num() == 1)
-				{
-					auto It = AudioDevice->ActivatedReverbs.CreateConstIterator();
-					TheString = FString::Printf(TEXT("  Activated Reverb Effect: %s (Priority: %g Tag: '%s')"), *It.Value().ReverbSettings.ReverbEffect->GetName(), It.Value().Priority, *It.Key().ToString());
-					Canvas->DrawShadowedString(X, Y, *TheString, GetSmallFont(), FLinearColor::White);
-					Y += 12;
-				}
-				else
-				{
-					Canvas->DrawShadowedString(X, Y, TEXT("  Activated Reverb Effects:"), GetSmallFont(), FLinearColor::White);
-					Y += 12;
-					TMap<int32, FString> PrioritySortedActivatedReverbs;
-					for (auto It = AudioDevice->ActivatedReverbs.CreateConstIterator(); It; ++It)
-					{
-						TheString = FString::Printf(TEXT("    %s (Priority: %g Tag: '%s')"), *It.Value().ReverbSettings.ReverbEffect->GetName(), It.Value().Priority, *It.Key().ToString());
-						PrioritySortedActivatedReverbs.Add(It.Value().Priority, TheString);
-					}
-					for (auto It = PrioritySortedActivatedReverbs.CreateConstIterator(); It; ++It)
-					{
-						Canvas->DrawShadowedString(X, Y, *It.Value(), GetSmallFont(), FLinearColor::White);
-						Y += 12;
-					}
-				}
+				CurrentAudioVolume = PlayerAudioVolume;
 			}
 		}
+		if (CurrentAudioVolume && CurrentAudioVolume->GetReverbSettings().ReverbEffect)
+		{
+			TheString = FString::Printf(TEXT("  Audio Volume Reverb Effect: %s (Priority: %g Volume Name: %s)"), *CurrentAudioVolume->GetReverbSettings().ReverbEffect->GetName(), CurrentAudioVolume->GetPriority(), *CurrentAudioVolume->GetName());
+		}
+		else
+		{
+			TheString = TEXT("  Audio Volume Reverb Effect: None");
+		}
+		Canvas->DrawShadowedString(X, Y, *TheString, UEngine::GetSmallFont(), FLinearColor::White);
+		Y += 12;
+		if (ActivatedReverbs.Num() == 0)
+		{
+			TheString = TEXT("  Activated Reverb: None");
+			Canvas->DrawShadowedString(X, Y, *TheString, UEngine::GetSmallFont(), FLinearColor::White);
+			Y += 12;
+		}
+		else if (ActivatedReverbs.Num() == 1)
+		{
+			auto It = ActivatedReverbs.CreateConstIterator();
+			TheString = FString::Printf(TEXT("  Activated Reverb Effect: %s (Priority: %g Tag: '%s')"), *It.Value().ReverbSettings.ReverbEffect->GetName(), It.Value().Priority, *It.Key().ToString());
+			Canvas->DrawShadowedString(X, Y, *TheString, UEngine::GetSmallFont(), FLinearColor::White);
+			Y += 12;
+		}
 		else
 		{
-			TheString = TEXT("Active Reverb Effect: None");
-			Canvas->DrawShadowedString(X, Y, *TheString, GetSmallFont(), FLinearColor::White);
+			Canvas->DrawShadowedString(X, Y, TEXT("  Activated Reverb Effects:"), UEngine::GetSmallFont(), FLinearColor::White);
 			Y += 12;
+			TMap<int32, FString> PrioritySortedActivatedReverbs;
+			for (auto It = ActivatedReverbs.CreateConstIterator(); It; ++It)
+			{
+				TheString = FString::Printf(TEXT("    %s (Priority: %g Tag: '%s')"), *It.Value().ReverbSettings.ReverbEffect->GetName(), It.Value().Priority, *It.Key().ToString());
+				PrioritySortedActivatedReverbs.Add(It.Value().Priority, TheString);
+			}
+			for (auto It = PrioritySortedActivatedReverbs.CreateConstIterator(); It; ++It)
+			{
+				Canvas->DrawShadowedString(X, Y, *It.Value(), UEngine::GetSmallFont(), FLinearColor::White);
+				Y += 12;
+			}
 		}
 	}
-	return Y;
+	else
+	{
+		TheString = TEXT("Active Reverb Effect: None");
+		Canvas->DrawShadowedString(X, Y, *TheString, UEngine::GetSmallFont(), FLinearColor::White);
+		Y += 12;
+	}
 }
 
 // SOUNDMIXES
 int32 UEngine::RenderStatSoundMixes(UWorld* World, FViewport* Viewport, FCanvas* Canvas, int32 X, int32 Y, const FVector* ViewLocation, const FRotator* ViewRotation)
 {
+#if 0
+	Canvas->DrawShadowedString(X, Y, TEXT("Active Sound Mixes:"), GetSmallFont(), FColor::Green);
+	Y += 12;
+
 	if (FAudioDevice* AudioDevice = World->GetAudioDevice())
 	{
-		Canvas->DrawShadowedString(X, Y, TEXT("Active Sound Mixes:"), GetSmallFont(), FColor::Green);
-		Y += 12;
-
 		if (AudioDevice->SoundMixModifiers.Num() > 0)
 		{
 			USoundMix* CurrentEQMix = AudioDevice->Effects->GetCurrentEQMix();
@@ -12241,64 +12259,142 @@ int32 UEngine::RenderStatSoundMixes(UWorld* World, FViewport* Viewport, FCanvas*
 			Y += 12;
 		}
 	}
+	else
+	{
+		Canvas->DrawShadowedString(X + 12, Y, TEXT("None"), GetSmallFont(), FColor::White);
+		Y += 12;
+	}
+#endif
 	return Y;
 }
 
-// SOUNDWAVES
-int32 UEngine::RenderStatSoundWaves(UWorld* World, FViewport* Viewport, FCanvas* Canvas, int32 X, int32 Y, const FVector* ViewLocation, const FRotator* ViewRotation)
+void FAudioDevice::UpdateSoundShowFlags(const uint8 OldSoundShowFlags, const uint8 NewSoundShowFlags)
 {
-	Canvas->DrawShadowedString(X, Y, TEXT("Active Sound Waves:"), GetSmallFont(), FLinearColor::White);
-	Y += 12;
+	if (NewSoundShowFlags != OldSoundShowFlags)
+	{
+		uint8 RequestedStatChange = 0;
+		if ((NewSoundShowFlags == FViewportClient::ESoundShowFlags::Disabled) || (OldSoundShowFlags == FViewportClient::ESoundShowFlags::Disabled))
+		{
+			RequestedStatChange |= ERequestedAudioStats::Sounds;
+		}
+		if ((NewSoundShowFlags ^ OldSoundShowFlags) & FViewportClient::ESoundShowFlags::Debug)
+		{
+			RequestedStatChange |= ERequestedAudioStats::DebugSounds;
+		}
+		if ((NewSoundShowFlags ^ OldSoundShowFlags) & FViewportClient::ESoundShowFlags::Long_Names)
+		{
+			RequestedStatChange |= ERequestedAudioStats::LongSoundNames;
+		}
+		if (RequestedStatChange != 0)
+		{
+			UpdateRequestedStat(RequestedStatChange);
+		}
+	}
+}
 
-	TSet<FActiveSound*> ActiveSounds;
+void FAudioDevice::UpdateRequestedStat(const uint8 RequestedStat)
+{
+	check(IsInGameThread());
+
+	DECLARE_CYCLE_STAT(TEXT("FAudioThreadTask.UpdateRequestedStat"), STAT_AudioUpdateRequestedStat, STATGROUP_TaskGraphTasks);
+
+	FAudioDevice* AudioDevice = this;
+	FAudioThread::RunCommandOnAudioThread([AudioDevice, RequestedStat]()
+	{
+		AudioDevice->RequestedAudioStats ^= RequestedStat;
+	}, GET_STATID(STAT_AudioUpdateRequestedStat));
+}
 
+bool UEngine::ToggleStatSoundWaves(UWorld* World, FCommonViewportClient* ViewportClient, const TCHAR* Stream)
+{
 	if (FAudioDevice* AudioDevice = World->GetAudioDevice())
 	{
-		TArray<FWaveInstance*> WaveInstances;
-		int32 FirstActiveIndex = AudioDevice->GetSortedActiveWaveInstances(WaveInstances, ESortedActiveWaveGetType::QueryOnly);
-		int32 ActiveInstances = 0;
+		AudioDevice->UpdateRequestedStat(ERequestedAudioStats::SoundWaves);
+	}
+	return true;
+}
+
+bool UEngine::ToggleStatSoundCues(UWorld* World, FCommonViewportClient* ViewportClient, const TCHAR* Stream)
+{
+	if (FAudioDevice* AudioDevice = World->GetAudioDevice())
+	{
+		AudioDevice->UpdateRequestedStat(ERequestedAudioStats::SoundCues);
+	}
+	return true;
+}
+
+// SOUNDWAVES
+int32 UEngine::RenderStatSoundWaves(UWorld* World, FViewport* Viewport, FCanvas* Canvas, int32 X, int32 Y, const FVector* ViewLocation, const FRotator* ViewRotation)
+{
+	if (FAudioDevice* AudioDevice = World->GetAudioDevice())
+	{
+		const FAudioStats& AudioStats = AudioDevice->GetAudioStats();
 
-		for (int32 InstanceIndex = FirstActiveIndex; InstanceIndex < WaveInstances.Num(); InstanceIndex++)
+		if (!AudioStats.bStale)
 		{
-			FWaveInstance* WaveInstance = WaveInstances[InstanceIndex];
+			Canvas->DrawShadowedString(X, Y, TEXT("Active Sound Waves:"), GetSmallFont(), FLinearColor::White);
+			Y += 12;
 
-			if (WaveInstance->GetActualVolume() >= 0.01f)
+			typedef TPair<const FAudioStats::FStatWaveInstanceInfo*, const FAudioStats::FStatSoundInfo*> FWaveInstancePair;
+
+			TArray<FWaveInstancePair> WaveInstances;
+
+			for (const FAudioStats::FStatSoundInfo& StatSoundInfo : AudioStats.StatSoundInfos)
 			{
-				++ActiveInstances;
-				ActiveSounds.Add(WaveInstance->ActiveSound);
+				for (const FAudioStats::FStatWaveInstanceInfo& WaveInstanceInfo : StatSoundInfo.WaveInstanceInfos)
+				{
+					if (WaveInstanceInfo.ActualVolume >= 0.01f)
+					{
+						WaveInstances.Add(TPairInitializer<const FAudioStats::FStatWaveInstanceInfo*, const FAudioStats::FStatSoundInfo*>(&WaveInstanceInfo, &StatSoundInfo));
+					}
+				}
+			}
+
+			WaveInstances.Sort([](const FWaveInstancePair& A, const FWaveInstancePair& B) { return A.Key->InstanceIndex < B.Key->InstanceIndex; });
 
-				UAudioComponent* AudioComponent = WaveInstance->ActiveSound->GetAudioComponent();
+			for (const FWaveInstancePair& WaveInstanceInfo : WaveInstances)
+			{
+				UAudioComponent* AudioComponent = UAudioComponent::GetAudioComponentFromID(WaveInstanceInfo.Value->AudioComponentID);
 				AActor* SoundOwner = AudioComponent ? AudioComponent->GetOwner() : nullptr;
-				USoundClass* SoundClass = WaveInstance->SoundClass;
 
 				FString TheString = *FString::Printf(TEXT("%4i.    %6.2f  %s   Owner: %s   SoundClass: %s"),
-													 InstanceIndex,
-													 WaveInstance->GetActualVolume(),
-													 *WaveInstance->WaveData->GetPathName(),
-													 SoundOwner ? *SoundOwner->GetName() : TEXT("None"),
-													 SoundClass ? *SoundClass->GetName() : TEXT("None"));
+					WaveInstanceInfo.Key->InstanceIndex,
+					WaveInstanceInfo.Key->ActualVolume,
+					*WaveInstanceInfo.Value->SoundName,
+					SoundOwner ? *SoundOwner->GetName() : TEXT("None"),
+					*WaveInstanceInfo.Value->SoundClassName.ToString());
 
 				Canvas->DrawShadowedString(X, Y, *TheString, GetSmallFont(), FColor::White);
 				Y += 12;
 			}
-		}
 
-		int32 Max = AudioDevice->MaxChannels / 2;
-		float f = FMath::Clamp<float>((float)(ActiveInstances - Max) / (float)Max, 0.f, 1.f);
-		int32 R = FMath::TruncToInt(f * 255);
+			const int32 ActiveInstances = WaveInstances.Num();
 
-		if (ActiveInstances > Max)
-		{
-			f = FMath::Clamp<float>((float)(Max - ActiveInstances) / (float)Max, 0.5f, 1.f);
-		}
-		else
-		{
-			f = 1.0f;
+			const int32 Max = AudioDevice->MaxChannels / 2;
+			float f = FMath::Clamp<float>((float)(ActiveInstances - Max) / (float)Max, 0.f, 1.f);
+			const int32 R = FMath::TruncToInt(f * 255);
+
+			if (ActiveInstances > Max)
+			{
+				f = FMath::Clamp<float>((float)(Max - ActiveInstances) / (float)Max, 0.5f, 1.f);
+			}
+			else
+			{
+				f = 1.0f;
+			}
+			const int32 G = FMath::TruncToInt(f * 255);
+			const int32 B = 0;
+
+			Canvas->DrawShadowedString(X, Y, *FString::Printf(TEXT(" Total: %i"), ActiveInstances), GetSmallFont(), FColor(R, G, B));
+			Y += 12;
 		}
-		int32 G = FMath::TruncToInt(f * 255);
-		int32 B = 0;
+	}
+	else
+	{
+		Canvas->DrawShadowedString(X, Y, TEXT("Active Sound Waves:"), GetSmallFont(), FLinearColor::White);
+		Y += 12;
 
-		Canvas->DrawShadowedString(X, Y, *FString::Printf(TEXT(" Total: %i"), ActiveInstances), GetSmallFont(), FColor(R, G, B));
+		Canvas->DrawShadowedString(X, Y, TEXT(" Total: 0"), GetSmallFont(), FLinearColor::White);
 		Y += 12;
 	}
 	return Y;
@@ -12307,40 +12403,27 @@ int32 UEngine::RenderStatSoundWaves(UWorld* World, FViewport* Viewport, FCanvas*
 // SOUNDCUES
 int32 UEngine::RenderStatSoundCues(UWorld* World, FViewport* Viewport, FCanvas* Canvas, int32 X, int32 Y, const FVector* ViewLocation, const FRotator* ViewRotation)
 {
-	TSet<FActiveSound*> ActiveSounds;
+	Canvas->DrawShadowedString(X, Y, TEXT("Active Sound Cues:"), GetSmallFont(), FColor::Green);
+	Y += 12;
+
+	int32 ActiveSoundCount = 0;
 
 	if (FAudioDevice* AudioDevice = World->GetAudioDevice())
 	{
-		TArray<FWaveInstance*> WaveInstances;
-		int32 FirstActiveIndex = AudioDevice->GetSortedActiveWaveInstances(WaveInstances, ESortedActiveWaveGetType::QueryOnly);
+		const FAudioStats& AudioStats = AudioDevice->GetAudioStats();
 
-		for (int32 InstanceIndex = FirstActiveIndex; InstanceIndex < WaveInstances.Num(); InstanceIndex++)
+		for (const FAudioStats::FStatSoundInfo& StatSoundInfo : AudioDevice->GetAudioStats().StatSoundInfos)
 		{
-			FWaveInstance* WaveInstance = WaveInstances[InstanceIndex];
-			if (WaveInstance->GetActualVolume() >= 0.01f)
-			{
-				ActiveSounds.Add(WaveInstance->ActiveSound);
-			}
+			const FString TheString = FString::Printf(TEXT("%4i. %s %s"), ActiveSoundCount++, *StatSoundInfo.SoundName, *StatSoundInfo.SoundClassName.ToString());
+			Canvas->DrawShadowedString(X, Y, *TheString, GetSmallFont(), FColor::White);
+			Y += 12;
 		}
 	}
 
-	Canvas->DrawShadowedString(X, Y, TEXT("Active Sound Cues:"), GetSmallFont(), FColor::Green);
-	Y += 12;
-
-	int32 ActiveSoundCount = 0;
-	for (FActiveSound* ActiveSound : ActiveSounds)
-	{
-		USoundClass* SoundClass = ActiveSound->GetSoundClass();
-		const FString TheString = FString::Printf(TEXT("%4i. %s %s"), ActiveSoundCount++, *ActiveSound->Sound->GetPathName(), (SoundClass ? *SoundClass->GetName() : TEXT("None")));
-		Canvas->DrawShadowedString(X, Y, *TheString, GetSmallFont(), FColor::White);
-		Y += 12;
-	}
-
-	Canvas->DrawShadowedString(X, Y, *FString::Printf(TEXT("Total: %i"), ActiveSounds.Num()), GetSmallFont(), FColor::Green);
+	Canvas->DrawShadowedString(X, Y, *FString::Printf(TEXT("Total: %i"), ActiveSoundCount), GetSmallFont(), FColor::Green);
 	Y += 12;
 	return Y;
 }
-#endif // !UE_BUILD_SHIPPING
 
 // SOUNDS
 bool UEngine::ToggleStatSounds(UWorld* World, FCommonViewportClient* ViewportClient, const TCHAR* Stream)
@@ -12367,209 +12450,160 @@ bool UEngine::ToggleStatSounds(UWorld* World, FCommonViewportClient* ViewportCli
 		GLog->Logf(TEXT(" This will show only debug sounds sorted by sound class"));
 	}
 
+	uint32 OldSoundShowFlags = ViewportClient->GetSoundShowFlags();
 
 	uint32 ShowSounds = FViewportClient::ESoundShowFlags::Disabled;
 	
-	const bool bDebug = Stream ? FParse::Param(Stream, TEXT("debug")) : false;
-	ShowSounds |= bDebug ? FViewportClient::ESoundShowFlags::Debug : 0;
-
-	const bool bLongNames = Stream ? FParse::Param(Stream, TEXT("longnames")) : false;
-	ShowSounds |= bLongNames ? FViewportClient::ESoundShowFlags::Long_Names : 0;
-
-	FString SortStr;
 	if (Stream)
 	{
-		FParse::Value(Stream, TEXT("sort="), SortStr);
-	}
-	if (SortStr == TEXT("distance"))
-	{
-		ShowSounds |= FViewportClient::ESoundShowFlags::Sort_Distance;
-	}
-	else if (SortStr == TEXT("class"))
-	{
-		ShowSounds |= FViewportClient::ESoundShowFlags::Sort_Class;
-	}
-	else if (SortStr == TEXT("name"))
-	{
-		ShowSounds |= FViewportClient::ESoundShowFlags::Sort_Name;
-	}
-	else if (SortStr == TEXT("waves"))
-	{
-		ShowSounds |= FViewportClient::ESoundShowFlags::Sort_WavesNum;
-	}
-	else
-	{
-		ShowSounds |= FViewportClient::ESoundShowFlags::Sort_Disabled;
-	}
+		const bool bHide = FParse::Command(&Stream, TEXT("off"));
+		if (bHide)
+		{
+			ShowSounds = FViewportClient::ESoundShowFlags::Disabled;
+		}
+		else
+		{
+			const bool bDebug = FParse::Param(Stream, TEXT("debug"));
+			if (bDebug)
+			{
+				ShowSounds |= FViewportClient::ESoundShowFlags::Debug;
+			}
 
-	const bool bHide = Stream ? FParse::Command(&Stream, TEXT("off")) : false;
-	if (bHide)
-	{
-		ShowSounds = FViewportClient::ESoundShowFlags::Disabled;
+			const bool bLongNames = FParse::Param(Stream, TEXT("longnames"));
+			if (bLongNames)
+			{
+				ShowSounds |= FViewportClient::ESoundShowFlags::Long_Names;
+			}
+
+			FString SortStr;
+			FParse::Value(Stream, TEXT("sort="), SortStr);
+			if (SortStr == TEXT("distance"))
+			{
+				ShowSounds |= FViewportClient::ESoundShowFlags::Sort_Distance;
+			}
+			else if (SortStr == TEXT("class"))
+			{
+				ShowSounds |= FViewportClient::ESoundShowFlags::Sort_Class;
+			}
+			else if (SortStr == TEXT("name"))
+			{
+				ShowSounds |= FViewportClient::ESoundShowFlags::Sort_Name;
+			}
+			else if (SortStr == TEXT("waves"))
+			{
+				ShowSounds |= FViewportClient::ESoundShowFlags::Sort_WavesNum;
+			}
+			else
+			{
+				ShowSounds |= FViewportClient::ESoundShowFlags::Sort_Disabled;
+			}
+		}
 	}
 
 	ViewportClient->SetSoundShowFlags((FViewportClient::ESoundShowFlags::Type)ShowSounds);
 
+	if (FAudioDevice* AudioDevice = World->GetAudioDevice())
+	{
+		AudioDevice->UpdateSoundShowFlags(OldSoundShowFlags, ShowSounds);
+	}
+	
 	return true;
 }
 
 int32 UEngine::RenderStatSounds(UWorld* World, FViewport* Viewport, FCanvas* Canvas, int32 X, int32 Y, const FVector* ViewLocation, const FRotator* ViewRotation)
 {
-#if !UE_BUILD_SHIPPING
-
-#if UE_BUILD_DEBUG
-
-	typedef TMap< const FActiveSound*, FSoundInfo* > TMapSounds;
-
-#else
-
-	typedef TMemStackAllocator<> TMemStackGameAllocator;
-	typedef TSetAllocator<TSparseArrayAllocator<TMemStackGameAllocator, TMemStackGameAllocator>, TMemStackGameAllocator> TMapStackAllocator;
-	typedef TMap< const FActiveSound*, FSoundInfo*, TMapStackAllocator > TMapSounds;
-
-#endif // UE_BUILD_DEBUG
-
-	TMapSounds SoundInfos;
 	const FViewportClient::ESoundShowFlags::Type ShowSounds = Viewport->GetClient() ? Viewport->GetClient()->GetSoundShowFlags() : FViewportClient::ESoundShowFlags::Disabled;
 	const bool bDebug = ShowSounds & FViewportClient::ESoundShowFlags::Debug;
 
 	if (FAudioDevice* AudioDevice = World->GetAudioDevice())
 	{
-		// Refresh the wave instances inside audio components.
-		static TArray<FWaveInstance*> WaveInstances;
-		WaveInstances.Reset();
-
-		int32 FirstActiveIndex = AudioDevice->GetSortedActiveWaveInstances(WaveInstances, ESortedActiveWaveGetType::QueryOnly);
-
-		// Grab the list of all active sound cues.
-		const FVector ListenerPosition = AudioDevice->Listeners[0].Transform.GetTranslation();
-
-		const TArray<FActiveSound*>& ActiveSounds = AudioDevice->GetActiveSounds();
-
-		for (int32 Nx = 0; Nx < ActiveSounds.Num(); ++Nx)
+		FAudioStats& AudioStats = AudioDevice->GetAudioStats();
+		if (!AudioStats.bStale)
 		{
-			const FActiveSound* ActiveSound = ActiveSounds[Nx];
+			FString SortingName = TEXT("disabled");
 
-			if (ActiveSound->Sound)
+			// Sort the list.
+			if (ShowSounds & FViewportClient::ESoundShowFlags::Sort_Name)
 			{
-				if (!bDebug || ActiveSound->Sound->bDebug)
-				{
-					const FString PathName = ActiveSound->Sound->GetPathName();
-					const float Distance = (ListenerPosition - ActiveSound->Transform.GetTranslation()).Size();
-					const FName ClassName = (ActiveSound->GetSoundClass() ? ActiveSound->GetSoundClass()->GetFName() : NAME_None);
-
-					SoundInfos.Add(ActiveSound, new FSoundInfo(PathName, Distance, ClassName));
-				}
+				AudioStats.StatSoundInfos.Sort([](const FAudioStats::FStatSoundInfo& A, const FAudioStats::FStatSoundInfo& B) { return A.SoundName < B.SoundName; });
+				SortingName = TEXT("pathname");
 			}
-		}
-
-		// Iterate through all wave instances.
-		for (int32 InstanceIndex = FirstActiveIndex; InstanceIndex < WaveInstances.Num(); ++InstanceIndex)
-		{
-			FWaveInstance* WaveInstance = WaveInstances[InstanceIndex];
-			FSoundInfo* SoundInfo = SoundInfos.FindRef(WaveInstance->ActiveSound);
-			if (SoundInfo)
+			else if (ShowSounds & FViewportClient::ESoundShowFlags::Sort_Distance)
 			{
-				SoundInfo->WaveInstances.Add(WaveInstance);
+				AudioStats.StatSoundInfos.Sort([](const FAudioStats::FStatSoundInfo& A, const FAudioStats::FStatSoundInfo& B) { return A.Distance < B.Distance; });
+				SortingName = TEXT("distance");
+			}
+			else if (ShowSounds & FViewportClient::ESoundShowFlags::Sort_Class)
+			{
+				AudioStats.StatSoundInfos.Sort([](const FAudioStats::FStatSoundInfo& A, const FAudioStats::FStatSoundInfo& B) { return A.SoundClassName < B.SoundClassName; });
+				SortingName = TEXT("class");
+			}
+			else if (ShowSounds & FViewportClient::ESoundShowFlags::Sort_WavesNum)
+			{
+				AudioStats.StatSoundInfos.Sort([](const FAudioStats::FStatSoundInfo& A, const FAudioStats::FStatSoundInfo& B) { return A.WaveInstanceInfos.Num() > B.WaveInstanceInfos.Num(); });
+				SortingName = TEXT("waves' num");
 			}
-		}
-
-		FString SortingName = TEXT("disabled");
-
-		// Sort the list.
-		if (ShowSounds & FViewportClient::ESoundShowFlags::Sort_Name)
-		{
-			SoundInfos.ValueSort(FCompareFSoundInfoByName());
-			SortingName = TEXT("pathname");
-		}
-		else if (ShowSounds & FViewportClient::ESoundShowFlags::Sort_Distance)
-		{
-			SoundInfos.ValueSort(FCompareFSoundInfoByDistance());
-			SortingName = TEXT("distance");
-		}
-		else if (ShowSounds & FViewportClient::ESoundShowFlags::Sort_Class)
-		{
-			SoundInfos.ValueSort(FCompareFSoundInfoByClass());
-			SortingName = TEXT("class");
-		}
-		else if (ShowSounds & FViewportClient::ESoundShowFlags::Sort_WavesNum)
-		{
-			SoundInfos.ValueSort(FCompareFSoundInfoByWaveInstNum());
-			SortingName = TEXT("waves' num");
-		}
-
-
-		Canvas->DrawShadowedString(X, Y, TEXT("Active Sounds:"), GetSmallFont(), FColor::Green);
-		Y += 12;
 
-		const FString InfoText = FString::Printf(TEXT(" Sorting: %s Debug: %s"), *SortingName, bDebug ? TEXT("enabled") : TEXT("disabled"));
-		Canvas->DrawShadowedString(X, Y, *InfoText, GetSmallFont(), FColor(128, 255, 128));
-		Y += 12;
+			Canvas->DrawShadowedString(X, Y, TEXT("Active Sounds:"), GetSmallFont(), FColor::Green);
+			Y += 12;
 
-		Canvas->DrawShadowedString(X, Y, TEXT("Index Path (Class) Distance"), GetSmallFont(), FColor::Green);
-		Y += 12;
+			const FString InfoText = FString::Printf(TEXT(" Sorting: %s Debug: %s"), *SortingName, bDebug ? TEXT("enabled") : TEXT("disabled"));
+			Canvas->DrawShadowedString(X, Y, *InfoText, GetSmallFont(), FColor(128, 255, 128));
+			Y += 12;
 
-		int32 TotalSoundWavesNum = 0;
-		int32 SoundIndex = 0;
-		for (TMapSounds::TConstIterator It(SoundInfos); It; ++It)
-		{
-			const FSoundInfo& SoundInfo = *It.Value();
-			const int32 WaveInstancesNum = SoundInfo.WaveInstances.Num();
+			Canvas->DrawShadowedString(X, Y, TEXT("Index Path (Class) Distance"), GetSmallFont(), FColor::Green);
+			Y += 12;
 
-			if (WaveInstancesNum > 0)
+			int32 TotalSoundWavesNum = 0;
+			for (int32 SoundIndex = 0 ; SoundIndex < AudioStats.StatSoundInfos.Num(); ++SoundIndex)
 			{
+				const FAudioStats::FStatSoundInfo& StatSoundInfo = AudioStats.StatSoundInfos[SoundIndex];
+				const int32 WaveInstancesNum = StatSoundInfo.WaveInstanceInfos.Num();
+				if (WaveInstancesNum > 0)
 				{
-					const FString TheString = FString::Printf(TEXT("%4i. %s (%s) %6.2f"), SoundIndex, *SoundInfo.PathName, *SoundInfo.ClassName.ToString(), SoundInfo.Distance);
-					Canvas->DrawShadowedString(X, Y, *TheString, GetSmallFont(), FColor::White);
-					Y += 12;
-				}
-
-				// Get the active sound waves.
-				for (int32 WaveIndex = 0; WaveIndex < WaveInstancesNum; WaveIndex++)
-				{
-					FWaveInstance* WaveInstance = SoundInfo.WaveInstances[WaveIndex];
-					FSoundSource* Source = AudioDevice->WaveInstanceSourceMap.FindRef(WaveInstance);
-
-					FString SourceDesc = Source ? Source->Describe((ShowSounds & FViewportClient::ESoundShowFlags::Long_Names) != 0) : FString(TEXT("No source"));
-					FString TheString = *FString::Printf(TEXT("    %4i. %s"), WaveIndex, *SourceDesc);
+					{
+						const FString TheString = FString::Printf(TEXT("%4i. %s (%s) %6.2f"), SoundIndex, *StatSoundInfo.SoundName, *StatSoundInfo.SoundClassName.ToString(), StatSoundInfo.Distance);
+						Canvas->DrawShadowedString(X, Y, *TheString, GetSmallFont(), FColor::White);
+						Y += 12;
+					}
 
-					Canvas->DrawShadowedString(X, Y, *TheString, GetSmallFont(), FColor(205, 205, 205));
-					Y += 12;
+					TotalSoundWavesNum += WaveInstancesNum;
 
-					TotalSoundWavesNum++;
+					// Get the active sound waves.
+					for (int32 WaveIndex = 0; WaveIndex < WaveInstancesNum; WaveIndex++)
+					{
+						const FString TheString = *FString::Printf(TEXT("    %4i. %s"), WaveIndex, *StatSoundInfo.WaveInstanceInfos[WaveIndex].Description);
+						Canvas->DrawShadowedString(X, Y, *TheString, GetSmallFont(), FColor(205, 205, 205));
+						Y += 12;
+					}
 				}
-				++SoundIndex;
 			}
-		}
 
-		Canvas->DrawShadowedString(X, Y, *FString::Printf(TEXT("Total sounds: %i, sound waves: %i"), SoundIndex, TotalSoundWavesNum), GetSmallFont(), FColor::Green);
-		Y += 12;
+			Canvas->DrawShadowedString(X, Y, *FString::Printf(TEXT("Total sounds: %i, sound waves: %i"), AudioStats.StatSoundInfos.Num(), TotalSoundWavesNum), GetSmallFont(), FColor::Green);
+			Y += 12;
 
-		Canvas->DrawShadowedString(X, Y, *FString::Printf(TEXT("Listener position: %s"), *ListenerPosition.ToString()), GetSmallFont(), FColor::Green);
-		Y += 12;
+			Canvas->DrawShadowedString(X, Y, *FString::Printf(TEXT("Listener position: %s"), *AudioStats.ListenerLocation.ToString()), GetSmallFont(), FColor::Green);
+			Y += 12;
+		}
 
 		// Draw sound cue's sphere.
 		if (bDebug)
 		{
-			for (TMapSounds::TConstIterator SoundInfoIt(SoundInfos); SoundInfoIt; ++SoundInfoIt)
+			for (const FAudioStats::FStatSoundInfo& StatSoundInfo : AudioStats.StatSoundInfos)
 			{
-				const FActiveSound& ActiveSound = *SoundInfoIt.Key();
-				const FSoundInfo& SoundInfo = *SoundInfoIt.Value();
-				const int32 WaveInstancesNum = SoundInfo.WaveInstances.Num();
+				const FTransform& SoundTransform = StatSoundInfo.Transform;
+				const int32 WaveInstancesNum = StatSoundInfo.WaveInstanceInfos.Num();
 
-				if (ActiveSound.Sound->bDebug && SoundInfo.Distance > 100.0f && WaveInstancesNum > 0)
+				if (StatSoundInfo.Distance > 100.0f && WaveInstancesNum > 0)
 				{
 					float SphereRadius = 0.f;
 					float SphereInnerRadius = 0.f;
 
-					TMultiMap<EAttenuationShape::Type, FAttenuationSettings::AttenuationShapeDetails> ShapeDetailsMap;
-					ActiveSound.CollectAttenuationShapesForVisualization(ShapeDetailsMap);
-
-					if (ShapeDetailsMap.Num() > 0)
+					if (StatSoundInfo.ShapeDetailsMap.Num() > 0)
 					{
-						DrawDebugString(World, ActiveSound.Transform.GetTranslation(), SoundInfo.PathName, NULL, FColor::White, 0.01f);
+						DrawDebugString(World, SoundTransform.GetTranslation(), StatSoundInfo.SoundName, NULL, FColor::White, 0.01f);
 
-						for (auto ShapeDetailsIt = ShapeDetailsMap.CreateConstIterator(); ShapeDetailsIt; ++ShapeDetailsIt)
+						for (auto ShapeDetailsIt = StatSoundInfo.ShapeDetailsMap.CreateConstIterator(); ShapeDetailsIt; ++ShapeDetailsIt)
 						{
 							const FAttenuationSettings::AttenuationShapeDetails& ShapeDetails = ShapeDetailsIt.Value();
 							switch (ShapeDetailsIt.Key())
@@ -12577,24 +12611,24 @@ int32 UEngine::RenderStatSounds(UWorld* World, FViewport* Viewport, FCanvas* Can
 							case EAttenuationShape::Sphere:
 								if (ShapeDetails.Falloff > 0.f)
 								{
-									DrawDebugSphere(World, ActiveSound.Transform.GetTranslation(), ShapeDetails.Extents.X + ShapeDetails.Falloff, 10, FColor(155, 155, 255));
-									DrawDebugSphere(World, ActiveSound.Transform.GetTranslation(), ShapeDetails.Extents.X, 10, FColor(55, 55, 255));
+									DrawDebugSphere(World, SoundTransform.GetTranslation(), ShapeDetails.Extents.X + ShapeDetails.Falloff, 10, FColor(155, 155, 255));
+									DrawDebugSphere(World, SoundTransform.GetTranslation(), ShapeDetails.Extents.X, 10, FColor(55, 55, 255));
 								}
 								else
 								{
-									DrawDebugSphere(World, ActiveSound.Transform.GetTranslation(), ShapeDetails.Extents.X, 10, FColor(155, 155, 255));
+									DrawDebugSphere(World, SoundTransform.GetTranslation(), ShapeDetails.Extents.X, 10, FColor(155, 155, 255));
 								}
 								break;
 
 							case EAttenuationShape::Box:
 								if (ShapeDetails.Falloff > 0.f)
 								{
-									DrawDebugBox(World, ActiveSound.Transform.GetTranslation(), ShapeDetails.Extents + FVector(ShapeDetails.Falloff), ActiveSound.Transform.GetRotation(), FColor(155, 155, 255));
-									DrawDebugBox(World, ActiveSound.Transform.GetTranslation(), ShapeDetails.Extents, ActiveSound.Transform.GetRotation(), FColor(55, 55, 255));
+									DrawDebugBox(World, SoundTransform.GetTranslation(), ShapeDetails.Extents + FVector(ShapeDetails.Falloff), SoundTransform.GetRotation(), FColor(155, 155, 255));
+									DrawDebugBox(World, SoundTransform.GetTranslation(), ShapeDetails.Extents, SoundTransform.GetRotation(), FColor(55, 55, 255));
 								}
 								else
 								{
-									DrawDebugBox(World, ActiveSound.Transform.GetTranslation(), ShapeDetails.Extents, ActiveSound.Transform.GetRotation(), FColor(155, 155, 255));
+									DrawDebugBox(World, SoundTransform.GetTranslation(), ShapeDetails.Extents, SoundTransform.GetRotation(), FColor(155, 155, 255));
 								}
 								break;
 
@@ -12602,30 +12636,30 @@ int32 UEngine::RenderStatSounds(UWorld* World, FViewport* Viewport, FCanvas* Can
 
 								if (ShapeDetails.Falloff > 0.f)
 								{
-									DrawDebugCapsule(World, ActiveSound.Transform.GetTranslation(), ShapeDetails.Extents.X + ShapeDetails.Falloff, ShapeDetails.Extents.Y + ShapeDetails.Falloff, ActiveSound.Transform.GetRotation(), FColor(155, 155, 255));
-									DrawDebugCapsule(World, ActiveSound.Transform.GetTranslation(), ShapeDetails.Extents.X, ShapeDetails.Extents.Y, ActiveSound.Transform.GetRotation(), FColor(55, 55, 255));
+									DrawDebugCapsule(World, SoundTransform.GetTranslation(), ShapeDetails.Extents.X + ShapeDetails.Falloff, ShapeDetails.Extents.Y + ShapeDetails.Falloff, SoundTransform.GetRotation(), FColor(155, 155, 255));
+									DrawDebugCapsule(World, SoundTransform.GetTranslation(), ShapeDetails.Extents.X, ShapeDetails.Extents.Y, SoundTransform.GetRotation(), FColor(55, 55, 255));
 								}
 								else
 								{
-									DrawDebugCapsule(World, ActiveSound.Transform.GetTranslation(), ShapeDetails.Extents.X, ShapeDetails.Extents.Y, ActiveSound.Transform.GetRotation(), FColor(155, 155, 255));
+									DrawDebugCapsule(World, SoundTransform.GetTranslation(), ShapeDetails.Extents.X, ShapeDetails.Extents.Y, SoundTransform.GetRotation(), FColor(155, 155, 255));
 								}
 								break;
 
 							case EAttenuationShape::Cone:
 							{
-								const FVector Origin = ActiveSound.Transform.GetTranslation() - (ActiveSound.Transform.GetUnitAxis(EAxis::X) * ShapeDetails.ConeOffset);
+								const FVector Origin = SoundTransform.GetTranslation() - (SoundTransform.GetUnitAxis(EAxis::X) * ShapeDetails.ConeOffset);
 
 								if (ShapeDetails.Falloff > 0.f || ShapeDetails.Extents.Z > 0.f)
 								{
 									const float OuterAngle = FMath::DegreesToRadians(ShapeDetails.Extents.Y + ShapeDetails.Extents.Z);
 									const float InnerAngle = FMath::DegreesToRadians(ShapeDetails.Extents.Y);
-									DrawDebugCone(World, Origin, ActiveSound.Transform.GetUnitAxis(EAxis::X), ShapeDetails.Extents.X + ShapeDetails.Falloff + ShapeDetails.ConeOffset, OuterAngle, OuterAngle, 10, FColor(155, 155, 255));
-									DrawDebugCone(World, Origin, ActiveSound.Transform.GetUnitAxis(EAxis::X), ShapeDetails.Extents.X + ShapeDetails.ConeOffset, InnerAngle, InnerAngle, 10, FColor(55, 55, 255));
+									DrawDebugCone(World, Origin, SoundTransform.GetUnitAxis(EAxis::X), ShapeDetails.Extents.X + ShapeDetails.Falloff + ShapeDetails.ConeOffset, OuterAngle, OuterAngle, 10, FColor(155, 155, 255));
+									DrawDebugCone(World, Origin, SoundTransform.GetUnitAxis(EAxis::X), ShapeDetails.Extents.X + ShapeDetails.ConeOffset, InnerAngle, InnerAngle, 10, FColor(55, 55, 255));
 								}
 								else
 								{
 									const float Angle = FMath::DegreesToRadians(ShapeDetails.Extents.Y);
-									DrawDebugCone(World, Origin, ActiveSound.Transform.GetUnitAxis(EAxis::X), ShapeDetails.Extents.X + ShapeDetails.ConeOffset, Angle, Angle, 10, FColor(155, 155, 255));
+									DrawDebugCone(World, Origin, SoundTransform.GetUnitAxis(EAxis::X), ShapeDetails.Extents.X + ShapeDetails.ConeOffset, Angle, Angle, 10, FColor(155, 155, 255));
 								}
 								break;
 							}
@@ -12638,15 +12672,10 @@ int32 UEngine::RenderStatSounds(UWorld* World, FViewport* Viewport, FCanvas* Can
 				}
 			}
 		}
-
-		for (TMapSounds::TConstIterator It(SoundInfos); It; ++It)
-		{
-			delete It.Value();
-		}
 	}
-#endif // !UE_BUILD_SHIPPING
 	return Y;
 }
+#endif // !UE_BUILD_SHIPPING
 
 // AI
 int32 UEngine::RenderStatAI(UWorld* World, FViewport* Viewport, FCanvas* Canvas, int32 X, int32 Y, const FVector* ViewLocation, const FRotator* ViewRotation)
diff --git a/Engine/Source/Runtime/Engine/Private/Vehicles/WheeledVehicleMovementComponent.cpp b/Engine/Source/Runtime/Engine/Private/Vehicles/WheeledVehicleMovementComponent.cpp
--- a/Engine/Source/Runtime/Engine/Private/Vehicles/WheeledVehicleMovementComponent.cpp
+++ b/Engine/Source/Runtime/Engine/Private/Vehicles/WheeledVehicleMovementComponent.cpp
@@ -10,6 +10,7 @@
 #include "Vehicles/TireType.h"
 #include "DisplayDebugHelpers.h"
 #include "PhysicsEngine/ConstraintInstance.h"
+#include "PhysicsEngine/PhysicsAsset.h"
 
 #if WITH_PHYSX
 #include "../PhysicsEngine/PhysXSupport.h"
@@ -1416,7 +1417,7 @@ void UWheeledVehicleMovementComponent::FixupSkeletalMesh()
 						FBodyInstance* BodyInstance = Mesh->Bodies[BodySetupIdx];
 						BodyInstance->SetResponseToAllChannels(ECR_Ignore);	//turn off collision for wheel automatically
 
-						if (UBodySetup * BodySetup = PhysicsAsset->BodySetup[BodySetupIdx])
+						if (UBodySetup * BodySetup = PhysicsAsset->SkeletalBodySetups[BodySetupIdx])
 						{
 
 							if (BodySetup->PhysicsType == PhysType_Default) 	//if they set it to unfixed we don't fixup because they are explicitely saying Unfixed
diff --git a/Engine/Source/Runtime/Engine/Private/WorldSettings.cpp b/Engine/Source/Runtime/Engine/Private/WorldSettings.cpp
--- a/Engine/Source/Runtime/Engine/Private/WorldSettings.cpp
+++ b/Engine/Source/Runtime/Engine/Private/WorldSettings.cpp
@@ -116,6 +116,17 @@ void AWorldSettings::PostInitializeComponents()
 	}
 }
 
+void AWorldSettings::PostRegisterAllComponents()
+{
+	Super::PostRegisterAllComponents();
+
+	UWorld* World = GetWorld();
+	if (FAudioDevice* AudioDevice = World->GetAudioDevice())
+	{
+		AudioDevice->SetDefaultAudioSettings(World, DefaultReverbSettings, DefaultAmbientZoneSettings);
+	}
+}
+
 float AWorldSettings::GetGravityZ() const
 {
 	if (!bWorldGravitySet)
@@ -307,17 +318,33 @@ bool AWorldSettings::CanEditChange(const UProperty* InProperty) const
 	return Super::CanEditChange(InProperty);
 }
 
+void AWorldSettings::PostEditChangeChainProperty(FPropertyChangedChainEvent& PropertyChangedEvent)
+{
+	Super::PostEditChangeChainProperty(PropertyChangedEvent);
+
+	const FName MemberPropertyName = PropertyChangedEvent.PropertyChain.GetActiveMemberNode()->GetValue()->GetFName();
+
+	if (MemberPropertyName == GET_MEMBER_NAME_CHECKED(AWorldSettings, DefaultReverbSettings) || MemberPropertyName == GET_MEMBER_NAME_CHECKED(AWorldSettings, DefaultAmbientZoneSettings))
+	{
+		UWorld* World = GetWorld();
+		if (FAudioDevice* AudioDevice = World->GetAudioDevice())
+		{
+			AudioDevice->SetDefaultAudioSettings(World, DefaultReverbSettings, DefaultAmbientZoneSettings);
+		}
+	}
+}
+
 void AWorldSettings::PostEditChangeProperty(FPropertyChangedEvent& PropertyChangedEvent)
 {
 	UProperty* PropertyThatChanged = PropertyChangedEvent.Property;
 	if (PropertyThatChanged)
 	{
-		if (PropertyThatChanged->GetName()==TEXT("bForceNoPrecomputedLighting") && bForceNoPrecomputedLighting)
+		if (PropertyThatChanged->GetFName()==GET_MEMBER_NAME_CHECKED(AWorldSettings,bForceNoPrecomputedLighting) && bForceNoPrecomputedLighting)
 		{
 			FMessageDialog::Open( EAppMsgType::Ok, LOCTEXT("bForceNoPrecomputedLightingIsEnabled", "bForceNoPrecomputedLighting is now enabled, build lighting once to propagate the change (will remove existing precomputed lighting data)."));
 		}
 
-		if (PropertyThatChanged->GetName()==TEXT("bEnableWorldComposition"))
+		else if (PropertyThatChanged->GetFName()==GET_MEMBER_NAME_CHECKED(AWorldSettings,bEnableWorldComposition))
 		{
 			if (UWorldComposition::EnableWorldCompositionEvent.IsBound())
 			{
@@ -354,12 +381,12 @@ void AWorldSettings::PostEditChangeProperty(FPropertyChangedEvent& PropertyChang
 			GEngine->DeferredCommands.AddUnique(TEXT("UpdateLandscapeSetup"));
 		}
 
-		if (PropertyThatChanged->GetName() == TEXT("TransitionScreenSize"))
+		if (PropertyThatChanged->GetFName() == GET_MEMBER_NAME_CHECKED(FHierarchicalSimplification,TransitionScreenSize))
 		{
 			GEditor->BroadcastHLODTransitionScreenSizeChanged();
 		}
 
-		if (PropertyThatChanged->GetName() == TEXT("HierarchicalLODSetup"))
+		else if (PropertyThatChanged->GetFName() == GET_MEMBER_NAME_CHECKED(AWorldSettings,HierarchicalLODSetup))
 		{
 			GEditor->BroadcastHLODLevelsArrayChanged();
 			NumHLODLevels = HierarchicalLODSetup.Num();			
diff --git a/Engine/Source/Runtime/Engine/Public/ActiveSound.h b/Engine/Source/Runtime/Engine/Public/ActiveSound.h
--- a/Engine/Source/Runtime/Engine/Public/ActiveSound.h
+++ b/Engine/Source/Runtime/Engine/Public/ActiveSound.h
@@ -7,6 +7,8 @@
 #include "Sound/SoundAttenuation.h"
 #include "Components/AudioComponent.h"
 
+class UWorld;
+
 /**
  *	Struct used for gathering the final parameters to apply to a wave instance
  */
@@ -27,6 +29,9 @@ struct FSoundParseParameters
 	float Volume;
 	float VolumeMultiplier;
 
+	// The multiplier to apply if the sound class desires
+	float InteriorVolumeMultiplier;
+
 	// The priority of sound, which is the product of the component priority and the USoundBased priority
 	float Priority;
 
@@ -73,6 +78,7 @@ struct FSoundParseParameters
 		, Velocity(ForceInit)
 		, Volume(1.f)
 		, VolumeMultiplier(1.f)
+		, InteriorVolumeMultiplier(1.f)
 		, Pitch(1.f)
 		, StartTime(-1.f)
 		, OmniRadius(0.0f)
@@ -97,23 +103,39 @@ public:
 	FActiveSound();
 	~FActiveSound();
 
-	class USoundBase* Sound;
-	TWeakObjectPtr<class UWorld> World;
 
 private:
-	TWeakObjectPtr<class UAudioComponent> AudioComponent;
-	UPTRINT AudioComponentIndex;
+	TWeakObjectPtr<UWorld> World;
+	uint32 WorldID;
+
+	class USoundBase* Sound;
+
+	uint64 AudioComponentID;
+	uint32 OwnerID;
+	
+	FName AudioComponentName;
+	FName OwnerName;
 
 public:
 
-	bool IsAudioComponentValid() const { return AudioComponent.IsValid(); }
-	UAudioComponent* GetAudioComponent() const { return AudioComponent.Get(); }
-	UPTRINT GetAudioComponentIndex() const { return AudioComponentIndex; }
-	void SetAudioComponent(UAudioComponent* Component)
-	{
-		AudioComponent = Component;
-		AudioComponentIndex = (UPTRINT)Component;
+	uint64 GetAudioComponentID() const { return AudioComponentID; }
+	void SetAudioComponent(UAudioComponent* Component);
+	FString GetAudioComponentName() const;
+	FString GetOwnerName() const;
+
+	uint32 GetWorldID() const { return WorldID; }
+	TWeakObjectPtr<UWorld> GetWeakWorld() const { return World; }
+	UWorld* GetWorld() const 
+	{ 
+		check(IsInGameThread()); 
+		return World.Get();
 	}
+	void SetWorld(UWorld* World);
+
+	USoundBase* GetSound() const { return Sound; }
+	void SetSound(USoundBase* InSound);
+
+	void SetSoundClass(USoundClass* SoundClass);
 
 	void SetAudioDevice(FAudioDevice* InAudioDevice)
 	{
@@ -134,81 +156,94 @@ public:
 	/** Optional USoundConcurrency to override sound */
 	USoundConcurrency* ConcurrencySettings;
 
+private:
 	/** Optional SoundClass to override Sound */
 	USoundClass* SoundClassOverride;
 
+public:
 	/** Whether or not the sound has checked if it was occluded already. Used to initialize a sound as occluded and bypassing occlusion interpolation. */
-	uint32 bHasCheckedOcclusion:1;
+	uint8 bHasCheckedOcclusion:1;
 
 	/** Flag to trigger binding our trace delegate for async trace calls */
-	uint32 bIsTraceDelegateBound:1;
+	uint8 bIsTraceDelegateBound:1;
 
 	/** Is this sound allowed to be spatialized? */
-	uint32 bAllowSpatialization:1;
+	uint8 bAllowSpatialization:1;
 
 	/** Does this sound have attenuation settings specified */
-	uint32 bHasAttenuationSettings:1;
+	uint8 bHasAttenuationSettings:1;
 
 	/** Whether the wave instances should remain active if they're dropped by the prioritization code. Useful for e.g. vehicle sounds that shouldn't cut out. */
-	uint32 bShouldRemainActiveIfDropped:1;
+	uint8 bShouldRemainActiveIfDropped:1;
 
 	/** Is the audio component currently fading out */
-	uint32 bFadingOut:1;
+	uint8 bFadingOut:1;
 
 	/** Whether the current component has finished playing */
-	uint32 bFinished:1;
+	uint8 bFinished:1;
 
 	/** Whether or not to stop this active sound due to max concurrency */
-	uint32 bShouldStopDueToMaxConcurrency:1;
+	uint8 bShouldStopDueToMaxConcurrency:1;
 
 	/** If true, the decision on whether to apply the radio filter has been made. */
-	uint32 bRadioFilterSelected:1;
+	uint8 bRadioFilterSelected:1;
 
 	/** If true, this sound will not be stopped when flushing the audio device. */
-	uint32 bApplyRadioFilter:1;
+	uint8 bApplyRadioFilter:1;
 
 	/** If true, the AudioComponent will be notified when a Wave is started to handle subtitles */
-	uint32 bHandleSubtitles:1;
+	uint8 bHandleSubtitles:1;
 
 	/** Whether the Location of the component is well defined */
-	uint32 bLocationDefined:1;
+	uint8 bLocationDefined:1;
 
 	/** If true, this sound will not be stopped when flushing the audio device. */
-	uint32 bIgnoreForFlushing:1;
+	uint8 bIgnoreForFlushing:1;
 
 	/** Whether audio effects are applied */
-	uint32 bEQFilterApplied:1;
+	uint8 bEQFilterApplied:1;
 
 	/** Whether to artificially prioritize the component to play */
-	uint32 bAlwaysPlay:1;
+	uint8 bAlwaysPlay:1;
 
 	/** Whether or not this sound plays when the game is paused in the UI */
-	uint32 bIsUISound:1;
+	uint8 bIsUISound:1;
 
 	/** Whether or not this audio component is a music clip */
-	uint32 bIsMusic:1;
+	uint8 bIsMusic:1;
 
 	/** Whether or not the audio component should be excluded from reverb EQ processing */
-	uint32 bReverb:1;
+	uint8 bReverb:1;
 
 	/** Whether or not this sound class forces sounds to the center channel */
-	uint32 bCenterChannelOnly:1;
+	uint8 bCenterChannelOnly:1;
 
 	/** Whether or not this active sound is a preview sound */
-	uint32 bIsPreviewSound:1;
+	uint8 bIsPreviewSound:1;
 
 	/** Whether we have queried for the interior settings at least once */
-	uint32 bGotInteriorSettings:1;
+	uint8 bGotInteriorSettings:1;
+
+	/** Whether some part of this sound will want interior sounds to be applied */
+	uint8 bApplyInteriorVolumes:1;
 
 #if !(NO_LOGGING || UE_BUILD_SHIPPING || UE_BUILD_TEST)
 	/** For debugging purposes, output to the log once that a looping sound has been orphaned */
-	uint32 bWarnedAboutOrphanedLooping:1;
+	uint8 bWarnedAboutOrphanedLooping:1;
 #endif
 
 	/** Whether or not we have a low-pass filter enabled on this active sound. */
-	uint32 bEnableLowPassFilter : 1;
+	uint8 bEnableLowPassFilter : 1;
+
 	/** Whether or not to use an async trace for occlusion */
-	uint32 bOcclusionAsyncTrace : 1;
+	uint8 bOcclusionAsyncTrace : 1;
+
+private:
+	/** Whether or not this sound is audible */
+	uint8 bIsAudible : 1;
+
+public:
+	uint8 UserIndex;
 
 	/** whether we were occluded the last time we checked */
 	FThreadSafeBool bIsOccluded;
@@ -216,8 +251,6 @@ public:
 	/** Whether or not there is an async occlusion trace pending */
 	FThreadSafeBool bAsyncOcclusionPending;
 
-	uint8 UserIndex;
-
 	float PlaybackTime;
 	float RequestedStartTime;
 
@@ -275,9 +308,9 @@ public:
 	FAttenuationSettings AttenuationSettings;
 
 	/** Cache what volume settings we had last time so we don't have to search again if we didn't move */
-	FInteriorSettings LastInteriorSettings;
+	FInteriorSettings InteriorSettings;
 
-	class AAudioVolume* LastAudioVolume;
+	uint32 AudioVolumeID;
 
 	// To remember where the volumes are interpolating to and from
 	double LastUpdateTime; 
@@ -367,7 +400,7 @@ public:
 	int32 FindClosestListener( const TArray<struct FListener>& InListeners ) const;
 	
 	/** Returns the unique ID of the active sound's owner if it exists. Returns 0 if the sound doesn't have an owner. */
-	uint32 TryGetOwnerID() const;
+	uint32 GetOwnerID() const { return OwnerID; }
 
 	/** Gets the sound concurrency to apply on this active sound instance */
 	const FSoundConcurrencySettings* GetSoundConcurrencySettingsToApply() const;
@@ -384,9 +417,6 @@ public:
 
 private:
 	
-	/** Whether or not this sound is audible */
-	bool bIsAudible;
-
 	/** Cached ptr to the closest listener. So we don't have to do the work to find it twice. */
 	const FListener* ClosestListenerPtr;
 
diff --git a/Engine/Source/Runtime/Engine/Public/Animation/AnimCurveTypes.h b/Engine/Source/Runtime/Engine/Public/Animation/AnimCurveTypes.h
--- a/Engine/Source/Runtime/Engine/Public/Animation/AnimCurveTypes.h
+++ b/Engine/Source/Runtime/Engine/Public/Animation/AnimCurveTypes.h
@@ -11,26 +11,26 @@
 enum EAnimCurveFlags
 {
 	// Used as morph target curve
-	ACF_DrivesMorphTarget	= 0x00000001,
+	ACF_DriveMorphTarget	= 0x00000001,
 	// Used as triggering event
-	ACF_TriggerEvent		= 0x00000002,
+	ACF_DriveAttribute		= 0x00000002,
 	// Is editable in Sequence Editor
 	ACF_Editable			= 0x00000004,
 	// Used as a material curve
-	ACF_DrivesMaterial		= 0x00000008,
+	ACF_DriveMaterial		= 0x00000008,
 	// Is a metadata 'curve'
 	ACF_Metadata			= 0x00000010,
 	// motifies bone track
 	ACF_DriveTrack			= 0x00000020,
 	// disabled, right now it's used by track
 	ACF_Disabled			= 0x00000040,
-	// pose curve
-	ACF_DrivesPose			= 0x00000080,
 
 	// default flag when created
-	ACF_DefaultCurve		= ACF_TriggerEvent | ACF_Editable,
+	ACF_DefaultCurve		= ACF_DriveAttribute | ACF_Editable,
 	// curves created from Morph Target
-	ACF_MorphTargetCurve	= ACF_DrivesMorphTarget
+	ACF_MorphTargetCurve	= ACF_DriveMorphTarget, 
+	// all editor preview curves - this will go away soon once skeleton contains curve types
+	ACF_EditorPreviewCurves	= ACF_DriveMorphTarget | ACF_DriveAttribute | ACF_DriveMaterial
 };
 
 /** UI representation of EAnimCurveFlags. This is used in Animation Nodes to set custom curves 
@@ -64,17 +64,17 @@ struct ENGINE_API FAnimCurveType
 
 		if (bMorphtarget)
 		{
-			OutFlags |= ACF_DrivesMorphTarget;
+			OutFlags |= ACF_DriveMorphTarget;
 		}
 
 		if (bEvent)
 		{
-			OutFlags |= ACF_TriggerEvent;
+			OutFlags |= ACF_DriveAttribute;
 		}
 
 		if (bMaterial)
 		{
-			OutFlags |= ACF_DrivesMaterial;
+			OutFlags |= ACF_DriveMaterial;
 		}
 
 		return OutFlags;
@@ -391,7 +391,7 @@ struct FBaseBlendedCurve
 	 * Blend (A, B) using Alpha, same as Lerp
 	 */
 	//@Todo curve flags won't transfer over - it only overwrites
-	void Blend(const FBaseBlendedCurve& A, const FBaseBlendedCurve& B, float Alpha)
+	void Lerp(const FBaseBlendedCurve& A, const FBaseBlendedCurve& B, float Alpha)
 	{
 		check(A.Num() == B.Num());
 		if (FMath::Abs(Alpha) <= ZERO_ANIMWEIGHT_THRESH)
@@ -418,7 +418,7 @@ struct FBaseBlendedCurve
 	/**
 	 * Blend with Other using Alpha, same as Lerp 
 	 */
-	void BlendWith(const FBaseBlendedCurve& Other, float Alpha)
+	void LerpTo(const FBaseBlendedCurve& Other, float Alpha)
 	{
 		check(Num() == Other.Num());
 		if (FMath::Abs(Alpha) <= ZERO_ANIMWEIGHT_THRESH)
diff --git a/Engine/Source/Runtime/Engine/Public/Animation/AnimInstanceProxy.h b/Engine/Source/Runtime/Engine/Public/Animation/AnimInstanceProxy.h
--- a/Engine/Source/Runtime/Engine/Public/Animation/AnimInstanceProxy.h
+++ b/Engine/Source/Runtime/Engine/Public/Animation/AnimInstanceProxy.h
@@ -10,6 +10,7 @@
 
 struct FAnimNode_Base;
 struct FAnimNode_SaveCachedPose;
+struct FAnimNode_SubInput;
 
 /** Proxy object passed around during animation tree update in lieu of a UAnimInstance */
 USTRUCT(meta = (DisplayName = "Native Variables"))
@@ -32,6 +33,7 @@ public:
 		, bIsBeingDebugged(false)
 #endif
 		, RootNode(nullptr)
+		, SubInstanceInputNode(nullptr)
 		, SyncGroupWriteIndex(0)
 		, RootMotionMode(ERootMotionMode::NoRootMotionExtraction)
 		, bShouldExtractRootMotion(false)
@@ -182,7 +184,7 @@ public:
 	}
 
 	/** Get the current skeletal mesh component we are running on. Note that this will return nullptr outside of pre/post update */
-	USkeletalMeshComponent* GetSkelMeshComponent() 
+	USkeletalMeshComponent* GetSkelMeshComponent() const
 	{ 
 		// Skeleton is only available during update/eval. If you're calling this function outside of it, it will return null. 
 		// adding ensure here so that we can catch them earlier
@@ -276,6 +278,7 @@ public:
 	friend class UAnimInstance;
 	friend class UAnimSingleNodeInstance;
 	friend class USkeletalMeshComponent;
+	friend struct FAnimNode_SubInstance;
 
 protected:
 	/** Called when our anim instance is being initialized */
@@ -519,6 +522,9 @@ private:
 	/** Anim graph */
 	FAnimNode_Base* RootNode;
 
+	/** Subinstance input node if available */
+	FAnimNode_SubInput* SubInstanceInputNode;
+
 	/** List of saved pose nodes to process after the graph has been updated */
 	TArray<FAnimNode_SaveCachedPose*> SavedPoseQueue;
 
