Copying //UE4/Dev-Framework to //UE4/Dev-Main (Source: //UE4/Dev-Framework @ 3038004)

diff --git a/Engine/Source/Runtime/Engine/Public/Animation/AnimSingleNodeInstanceProxy.h b/Engine/Source/Runtime/Engine/Public/Animation/AnimSingleNodeInstanceProxy.h
--- a/Engine/Source/Runtime/Engine/Public/Animation/AnimSingleNodeInstanceProxy.h
+++ b/Engine/Source/Runtime/Engine/Public/Animation/AnimSingleNodeInstanceProxy.h
@@ -114,7 +114,7 @@ public:
 		return bCanProcessAdditiveAnimations;
 	}
 #endif
-	void SetPreviewPoseOverride(const FName& PoseName, float Value);
+	void SetPreviewCurveOverride(const FName& PoseName, float Value, bool bRemoveIfZero);
 
 private:
 	void InternalBlendSpaceEvaluatePose(class UBlendSpaceBase* BlendSpace, TArray<FBlendSampleData>& BlendSampleDataCache, FPoseContext& OutContext);
@@ -127,7 +127,7 @@ protected:
 #endif
 
 	/** Pose Weight value that can override curve data. In the future, we'd like to have UCurveSet that can play by default**/
-	TMap<FName, float> PreviewPoseOverride;
+	TMap<FName, float> PreviewCurveOverride;
 
 	/** Current Asset being played. Note that this will be nullptr outside of pre/post update **/
 	UAnimationAsset* CurrentAsset;
diff --git a/Engine/Source/Runtime/Engine/Public/Animation/AnimTypes.h b/Engine/Source/Runtime/Engine/Public/Animation/AnimTypes.h
--- a/Engine/Source/Runtime/Engine/Public/Animation/AnimTypes.h
+++ b/Engine/Source/Runtime/Engine/Public/Animation/AnimTypes.h
@@ -21,6 +21,16 @@ namespace EAnimEventTriggerOffsets
 
 ENGINE_API float GetTriggerTimeOffsetForType(EAnimEventTriggerOffsets::Type OffsetType);
 
+/** Enum for specifying a specific axis of a bone */
+UENUM()
+enum EBoneAxis
+{
+	BA_X UMETA(DisplayName = "X Axis"),
+	BA_Y UMETA(DisplayName = "Y Axis"),
+	BA_Z UMETA(DisplayName = "Z Axis"),
+};
+
+
 /** Enum for controlling which reference frame a controller is applied in. */
 UENUM()
 enum EBoneControlSpace
diff --git a/Engine/Source/Runtime/Engine/Public/Audio.h b/Engine/Source/Runtime/Engine/Public/Audio.h
--- a/Engine/Source/Runtime/Engine/Public/Audio.h
+++ b/Engine/Source/Runtime/Engine/Public/Audio.h
@@ -61,7 +61,6 @@ ENGINE_API DECLARE_LOG_CATEGORY_EXTERN(LogAudioDebug, Display, All);
 /**
  * Audio stats
  */
-DECLARE_CYCLE_STAT_EXTERN( TEXT( "Audio Update Time" ), STAT_AudioUpdateTime, STATGROUP_Audio , );
 DECLARE_DWORD_COUNTER_STAT_EXTERN( TEXT( "Active Sounds" ), STAT_ActiveSounds, STATGROUP_Audio , );
 DECLARE_CYCLE_STAT_EXTERN(TEXT("Audio Evaluate Concurrency"), STAT_AudioEvaluateConcurrency, STATGROUP_Audio, );
 DECLARE_DWORD_COUNTER_STAT_EXTERN( TEXT( "Audio Sources" ), STAT_AudioSources, STATGROUP_Audio , );
diff --git a/Engine/Source/Runtime/Engine/Public/AudioDevice.h b/Engine/Source/Runtime/Engine/Public/AudioDevice.h
--- a/Engine/Source/Runtime/Engine/Public/AudioDevice.h
+++ b/Engine/Source/Runtime/Engine/Public/AudioDevice.h
@@ -14,6 +14,7 @@ class FAudioEffectsManager;
 class ICompressedAudioInfo;
 class IAudioSpatializationPlugin;
 class IAudioSpatializationAlgorithm;
+class UReverbEffect;
 
 /** 
  * Debug state of the audio system
@@ -69,6 +70,15 @@ namespace ESortedActiveWaveGetType
 	};
 }
 
+namespace ERequestedAudioStats
+{
+	static const uint8 SoundWaves = 0x1;
+	static const uint8 SoundCues = 0x2;
+	static const uint8 Sounds = 0x4;
+	static const uint8 DebugSounds = 0x8;
+	static const uint8 LongSoundNames = 0x01;
+};
+
 /** 
  * Defines the properties of the listener
  */
@@ -79,8 +89,11 @@ struct FListener
 
 	struct FInteriorSettings InteriorSettings;
 
-	/** The volume the listener resides in */
-	AAudioVolume* Volume;
+	/** The ID of the volume the listener resides in */
+	uint32 AudioVolumeID;
+
+	/** The ID of the world the listener resides in */
+	uint32 WorldID;
 
 	/** The times of interior volumes fading in and out */
 	double InteriorStartTime;
@@ -110,12 +123,12 @@ struct FListener
 	/** 
 	 * Apply the interior settings to ambient sounds
 	 */
-	void ApplyInteriorSettings(AAudioVolume* Volume, const FInteriorSettings& Settings);
+	void ApplyInteriorSettings(uint32 AudioVolumeID, const FInteriorSettings& Settings);
 
 	FListener()
 		: Transform(FTransform::Identity)
 		, Velocity(ForceInit)
-		, Volume(nullptr)
+		, AudioVolumeID(0)
 		, InteriorStartTime(0.0)
 		, InteriorEndTime(0.0)
 		, ExteriorEndTime(0.0)
@@ -202,14 +215,13 @@ struct FActivatedReverb
 struct FAttenuationListenerData
 {
 	FVector ListenerToSoundDir;
-	const FListener* Listener;
+	FTransform ListenerTransform;
 	float AttenuationDistance;
 	float ListenerToSoundDistance;
 	bool bDataComputed;
 
 	FAttenuationListenerData()
 		: ListenerToSoundDir(FVector::ZeroVector)
-		, Listener(nullptr)
 		, AttenuationDistance(0.0f)
 		, ListenerToSoundDistance(0.0f)
 		, bDataComputed(false)
@@ -257,6 +269,39 @@ struct FGlobalFocusSettings
 	{}
 };
 
+#if !UE_BUILD_SHIPPING
+struct FAudioStats
+{
+	FAudioStats()
+		: bStale(true)
+	{
+	}
+
+	struct FStatWaveInstanceInfo
+	{
+		FString Description;
+		float ActualVolume;
+		int32 InstanceIndex;
+	};
+
+	struct FStatSoundInfo
+	{
+		FString SoundName;
+		FName SoundClassName;
+		float Distance;
+		uint32 AudioComponentID;
+		FTransform Transform;
+		TArray<FStatWaveInstanceInfo> WaveInstanceInfos;
+		TMultiMap<EAttenuationShape::Type, FAttenuationSettings::AttenuationShapeDetails> ShapeDetailsMap;
+	};
+
+	uint8 bStale:1;
+	FVector ListenerLocation;
+	TArray<FStatSoundInfo> StatSoundInfos;
+
+};
+#endif
+
 class ENGINE_API FAudioDevice : public FExec
 {
 public:
@@ -266,6 +311,7 @@ public:
 	//End FExec Interface
 
 #if !UE_BUILD_SHIPPING
+private:
 	/**
 	 * Exec command handlers
 	 */
@@ -309,8 +355,20 @@ public:
 	bool HandlePlayAllPIEAudioCommand(const TCHAR* Cmd, FOutputDevice& Ar);
 	bool HandleAudio3dVisualizeCommand(const TCHAR* Cmd, FOutputDevice& Ar);
 	bool HandleAudioMemoryInfo(const TCHAR* Cmd, FOutputDevice& Ar);
+
+	/**
+	* Lists a summary of loaded sound collated by class
+	*/
+	void ShowSoundClassHierarchy(FOutputDevice& Ar, USoundClass* SoundClass = nullptr, int32 Indent = 0) const;
+
+	/** 
+	* Gets a summary of loaded sound collated by class
+	*/
+	void GetSoundClassInfo(TMap<FName, FAudioClassInfo>& AudioClassInfos);
 #endif
 
+public:
+
 	/**
 	 * Constructor
 	 */
@@ -375,7 +433,7 @@ public:
 	 * @param SoundWave					Resource to stop any sounds that are using it
 	 * @param[out] StoppedComponents	List of Audio Components that were stopped
 	 */
-	void StopSoundsUsingResource(USoundWave* SoundWave, TArray<UAudioComponent*>& StoppedComponents);
+	void StopSoundsUsingResource(USoundWave* SoundWave, TArray<UAudioComponent*>* StoppedComponents = nullptr);
 
 #if WITH_EDITOR
 	/** Deals with anything audio related that should happen when PIE starts */
@@ -423,22 +481,19 @@ public:
 	 * @param   ListenerIndex		The index of the listener
 	 * @param   ListenerTransform   The listener's world transform
 	 * @param   DeltaSeconds		The amount of time over which velocity should be calculated.  If 0, then velocity will not be calculated.
-	 * @param   Volume				The audio volume this listener is in
-	 * @param	ReverbSettings		The reverb settings for this user to use.
-	 * @param	InteriorSettings	The interior settings for this user to use.
 	 */
-	void SetListener(const int32 InListenerIndex, const FTransform& ListenerTransform, const float InDeltaSeconds, AAudioVolume* Volume, const FInteriorSettings& InteriorSettings);
+	void SetListener(UWorld* World, int32 InListenerIndex, const FTransform& ListenerTransform, float InDeltaSeconds);
 
-	/**
-	 * Starts a transition to new reverb and interior settings
-	 *
-	 */
-	void SetReverbSettings(AAudioVolume* Volume, const FReverbSettings& ReverbSettings);
+	const TArray<FListener>& GetListeners() const { check(IsInAudioThread()); return Listeners; }
 
 	/** 
 	 * Returns the currently applied reverb effect if there is one.
 	 */
-	class UReverbEffect* GetCurrentReverbEffect();
+	UReverbEffect* GetCurrentReverbEffect() const
+	{
+		check(IsInGameThread());
+		return CurrentReverbEffect;
+	}
 
 	/**
 	 * Creates an audio component to handle playing a sound.
@@ -478,7 +533,7 @@ public:
 	/**
 	 * Removes the active sound for the specified audio component
 	 */
-	void StopActiveSound(UAudioComponent* AudioComponent);
+	void StopActiveSound(uint64 AudioComponentID);
 
 	/**
 	* Stops the active sound
@@ -486,36 +541,53 @@ public:
 	void StopActiveSound(FActiveSound* ActiveSound);
 
 	/**
-	 * Finds the active sound for the specified audio component
-	 */
-	FActiveSound* FindActiveSound(UAudioComponent* AudioComponent);
+	* Finds the active sound for the specified audio component ID
+	*/
+	FActiveSound* FindActiveSound(uint64 AudioComponentID);
 
 	/**
 	 * Removes an active sound from the active sounds array
 	 */
 	void RemoveActiveSound(FActiveSound* ActiveSound);
 
+	void AddAudioVolumeProxy(const FAudioVolumeProxy& Proxy);
+	void RemoveAudioVolumeProxy(uint32 AudioVolumeID);
+	void UpdateAudioVolumeProxy(const FAudioVolumeProxy& Proxy);
+
+	struct FAudioVolumeSettings
+	{
+		uint32 AudioVolumeID;
+		float Priority;
+		FReverbSettings ReverbSettings;
+		FInteriorSettings InteriorSettings;
+	};
+
+	void GetAudioVolumeSettings(const uint32 WorldID, const FVector& Location, FAudioVolumeSettings& OutSettings) const;
+
+public:
+
+	void SetDefaultAudioSettings(UWorld* World, const FReverbSettings& DefaultReverbSettings, const FInteriorSettings& DefaultInteriorSettings);
+
 	/** 
 	 * Gets the current audio debug state
 	 */
-	EDebugState GetMixDebugState();
+	EDebugState GetMixDebugState() const { return((EDebugState)DebugState); }
+
+	void SetMixDebugState(EDebugState DebugState);
 
 	/**
 	 * Set up the sound class hierarchy
 	 */
 	void InitSoundClasses();
 
+protected:
 	/**
 	 * Set up the initial sound sources
 	 * Allows us to initialize sound source early on, allowing for render callback hookups for iOS Audio.
 	 */
 	void InitSoundSources();
 
-	/** 
-	 * Gets a summary of loaded sound collated by class
-	 */
-	void GetSoundClassInfo(TMap<FName, FAudioClassInfo>& AudioClassInfos);
-
+public:
 	/**
 	 * Registers a sound class with the audio device
 	 *
@@ -540,12 +612,12 @@ public:
 	/**
 	* Checks to see if a coordinate is within a distance of any listener
 	*/
-	bool LocationIsAudible(const FVector& Location, const float MaxDistance);
+	bool LocationIsAudible(const FVector& Location, const float MaxDistance) const;
 
 	/**
 	* Checks to see if a coordinate is within a distance of the given listener
 	*/
-	bool LocationIsAudible(const FVector& Location, const FListener& Listener, const float MaxDistance);
+	static bool LocationIsAudible(const FVector& Location, const FTransform& ListenerTransform, const float MaxDistance);
 
 	/**
 	 * Sets the Sound Mix that should be active by default
@@ -569,7 +641,7 @@ public:
 	/**
 	 * Sets a new sound mix and applies it to all appropriate sound classes
 	 */
-	bool SetBaseSoundMix(USoundMix* SoundMix);
+	void SetBaseSoundMix(USoundMix* SoundMix);
 
 	/**
 	 * Push a SoundMix onto the Audio Device's list.
@@ -662,21 +734,20 @@ public:
 	virtual bool IsExernalBackgroundSoundActive() { return false; }
 
 	/** Whether or not HRTF spatialization is enabled for all. */
-	bool IsHRTFEnabledForAll() const
-	{
-		return bHRTFEnabledForAll && IsSpatializationPluginEnabled();
-	}
+	bool IsHRTFEnabledForAll() const;
 
 	bool IsAudioDeviceMuted() const;
 
+	void SetDeviceMuted(bool bMuted);
+
 	/** Computes and returns some geometry related to the listener and the given sound transform. */
-	void GetAttenuationListenerData(FAttenuationListenerData& OutListenerData, const FTransform& SoundTransform, const FAttenuationSettings& AttenuationSettings, const FListener* InListener = nullptr) const;
+	void GetAttenuationListenerData(FAttenuationListenerData& OutListenerData, const FTransform& SoundTransform, const FAttenuationSettings& AttenuationSettings, const FTransform* InListenerTransform = nullptr) const;
 
 	/** Returns the focus factor of a sound based on its position and listener data. */
-	float GetFocusFactor(FAttenuationListenerData& OutListenerData, const USoundBase* Sound, const FTransform& SoundTransform, const FAttenuationSettings& AttenuationSettings, const FListener* InListener = nullptr) const;
+	float GetFocusFactor(FAttenuationListenerData& OutListenerData, const USoundBase* Sound, const FTransform& SoundTransform, const FAttenuationSettings& AttenuationSettings, const FTransform* InListenerTransform = nullptr) const;
 
 	/** Gets the max distance and focus factor of a sound. */
-	void GetMaxDistanceAndFocusFactor(USoundBase* Sound, const UWorld* World, const FVector& Location, const FAttenuationSettings* AttenuationSettingsToApply, float *OutMaxDistance, float* OutFocusFactor);
+	void GetMaxDistanceAndFocusFactor(USoundBase* Sound, const UWorld* World, const FVector& Location, const FAttenuationSettings* AttenuationSettingsToApply, float& OutMaxDistance, float& OutFocusFactor);
 
 	/**
 	* Checks if the given sound would be audible.
@@ -722,23 +793,13 @@ protected:
 	{
 	}
 
-	/**
-	 * Lists a summary of loaded sound collated by class
-	 */
-	void ShowSoundClassHierarchy(FOutputDevice& Ar, USoundClass* SoundClass = nullptr, int32 Indent = 0) const;
+private:
 
 	/**
 	 * Parses the sound classes and propagates multiplicative properties down the tree.
 	 */
 	void ParseSoundClasses();
 
-	/**
-	 * Construct the CurrentSoundClassProperties map
-	 * @param DeltaTime The current frame delta. Used to interpolate sound class adjustments.
-	 *
-	 * This contains the original sound class properties propagated properly, and all adjustments due to the sound mixes
-	 */
-	void UpdateSoundClassProperties(float DeltaTime);
 
 	/**
 	 * Set the mix for altering sound class properties
@@ -806,6 +867,14 @@ protected:
 	 */
 	void ApplyClassAdjusters(USoundMix* SoundMix, float InterpValue, float DeltaTime);
 
+	/**
+	* Construct the CurrentSoundClassProperties map
+	* @param DeltaTime The current frame delta. Used to interpolate sound class adjustments.
+	*
+	* This contains the original sound class properties propagated properly, and all adjustments due to the sound mixes
+	*/
+	void UpdateSoundClassProperties(float DeltaTime);
+
 	/**
 	 * Recursively apply an adjuster to the passed in sound class and all children of the sound class
 	 * 
@@ -822,6 +891,8 @@ protected:
 		return Adjuster * InterpValue + 1.0f - InterpValue;
 	}
 
+public:
+
 	/**
 	 * Platform dependent call to init effect data on a sound source
 	 */
@@ -843,12 +914,11 @@ protected:
 	 */
 	FAudioEffectsManager* GetEffects()
 	{
+		check(IsInAudioThread());
 		return Effects;
 	}
 
-	/** Internal */
-	void SortWaveInstances(int32 MaxChannels);
-
+private:
 	/**
 	 * Internal helper function used by ParseSoundClasses to traverse the tree.
 	 *
@@ -862,6 +932,9 @@ protected:
 	 */
 	void UpdateHighestPriorityReverb();
 
+	void SendUpdateResultsToGameThread(int32 FirstActiveIndex);
+
+public:
 
 // If we make FAudioDevice not be subclassable, then all the functions following would move to IAudioDeviceModule
 
@@ -898,9 +971,6 @@ protected:
 
 	void AddSoundToStop(struct FActiveSound* SoundToStop);
 
-	/** Updates the listener transform. */
-	void UpdateListenerTransform();
-
 	/**   
 	* Gets the direction of the given position vector transformed relative to listener.   
 	* @param Position				Input position vector to transform relative to listener
@@ -909,12 +979,51 @@ protected:
 	*/
 	FVector GetListenerTransformedDirection(const FVector& Position, float* OutDistance);
 
+private:
 	/** Processes the set of pending sounds that need to be stopped */ 
 	void ProcessingPendingActiveSoundStops(bool bForceDelete = false);
 
+public:
+
 	/** Query if the editor is in VR Preview for the current play world. Returns false for non-editor builds */
 	static bool CanUseVRAudioDevice();
 
+#if !UE_BUILD_SHIPPING
+	void RenderStatReverb(UWorld* World, FViewport* Viewport, FCanvas* Canvas, int32 X, int32& Y, const FVector* ViewLocation, const FRotator* ViewRotation) const;
+
+	void UpdateSoundShowFlags(const uint8 OldSoundShowFlags, const uint8 NewSoundShowFlags);
+	void UpdateRequestedStat(const uint8 InRequestedStat);
+
+	FAudioStats& GetAudioStats()
+	{
+		check(IsInGameThread());
+		return AudioStats;
+	}
+#endif
+
+	bool AreStartupSoundsPreCached() const { return bStartupSoundsPreCached; }
+
+	float GetTransientMasterVolume() const { check(IsInAudioThread()); return TransientMasterVolume; }
+	void SetTransientMasterVolume(float TransientMasterVolume);
+
+	FSoundSource* GetSoundSource(FWaveInstance* WaveInstance) const;
+
+	const FGlobalFocusSettings& GetGlobalFocusSettings() const;
+	void SetGlobalFocusSettings(const FGlobalFocusSettings& NewFocusSettings);
+
+	const FDynamicParameter& GetGlobalPitchScale() const { check(IsInAudioThread()); return GlobalPitchScale; }
+	void SetGlobalPitchModulation(float PitchScale, float TimeSec);
+
+	float GetPlatformAudioHeadroom() const { check(IsInAudioThread()); return PlatformAudioHeadroom; }
+	void SetPlatformAudioHeadroom(float PlatformHeadRoom);
+
+	DEPRECATED(4.13, "Direct access of SoundClasses is no longer allowed. Instead you should use the SoundMixClassOverride system")
+	const TMap<USoundClass*, FSoundClassProperties>& GetSoundClassPropertyMap() const
+	{
+		check(IsInAudioThread());
+		return SoundClasses;
+	}
+
 public:
 
 	/** The maximum number of concurrent audible sounds */
@@ -932,16 +1041,22 @@ public:
 	/** Available size in permanent memory stack */
 	int32 CommonAudioPoolFreeBytes;
 
-	uint32 bGameWasTicking:1;
+	/** The handle for this audio device used in the audio device manager. */
+	uint32 DeviceHandle;
 
-	/* HACK: Temporarily disable audio caching.  This will be done better by changing the decompression pool size in the future */
-	uint32 bDisableAudioCaching:1;
+	/** Audio spatialization plugin. */
+	IAudioSpatializationPlugin* SpatializationPlugin;
 
-	/* True once the startup sounds have been precached */
-	uint32 bStartupSoundsPreCached:1;
+	/** Audio spatialization algorithm (derived from a plugin). */
+	IAudioSpatializationAlgorithm* SpatializeProcessor;
 
+private:
+	// Audio thread representation of listeners
 	TArray<FListener> Listeners;
 
+	// Game thread cache of listener transforms
+	TArray<FTransform> ListenerTransforms;
+
 	uint64 CurrentTick;
 
 	/** An AudioComponent to play test sounds on */
@@ -958,7 +1073,8 @@ public:
 
 	/** The global focus settings */
 	FGlobalFocusSettings GlobalFocusSettings;
-
+	FGlobalFocusSettings GlobalFocusSettings_GameThread;
+	
 	/** Timestamp of the last update */
 	double LastUpdateTime;
 
@@ -966,8 +1082,11 @@ public:
 	int32 NextResourceID;
 
 	/** Set of sources used to play sounds (platform will subclass these) */
+protected:
 	TArray<FSoundSource*> Sources;
 	TArray<FSoundSource*> FreeSources;
+
+private:
 	TMap<FWaveInstance*, FSoundSource*>	WaveInstanceSourceMap;
 
 	/** Current properties of all sound classes */
@@ -985,43 +1104,55 @@ public:
 	/** Map of sound mix sound class overrides. Will override any sound class effects for any sound mixes */
 	TMap<USoundMix*, FSoundMixClassOverrideMap> SoundMixClassEffectOverrides;
 
+protected:
 	/** Interface to audio effects processing */
 	FAudioEffectsManager* Effects;
 
-	/** The volume the listener resides in */
-	const AAudioVolume* CurrentAudioVolume;
+private:
+	UReverbEffect* CurrentReverbEffect;
 
-	/** A volume headroom to apply to specific platforms to achieve beter platform consistency. */
+	/** A volume headroom to apply to specific platforms to achieve better platform consistency. */
 	float PlatformAudioHeadroom;
 
-	/** Reverb Effects activated without volumes */
+	/** Reverb Effects activated without volumes - Game Thread owned */
 	TMap<FName, FActivatedReverb> ActivatedReverbs;
 
-	/** The activated reverb that currently has the highest priority */
-	const FActivatedReverb* HighestPriorityReverb;
+	/** The activated reverb that currently has the highest priority - Audio Thread owned */
+	FActivatedReverb HighestPriorityActivatedReverb;
 
-	/** Audio spatialization plugin. */
-	IAudioSpatializationPlugin* SpatializationPlugin;
+	/** Gamethread representation of whether HRTF is enabled for all 3d sounds. (not bitpacked to avoid thread issues) */
+	bool bHRTFEnabledForAll_OnGameThread:1;
 
-	/** Audio spatialization algorithm (derived from a plugin). */
-	IAudioSpatializationAlgorithm* SpatializeProcessor;
+	uint8 bGameWasTicking:1;
 
-	/** Whether or not the spatialization plugin is enabled. */
-	bool bSpatializationExtensionEnabled;
+public:
+	/* HACK: Temporarily disable audio caching.  This will be done better by changing the decompression pool size in the future */
+	uint8 bDisableAudioCaching:1;
 
-	/** Whether or HRTF is enabled for all 3d sounds. */
-	bool bHRTFEnabledForAll;
+private:
+	/* True once the startup sounds have been precached */
+	uint8 bStartupSoundsPreCached:1;
 
-	/** The handle for this audio device used in the audio device manager. */
-	uint32 DeviceHandle;
+	/** Whether or not the spatialization plugin is enabled. */
+	uint8 bSpatializationExtensionEnabled:1;
+
+	/** Whether HRTF is enabled for all 3d sounds. */
+	uint8 bHRTFEnabledForAll:1;
 
 	/** Whether the audio device is active (current audio device in-focus in PIE) */
-	uint32 bIsDeviceMuted:1;
+	uint8 bIsDeviceMuted:1;
 
 	/** Whether the audio device has been initialized */
-	uint32 bIsInitialized:1;
+	uint8 bIsInitialized:1;
 
-private:
+	/** Whether the value in HighestPriorityActivatedReverb should be used - Audio Thread owned */
+	uint8 bHasActivatedReverb:1;
+
+#if !UE_BUILD_SHIPPING
+	uint8 RequestedAudioStats;
+
+	FAudioStats AudioStats;
+#endif
 
 	TArray<FActiveSound*> ActiveSounds;
 	TArray<FWaveInstance*> ActiveWaveInstances;
@@ -1035,7 +1166,11 @@ private:
 	/** A set of sounds which need to be deleted but weren't able to be deleted due to pending async operations */
 	TArray<FActiveSound*> PendingSoundsToDelete;
 
-	TMap<UPTRINT, FActiveSound*> AudioComponentToActiveSoundMap;
+	TMap<uint64, FActiveSound*> AudioComponentIDToActiveSoundMap;
+
+	TMap<uint32, FAudioVolumeProxy> AudioVolumeProxies;
+
+	TMap<uint32, TPair<FReverbSettings,FInteriorSettings>> WorldIDToDefaultAudioVolumeSettingsMap;
 
 	/** List of passive SoundMixes active last frame */
 	TArray<USoundMix*> PrevPassiveSoundMixModifiers;
diff --git a/Engine/Source/Runtime/Engine/Public/AudioDeviceManager.h b/Engine/Source/Runtime/Engine/Public/AudioDeviceManager.h
--- a/Engine/Source/Runtime/Engine/Public/AudioDeviceManager.h
+++ b/Engine/Source/Runtime/Engine/Public/AudioDeviceManager.h
@@ -23,13 +23,22 @@ public:
 	* Registers the audio device module to the audio device manager. The audio device
 	* module is responsible for creating platform-dependent audio devices.
 	*/
-	void RegisterAudioDeviceModule(class IAudioDeviceModule * AudioDeviceModuleInput);
+	void RegisterAudioDeviceModule(class IAudioDeviceModule* AudioDeviceModuleInput);
+
+	struct FCreateAudioDeviceResults
+	{
+		uint32 Handle;
+		uint8  bNewDevice:1;
+		FAudioDevice* AudioDevice;
+
+		FCreateAudioDeviceResults();
+	};
 
 	/**
 	* Creates and audio device instance internally and returns a
 	* handle to the audio device. Returns true on success.
 	*/
-	class FAudioDevice* CreateAudioDevice(uint32 & HandleOut, bool bCreateNewDevice);
+	bool CreateAudioDevice(bool bCreateNewDevice, FCreateAudioDeviceResults& OutResults);
 
 	/**
 	* Returns whether the audio device handle is valid (i.e. points to
@@ -72,11 +81,8 @@ public:
 	/** Tracks objects in the active audio devices. */
 	void AddReferencedObjects(FReferenceCollector& Collector);
 
-	/** Stops any sounds on any device using the given soundwave. */
-	void StopSoundsUsingWave(class USoundWave* InSoundWave);
-
 	/** Stops sounds using the given resource on all audio devices. */
-	void StopSoundsUsingResource(class USoundWave* InSoundWave, TArray<UAudioComponent*>& StoppedComponents);
+	void StopSoundsUsingResource(class USoundWave* InSoundWave, TArray<UAudioComponent*>* StoppedComponents = nullptr);
 
 	/** Registers the Sound Class for all active devices. */
 	void RegisterSoundClass(USoundClass* SoundClass);
@@ -118,10 +124,10 @@ public:
 	void TogglePlayAllDeviceAudio();
 
 	/** Gets whether or not all devices should play their audio. */
-	bool IsPlayAllDeviceAudio() const;
+	bool IsPlayAllDeviceAudio() const { return bPlayAllDeviceAudio; }
 
 	/** Is debug visualization of 3d sounds enabled */
-	bool IsVisualizeDebug3dEnabled() const;
+	bool IsVisualizeDebug3dEnabled() const { return bVisualize3dDebug; }
 
 	/** Toggles 3d visualization of 3d sounds on/off */
 	void ToggleVisualize3dDebug();
@@ -137,6 +143,9 @@ public:
 
 private:
 
+	/** Call back for garbage collector, ensures no processing is happening on the thread before collecting resources */
+	void OnPreGarbageCollect();
+
 	/** Returns index of the given handle */
 	uint32 GetIndex(uint32 Handle) const;
 
diff --git a/Engine/Source/Runtime/Engine/Public/AudioThread.h b/Engine/Source/Runtime/Engine/Public/AudioThread.h
new file mode 100644
--- /dev/null
+++ b/Engine/Source/Runtime/Engine/Public/AudioThread.h
@@ -0,0 +1,143 @@
+// Copyright 1998-2016 Epic Games, Inc. All Rights Reserved.
+
+/*=============================================================================
+	AudioThread.h: Rendering thread definitions.
+=============================================================================*/
+
+#ifndef __AudioThread_H__
+#define __AudioThread_H__
+
+#include "TaskGraphInterfaces.h"
+
+////////////////////////////////////
+// Audio thread API
+////////////////////////////////////
+
+DECLARE_STATS_GROUP(TEXT("Audio Thread Commands"), STATGROUP_AudioThreadCommands, STATCAT_Advanced);
+
+class FAudioThread : public FRunnable
+{
+private:
+	/**
+	* Whether to run with an audio thread
+	*/
+	static bool bUseThreadedAudio;
+
+	/** 
+	* Sync event to make sure that audio thread is bound to the task graph before main thread queues work against it.
+	*/
+	FEvent* TaskGraphBoundSyncEvent;
+
+	/**
+	* Whether the audio thread is currently running
+	* If this is false, then we have no audio thread and audio commands will be issued directly on the game thread
+	*/
+	static bool bIsAudioThreadRunning;
+
+	static uint32 SuspensionCount;
+
+	/**
+	* Polled by the game thread to detect crashes in the audio thread.
+	* If the audio thread crashes, it sets this variable to false.
+	*/
+	volatile static bool bIsAudioThreadHealthy;
+
+	/** If the audio thread has been terminated by an unhandled exception, this contains the error message. */
+	static FString AudioThreadError;
+
+	/** The audio thread itself. */
+	static FRunnable* AudioThreadRunnable;
+
+	/** The stashed value of the audio thread as we clear it during GC */
+	static uint32 CachedAudioThreadId;
+
+	void OnPreGarbageCollect();
+	void OnPostGarbageCollect();
+
+public:
+
+	FAudioThread();
+	virtual ~FAudioThread();
+
+	// FRunnable interface.
+	virtual bool Init() override;
+	virtual void Exit() override;
+	virtual uint32 Run() override;
+
+	/** Starts the audio thread. */
+	static ENGINE_API void StartAudioThread();
+
+	/** Stops the audio thread. */
+	static ENGINE_API void StopAudioThread();
+
+	/** Execute a command on the audio thread. If GIsAudioThreadRunning is false the command will execute immediately */
+	static ENGINE_API void RunCommandOnAudioThread(TFunction<void()> InFunction, const TStatId InStatId = TStatId());
+
+	/** Execute a (presumably audio) command on the game thread. If GIsAudioThreadRunning is false the command will execute immediately */
+	static ENGINE_API void RunCommandOnGameThread(TFunction<void()> InFunction, const TStatId InStatId = TStatId());
+
+	/**
+	* Checks if the audio thread is healthy and running.
+	* If it has crashed, UE_LOG is called with the exception information.
+	*/
+	static ENGINE_API void CheckAudioThreadHealth();
+
+	/** Checks if the audio thread is healthy and running, without crashing */
+	static ENGINE_API bool IsAudioThreadHealthy();
+
+	static ENGINE_API void SetUseThreadedAudio(bool bInUseThreadedAudio);
+
+	static ENGINE_API bool IsAudioThreadRunning() { return bIsAudioThreadRunning; }
+
+	static ENGINE_API void SuspendAudioThread();
+	static ENGINE_API void ResumeAudioThread();
+
+};
+
+/** Suspends the audio thread for the duration of the lifetime of the object */
+struct FAudioThreadSuspendContext
+{
+	FAudioThreadSuspendContext()
+	{
+		FAudioThread::SuspendAudioThread();
+	}
+
+	~FAudioThreadSuspendContext()
+	{
+		FAudioThread::ResumeAudioThread();
+	}
+};
+
+////////////////////////////////////
+// Audio fences
+////////////////////////////////////
+
+/**
+* Used to track pending audio commands from the game thread.
+*/
+class ENGINE_API FAudioCommandFence
+{
+public:
+
+	/**
+	* Adds a fence command to the audio command queue.
+	* Conceptually, the pending fence count is incremented to reflect the pending fence command.
+	* Once the rendering thread has executed the fence command, it decrements the pending fence count.
+	*/
+	void BeginFence();
+
+	/**
+	* Waits for pending fence commands to retire.
+	* @param bProcessGameThreadTasks, if true we are on a short callstack where it is safe to process arbitrary game thread tasks while we wait
+	*/
+	void Wait(bool bProcessGameThreadTasks = false) const;
+
+	// return true if the fence is complete
+	bool IsFenceComplete() const;
+
+private:
+	/** Graph event that represents completion of this fence **/
+	mutable FGraphEventRef CompletionEvent;
+};
+
+#endif
diff --git a/Engine/Source/Runtime/Engine/Public/BonePose.h b/Engine/Source/Runtime/Engine/Public/BonePose.h
--- a/Engine/Source/Runtime/Engine/Public/BonePose.h
+++ b/Engine/Source/Runtime/Engine/Public/BonePose.h
@@ -236,11 +236,12 @@ public:
 	}
 
 	// Sets every bone transform to Identity
-	void ResetToIdentity()
+	void ResetToAdditiveIdentity()
 	{
 		for (FTransform& Bone : this->Bones)
 		{
 			Bone.SetIdentity();
+			Bone.SetScale3D(FVector::ZeroVector);
 		}
 	}
 
diff --git a/Engine/Source/Runtime/Engine/Public/CollisionQueryParams.h b/Engine/Source/Runtime/Engine/Public/CollisionQueryParams.h
--- a/Engine/Source/Runtime/Engine/Public/CollisionQueryParams.h
+++ b/Engine/Source/Runtime/Engine/Public/CollisionQueryParams.h
@@ -129,6 +129,9 @@ public:
 	/** Add an actor for this trace to ignore */
 	void AddIgnoredActor(const AActor* InIgnoreActor);
 
+	/** Add an actor by ID for this trace to ignore */
+	void AddIgnoredActor(const uint32 InIgnoreActorID);
+
 	/** Add a collection of actors for this trace to ignore */
 	void AddIgnoredActors(const TArray<AActor*>& InIgnoreActors);
 
diff --git a/Engine/Source/Runtime/Engine/Public/ContentStreaming.h b/Engine/Source/Runtime/Engine/Public/ContentStreaming.h
--- a/Engine/Source/Runtime/Engine/Public/ContentStreaming.h
+++ b/Engine/Source/Runtime/Engine/Public/ContentStreaming.h
@@ -524,14 +524,14 @@ struct FStreamingManagerCollection : public IStreamingManager
 	 *
 	 * @param StreamingManager	Streaming manager to add
 	 */
-	void AddStreamingManager( IStreamingManager* StreamingManager );
+	ENGINE_API void AddStreamingManager( IStreamingManager* StreamingManager );
 
 	/**
 	 * Removes a streaming manager from the array of managers to route function calls to.
 	 *
 	 * @param StreamingManager	Streaming manager to remove
 	 */
-	void RemoveStreamingManager( IStreamingManager* StreamingManager );
+	ENGINE_API void RemoveStreamingManager( IStreamingManager* StreamingManager );
 
 	/**
 	 * Sets the number of iterations to use for the next time UpdateResourceStreaming is being called. This 
diff --git a/Engine/Source/Runtime/Engine/Public/Engine.h b/Engine/Source/Runtime/Engine/Public/Engine.h
--- a/Engine/Source/Runtime/Engine/Public/Engine.h
+++ b/Engine/Source/Runtime/Engine/Public/Engine.h
@@ -211,7 +211,6 @@
 #include "Engine/PendingNetGame.h"
 #include "PhysicalMaterials/PhysicalMaterial.h"
 #include "PhysicalMaterials/PhysicalMaterialPropertyBase.h"
-#include "PhysicsEngine/PhysicsAsset.h"
 #include "PhysicsEngine/PhysicsCollisionHandler.h"
 #include "Engine/PlatformInterfaceBase.h"
 #include "Engine/CloudStorageBase.h"
diff --git a/Engine/Source/Runtime/Engine/Public/PhysicsPublic.h b/Engine/Source/Runtime/Engine/Public/PhysicsPublic.h
--- a/Engine/Source/Runtime/Engine/Public/PhysicsPublic.h
+++ b/Engine/Source/Runtime/Engine/Public/PhysicsPublic.h
@@ -10,7 +10,7 @@
 #include "PhysxUserData.h"
 #include "DynamicMeshBuilder.h"
 #include "LocalVertexFactory.h"
-#include "PhysicsEngine/PhysicsAsset.h"
+#include "PhysicsEngine/RigidBodyIndexPair.h"
 /**
  * Physics stats
  */
@@ -225,11 +225,17 @@ namespace SleepEvent
 
 }
 
+/** Builder to allow overriding the sim event callback. */
+typedef class FPhysXSimEventCallback* (CreateSimEventCallbackFn)(class FPhysScene* PhysScene, int32 SceneType);
+
 /** Container object for a physics engine 'scene'. */
 
 class FPhysScene
 {
 public:
+	/** If the function is set it will be used to create the sim event callback. Otherwise FPhysXSimEventCallback will be used. */
+	ENGINE_API static CreateSimEventCallbackFn* CreateSimEventCallback;
+
 	/** Indicates whether the async scene is enabled or not. */
 	bool							bAsyncSceneEnabled;
 
@@ -531,10 +537,10 @@ public:
 	}
 
 	/** Adds to queue of skelmesh we want to add to collision disable table */
-	void DeferredAddCollisionDisableTable(uint32 SkelMeshCompID, TMap<struct FRigidBodyIndexPair, bool> * CollisionDisableTable);
+	ENGINE_API void DeferredAddCollisionDisableTable(uint32 SkelMeshCompID, TMap<struct FRigidBodyIndexPair, bool> * CollisionDisableTable);
 
 	/** Adds to queue of skelmesh we want to remove from collision disable table */
-	void DeferredRemoveCollisionDisableTable(uint32 SkelMeshCompID);
+	ENGINE_API void DeferredRemoveCollisionDisableTable(uint32 SkelMeshCompID);
 
 #if WITH_APEX
 	/** Adds a damage event to be fired when fetchResults is done */
diff --git a/Engine/Source/Runtime/Engine/Public/SkeletalMeshMerge.h b/Engine/Source/Runtime/Engine/Public/SkeletalMeshMerge.h
--- a/Engine/Source/Runtime/Engine/Public/SkeletalMeshMerge.h
+++ b/Engine/Source/Runtime/Engine/Public/SkeletalMeshMerge.h
@@ -88,6 +88,15 @@ struct FSkelMeshMergeSectionMapping
 	TArray<int32> SectionIDs;
 };
 
+/** 
+* Info to map all the sections about how to transform their UVs
+*/
+struct FSkelMeshMergeUVTransforms
+{
+	/** For each UV channel on each mesh, how the UVS should be transformed. */
+	TArray<TArray<FTransform>> UVTransformsPerMesh;
+};
+
 /** 
 * Utility for merging a list of skeletal meshes into a single mesh.
 */
@@ -101,13 +110,15 @@ public:
 	* @param InForceSectionMapping - optional array to map sections from the source meshes to merged section entries
 	* @param StripTopLODs - number of high LODs to remove from input meshes
     * @param bMeshNeedsCPUAccess - (optional) if the resulting mesh needs to be accessed by the CPU for any reason (e.g. for spawning particle effects).
+	* @param UVTransforms - optional array to transform the UVs in each mesh
 	*/
 	FSkeletalMeshMerge( 
 		USkeletalMesh* InMergeMesh, 
 		const TArray<USkeletalMesh*>& InSrcMeshList, 
 		const TArray<FSkelMeshMergeSectionMapping>& InForceSectionMapping,
 		int32 StripTopLODs,
-        EMeshBufferAccess MeshBufferAccess=EMeshBufferAccess::Default
+        EMeshBufferAccess MeshBufferAccess=EMeshBufferAccess::Default,
+		FSkelMeshMergeUVTransforms* InSectionUVTransforms = nullptr
 		);
 
 	/**
@@ -160,6 +171,9 @@ private:
 	/** array to map sections from the source meshes to merged section entries */
 	const TArray<FSkelMeshMergeSectionMapping>& ForceSectionMapping;
 
+	/** optional array to transform UVs in each source mesh */
+	const FSkelMeshMergeUVTransforms* SectionUVTransforms;
+
 	/** Matches the Materials array in the final mesh - used for creating the right number of Material slots. */
 	TArray<int32>	MaterialIds;
 
@@ -172,10 +186,13 @@ private:
 		const FSkelMeshSection* Section;
 		/** mapping from the original BoneMap for this sections chunk to the new MergedBoneMap */
 		TArray<FBoneIndexType> BoneMapToMergedBoneMap;
+		/** transform from the original UVs */
+		TArray<FTransform> UVTransforms;
 
-		FMergeSectionInfo( const USkeletalMesh* InSkelMesh,const FSkelMeshSection* InSection )
+		FMergeSectionInfo( const USkeletalMesh* InSkelMesh,const FSkelMeshSection* InSection, TArray<FTransform> & InUVTransforms )
 			:	SkelMesh(InSkelMesh)
 			,	Section(InSection)
+			,	UVTransforms(InUVTransforms)
 		{}
 	};
 
diff --git a/Engine/Source/Runtime/Engine/Public/SubtitleManager.h b/Engine/Source/Runtime/Engine/Public/SubtitleManager.h
--- a/Engine/Source/Runtime/Engine/Public/SubtitleManager.h
+++ b/Engine/Source/Runtime/Engine/Public/SubtitleManager.h
@@ -12,12 +12,12 @@
 class FActiveSubtitle
 {
 public:
-	FActiveSubtitle( int32 InIndex, float InPriority, bool InbSplit, bool InbSingleLine, const TArray<FSubtitleCue>& InSubtitles )
+	FActiveSubtitle( int32 InIndex, float InPriority, bool InbSplit, bool InbSingleLine, TArray<FSubtitleCue> InSubtitles )
 		:	Index( InIndex )
 		,	Priority( InPriority )
 		,	bSplit( InbSplit )
 		,	bSingleLine( InbSingleLine )
-		,	Subtitles( InSubtitles )
+		,	Subtitles( MoveTemp(InSubtitles) )
 	{}
 
 	/** Index into the Subtitles array for the currently active subtitle in this set. */
@@ -32,6 +32,23 @@ public:
 	TArray<FSubtitleCue>	Subtitles;
 };
 
+struct FQueueSubtitleParams
+{
+	FQueueSubtitleParams(const TArray<FSubtitleCue>& InSubtitles)
+		: Subtitles(InSubtitles)
+	{ }
+
+	uint64 AudioComponentID;
+	TWeakObjectPtr<UWorld> WorldPtr;
+	PTRINT WaveInstance;
+	float SubtitlePriority;
+	float Duration;
+	uint8 bManualWordWrap : 1;
+	uint8 bSingleLine : 1;
+	const TArray<FSubtitleCue>& Subtitles;
+};
+
+
 /**
  * Subtitle manager.  Handles prioritization and rendering of subtitles.
  */
@@ -73,7 +90,9 @@ public:
 	 * @param  SoundDuration - time in seconds after which the subtitles do not display
 	 * @param  Subtitles - TArray of lines of subtitle and time offset to play them
 	 */
-	void		QueueSubtitles( PTRINT SubtitleID, float Priority, bool bManualWordWrap, bool bSingleLine, float SoundDuration, TArray<FSubtitleCue>& Subtitles, float InStartTime );
+	void		QueueSubtitles( PTRINT SubtitleID, float Priority, bool bManualWordWrap, bool bSingleLine, float SoundDuration, const TArray<FSubtitleCue>& Subtitles, float InStartTime );
+
+	static void QueueSubtitles(const FQueueSubtitleParams& QueueSubtitlesParams);
 
 	/**
 	 * Draws a subtitle at the specified pixel location.
diff --git a/Engine/Source/Runtime/GameplayAbilities/Private/Abilities/Tasks/AbilityTask_PlayMontageAndWait.cpp b/Engine/Source/Runtime/GameplayAbilities/Private/Abilities/Tasks/AbilityTask_PlayMontageAndWait.cpp
--- a/Engine/Source/Runtime/GameplayAbilities/Private/Abilities/Tasks/AbilityTask_PlayMontageAndWait.cpp
+++ b/Engine/Source/Runtime/GameplayAbilities/Private/Abilities/Tasks/AbilityTask_PlayMontageAndWait.cpp
@@ -183,7 +183,7 @@ bool UAbilityTask_PlayMontageAndWait::StopPlayingMontage()
 			&& AbilitySystemComponent->GetCurrentMontage() == MontageToPlay)
 		{
 			// Unbind delegates so they don't get called as well
-			FAnimMontageInstance* MontageInstance = AnimInstance->GetActiveInstanceForMontage(*MontageToPlay);
+			FAnimMontageInstance* MontageInstance = AnimInstance->GetActiveInstanceForMontage(MontageToPlay);
 			if (MontageInstance)
 			{
 				MontageInstance->OnMontageBlendingOutStarted.Unbind();
diff --git a/Engine/Source/Runtime/GameplayAbilities/Private/AbilitySystemComponent_Abilities.cpp b/Engine/Source/Runtime/GameplayAbilities/Private/AbilitySystemComponent_Abilities.cpp
--- a/Engine/Source/Runtime/GameplayAbilities/Private/AbilitySystemComponent_Abilities.cpp
+++ b/Engine/Source/Runtime/GameplayAbilities/Private/AbilitySystemComponent_Abilities.cpp
@@ -2347,7 +2347,7 @@ void UAbilitySystemComponent::OnRep_ReplicatedAnimMontage()
 			if ((CurrentSectionID == RepSectionID) && (FMath::Abs(CurrentPosition - RepAnimMontageInfo.Position) > MONTAGE_REP_POS_ERR_THRESH) && RepAnimMontageInfo.IsStopped == 0)
 			{
 				// fast forward to server position and trigger notifies
-				if (FAnimMontageInstance* MontageInstance = AnimInstance->GetActiveInstanceForMontage(*RepAnimMontageInfo.AnimMontage))
+				if (FAnimMontageInstance* MontageInstance = AnimInstance->GetActiveInstanceForMontage(RepAnimMontageInfo.AnimMontage))
 				{
 					MontageInstance->HandleEvents(CurrentPosition, RepAnimMontageInfo.Position, nullptr);
 					AnimInstance->TriggerAnimNotifies(0.f);
diff --git a/Engine/Source/Runtime/IOS/IOSAudio/Private/IOSAudioDevice.cpp b/Engine/Source/Runtime/IOS/IOSAudio/Private/IOSAudioDevice.cpp
--- a/Engine/Source/Runtime/IOS/IOSAudio/Private/IOSAudioDevice.cpp
+++ b/Engine/Source/Runtime/IOS/IOSAudio/Private/IOSAudioDevice.cpp
@@ -238,10 +238,11 @@ void FIOSAudioDevice::TeardownHardware()
 
 void FIOSAudioDevice::UpdateHardware()
 {
-	PlayerLocation	= Listeners[0].Transform.GetLocation();
-	PlayerFacing	= Listeners[0].GetFront();
-	PlayerUp		= Listeners[0].GetUp();
-	PlayerRight		= Listeners[0].GetRight();
+	const FListener& Listener = GetListeners()[0];
+	PlayerLocation	= Listener.Transform.GetLocation();
+	PlayerFacing	= Listener.GetFront();
+	PlayerUp		= Listener.GetUp();
+	PlayerRight		= Listener.GetRight();
 }
 
 void FIOSAudioDevice::HandleError(const TCHAR* InLogOutput, bool bTeardown)
diff --git a/Engine/Source/Runtime/IOS/IOSAudio/Private/IOSAudioSource.cpp b/Engine/Source/Runtime/IOS/IOSAudio/Private/IOSAudioSource.cpp
--- a/Engine/Source/Runtime/IOS/IOSAudio/Private/IOSAudioSource.cpp
+++ b/Engine/Source/Runtime/IOS/IOSAudio/Private/IOSAudioSource.cpp
@@ -245,7 +245,7 @@ void FIOSAudioSoundSource::Update(void)
 
 	AudioUnitParameterValue Volume = 0.0f;
 
-	if (!AudioDevice->bIsDeviceMuted)
+	if (!AudioDevice->IsAudioDeviceMuted())
 	{
 		Volume = WaveInstance->Volume * WaveInstance->VolumeMultiplier;
 	}
@@ -257,7 +257,7 @@ void FIOSAudioSoundSource::Update(void)
 	}
 
 	// Apply global multiplier to disable sound when not the foreground app
-	Volume *= AudioDevice->PlatformAudioHeadroom;
+	Volume *= AudioDevice->GetPlatformAudioHeadroom();
 	Volume = FMath::Clamp(Volume, 0.0f, 1.0f);
 
 	// Convert to dB
diff --git a/Engine/Source/Runtime/Launch/Private/LaunchEngineLoop.cpp b/Engine/Source/Runtime/Launch/Private/LaunchEngineLoop.cpp
--- a/Engine/Source/Runtime/Launch/Private/LaunchEngineLoop.cpp
+++ b/Engine/Source/Runtime/Launch/Private/LaunchEngineLoop.cpp
@@ -39,6 +39,7 @@
 #endif
 
 #if WITH_ENGINE
+	#include "AudioThread.h"
 	#include "AutomationController.h"
 	#include "Database.h"
 	#include "DerivedDataCacheInterface.h"
@@ -878,7 +879,7 @@ int32 FEngineLoop::PreInit( const TCHAR* CmdLine )
 	GIsGameThreadIdInitialized = true;
 
 	FPlatformProcess::SetThreadAffinityMask(FPlatformAffinity::GetMainGameMask());
-	FPlatformProcess::SetupGameOrRenderThread(false);
+	FPlatformProcess::SetupGameThread();
 
 	// Figure out whether we're the editor, ucc or the game.
 	const SIZE_T CommandLineSize = FCString::Strlen(CmdLine)+1;
@@ -1441,6 +1442,16 @@ int32 FEngineLoop::PreInit( const TCHAR* CmdLine )
 
 	EndInitTextLocalization();
 
+	if (FApp::ShouldUseThreadingForPerformance() && FPlatformMisc::AllowAudioThread())
+	{
+		bool bUseThreadedAudio = false;
+		if (!GIsEditor)
+		{
+			GConfig->GetBool(TEXT("Audio"), TEXT("UseAudioThread"), bUseThreadedAudio, GEngineIni);
+		}
+		FAudioThread::SetUseThreadedAudio(bUseThreadedAudio);
+	}
+
 	if (FPlatformProcess::SupportsMultithreading() && !IsRunningDedicatedServer() && (bIsRegularClient || bHasEditorToken))
 	{
 		FPlatformSplash::Show();
diff --git a/Engine/Source/Runtime/Mac/CoreAudio/Private/CoreAudioDevice.cpp b/Engine/Source/Runtime/Mac/CoreAudio/Private/CoreAudioDevice.cpp
--- a/Engine/Source/Runtime/Mac/CoreAudio/Private/CoreAudioDevice.cpp
+++ b/Engine/Source/Runtime/Mac/CoreAudio/Private/CoreAudioDevice.cpp
@@ -317,9 +317,10 @@ void FCoreAudioDevice::UpdateHardware()
 {
 	// Caches the matrix used to transform a sounds position into local space so we can just look
 	// at the Y component after normalization to determine spatialization.
-	const FVector Up = Listeners[0].GetUp();
-	const FVector Right = Listeners[0].GetFront();
-	InverseTransform = FMatrix(Up, Right, Up ^ Right, Listeners[0].Transform.GetTranslation()).InverseFast();
+	const FListener& Listener = GetListeners()[0];
+	const FVector Up = Listener.GetUp();
+	const FVector Right = Listener.GetFront();
+	InverseTransform = FMatrix(Up, Right, Up ^ Right, Listener.Transform.GetTranslation()).InverseFast();
 
 	UpdateAUGraph();
 }
diff --git a/Engine/Source/Runtime/Mac/CoreAudio/Private/CoreAudioSource.cpp b/Engine/Source/Runtime/Mac/CoreAudio/Private/CoreAudioSource.cpp
--- a/Engine/Source/Runtime/Mac/CoreAudio/Private/CoreAudioSource.cpp
+++ b/Engine/Source/Runtime/Mac/CoreAudio/Private/CoreAudioSource.cpp
@@ -60,7 +60,7 @@ FCoreAudioSoundSource::FCoreAudioSoundSource( FAudioDevice* InAudioDevice )
 {
 	AudioDevice = ( FCoreAudioDevice *)InAudioDevice;
 	check( AudioDevice );
-	Effects = ( FCoreAudioEffectsManager* )AudioDevice->Effects;
+	Effects = ( FCoreAudioEffectsManager* )AudioDevice->GetEffects();
 	check( Effects );
 
 	FMemory::Memzero( CoreAudioBuffers, sizeof( CoreAudioBuffers ) );
@@ -313,12 +313,12 @@ void FCoreAudioSoundSource::Update( void )
 
 	float Volume = 0.0f;
 	
-	if (!AudioDevice->bIsDeviceMuted)
+	if (!AudioDevice->IsAudioDeviceMuted())
 	{
 		Volume = WaveInstance->GetActualVolume();
 	}
 
-	Volume *= AudioDevice->PlatformAudioHeadroom;
+	Volume *= AudioDevice->GetPlatformAudioHeadroom();
 
 	if( Buffer->NumChannels < 3 )
 	{
diff --git a/Engine/Source/Runtime/MovieSceneTracks/Private/TrackInstances/MovieSceneAudioTrackInstance.cpp b/Engine/Source/Runtime/MovieSceneTracks/Private/TrackInstances/MovieSceneAudioTrackInstance.cpp
--- a/Engine/Source/Runtime/MovieSceneTracks/Private/TrackInstances/MovieSceneAudioTrackInstance.cpp
+++ b/Engine/Source/Runtime/MovieSceneTracks/Private/TrackInstances/MovieSceneAudioTrackInstance.cpp
@@ -8,6 +8,7 @@
 #include "Runtime/Engine/Public/AudioDecompress.h"
 #include "MovieSceneAudioTrack.h"
 #include "MovieSceneAudioSection.h"
+#include "AudioThread.h"
 
 
 FMovieSceneAudioTrackInstance::FMovieSceneAudioTrackInstance( UMovieSceneAudioTrack& InAudioTrack )
@@ -139,20 +140,24 @@ void FMovieSceneAudioTrackInstance::Update(EMovieSceneUpdateData& UpdateData, co
 		{
 			for (int32 ActorIndex = 0; ActorIndex < Actors.Num(); ++ActorIndex)
 			{
-				TWeakObjectPtr<UAudioComponent> Component = GetAudioComponent(Player, Actors[ActorIndex], RowIndex);
-				
-				if (!Component.IsValid() || !Component->IsPlaying())
+				UAudioComponent* Component = GetAudioComponent(Player, Actors[ActorIndex], RowIndex).Get();
+				if (Component && Component->IsPlaying())
 				{
-					continue;
-				}
-
-				FAudioDevice* AudioDevice = Component->GetAudioDevice();
-				FActiveSound* ActiveSound = AudioDevice->FindActiveSound(Component.Get());
+					if (FAudioDevice* AudioDevice = Component->GetAudioDevice())
+					{
+						DECLARE_CYCLE_STAT(TEXT("FAudioThreadTask.MovieSceneUpdateAudioTransform"), STAT_MovieSceneUpdateAudioTransform, STATGROUP_TaskGraphTasks);
 
-				if (ActiveSound != nullptr)
-				{
-					ActiveSound->bLocationDefined = true;
-					ActiveSound->Transform = Actors[ActorIndex]->GetTransform();
+						const FTransform ActorTransform = Actors[ActorIndex]->GetTransform();
+						const uint64 ActorComponentID = Component->GetUniqueID();
+						FAudioThread::RunCommandOnAudioThread([AudioDevice, ActorComponentID, ActorTransform]()
+						{
+							if (FActiveSound* ActiveSound = AudioDevice->FindActiveSound(ActorComponentID))
+							{
+								ActiveSound->bLocationDefined = true;
+								ActiveSound->Transform = ActorTransform;
+							}
+						}, GET_STATID(STAT_MovieSceneUpdateAudioTransform));
+					}
 				}
 			}
 		}
diff --git a/Engine/Source/Runtime/RenderCore/Private/RenderingThread.cpp b/Engine/Source/Runtime/RenderCore/Private/RenderingThread.cpp
--- a/Engine/Source/Runtime/RenderCore/Private/RenderingThread.cpp
+++ b/Engine/Source/Runtime/RenderCore/Private/RenderingThread.cpp
@@ -427,7 +427,7 @@ public:
 	virtual uint32 Run(void) override
 	{
 		FMemory::SetupTLSCachesOnCurrentThread();
-		FPlatformProcess::SetupGameOrRenderThread(true);
+		FPlatformProcess::SetupRenderThread();
 
 #if PLATFORM_WINDOWS
 		if ( !FPlatformMisc::IsDebuggerPresent() || GAlwaysReportCrash )
diff --git a/Engine/Source/Runtime/Windows/XAudio2/Private/XAudio2Source.cpp b/Engine/Source/Runtime/Windows/XAudio2/Private/XAudio2Source.cpp
--- a/Engine/Source/Runtime/Windows/XAudio2/Private/XAudio2Source.cpp
+++ b/Engine/Source/Runtime/Windows/XAudio2/Private/XAudio2Source.cpp
@@ -53,7 +53,7 @@ FXAudio2SoundSource::FXAudio2SoundSource(FAudioDevice* InAudioDevice)
 
 	AudioDevice = ( FXAudio2Device* )InAudioDevice;
 	check( AudioDevice );
-	Effects = (FXAudio2EffectsManager*)AudioDevice->Effects;
+	Effects = (FXAudio2EffectsManager*)AudioDevice->GetEffects();
 	check( Effects );
 
 	Destinations[DEST_DRY].Flags = 0;
@@ -145,17 +145,16 @@ void FXAudio2SoundSource::FreeResources( void )
 		PendingTaskInfo.Buffer = Buffer;
 
 		AudioDevice->DeviceProperties->AddPendingTaskToCleanup(PendingTaskInfo);
-
-		// Clear our buffer ptr here
-		Buffer = XAudio2Buffer = nullptr;
 	}
 	else if (bResourcesNeedFreeing && Buffer)
 	{
 		check(Buffer->ResourceID == 0);
 		delete Buffer;
-		Buffer = XAudio2Buffer = nullptr;
 	}
 
+	// Make sure to nullify the buffer ptrs so that on re-use the source will have a clean buffer
+	// Note that most cases will not require a delete since they are cached and owned by the audio device manager.
+	Buffer = XAudio2Buffer = nullptr;
 	CurrentBuffer = 0;
 }
 
@@ -394,7 +393,10 @@ bool FXAudio2SoundSource::CreateSource( void )
 
 	// Create a source that goes to the spatialisation code and reverb effect
 	Destinations[NumSends].pOutputVoice = Effects->DryPremasterVoice;
-	if( IsEQFilterApplied() )
+
+	// EQFilter Causes sound devices on AMD boards to lag and starve important game threads. Hack disable for AMD until a long term solution is put into place.
+	static const bool bIsAMD = (FPlatformMisc::GetCPUVendor() == TEXT("AuthenticAMD"));
+	if (!bIsAMD && IsEQFilterApplied())
 	{
 		Destinations[NumSends].pOutputVoice = Effects->EQPremasterVoice;
 	}
@@ -492,19 +494,10 @@ bool FXAudio2SoundSource::PrepareForInitialization(FWaveInstance* InWaveInstance
 	bEditorWarnedChangedSpatialization = false;
 
 	// Find matching buffer.
-	FAudioDevice* BestAudioDevice = nullptr;
-	if (UAudioComponent* AudioComponent = InWaveInstance->ActiveSound->GetAudioComponent())
-	{
-		BestAudioDevice = AudioComponent->GetAudioDevice();
-	}
-	else
-	{
-		BestAudioDevice = GEngine->GetMainAudioDevice();
-	}
-	check(BestAudioDevice);
+	check(InWaveInstance->ActiveSound->AudioDevice);
 
 	check(XAudio2Buffer == nullptr);
-	XAudio2Buffer = FXAudio2SoundBuffer::Init(BestAudioDevice, InWaveInstance->WaveData, InWaveInstance->StartTime > 0.f);
+	XAudio2Buffer = FXAudio2SoundBuffer::Init(InWaveInstance->ActiveSound->AudioDevice, InWaveInstance->WaveData, InWaveInstance->StartTime > 0.f);
 	if (XAudio2Buffer)
 	{
 		// If our realtime source is not ready, then we will need to free our resources because this 
@@ -556,51 +549,58 @@ bool FXAudio2SoundSource::Init(FWaveInstance* InWaveInstance)
 		SetReverbApplied(Effects->ReverbEffectVoice != nullptr);
 
 		// Create a new source if we haven't already
-		if (!CreateSource())
-		{
-			return false;
-		}
-
-		check(WaveInstance);
-		if (WaveInstance->StartTime)
+		if (CreateSource())
 		{
-			XAudio2Buffer->Seek(WaveInstance->StartTime);
-		}
+			check(WaveInstance);
+			if (WaveInstance->StartTime)
+			{
+				XAudio2Buffer->Seek(WaveInstance->StartTime);
+			}
 
-		// Submit audio buffers
-		switch (XAudio2Buffer->SoundFormat)
-		{
+			// Submit audio buffers
+			switch (XAudio2Buffer->SoundFormat)
+			{
 			case SoundFormat_PCM:
 			case SoundFormat_PCMPreview:
-			SubmitPCMBuffers();
-			break;
+				SubmitPCMBuffers();
+				break;
 
 			case SoundFormat_PCMRT:
 			case SoundFormat_Streaming:
-			SubmitPCMRTBuffers();
-			break;
+				SubmitPCMRTBuffers();
+				break;
 
 			case SoundFormat_XMA2:
-			SubmitXMA2Buffers();
-			break;
+				SubmitXMA2Buffers();
+				break;
 
 			case SoundFormat_XWMA:
-			SubmitXWMABuffers();
-			break;
-		}
+				SubmitXWMABuffers();
+				break;
+			}
 
-		// First updates of the source which e.g. sets the pitch and volume.
-		bInitialized = true;
-		bFirstRTBuffersSubmitted = false;
+			// First updates of the source which e.g. sets the pitch and volume.
+			bInitialized = true;
+			bFirstRTBuffersSubmitted = false;
 
-		Update();
+			Update();
 
-		// Now set the source state to initialized so it can be played
-		// Initialization succeeded.
-		return true;
+			// Now set the source state to initialized so it can be played
+			// Initialization succeeded.
+			return true;
+		}
+		else
+		{
+			UE_LOG(LogXAudio2, Warning, TEXT("Failed to init sound source for wave instance '%s' due to being unable to create an IXAudio2SourceVoice."), *InWaveInstance->GetName());
+		}
+	}
+	else
+	{
+		UE_LOG(LogXAudio2, Warning, TEXT("Failed to init sound source for wave instance '%s' due to invalid buffer or error in compression."), *InWaveInstance->GetName());
 	}
 
 	// Initialization failed.
+	FreeResources();
 	return false;
 }
 
@@ -672,7 +672,7 @@ void FXAudio2SoundSource::GetChannelVolumes(float ChannelVolumes[CHANNEL_MATRIX_
 			UE_LOG(LogXAudio2, Warning, TEXT("FXAudio2SoundSource contains unreasonble value %f in channel %d: %s"), ChannelVolumes[i], i, *Describe_Internal(true, false));
 		}
 
-		ChannelVolumes[i] = FMath::Clamp<float>(ChannelVolumes[i] * AudioDevice->PlatformAudioHeadroom, 0.0f, MAX_VOLUME);
+		ChannelVolumes[i] = FMath::Clamp<float>(ChannelVolumes[i] * AudioDevice->GetPlatformAudioHeadroom(), 0.0f, MAX_VOLUME);
 	}
 }
 
@@ -1237,16 +1237,6 @@ FString FXAudio2SoundSource::Describe(bool bUseLongName)
 
 FString FXAudio2SoundSource::Describe_Internal(bool bUseLongName, bool bIncludeChannelVolumes)
 {
-	// look for a component and its owner
-	AActor* SoundOwner = NULL;
-
-	// TODO - Audio Threading. This won't work cross thread.
-	UAudioComponent* AudioComponent = (WaveInstance->ActiveSound ?  WaveInstance->ActiveSound->GetAudioComponent() : nullptr);
-	if (AudioComponent)
-	{
-		SoundOwner = AudioComponent->GetOwner();
-	}
-
 	FString SpatializedVolumeInfo;
 	if (bIncludeChannelVolumes && WaveInstance->bUseSpatialization)
 	{
@@ -1281,11 +1271,13 @@ FString FXAudio2SoundSource::Describe_Internal(bool bUseLongName, bool bIncludeC
 		}
 	}
 
+	const FString SoundOwnerName = (WaveInstance->ActiveSound ? WaveInstance->ActiveSound->GetOwnerName() : TEXT("None"));
+
 	return FString::Printf(TEXT("Wave: %s, Volume: %6.2f%s, Owner: %s"), 
 		bUseLongName ? *WaveInstance->WaveData->GetPathName() : *WaveInstance->WaveData->GetName(),
 		WaveInstance->GetActualVolume(),
 		*SpatializedVolumeInfo,
-		SoundOwner ? *SoundOwner->GetName() : TEXT("None"));
+		*SoundOwnerName);
 
 }
 
@@ -1303,7 +1295,7 @@ void FXAudio2SoundSource::Update()
 	// Don't apply global pitch scale to UI sounds
 	if (!WaveInstance->bIsUISound)
 	{
-		Pitch *= AudioDevice->GlobalPitchScale.GetValue();
+		Pitch *= AudioDevice->GetGlobalPitchScale().GetValue();
 	}
 
 	Pitch = FMath::Clamp<float>(Pitch, MIN_PITCH, MAX_PITCH);
@@ -1418,7 +1410,6 @@ void FXAudio2SoundSource::Stop()
 		// Free resources
 		FreeResources();
 
-		Buffer = XAudio2Buffer = nullptr;
 		bBuffersToFlush = false;
 		bLoopCallback = false;
 		bResourcesNeedFreeing = false;
