Fix shipping compile errors

diff --git a/Engine/Source/Runtime/Engine/Private/GameViewportClient.cpp b/Engine/Source/Runtime/Engine/Private/GameViewportClient.cpp
--- a/Engine/Source/Runtime/Engine/Private/GameViewportClient.cpp
+++ b/Engine/Source/Runtime/Engine/Private/GameViewportClient.cpp
@@ -292,10 +292,12 @@ void UGameViewportClient::Init(struct FWorldContext& WorldContext, UGameInstance
 			{
 				AudioDeviceHandle = NewDeviceResults.Handle;
 
+#if !UE_BUILD_SHIPPING
 				if (NewDeviceResults.bNewDevice)
 				{
 					NewDeviceResults.AudioDevice->UpdateSoundShowFlags(0, GetSoundShowFlags());
 				}
+#endif
 
 				// Set the base mix of the new device based on the world settings of the world
 				if (World)
