Fix shipping compile errors

diff --git a/Engine/Source/Runtime/Engine/Private/AudioDevice.cpp b/Engine/Source/Runtime/Engine/Private/AudioDevice.cpp
--- a/Engine/Source/Runtime/Engine/Private/AudioDevice.cpp
+++ b/Engine/Source/Runtime/Engine/Private/AudioDevice.cpp
@@ -847,17 +847,6 @@ bool FAudioDevice::HandleEnableHRTFForAllCommand(const TCHAR* Cmd, FOutputDevice
 	return true;
 }
 
-bool FAudioDevice::IsHRTFEnabledForAll() const
-{
-	if (IsInAudioThread())
-	{
-		return bHRTFEnabledForAll && IsSpatializationPluginEnabled();
-	}
-
-	check(IsInGameThread());
-	return bHRTFEnabledForAll_OnGameThread && IsSpatializationPluginEnabled();
-}
-
 bool FAudioDevice::HandleSoloCommand(const TCHAR* Cmd, FOutputDevice& Ar)
 {
 	// Apply the solo to the given device
@@ -1129,6 +1118,17 @@ bool FAudioDevice::HandleAudioMemoryInfo(const TCHAR* Cmd, FOutputDevice& Ar)
 
 #endif // !UE_BUILD_SHIPPING
 
+bool FAudioDevice::IsHRTFEnabledForAll() const
+{
+	if (IsInAudioThread())
+	{
+		return bHRTFEnabledForAll && IsSpatializationPluginEnabled();
+	}
+
+	check(IsInGameThread());
+	return bHRTFEnabledForAll_OnGameThread && IsSpatializationPluginEnabled();
+}
+
 void FAudioDevice::SetMixDebugState(EDebugState InDebugState)
 {
 	if (!IsInAudioThread())
@@ -2791,7 +2791,7 @@ void FAudioDevice::Update(bool bGameTicking)
 	}
 
 	DECLARE_CYCLE_STAT(TEXT("FAudioThreadTask.AudioUpdateTime"), STAT_AudioUpdateTime, STATGROUP_AudioThreadCommands);
-	FScopeCycleCounter(GET_STATID(STAT_AudioUpdateTime));
+	FScopeCycleCounter AudioUpdateTimeCounter(GET_STATID(STAT_AudioUpdateTime));
 
 	double CurrTime = FPlatformTime::Seconds();
 	double DeltaTime = CurrTime - LastUpdateTime;
