Copy CL 3042565 from //Orion/Dev-General to //UE4/Main
Fix bug where MaxBoneInfluences was being lost when converting old FSkelMeshChunk's to FSkelMeshSection. For old content we recalculate max influences and update the gpu vertex buffer.

diff --git a/Engine/Source/Runtime/Engine/Private/SkeletalMesh.cpp b/Engine/Source/Runtime/Engine/Private/SkeletalMesh.cpp
--- a/Engine/Source/Runtime/Engine/Private/SkeletalMesh.cpp
+++ b/Engine/Source/Runtime/Engine/Private/SkeletalMesh.cpp
@@ -51,6 +51,31 @@ DEFINE_LOG_CATEGORY(LogSkeletalMesh);
 
 DECLARE_CYCLE_STAT(TEXT("GetShadowShapes"), STAT_GetShadowShapes, STATGROUP_Anim);
 
+// Custom serialization version for RecomputeTangent
+struct FSkeletalMeshCustomVersion
+{
+	enum Type
+	{
+		// Before any version changes were made
+		BeforeCustomVersionWasAdded = 0,
+		// Remove Chunks array in FStaticLODModel and combine with Sections array
+		CombineSectionWithChunk = 1,
+		// Remove FRigidSkinVertex and combine with FSoftSkinVertex array
+		CombineSoftAndRigidVerts = 2,
+		// Need to recalc max bone influences
+		RecalcMaxBoneInfluences = 3,
+		// -----<new versions can be added above this line>-------------------------------------------------
+		VersionPlusOne,
+		LatestVersion = VersionPlusOne - 1
+	};
+
+	// The GUID for this custom version number
+	const static FGuid GUID;
+
+private:
+	FSkeletalMeshCustomVersion() {}
+};
+
 const FGuid FSkeletalMeshCustomVersion::GUID(0xD78A4A00, 0xE8584697, 0xBAA819B5, 0x487D46B4);
 FCustomVersionRegistration GRegisterSkeletalMeshCustomVersion(FSkeletalMeshCustomVersion::GUID, FSkeletalMeshCustomVersion::LatestVersion, TEXT("SkeletalMeshVer"));
 
@@ -1268,6 +1293,15 @@ FArchive& operator<<(FArchive& Ar,FSkelMeshSection& S)
 		}
 
 		Ar << S.MaxBoneInfluences;
+
+#if WITH_EDITOR
+		// If loading content where we need to recalc 'max bone influences' instead of using loaded version, do that now
+		if (!StripFlags.IsEditorDataStripped() && Ar.IsLoading() && Ar.CustomVer(FSkeletalMeshCustomVersion::GUID) < FSkeletalMeshCustomVersion::RecalcMaxBoneInfluences)
+		{
+			S.CalcMaxBoneInfluences();
+		}
+#endif
+
 		Ar << S.ApexClothMappingData;
 		Ar << S.PhysicalMeshVertices;
 		Ar << S.PhysicalMeshNormals;
@@ -1343,6 +1377,7 @@ struct FLegacySkelMeshChunk
 		Section.PhysicalMeshVertices = PhysicalMeshVertices;
 		Section.PhysicalMeshNormals = PhysicalMeshNormals;
 		Section.BoneMap = BoneMap;
+		Section.MaxBoneInfluences = MaxBoneInfluences;
 		Section.CorrespondClothAssetIndex = CorrespondClothAssetIndex;
 		Section.ClothAssetSubmeshIndex = ClothAssetSubmeshIndex;
 	}
@@ -3167,6 +3202,19 @@ void USkeletalMesh::PostLoad()
 	// initialize rendering resources
 	if (FApp::CanEverRender())
 	{
+#if WITH_EDITOR
+		// If we needed to recalc max bone influences, also need to recreate gpu vertex buffer with that info
+		if (GetLinkerCustomVersion(FSkeletalMeshCustomVersion::GUID) < FSkeletalMeshCustomVersion::RecalcMaxBoneInfluences)
+		{
+			FSkeletalMeshResource* Resource = GetImportedResource();
+			uint32 VertexFlags = GetVertexBufferFlags();
+			for (int32 LODIndex = 0; LODIndex < Resource->LODModels.Num(); LODIndex++)
+			{
+				Resource->LODModels[LODIndex].BuildVertexBuffers(VertexFlags);
+			}
+		}
+#endif // WITH_EDITOR
+
 		InitResources();
 	}
 
diff --git a/Engine/Source/Runtime/Engine/Public/SkeletalMeshTypes.h b/Engine/Source/Runtime/Engine/Public/SkeletalMeshTypes.h
--- a/Engine/Source/Runtime/Engine/Public/SkeletalMeshTypes.h
+++ b/Engine/Source/Runtime/Engine/Public/SkeletalMeshTypes.h
@@ -54,29 +54,6 @@ enum ECustomSortAlternateIndexMode
 };
 
 
-// Custom serialization version for RecomputeTangent
-struct FSkeletalMeshCustomVersion
-{
-	enum Type
-	{
-		// Before any version changes were made
-		BeforeCustomVersionWasAdded = 0,
-		// Remove Chunks array in FStaticLODModel and combine with Sections array
-		CombineSectionWithChunk = 1,
-		// Remove FRigidSkinVertex and combine with FSoftSkinVertex array
-		CombineSoftAndRigidVerts = 2,
-		// -----<new versions can be added above this line>-------------------------------------------------
-		VersionPlusOne,
-		LatestVersion = VersionPlusOne - 1
-	};
-
-	// The GUID for this custom version number
-	const static FGuid GUID;
-
-private:
-	FSkeletalMeshCustomVersion() {}
-};
-
 class USkeletalMesh;
 
 
