Modified fix for templates to not also break games :/

diff --git a/Engine/Source/Runtime/Engine/Classes/GameFramework/GameMode.h b/Engine/Source/Runtime/Engine/Classes/GameFramework/GameMode.h
--- a/Engine/Source/Runtime/Engine/Classes/GameFramework/GameMode.h
+++ b/Engine/Source/Runtime/Engine/Classes/GameFramework/GameMode.h
@@ -8,7 +8,7 @@
 
 class FDebugDisplayInfo;
 class FUniqueNetId;
-class FUniqueNetIdRepl;
+struct FUniqueNetIdRepl;
 
 //~=============================================================================
 //  GameMode defines the rules and mechanics of the game.  It is only 
