Copying //UE4/Dev-Rendering to //UE4/Dev-Main (Source: //UE4/Dev-Rendering @ 3045398)

diff --git a/Engine/Source/Editor/UnrealEd/Private/EditorBuildUtils.cpp b/Engine/Source/Editor/UnrealEd/Private/EditorBuildUtils.cpp
--- a/Engine/Source/Editor/UnrealEd/Private/EditorBuildUtils.cpp
+++ b/Engine/Source/Editor/UnrealEd/Private/EditorBuildUtils.cpp
@@ -1042,31 +1042,33 @@ void FEditorBuildUtils::TriggerHierarchicalLODBuilder(UWorld* InWorld, FName Id)
 	InWorld->HierarchicalLODBuilder->BuildMeshesForLODActors();
 }
 
-bool FEditorBuildUtils::EditorBuildTextureStreaming(UWorld* InWorld)
+bool FEditorBuildUtils::EditorBuildTextureStreaming(UWorld* InWorld, bool bWithTexCoordScales, bool bDebugDataOnly)
 {
-	FScopedSlowTask BuildTextureStreamingTask(5.f, (LOCTEXT("TextureStreamingBuild", "Building Texture Streaming")));
+	FScopedSlowTask BuildTextureStreamingTask(bWithTexCoordScales ? 5.f : 1.f, bDebugDataOnly ? LOCTEXT("TextureStreamingDataUpdate", "Updating Texture Streaming Data") : LOCTEXT("TextureStreamingBuild", "Building Texture Streaming"));
 	BuildTextureStreamingTask.MakeDialog(true);
 
 	const EMaterialQualityLevel::Type QualityLevel = EMaterialQualityLevel::High;
 	const ERHIFeatureLevel::Type FeatureLevel = GMaxRHIFeatureLevel;
 
+
 	CollectGarbage( GARBAGE_COLLECTION_KEEPFLAGS );
 
 	FTexCoordScaleMap TexCoordScales;
-	if (!BuildTextureStreamingShaders(InWorld, QualityLevel, FeatureLevel, TexCoordScales, BuildTextureStreamingTask))
+	if (bWithTexCoordScales)
 	{
-		return false;
-	}
+		if (!BuildTextureStreamingShaders(InWorld, QualityLevel, FeatureLevel, TexCoordScales, bDebugDataOnly, BuildTextureStreamingTask))
+		{
+			return false;
+		}
 
-	// Exporting Material TexCoord Scales
-	{
+		// Exporting Material TexCoord Scales
 		FScopedSlowTask SlowTask((float)TexCoordScales.Num(), (LOCTEXT("TextureStreamingBuild_ExportingMaterialScales", "Exporting Material TexCoord Scales")));
 
 		const double StartTime = FPlatformTime::Seconds();
 		for (FTexCoordScaleMap::TIterator It(TexCoordScales); It; ++It)
 		{
 			SlowTask.EnterProgressFrame();
-			BuildTextureStreamingTask.EnterProgressFrame((1.f - KINDA_SMALL_NUMBER) / TexCoordScales.Num());
+			BuildTextureStreamingTask.EnterProgressFrame(1.f / TexCoordScales.Num());
 			if (GWarn->ReceivedUserCancel()) return false;
 
 			UMaterialInterface* MaterialInterface = It.Key();
@@ -1078,10 +1080,21 @@ bool FEditorBuildUtils::EditorBuildTextureStreaming(UWorld* InWorld)
 		UE_LOG(LogLevel, Display, TEXT("Export Material TexCoord Scales took %.3f seconds."), FPlatformTime::Seconds() - StartTime);
 	}
 
-	if (!BuildTextureStreamingData(InWorld, TexCoordScales, QualityLevel, FeatureLevel, BuildTextureStreamingTask))
+	if (bDebugDataOnly)
 	{
-		return false;
+		if (!UpdateComponentStreamingSectionData(InWorld, TexCoordScales, true, BuildTextureStreamingTask))
+		{
+			return false;
+		}
 	}
+	else
+	{
+		if (!BuildTextureStreamingData(InWorld, TexCoordScales, QualityLevel, FeatureLevel, BuildTextureStreamingTask))
+		{
+			return false;
+		}
+	}
+
 
 	CollectGarbage( GARBAGE_COLLECTION_KEEPFLAGS );
 	return true;
diff --git a/Engine/Source/Editor/UnrealEd/Private/EditorEngine.cpp b/Engine/Source/Editor/UnrealEd/Private/EditorEngine.cpp
--- a/Engine/Source/Editor/UnrealEd/Private/EditorEngine.cpp
+++ b/Engine/Source/Editor/UnrealEd/Private/EditorEngine.cpp
@@ -656,7 +656,8 @@ void UEditorEngine::Init(IEngineLoop* InEngineLoop)
 			TEXT("ReferenceViewer"),
 			TEXT("TreeMap"),
 			TEXT("SizeMap"),
-			TEXT("MergeActors")
+			TEXT("MergeActors"),
+			TEXT("NiagaraEditor")
 		};
 
 		FScopedSlowTask ModuleSlowTask(ARRAY_COUNT(ModuleNames));
diff --git a/Engine/Source/Editor/UnrealEd/Private/EditorViewportClient.cpp b/Engine/Source/Editor/UnrealEd/Private/EditorViewportClient.cpp
--- a/Engine/Source/Editor/UnrealEd/Private/EditorViewportClient.cpp
+++ b/Engine/Source/Editor/UnrealEd/Private/EditorViewportClient.cpp
@@ -31,6 +31,7 @@
 #include "IHeadMountedDisplay.h"
 #include "SceneViewExtension.h"
 #include "ComponentRecreateRenderStateContext.h"
+#include "EditorBuildUtils.h"
 
 #define LOCTEXT_NAMESPACE "EditorViewportClient"
 
@@ -4619,6 +4620,10 @@ void FEditorViewportClient::SetViewMode(EViewModeIndex InViewModeIndex)
 {
 	if (IsPerspective())
 	{
+		if (InViewModeIndex == VMI_PrimitiveDistanceAccuracy || InViewModeIndex == VMI_MeshTexCoordSizeAccuracy || InViewModeIndex == VMI_MaterialTexCoordScalesAccuracy)
+		{
+			FEditorBuildUtils::EditorBuildTextureStreaming(GetWorld(), InViewModeIndex == VMI_MaterialTexCoordScalesAccuracy, true);
+		}
 		PerspViewModeIndex = InViewModeIndex;
 		ApplyViewMode(PerspViewModeIndex, true, EngineShowFlags);
 		bForcingUnlitForNewMap = false;
diff --git a/Engine/Source/Editor/UnrealEd/Private/Factories/ActorFactory.cpp b/Engine/Source/Editor/UnrealEd/Private/Factories/ActorFactory.cpp
--- a/Engine/Source/Editor/UnrealEd/Private/Factories/ActorFactory.cpp
+++ b/Engine/Source/Editor/UnrealEd/Private/Factories/ActorFactory.cpp
@@ -644,82 +644,6 @@ void UActorFactoryEmitter::PostCreateBlueprint( UObject* Asset, AActor* CDO )
 }
 
 
-
-
-
-/*-----------------------------------------------------------------------------
-UActorFactoryNiagara
------------------------------------------------------------------------------*/
-UActorFactoryNiagara::UActorFactoryNiagara(const FObjectInitializer& ObjectInitializer)
-: Super(ObjectInitializer)
-{
-	DisplayName = LOCTEXT("NiagaraEffectDisplayName", "NiagaraEffect");
-	NewActorClass = ANiagaraActor::StaticClass();
-}
-
-bool UActorFactoryNiagara::CanCreateActorFrom(const FAssetData& AssetData, FText& OutErrorMsg)
-{
-	if (!AssetData.IsValid() || !AssetData.GetClass()->IsChildOf(UNiagaraEffect::StaticClass()))
-	{
-		OutErrorMsg = NSLOCTEXT("CanCreateActor", "NoEffect", "A valid Niagara effect must be specified.");
-		return false;
-	}
-
-	return true;
-}
-
-void UActorFactoryNiagara::PostSpawnActor(UObject* Asset, AActor* NewActor)
-{
-	Super::PostSpawnActor(Asset, NewActor);
-
-	UNiagaraEffect* Effect = CastChecked<UNiagaraEffect>(Asset);
-	ANiagaraActor* NiagaraActor = CastChecked<ANiagaraActor>(NewActor);
-
-	// Term Component
-	NiagaraActor->GetNiagaraComponent()->UnregisterComponent();
-
-	// Change properties
-	NiagaraActor->GetNiagaraComponent()->SetAsset(Effect);
-
-	// if we're created by Kismet on the server during gameplay, we need to replicate the emitter
-	if (NiagaraActor->GetWorld()->HasBegunPlay() && NiagaraActor->GetWorld()->GetNetMode() != NM_Client)
-	{
-		NiagaraActor->SetReplicates(true);
-		NiagaraActor->bAlwaysRelevant = true;
-		NiagaraActor->NetUpdateFrequency = 0.1f; // could also set bNetTemporary but LD might further trigger it or something
-	}
-
-	// Init Component
-	NiagaraActor->GetNiagaraComponent()->RegisterComponent();
-}
-
-UObject* UActorFactoryNiagara::GetAssetFromActorInstance(AActor* Instance)
-{
-	check(Instance->IsA(NewActorClass));
-	ANiagaraActor* NewActor = CastChecked<ANiagaraActor>(Instance);
-	if (NewActor->GetNiagaraComponent())
-	{
-		return NewActor->GetNiagaraComponent()->GetAsset();
-	}
-	else
-	{
-		return NULL;
-	}
-}
-
-void UActorFactoryNiagara::PostCreateBlueprint(UObject* Asset, AActor* CDO)
-{
-	if (Asset != NULL && CDO != NULL)
-	{
-		UNiagaraEffect* Effect = CastChecked<UNiagaraEffect>(Asset);
-		ANiagaraActor* Actor = CastChecked<ANiagaraActor>(CDO);
-		Actor->GetNiagaraComponent()->SetAsset(Effect);
-	}
-}
-
-
-
-
 /*-----------------------------------------------------------------------------
 UActorFactoryPlayerStart
 -----------------------------------------------------------------------------*/
diff --git a/Engine/Source/Editor/UnrealEd/Private/Factories/EditorFactories.cpp b/Engine/Source/Editor/UnrealEd/Private/Factories/EditorFactories.cpp
--- a/Engine/Source/Editor/UnrealEd/Private/Factories/EditorFactories.cpp
+++ b/Engine/Source/Editor/UnrealEd/Private/Factories/EditorFactories.cpp
@@ -4245,8 +4245,8 @@ UObject* UTextureFactory::FactoryCreateBinary
 	TextureFilter						ExistingFilter		= TF_Default;
 	TextureGroup						ExistingLODGroup	= TEXTUREGROUP_World;
 	TextureCompressionSettings			ExistingCompressionSettings = TC_Default;
-	int32									ExistingLODBias		= 0;
-	int32									ExistingNumCinematicMipLevels = 0;
+	int32								ExistingLODBias		= 0;
+	int32								ExistingNumCinematicMipLevels = 0;
 	bool								ExistingNeverStream = false;
 	bool								ExistingSRGB		= false;
 	bool								ExistingPreserveBorder = false;
@@ -4263,6 +4263,7 @@ UObject* UTextureFactory::FactoryCreateBinary
 	float								ExistingAdjustHue = 0.0f;
 	float								ExistingAdjustMinAlpha = 0.0f;
 	float								ExistingAdjustMaxAlpha = 1.0f;
+	FVector4							ExistingAlphaCoverageThresholds = FVector4(0, 0, 0, 0);
 	TextureMipGenSettings				ExistingMipGenSettings = TextureMipGenSettings(0);
 
 	bUsingExistingSettings = bSuppressImportOverwriteDialog;
@@ -4323,6 +4324,7 @@ UObject* UTextureFactory::FactoryCreateBinary
 		ExistingDeferCompression = ExistingTexture->DeferCompression;
 		ExistingFlipGreenChannel = ExistingTexture->bFlipGreenChannel;
 		ExistingDitherMipMapAlpha = ExistingTexture->bDitherMipMapAlpha;
+		ExistingAlphaCoverageThresholds = ExistingTexture->AlphaCoverageThresholds;
 		ExistingAdjustBrightness = ExistingTexture->AdjustBrightness;
 		ExistingAdjustBrightnessCurve = ExistingTexture->AdjustBrightnessCurve;
 		ExistingAdjustVibrance = ExistingTexture->AdjustVibrance;
@@ -4410,6 +4412,7 @@ UObject* UTextureFactory::FactoryCreateBinary
 	Texture->CompressionNoAlpha		= NoAlpha;
 	Texture->DeferCompression		= bDeferCompression;
 	Texture->bDitherMipMapAlpha		= bDitherMipMapAlpha;
+	Texture->AlphaCoverageThresholds = AlphaCoverageThresholds;
 	
 	if(Texture->MipGenSettings == TMGS_FromTextureGroup)
 	{
@@ -4444,6 +4447,7 @@ UObject* UTextureFactory::FactoryCreateBinary
 		Texture->CompressionNoAlpha = ExistingNoAlpha;
 		Texture->DeferCompression = ExistingDeferCompression;
 		Texture->bDitherMipMapAlpha = ExistingDitherMipMapAlpha;
+		Texture->AlphaCoverageThresholds = ExistingAlphaCoverageThresholds;
 		Texture->bFlipGreenChannel = ExistingFlipGreenChannel;
 		Texture->AdjustBrightness = ExistingAdjustBrightness;
 		Texture->AdjustBrightnessCurve = ExistingAdjustBrightnessCurve;
diff --git a/Engine/Source/Editor/UnrealEd/Private/PerformanceMonitor.cpp b/Engine/Source/Editor/UnrealEd/Private/PerformanceMonitor.cpp
--- a/Engine/Source/Editor/UnrealEd/Private/PerformanceMonitor.cpp
+++ b/Engine/Source/Editor/UnrealEd/Private/PerformanceMonitor.cpp
@@ -268,6 +268,7 @@ void FPerformanceMonitor::Tick(float DeltaTime)
 	{
 		static IConsoleVariable* CVarMinFPS = IConsoleManager::Get().FindConsoleVariable(TEXT("PerfWarn.FineMinFPS"));
 		static IConsoleVariable* CVarPercentThreshold = IConsoleManager::Get().FindConsoleVariable(TEXT("PerfWarn.FinePercentThreshold"));
+		static IConsoleVariable* CVarSampleTime = IConsoleManager::Get().FindConsoleVariable(TEXT("PerfWarn.FineSampleTime"));
 
 		PercentUnderTarget = FineMovingAverage.PercentageBelowThreshold(CVarMinFPS->GetFloat());
 
@@ -275,7 +276,7 @@ void FPerformanceMonitor::Tick(float DeltaTime)
 		{
 			Arguments.Add(TEXT("Framerate"), CVarMinFPS->GetInt());
 			Arguments.Add(TEXT("Percentage"), FMath::FloorToFloat(PercentUnderTarget));
-			SampleTime = IConsoleManager::Get().FindConsoleVariable(TEXT("PerfWarn.FineSampleTime"))->GetInt();
+			SampleTime = CVarSampleTime->GetInt();
 
 			bLowFramerate = true;
 		}
diff --git a/Engine/Source/Editor/UnrealEd/Public/EditorBuildUtils.h b/Engine/Source/Editor/UnrealEd/Public/EditorBuildUtils.h
--- a/Engine/Source/Editor/UnrealEd/Public/EditorBuildUtils.h
+++ b/Engine/Source/Editor/UnrealEd/Public/EditorBuildUtils.h
@@ -123,11 +123,13 @@ public:
 	/**
 	 * Perform an editor build for texture streaming
 	 *
-	 * @param	InWorld				WorldContext
+	 * @param	InWorld					WorldContext
+	 * @param	bWithTexCoordScales		true if the build should compute material texcoord scales.
+	 * @param	bDebugDataOnly			true if the build should only update data useful for accuracy viewmodes.
 	 *
 	 * @return	true if the build completed successfully; false if it did not (or was manually canceled)
 	 */
-	static UNREALED_API bool EditorBuildTextureStreaming( UWorld* InWorld );
+	static UNREALED_API bool EditorBuildTextureStreaming( UWorld* InWorld, bool bWithTexCoordScales = true, bool bDebugDataOnly = false);
 
 
 	/** 
diff --git a/Engine/Source/Editor/UnrealEd/Public/UnrealEd.h b/Engine/Source/Editor/UnrealEd/Public/UnrealEd.h
--- a/Engine/Source/Editor/UnrealEd/Public/UnrealEd.h
+++ b/Engine/Source/Editor/UnrealEd/Public/UnrealEd.h
@@ -38,7 +38,6 @@
 #include "Classes/ActorFactories/ActorFactoryPawn.h"
 #include "Classes/ActorFactories/ActorFactoryExponentialHeightFog.h"
 #include "Classes/ActorFactories/ActorFactoryMatineeActor.h"
-#include "Classes/ActorFactories/ActorFactoryNiagara.h"
 #include "Classes/ActorFactories/ActorFactoryNote.h"
 #include "Classes/ActorFactories/ActorFactoryPhysicsAsset.h"
 #include "Classes/ActorFactories/ActorFactoryPlaneReflectionCapture.h"
@@ -85,7 +84,6 @@
 #include "Classes/Commandlets/DumpHiddenCategoriesCommandlet.h"
 #include "Classes/MaterialGraph/MaterialGraph.h"
 #include "Classes/MaterialGraph/MaterialGraphNode_Comment.h"
-#include "Classes/Niagara/EdGraphSchema_Niagara.h"
 #include "Classes/Animation/EditorAnimBaseObj.h"
 #include "Classes/Animation/EditorAnimCompositeSegment.h"
 #include "Classes/Animation/EditorAnimSegment.h"
@@ -158,8 +156,6 @@
 #include "Classes/Factories/MaterialInstanceConstantFactoryNew.h"
 #include "Classes/Factories/MaterialParameterCollectionFactoryNew.h"
 #include "Classes/Factories/ModelFactory.h"
-#include "Classes/Factories/NiagaraEffectFactoryNew.h"
-#include "Classes/Factories/NiagaraScriptFactoryNew.h"
 #include "Classes/Factories/ObjectLibraryFactory.h"
 #include "Classes/Factories/PackageFactory.h"
 #include "Classes/Factories/PackFactory.h"
@@ -238,18 +234,6 @@
 #include "Classes/MaterialGraph/MaterialGraphNode_Root.h"
 #include "Classes/MaterialGraph/MaterialGraphSchema.h"
 #include "Classes/Exporters/ModelExporterT3D.h"
-#include "Classes/Niagara/NiagaraGraph.h"
-#include "Classes/Niagara/NiagaraNode.h"
-#include "Classes/Niagara/NiagaraNodeConstant.h"
-#include "Classes/Niagara/NiagaraNodeFunctionCall.h"
-#include "Classes/Niagara/NiagaraNodeGetAttr.h"
-#include "Classes/Niagara/NiagaraNodeInput.h"
-#include "Classes/Niagara/NiagaraNodeOp.h"
-#include "Classes/Niagara/NiagaraNodeOutput.h"
-#include "Classes/Niagara/NiagaraNodeOutputUpdate.h"
-#include "Classes/Niagara/NiagaraNodeReadDataSet.h"
-#include "Classes/Niagara/NiagaraNodeWriteDataSet.h"
-#include "Classes/Niagara/NiagaraScriptSource.h"
 #include "Classes/Exporters/ObjectExporterT3D.h"
 #include "Classes/Commandlets/ParticleSystemAuditCommandlet.h"
 #include "Classes/Preferences/PersonaOptions.h"
diff --git a/Engine/Source/Editor/UnrealEd/UnrealEd.Build.cs b/Engine/Source/Editor/UnrealEd/UnrealEd.Build.cs
--- a/Engine/Source/Editor/UnrealEd/UnrealEd.Build.cs
+++ b/Engine/Source/Editor/UnrealEd/UnrealEd.Build.cs
@@ -86,8 +86,6 @@ public class UnrealEd : ModuleRules
 				"BlueprintGraph",
                 "Http",
 				"UnrealAudio",
-                "Niagara",
-                "VectorVM",
                 "FunctionalTesting"
             }
 		);
diff --git a/Engine/Source/Programs/UnrealLightmass/Private/Lighting/Embree.cpp b/Engine/Source/Programs/UnrealLightmass/Private/Lighting/Embree.cpp
--- a/Engine/Source/Programs/UnrealLightmass/Private/Lighting/Embree.cpp
+++ b/Engine/Source/Programs/UnrealLightmass/Private/Lighting/Embree.cpp
@@ -137,63 +137,42 @@ struct FEmbreeFilterProcessor
      */
 	EMBREE_INLINE bool HitRejectedByHLODTest() const
 	{
+		const uint32 InvalidIndex = 0xFFFF;
 		uint32 GeoHLODTreeIndex = (Geo.Mesh->GetLODIndices() & 0xFFFF0000) >> 16;
+		uint32 RayHLODTreeIndex = Ray.MappingMesh ? (Ray.MappingMesh->GetLODIndices() & 0xFFFF0000) >> 16 : InvalidIndex;
 
-		if (Ray.MappingMesh)
+		// If either Geo or Ray is a HLOD (0xFFFF being invalid HLOD).
+		if (GeoHLODTreeIndex != InvalidIndex || RayHLODTreeIndex != InvalidIndex)
 		{
-			uint32 RayHLODTreeIndex = (Ray.MappingMesh->GetLODIndices() & 0xFFFF0000) >> 16;
-
-			// If either Geo or Ray is a HLOD (0xFFFF being invalid HLOD).
-			if (GeoHLODTreeIndex != 0xFFFF || RayHLODTreeIndex != 0xFFFF)
-			{
-				uint32 GeoHLODRange = Geo.Mesh->GetHLODRange();
-				uint32 GeoHLODRangeStart = GeoHLODRange & 0xFFFF;
-				uint32 GeoHLODRangeEnd = (GeoHLODRange & 0xFFFF0000) >> 16;
+			uint32 GeoHLODRange = Geo.Mesh->GetHLODRange();
+			uint32 GeoHLODRangeStart = GeoHLODRange & 0xFFFF;
+			uint32 GeoHLODRangeEnd = (GeoHLODRange & 0xFFFF0000) >> 16;
 
-				uint32 RayHLODRange = Ray.MappingMesh->GetHLODRange();
-				uint32 RayHLODRangeStart = RayHLODRange & 0xFFFF;
-				uint32 RayHLODRangeEnd = (RayHLODRange & 0xFFFF0000) >> 16;
+			uint32 RayHLODRange = Ray.MappingMesh ? Ray.MappingMesh->GetHLODRange() : 0;
+			uint32 RayHLODRangeStart = RayHLODRange & 0xFFFF;
+			uint32 RayHLODRangeEnd = (RayHLODRange & 0xFFFF0000) >> 16;
 
-				// Different rules apply if we're dealing with the same cluster.
-				if (GeoHLODTreeIndex != RayHLODTreeIndex)
-				{
-					// Allow HLOD leaf nodes (Start == End) to interact with other meshes, else reject.
-					if (GeoHLODRangeStart != GeoHLODRangeEnd)
-					{
-						return true;
-					}
-				}
-				else
+			// Different rules apply if we're dealing with the same cluster.
+			if (GeoHLODTreeIndex != RayHLODTreeIndex)
+			{
+				// Allow LOD0 nodes to interact with other meshes, else reject.
+				// Note: If the Start == End, it's an HLOD leaf which means LOD0 (inverted tree)
+				if (GeoHLODRangeStart != GeoHLODRangeEnd)
 				{
-					// Allow HLOD nodes to self-shadow (identical ranges), else reject where range intersects.
-					if (GeoHLODRange != RayHLODRange && 
-						GeoHLODRangeStart <= RayHLODRangeEnd && RayHLODRangeStart <= GeoHLODRangeEnd)
-					{
-						return true;
-					}
+					return true;
 				}
 			}
 			else
 			{
-				return false;
+				// Allow HLOD nodes to self-shadow (identical ranges), else reject where range intersects.
+				if (GeoHLODRange != RayHLODRange && 
+					GeoHLODRangeStart <= RayHLODRangeEnd && RayHLODRangeStart <= GeoHLODRangeEnd)
+				{
+					return true;
+				}
 			}
 		}
 
-		// If the ray didn't originate from a mesh, only intersect against LOD0 of the HLOD
-		if (GeoHLODTreeIndex > 0 && GeoHLODTreeIndex != 0xFFFF)
-		{
-			uint32 GeoHLODRange = Geo.Mesh->GetHLODRange();
-			uint32 GeoHLODRangeStart = GeoHLODRange & 0xFFFF;
-			uint32 GeoHLODRangeEnd = (GeoHLODRange & 0xFFFF0000) >> 16;
-
-			// If the Start == End, it's an HLOD leaf which means the LOD0 (inverted tree)
-			// So reject any intersections with non-leaves
-			if (GeoHLODRangeStart != GeoHLODRangeEnd)
-			{
-				return true;
-			}
-		}
-		
 		return false;
 	}
 
diff --git a/Engine/Source/Runtime/Core/Private/HAL/ConsoleManager.cpp b/Engine/Source/Runtime/Core/Private/HAL/ConsoleManager.cpp
--- a/Engine/Source/Runtime/Core/Private/HAL/ConsoleManager.cpp
+++ b/Engine/Source/Runtime/Core/Private/HAL/ConsoleManager.cpp
@@ -1954,14 +1954,6 @@ static TAutoConsoleVariable<int32> CVarUseDXT5NormalMaps(
 	// Changing this causes a full shader recompile
 	ECVF_ReadOnly);
 
-// Changing this is currently unsupported after content has been chunked with the previous setting
-// Changing this causes a full shader recompile
-static TAutoConsoleVariable<int32> CVarMaxGPUSkinBones(
-	TEXT("Compat.MAX_GPUSKIN_BONES"),
-	256,
-	TEXT("Max number of bones that can be skinned on the GPU in a single draw call. Cannot be changed at runtime."),
-	ECVF_ReadOnly);
-
 static TAutoConsoleVariable<int32> CVarContactShadows(
 	TEXT("r.ContactShadows"),
 	1,
diff --git a/Engine/Source/Runtime/Core/Public/Math/UnrealMath.h b/Engine/Source/Runtime/Core/Public/Math/UnrealMath.h
--- a/Engine/Source/Runtime/Core/Public/Math/UnrealMath.h
+++ b/Engine/Source/Runtime/Core/Public/Math/UnrealMath.h
@@ -469,11 +469,17 @@ inline FString FMath::FormatIntToHumanReadable(int32 Val)
 }
 
 template<class U>
-FORCEINLINE_DEBUGGABLE FRotator FMath::Lerp( const FRotator& A, const FRotator& B, const U& Alpha)
+FORCEINLINE_DEBUGGABLE FRotator FMath::Lerp(const FRotator& A, const FRotator& B, const U& Alpha)
 {
-	return (A * (1 - Alpha) + B * Alpha).GetNormalized();
+	return A + (B - A).GetNormalized() * Alpha;
 }
 
+template<class U>
+FORCEINLINE_DEBUGGABLE FRotator FMath::LerpRange(const FRotator& A, const FRotator& B, const U& Alpha)
+{
+	// Similar to Lerp, but does not take the shortest path. Allows interpolation over more than 180 degrees.
+	return (A * (1 - Alpha) + B * Alpha).GetNormalized();
+}
 
 template<class U>
 FORCEINLINE_DEBUGGABLE FQuat FMath::Lerp( const FQuat& A, const FQuat& B, const U& Alpha)
diff --git a/Engine/Source/Runtime/Core/Public/Math/UnrealMathUtility.h b/Engine/Source/Runtime/Core/Public/Math/UnrealMathUtility.h
--- a/Engine/Source/Runtime/Core/Public/Math/UnrealMathUtility.h
+++ b/Engine/Source/Runtime/Core/Public/Math/UnrealMathUtility.h
@@ -818,11 +818,12 @@ struct FMath : public FPlatformMath
 	}
 
 	// Rotator specific interpolation
-	template< class U > static FRotator Lerp( const FRotator& A, const FRotator& B, const U& Alpha);
+	template< class U > static FRotator Lerp(const FRotator& A, const FRotator& B, const U& Alpha);
+	template< class U > static FRotator LerpRange(const FRotator& A, const FRotator& B, const U& Alpha);
 
 	// Quat-specific interpolation
 
-	template< class U > static FQuat Lerp( const FQuat& A, const FQuat& B, const U& Alpha);
+	template< class U > static FQuat Lerp(const FQuat& A, const FQuat& B, const U& Alpha);
 	template< class U > static FQuat BiLerp(const FQuat& P00, const FQuat& P10, const FQuat& P01, const FQuat& P11, float FracX, float FracY);
 
 	/**
diff --git a/Engine/Source/Runtime/Core/Public/Misc/FeedbackContext.h b/Engine/Source/Runtime/Core/Public/Misc/FeedbackContext.h
--- a/Engine/Source/Runtime/Core/Public/Misc/FeedbackContext.h
+++ b/Engine/Source/Runtime/Core/Public/Misc/FeedbackContext.h
@@ -332,7 +332,8 @@ public:
 		CompletedWork += CurrentFrameScope;
 
 		const float WorkRemaining = TotalAmountOfWork - CompletedWork;
-		ensureMsgf(ExpectedWorkThisFrame <= WorkRemaining, TEXT("Work overflow in slow task. Please revise call-site to account for entire progress range."));
+		// Add a small threshold here because when there are a lot of tasks, numerical imprecision can add up and trigger this.
+		ensureMsgf(ExpectedWorkThisFrame <= 1.01f * TotalAmountOfWork - CompletedWork, TEXT("Work overflow in slow task. Please revise call-site to account for entire progress range."));
 		CurrentFrameScope = FMath::Min(WorkRemaining, ExpectedWorkThisFrame);
 
 		if (!bCreatedDialog && OpenDialogThreshold.IsSet() && static_cast<float>(FPlatformTime::Seconds() - StartTime) > OpenDialogThreshold.GetValue())
diff --git a/Engine/Source/Runtime/Core/Public/UObject/UnrealNames.inl b/Engine/Source/Runtime/Core/Public/UObject/UnrealNames.inl
--- a/Engine/Source/Runtime/Core/Public/UObject/UnrealNames.inl
+++ b/Engine/Source/Runtime/Core/Public/UObject/UnrealNames.inl
@@ -123,6 +123,7 @@ REGISTER_NAME(249,Rotation)
 REGISTER_NAME(250,BSP)
 REGISTER_NAME(251,EditorSettings)
 REGISTER_NAME(252,AudioThread)
+REGISTER_NAME(253,ID)
 
 // Online
 REGISTER_NAME(280,DGram)
diff --git a/Engine/Source/Runtime/Engine/Classes/Components/PrimitiveComponent.h b/Engine/Source/Runtime/Engine/Classes/Components/PrimitiveComponent.h
--- a/Engine/Source/Runtime/Engine/Classes/Components/PrimitiveComponent.h
+++ b/Engine/Source/Runtime/Engine/Classes/Components/PrimitiveComponent.h
@@ -6,6 +6,7 @@
 #include "PhysicsEngine/BodyInstance.h"
 #include "Components/SceneComponent.h"
 #include "SceneTypes.h"
+#include "Engine/TextureStreamingTypes.h"
 #include "Engine/EngineTypes.h"
 #include "AI/Navigation/NavRelevantInterface.h"
 #include "PrimitiveComponent.generated.h"
@@ -1177,6 +1178,22 @@ public:
 	 */
 	void GetStreamingTextureInfoWithNULLRemoval(FStreamingTextureLevelContext& LevelContext, TArray<FStreamingTexturePrimitiveInfo>& OutStreamingTextures) const;
 
+	/**
+	* Return whether this primitive should have section data for texture streaming but it is missing. Used for incremental updates.
+	*
+	* @param bCheckTexCoordScales - If true, section data must contains texcoord scales to be valid.
+	*
+	* @return - true if some sections have missing data. If this component is not expected to have data, this should return false.
+	*/
+	virtual bool HasMissingStreamingSectionData(bool bCheckTexCoordScales) const { return false; }
+
+	/**
+	* Update section data for texture streaming. Note that this data is expected to be transient.
+	* Only useful within the texture streaming build, or streaming accuracy viewmodes.
+	*
+	* @param	TexCoordScales - The texcoord scales for each texture register of each relevant materials.
+	*/
+	virtual void UpdateStreamingSectionData(const FTexCoordScaleMap& TexCoordScales) {}
 
 	/**
 	 * Determines the DPG the primitive's primary elements are drawn in.
diff --git a/Engine/Source/Runtime/Engine/Classes/Components/StaticMeshComponent.h b/Engine/Source/Runtime/Engine/Classes/Components/StaticMeshComponent.h
--- a/Engine/Source/Runtime/Engine/Classes/Components/StaticMeshComponent.h
+++ b/Engine/Source/Runtime/Engine/Classes/Components/StaticMeshComponent.h
@@ -3,7 +3,6 @@
 #pragma once
 
 #include "SceneTypes.h"
-#include "Engine/TextureStreamingTypes.h"
 #include "Components/MeshComponent.h"
 #include "Runtime/RenderCore/Public/PackedNormal.h"
 #include "RawIndexBuffer.h"
@@ -338,6 +337,22 @@ public:
 	}
 	virtual ELightMapInteractionType GetStaticLightingType() const override;
 
+	/**
+	* Return whether this primitive should have section data for texture streaming but it is missing. Used for incremental updates.
+	*
+	* @param	bCheckTexCoordScales		If true, section data must contains texcoord scales to be valid.
+	*
+	* @return	true if some sections have missing data. If this component is not expected to have data, this should return false.
+	*/
+	virtual bool HasMissingStreamingSectionData(bool bCheckTexCoordScales) const override;
+
+	/**
+	*	Update the precomputed streaming debug data of this component.
+	*
+	*	@param	TexCoordScales				The texcoord scales for each texture register of each relevant materials.
+	*/
+	virtual void UpdateStreamingSectionData(const FTexCoordScaleMap& TexCoordScales) override;
+
 	/**
 	 *	Update the precomputed streaming data of this component.
 	 *
@@ -348,13 +363,6 @@ public:
 	 */
 	virtual void UpdateStreamingTextureData(TArray<UTexture2D*>& LevelTextures, const FTexCoordScaleMap& TexCoordScales, EMaterialQualityLevel::Type QualityLevel, ERHIFeatureLevel::Type FeatureLevel);
 
-	/**
-	*	Update the precomputed streaming debug data of this component.
-	*
-	*	@param	TexCoordScales				The texcoord scales for each texture register of each relevant materials.
-	*/
-	virtual void UpdateStreamingSectionData(const FTexCoordScaleMap& TexCoordScales);
-
 	virtual bool GetStreamingTextureFactors(float& OutWorldTexelFactor, float& OutWorldLightmapFactor) const;
 	virtual void GetStreamingTextureInfo(FStreamingTextureLevelContext& LevelContext, TArray<FStreamingTexturePrimitiveInfo>& OutStreamingTextures) const override;
 
diff --git a/Engine/Source/Runtime/Engine/Classes/Engine/EngineTypes.h b/Engine/Source/Runtime/Engine/Classes/Engine/EngineTypes.h
--- a/Engine/Source/Runtime/Engine/Classes/Engine/EngineTypes.h
+++ b/Engine/Source/Runtime/Engine/Classes/Engine/EngineTypes.h
@@ -231,7 +231,7 @@ enum ETranslucencyLightingMode
 {
 	/** 
 	 * Lighting will be calculated for a volume, without directionality.  Use this on particle effects like smoke and dust.
-	 * This is the cheapest lighting method, however the material normal is not taken into account.
+	 * This is the cheapest per-pixel lighting method, however the material normal is not taken into account.
 	 */
 	TLM_VolumetricNonDirectional UMETA(DisplayName="Volumetric NonDirectional"),
 
@@ -262,8 +262,10 @@ enum ETranslucencyLightingMode
 
 	/** 
 	 * Lighting will be calculated for a surface. Use this on translucent surfaces like glass and water.
+	 * This is implemented with forward shading so specular highlights from local lights are supported, however many deferred-only features are not.
+	 * This is the most expensive translucency lighting method as each light's contribution is computed per-pixel.
 	 */
-	TLM_SurfacePerPixelLighting UMETA(DisplayName="Surface PerPixel (currently uses TranslucencyVolume)"),
+	TLM_SurfacePerPixelLighting UMETA(DisplayName="Surface ForwardShading"),
 
 	TLM_MAX,
 };
diff --git a/Engine/Source/Runtime/Engine/Classes/Engine/RendererSettings.h b/Engine/Source/Runtime/Engine/Classes/Engine/RendererSettings.h
--- a/Engine/Source/Runtime/Engine/Classes/Engine/RendererSettings.h
+++ b/Engine/Source/Runtime/Engine/Classes/Engine/RendererSettings.h
@@ -181,6 +181,12 @@ class ENGINE_API URendererSettings : public UDeveloperSettings
 		ToolTip = "The cubemap resolution for all reflection capture probes. Must be power of 2. Note that for very high values the memory and performance impact may be severe."))
 	int32 ReflectionCaptureResolution;
 
+	UPROPERTY(config, EditAnywhere, Category=ForwardShading, meta=(
+		ConsoleVariable="r.ForwardShading",
+		ToolTip="Whether to use forward shading on desktop platforms.  Requires Shader Model 5 hardware.  Forward shading has lower constant cost, but fewer features supported.  Changing this setting requires restarting the editor.",
+		ConfigRestartRequired=true))
+	uint32 bForwardShading:1;
+
 	UPROPERTY(config, EditAnywhere, Category=Lighting, meta=(
 		ConsoleVariable="r.AllowStaticLighting",
 		ToolTip="Whether to allow any static lighting to be generated and used, like lightmaps and shadowmaps. Games that only use dynamic lighting should set this to 0 to save some static lighting overhead. Changing this setting requires restarting the editor.",
diff --git a/Engine/Source/Runtime/Engine/Classes/Engine/Texture.h b/Engine/Source/Runtime/Engine/Classes/Engine/Texture.h
--- a/Engine/Source/Runtime/Engine/Classes/Engine/Texture.h
+++ b/Engine/Source/Runtime/Engine/Classes/Engine/Texture.h
@@ -419,6 +419,10 @@ public:
 	UPROPERTY(EditAnywhere, BlueprintReadWrite, Category=Texture, AdvancedDisplay)
 	uint32 bDitherMipMapAlpha:1;
 
+	/** Alpha values per channel to compare to when preserving alpha coverage. */
+	UPROPERTY(EditAnywhere, BlueprintReadWrite, Category = Texture, meta=(ClampMin = "0", ClampMax = "1"), AdvancedDisplay)
+	FVector4 AlphaCoverageThresholds;
+
 	/** When true the texture's border will be preserved during mipmap generation. */
 	UPROPERTY(EditAnywhere, BlueprintReadWrite, Category=LevelOfDetail, AdvancedDisplay)
 	uint32 bPreserveBorder:1;
diff --git a/Engine/Source/Runtime/Engine/Classes/Engine/Texture2D.h b/Engine/Source/Runtime/Engine/Classes/Engine/Texture2D.h
--- a/Engine/Source/Runtime/Engine/Classes/Engine/Texture2D.h
+++ b/Engine/Source/Runtime/Engine/Classes/Engine/Texture2D.h
@@ -65,7 +65,7 @@ public:
 private:
 	/** WorldSettings timestamp that tells the streamer to force all miplevels to be resident up until that time. */
 	UPROPERTY(transient)
-	float ForceMipLevelsToBeResidentTimestamp;
+	double ForceMipLevelsToBeResidentTimestamp;
 
 	/** True if streaming is temporarily disabled so we can update subregions of this texture's resource 
 	without streaming clobbering it. Automatically cleared before saving. */
diff --git a/Engine/Source/Runtime/Engine/Classes/Engine/TextureStreamingTypes.h b/Engine/Source/Runtime/Engine/Classes/Engine/TextureStreamingTypes.h
--- a/Engine/Source/Runtime/Engine/Classes/Engine/TextureStreamingTypes.h
+++ b/Engine/Source/Runtime/Engine/Classes/Engine/TextureStreamingTypes.h
@@ -225,6 +225,7 @@ public:
 typedef TMap<UMaterialInterface*, TArray<FMaterialTexCoordBuildInfo> > FTexCoordScaleMap;
 
 /** Build the shaders required for the texture streaming build. Returns whether or not the action was successful. */
-ENGINE_API bool BuildTextureStreamingShaders(UWorld* InWorld, EMaterialQualityLevel::Type QualityLevel, ERHIFeatureLevel::Type FeatureLevel, OUT FTexCoordScaleMap& TexCoordScales, FSlowTask& BuildTextureStreamingTask);
+ENGINE_API bool BuildTextureStreamingShaders(UWorld* InWorld, EMaterialQualityLevel::Type QualityLevel, ERHIFeatureLevel::Type FeatureLevel, OUT FTexCoordScaleMap& TexCoordScales, bool bIncremental, FSlowTask& BuildTextureStreamingTask);
+ENGINE_API bool UpdateComponentStreamingSectionData(UWorld* InWorld, const FTexCoordScaleMap& InTexCoordScales, bool bIncremental, FSlowTask& BuildTextureStreamingTask);
 ENGINE_API bool BuildTextureStreamingData(UWorld* InWorld, const FTexCoordScaleMap& InTexCoordScales, EMaterialQualityLevel::Type QualityLevel, ERHIFeatureLevel::Type FeatureLevel, FSlowTask& BuildTextureStreamingTask);
 extern ENGINE_API TAutoConsoleVariable<int32> CVarStreamingUseNewMetrics;
\ No newline at end of file
diff --git a/Engine/Source/Runtime/Engine/Classes/Materials/Material.h b/Engine/Source/Runtime/Engine/Classes/Materials/Material.h
--- a/Engine/Source/Runtime/Engine/Classes/Materials/Material.h
+++ b/Engine/Source/Runtime/Engine/Classes/Materials/Material.h
@@ -435,7 +435,7 @@ public:
 	UPROPERTY(EditAnywhere, Category=Material)
 	uint32 TwoSided:1;
 
-	/** Whether the material should support a dithered LOD transition when used with the foliage system. */
+	/** Whether the material should support a dithered LOD transition. */
 	UPROPERTY(EditAnywhere, Category=Material, AdvancedDisplay)
 	uint32 DitheredLODTransition:1;
 
@@ -628,7 +628,7 @@ public:
 	uint32 bAutomaticallySetUsageInEditor:1;
 
 	/* Forces the material to be completely rough. Saves a number of instructions and one sampler. */
-	UPROPERTY(EditAnywhere, BlueprintReadOnly, Category=Mobile)
+	UPROPERTY(EditAnywhere, BlueprintReadOnly, Category=Material, AdvancedDisplay)
 	uint32 bFullyRough:1;
 
 	/** 
@@ -647,8 +647,8 @@ public:
 	UPROPERTY(EditAnywhere, BlueprintReadOnly, Category = Mobile, meta = (DisplayName = "High Quality Reflections"))
 	uint32 bUseHQForwardReflections : 1;
 
-	/* Enables planar reflection when using the forward renderer. Enabling this setting reduces the number of samplers available to the material as one more sampler will be used for the planar reflection. */
-	UPROPERTY(EditAnywhere, BlueprintReadOnly, Category = Mobile, meta = (DisplayName = "Planar Reflections"))
+	/* Enables planar reflection when using the forward renderer or mobile. Enabling this setting reduces the number of samplers available to the material as one more sampler will be used for the planar reflection. */
+	UPROPERTY(EditAnywhere, BlueprintReadOnly, Category = ForwardShading, meta = (DisplayName = "Planar Reflections"))
 	uint32 bUsePlanarForwardReflections : 1;
 
 	/** The type of tessellation to apply to this object.  Note D3D11 required for anything except MTM_NoTessellation. */
diff --git a/Engine/Source/Runtime/Engine/Classes/Materials/MaterialExpressionNoise.h b/Engine/Source/Runtime/Engine/Classes/Materials/MaterialExpressionNoise.h
--- a/Engine/Source/Runtime/Engine/Classes/Materials/MaterialExpressionNoise.h
+++ b/Engine/Source/Runtime/Engine/Classes/Materials/MaterialExpressionNoise.h
@@ -8,14 +8,39 @@
 UENUM()
 enum ENoiseFunction
 {
-	/** Fast (~94 instructions per level). */
-	NOISEFUNCTION_Simplex UMETA(DisplayName="Simplex"),
-	/** Fast (~77 instructions per level) but low quality. */
-	NOISEFUNCTION_Perlin UMETA(DisplayName="Perlin"),
-	/** Very slow (~393 instructions per level). */
-	NOISEFUNCTION_Gradient UMETA(DisplayName="Gradient"),
-	/** Very fast (1 texture lookup, ~33 instructions per level), need to test more on every hardware, requires high quality texture filtering for bump mapping. */
-	NOISEFUNCTION_FastGradient UMETA(DisplayName="FastGradient"),
+	/** High quality for direct use and bumps 
+	 * ~77 instructions per level, 4 texture lookups
+	 * Cannot tile
+	 */
+	NOISEFUNCTION_SimplexTex UMETA(DisplayName="Simplex - Texture Based"),
+
+	/** High quality for direct use and bumps
+	 * Non-tiled: ~61 instructions per level, 8 texture lookups
+	 * Tiling: ~74 instructions per level, 8 texture lookups
+	 * Even "non-tiled" mode has a repeat of 128. Useful Repeat Size range <= 128
+	 * Formerly labeled as Perlin noise
+	 */
+	NOISEFUNCTION_GradientTex UMETA(DisplayName="Gradient - Texture Based"),
+
+	/** High quality for direct use, BAD for bumps
+	 * ~16 instructions per level, 1 texture lookup
+	 * Always tiles with a repeat of 16, "Tiling" mode is not an option for Fast Gradient noise
+	 */
+	NOISEFUNCTION_GradientTex3D UMETA(DisplayName = "Fast Gradient - 3D Texture"),
+
+	/** High quality for direct use and bumps
+	 * Non-tiled: ~103 instructions per level, no textures
+	 * Tiling: ~174 instructions per level, no textures
+	 */
+	NOISEFUNCTION_GradientALU UMETA(DisplayName = "Gradient - Computational"),
+
+	/** Low quality, but pure computation
+	 * Non-tiled: ~53 instructions per level, no textures
+	 * Tiling: ~118 instructions per level, no textures
+	 * Formerly mis-labeled as Gradient noise
+	 */
+	NOISEFUNCTION_ValueALU UMETA(DisplayName="Value - Computational"),
+
 	NOISEFUNCTION_MAX,
 };
 
@@ -48,7 +73,7 @@ class UMaterialExpressionNoise : public UMaterialExpression
 	UPROPERTY(EditAnywhere, Category=MaterialExpressionNoise)
 	uint32 bTurbulence:1;
 
-	/** 1 = fast but little detail, .. larger numbers cost more performance, only used for turbulence */
+	/** 1 = fast but little detail, .. larger numbers cost more performance */
 	UPROPERTY(EditAnywhere, Category=MaterialExpressionNoise, meta=(UIMin = "1", UIMax = "10"))
 	int32 Levels;
 
@@ -62,8 +87,17 @@ class UMaterialExpressionNoise : public UMaterialExpression
 	UPROPERTY(EditAnywhere, Category=MaterialExpressionNoise, meta=(UIMin = "2", UIMax = "8"))
 	float LevelScale;
 
+	/** Whether to use tiling noise pattern, useful for baking to seam-free repeating textures */
+	UPROPERTY(EditAnywhere, Category = MaterialExpressionNoise)
+	uint32 bTiling:1;
+
+	/** How many units in each tile (if Tiling is on) */
+	UPROPERTY(EditAnywhere, Category = MaterialExpressionNoise, meta=(UIMin = "4"))
+	uint32 RepeatSize;
+
 	//~ Begin UMaterialExpression Interface
 #if WITH_EDITOR
+	virtual bool CanEditChange(const UProperty* InProperty) const override;
 	virtual int32 Compile(class FMaterialCompiler* Compiler, int32 OutputIndex, int32 MultiplexIndex) override;
 	virtual void GetCaption(TArray<FString>& OutCaptions) const override;
 #endif
diff --git a/Engine/Source/Runtime/Engine/Classes/Materials/MaterialInterface.h b/Engine/Source/Runtime/Engine/Classes/Materials/MaterialInterface.h
--- a/Engine/Source/Runtime/Engine/Classes/Materials/MaterialInterface.h
+++ b/Engine/Source/Runtime/Engine/Classes/Materials/MaterialInterface.h
@@ -45,6 +45,7 @@ struct ENGINE_API FMaterialRelevance
 	uint32 bUsesGlobalDistanceField : 1;
 	uint32 bUsesWorldPositionOffset : 1;
 	uint32 bDecal : 1;
+	uint32 bTranslucentSurfaceLighting : 1;
 
 	/** Default constructor */
 	FMaterialRelevance()
diff --git a/Engine/Source/Runtime/Engine/Engine.Build.cs b/Engine/Source/Runtime/Engine/Engine.Build.cs
--- a/Engine/Source/Runtime/Engine/Engine.Build.cs
+++ b/Engine/Source/Runtime/Engine/Engine.Build.cs
@@ -33,7 +33,7 @@ public class Engine : ModuleRules
 				"AutomationWorker",
 				"MovieSceneCapture",
 				"DesktopPlatform",
-			}
+            }
 		);
 
 		if (Target.Configuration != UnrealTargetConfiguration.Shipping)
@@ -98,7 +98,6 @@ public class Engine : ModuleRules
 				"Landscape",
                 "UMG",
 				"Projects",
-				"Niagara",
                 "Internationalization",
                 "MaterialShaderQualitySettings",
                 "CinematicCamera",
@@ -155,7 +154,6 @@ public class Engine : ModuleRules
         CircularlyReferencedDependentModules.Add("AIModule");
 		CircularlyReferencedDependentModules.Add("Landscape");
         CircularlyReferencedDependentModules.Add("UMG");
-        CircularlyReferencedDependentModules.Add("Niagara");
         CircularlyReferencedDependentModules.Add("MaterialShaderQualitySettings");
         CircularlyReferencedDependentModules.Add("CinematicCamera");
 
@@ -410,5 +408,7 @@ public class Engine : ModuleRules
 		{
 			RuntimeDependencies.Add("$(EngineDir)/Content/Stats/...", StagedFileType.UFS);
 		}
-	}
+
+        DynamicallyLoadedModuleNames.Add("Niagara");
+    }
 }
diff --git a/Engine/Source/Runtime/Engine/Private/Components/DirectionalLightComponent.cpp b/Engine/Source/Runtime/Engine/Private/Components/DirectionalLightComponent.cpp
--- a/Engine/Source/Runtime/Engine/Private/Components/DirectionalLightComponent.cpp
+++ b/Engine/Source/Runtime/Engine/Private/Components/DirectionalLightComponent.cpp
@@ -135,7 +135,7 @@ public:
 		{
 			const FSceneInterface* Scene = Component->GetScene();
 			// ensure bUseWholeSceneCSMForMovableObjects is only be used with the forward renderer.
-			const bool bUsingDeferredRenderer = Scene == nullptr ? true : Scene->ShouldUseDeferredRenderer();
+			const bool bUsingDeferredRenderer = Scene == nullptr ? true : Scene->GetShadingPath() == EShadingPath::Deferred;
 			bUseWholeSceneCSMForMovableObjects = Component->Mobility == EComponentMobility::Stationary && !Component->bUseInsetShadowsForMovableObjects && !bUsingDeferredRenderer;
 		}
 		bCastModulatedShadows = Component->bCastModulatedShadows;
diff --git a/Engine/Source/Runtime/Engine/Private/Components/ReflectionCaptureComponent.cpp b/Engine/Source/Runtime/Engine/Private/Components/ReflectionCaptureComponent.cpp
--- a/Engine/Source/Runtime/Engine/Private/Components/ReflectionCaptureComponent.cpp
+++ b/Engine/Source/Runtime/Engine/Private/Components/ReflectionCaptureComponent.cpp
@@ -1381,7 +1381,7 @@ void ReadbackFromSM4Cubemap(FReflectionTextureCubeResource* SM4FullHDRCubemapTex
 
 void UReflectionCaptureComponent::ReadbackFromGPU(UWorld* WorldToUpdate)
 {
-	if (bDerivedDataDirty && !IsRunningCommandlet() && WorldToUpdate->FeatureLevel >= ERHIFeatureLevel::SM4)
+	if (bDerivedDataDirty && !IsRunningCommandlet() && WorldToUpdate && WorldToUpdate->FeatureLevel >= ERHIFeatureLevel::SM4)
 	{
 		FReflectionCaptureFullHDR* NewDerivedData = new FReflectionCaptureFullHDR();
 
diff --git a/Engine/Source/Runtime/Engine/Private/Components/SkyLightComponent.cpp b/Engine/Source/Runtime/Engine/Private/Components/SkyLightComponent.cpp
--- a/Engine/Source/Runtime/Engine/Private/Components/SkyLightComponent.cpp
+++ b/Engine/Source/Runtime/Engine/Private/Components/SkyLightComponent.cpp
@@ -473,8 +473,9 @@ void USkyLightComponent::UpdateSkyCaptureContentsArray(UWorld* WorldToUpdate, TA
 	for (int32 CaptureIndex = ComponentArray.Num() - 1; CaptureIndex >= 0; CaptureIndex--)
 	{
 		USkyLightComponent* CaptureComponent = ComponentArray[CaptureIndex];
+		AActor* Owner = CaptureComponent->GetOwner();
 
-		if ((!CaptureComponent->GetOwner() || WorldToUpdate->ContainsActor(CaptureComponent->GetOwner()))
+		if ((!Owner || !Owner->GetLevel() || (WorldToUpdate->ContainsActor(Owner) && Owner->GetLevel()->bIsVisible))
 			// Only process sky capture requests once async shader compiling completes, otherwise we will capture the scene with temporary shaders
 			&& (!bIsCompilingShaders || CaptureComponent->SourceType == SLS_SpecifiedCubemap))
 		{
diff --git a/Engine/Source/Runtime/Engine/Private/Components/StaticMeshComponent.cpp b/Engine/Source/Runtime/Engine/Private/Components/StaticMeshComponent.cpp
--- a/Engine/Source/Runtime/Engine/Private/Components/StaticMeshComponent.cpp
+++ b/Engine/Source/Runtime/Engine/Private/Components/StaticMeshComponent.cpp
@@ -636,6 +636,39 @@ void UStaticMeshComponent::UpdateStreamingTextureData(TArray<UTexture2D*>& Level
 #endif
 }
 
+bool UStaticMeshComponent::HasMissingStreamingSectionData(bool bCheckTexCoordScales) const
+{
+#if WITH_EDITORONLY_DATA
+	const bool bUseMetrics = CVarStreamingUseNewMetrics.GetValueOnGameThread() != 0;
+	const bool bNeedsPrecomputedData = !bIgnoreInstanceForTextureStreaming && Mobility == EComponentMobility::Static && StaticMesh && StaticMesh->RenderData;
+
+	if (bUseMetrics && bNeedsPrecomputedData)
+	{
+		if (!StreamingSectionData.IsValid())
+		{
+			return true;
+		}
+		else if (bCheckTexCoordScales)
+		{
+			for (const FStreamingSectionBuildInfo& SectionData : *StreamingSectionData)
+			{
+				// If at least one section has data, then it's enough to say it was built with TexCoordScales.
+				if (SectionData.TexCoordData.Num())
+				{
+					return false;
+				}
+			}
+			return true;
+		}
+		else
+		{
+			return !StreamingSectionData->Num();
+		}
+	}
+#endif
+	return false;
+}
+
 void UStaticMeshComponent::UpdateStreamingSectionData(const FTexCoordScaleMap& TexCoordScales)
 {
 #if WITH_EDITORONLY_DATA
diff --git a/Engine/Source/Runtime/Engine/Private/ContentStreaming.cpp b/Engine/Source/Runtime/Engine/Private/ContentStreaming.cpp
--- a/Engine/Source/Runtime/Engine/Private/ContentStreaming.cpp
+++ b/Engine/Source/Runtime/Engine/Private/ContentStreaming.cpp
@@ -199,8 +199,6 @@ bool TrackTextureEvent( FStreamingTexture* StreamingTexture, UTexture2D* Texture
 		TrackTextureInit();
 	}
 
-	float MipBias =  Manager ? Manager->MipBias : 0;
-
 	int32 NumTrackedTextures = GTrackedTextureNames.Num();
 	if ( NumTrackedTextures )
 	{
@@ -1063,6 +1061,12 @@ void FStreamingManagerCollection::AddOrRemoveTextureStreamingManagerIfNeeded(boo
 		//Remove the texture streaming manager if needed.
 		if( TextureStreamingManager )
 		{
+			TextureStreamingManager->BlockTillAllRequestsFinished();
+
+			RemoveStreamingManager(TextureStreamingManager);
+			delete TextureStreamingManager;
+			TextureStreamingManager = nullptr;
+
 			for( TObjectIterator<UTexture2D>It; It; ++It )
 			{
 				if( It->bIsStreamable )
@@ -1070,11 +1074,6 @@ void FStreamingManagerCollection::AddOrRemoveTextureStreamingManagerIfNeeded(boo
 					It->UpdateResource();
 				}
 			}
-			TextureStreamingManager->BlockTillAllRequestsFinished();
-
-			RemoveStreamingManager(TextureStreamingManager);
-			delete TextureStreamingManager;
-			TextureStreamingManager = nullptr;
 		}
 	}
 }
diff --git a/Engine/Source/Runtime/Engine/Private/DebugViewModeMaterialProxy.h b/Engine/Source/Runtime/Engine/Private/DebugViewModeMaterialProxy.h
--- a/Engine/Source/Runtime/Engine/Private/DebugViewModeMaterialProxy.h
+++ b/Engine/Source/Runtime/Engine/Private/DebugViewModeMaterialProxy.h
@@ -103,6 +103,7 @@ public:
 	static void AddShader(UMaterialInterface* InMaterialInterface, EMaterialQualityLevel::Type QualityLevel, ERHIFeatureLevel::Type FeatureLevel, EMaterialShaderMapUsage::Type InUsage);
 	static const FMaterial* GetShader(EDebugViewShaderMode DebugViewShaderMode, const FMaterial* Material);
 	static void ClearAllShaders();
+	static bool HasAnyShaders() { return DebugMaterialShaderMap.Num() > 0; }
 	static void ValidateAllShaders(OUT FTexCoordScaleMap& TexCoordScales);
 
 private:
diff --git a/Engine/Source/Runtime/Engine/Private/GPUSkinVertexFactory.cpp b/Engine/Source/Runtime/Engine/Private/GPUSkinVertexFactory.cpp
--- a/Engine/Source/Runtime/Engine/Private/GPUSkinVertexFactory.cpp
+++ b/Engine/Source/Runtime/Engine/Private/GPUSkinVertexFactory.cpp
@@ -10,6 +10,15 @@
 #include "ShaderParameters.h"
 #include "ShaderParameterUtils.h"
 
+// Changing this is currently unsupported after content has been chunked with the previous setting
+// Changing this causes a full shader recompile
+static int32 GCVarMaxGPUSkinBones = FGPUBaseSkinVertexFactory::GHardwareMaxGPUSkinBones;
+static FAutoConsoleVariableRef CVarMaxGPUSkinBones(
+	TEXT("Compat.MAX_GPUSKIN_BONES"),
+	GCVarMaxGPUSkinBones,
+	TEXT("Max number of bones that can be skinned on the GPU in a single draw call. Cannot be changed at runtime."),
+	ECVF_ReadOnly);
+
 IMPLEMENT_UNIFORM_BUFFER_STRUCT(FAPEXClothUniformShaderParameters,TEXT("APEXClothParam"));
 
 IMPLEMENT_UNIFORM_BUFFER_STRUCT(FBoneMatricesUniformShaderParameters,TEXT("Bones"));
@@ -289,6 +298,11 @@ bool FGPUBaseSkinVertexFactory::FShaderDataType::UpdateBoneData(FRHICommandListI
 	return false;
 }
 
+int32 FGPUBaseSkinVertexFactory::GetMaxGPUSkinBones()
+{
+	return GCVarMaxGPUSkinBones;
+}
+
 /*-----------------------------------------------------------------------------
 	FBoneDataTexture
 	SizeX(32 * 1024) - Good size for UE3
diff --git a/Engine/Source/Runtime/Engine/Private/Materials/HLSLMaterialTranslator.h b/Engine/Source/Runtime/Engine/Private/Materials/HLSLMaterialTranslator.h
--- a/Engine/Source/Runtime/Engine/Private/Materials/HLSLMaterialTranslator.h
+++ b/Engine/Source/Runtime/Engine/Private/Materials/HLSLMaterialTranslator.h
@@ -4120,7 +4120,7 @@ protected:
 			*GetParameterCode(Depth), FunctionValueIndex);
 	}
 
-	virtual int32 Noise(int32 Position, float Scale, int32 Quality, uint8 NoiseFunction, bool bTurbulence, int32 Levels, float OutputMin, float OutputMax, float LevelScale, int32 FilterWidth) override
+	virtual int32 Noise(int32 Position, float Scale, int32 Quality, uint8 NoiseFunction, bool bTurbulence, int32 Levels, float OutputMin, float OutputMax, float LevelScale, int32 FilterWidth, bool bTiling, uint32 RepeatSize) override
 	{
 		if (ErrorUnlessFeatureLevelSupported(ERHIFeatureLevel::SM4) == INDEX_NONE)
 		{
@@ -4143,9 +4143,11 @@ protected:
 		int32 OutputMinConst = Constant(OutputMin);
 		int32 OutputMaxConst = Constant(OutputMax);
 		int32 LevelScaleConst = Constant(LevelScale);
+		int32 TilingConst = Constant(bTiling);
+		int32 RepeatSizeConst = Constant(RepeatSize);
 
 		return AddCodeChunk(MCT_Float, 
-			TEXT("MaterialExpressionNoise(%s,%s,%s,%s,%s,%s,%s,%s,%s,%s)"), 
+			TEXT("MaterialExpressionNoise(%s,%s,%s,%s,%s,%s,%s,%s,%s,%s,%s,%s)"), 
 			*GetParameterCode(Position),
 			*GetParameterCode(ScaleConst),
 			*GetParameterCode(QualityConst),
@@ -4155,7 +4157,9 @@ protected:
 			*GetParameterCode(OutputMinConst),
 			*GetParameterCode(OutputMaxConst),
 			*GetParameterCode(LevelScaleConst),
-			*GetParameterCode(FilterWidth));
+			*GetParameterCode(FilterWidth),
+			*GetParameterCode(TilingConst),
+			*GetParameterCode(RepeatSizeConst));
 	}
 
 	virtual int32 BlackBody( int32 Temp ) override
diff --git a/Engine/Source/Runtime/Engine/Private/Materials/MaterialExpressions.cpp b/Engine/Source/Runtime/Engine/Private/Materials/MaterialExpressions.cpp
--- a/Engine/Source/Runtime/Engine/Private/Materials/MaterialExpressions.cpp
+++ b/Engine/Source/Runtime/Engine/Private/Materials/MaterialExpressions.cpp
@@ -8813,14 +8813,39 @@ UMaterialExpressionNoise::UMaterialExpressionNoise(const FObjectInitializer& Obj
 	OutputMin = -1.0f;
 	OutputMax = 1.0f;
 	LevelScale = 2.0f;
-	NoiseFunction = NOISEFUNCTION_Simplex;
+	NoiseFunction = NOISEFUNCTION_SimplexTex;
 	bTurbulence = true;
+	bTiling = false;
+	RepeatSize = 512;
 
 	MenuCategories.Add(ConstructorStatics.NAME_Utility);
 
 }
 
 #if WITH_EDITOR
+bool UMaterialExpressionNoise::CanEditChange(const UProperty* InProperty) const
+{
+	bool bIsEditable = Super::CanEditChange(InProperty);
+	if (bIsEditable && InProperty != NULL)
+	{
+		FName PropertyFName = InProperty->GetFName();
+
+		bool bTilableNoiseType = NoiseFunction == NOISEFUNCTION_GradientALU || NoiseFunction == NOISEFUNCTION_ValueALU 
+			|| NoiseFunction == NOISEFUNCTION_GradientTex;
+
+		if (PropertyFName == GET_MEMBER_NAME_CHECKED(UMaterialExpressionNoise, bTiling))
+		{
+			bIsEditable = bTilableNoiseType;
+		}
+		else if (PropertyFName == GET_MEMBER_NAME_CHECKED(UMaterialExpressionNoise, RepeatSize))
+		{
+			bIsEditable = bTilableNoiseType && bTiling;
+		}
+	}
+
+	return bIsEditable;
+}
+
 int32 UMaterialExpressionNoise::Compile(class FMaterialCompiler* Compiler, int32 OutputIndex, int32 MultiplexIndex)
 {
 	int32 PositionInput;
@@ -8845,7 +8870,7 @@ int32 UMaterialExpressionNoise::Compile(class FMaterialCompiler* Compiler, int32
 		FilterWidthInput = Compiler->Constant(0);
 	}
 
-	return Compiler->Noise(PositionInput, Scale, Quality, NoiseFunction, bTurbulence, Levels, OutputMin, OutputMax, LevelScale, FilterWidthInput);
+	return Compiler->Noise(PositionInput, Scale, Quality, NoiseFunction, bTurbulence, Levels, OutputMin, OutputMax, LevelScale, FilterWidthInput, bTiling, RepeatSize);
 }
 
 void UMaterialExpressionNoise::GetCaption(TArray<FString>& OutCaptions) const
diff --git a/Engine/Source/Runtime/Engine/Private/Materials/MaterialInterface.cpp b/Engine/Source/Runtime/Engine/Private/Materials/MaterialInterface.cpp
--- a/Engine/Source/Runtime/Engine/Private/Materials/MaterialInterface.cpp
+++ b/Engine/Source/Runtime/Engine/Private/Materials/MaterialInterface.cpp
@@ -26,6 +26,7 @@ void FMaterialRelevance::SetPrimitiveViewRelevance(FPrimitiveViewRelevance& OutV
 	OutViewRelevance.bUsesGlobalDistanceField = bUsesGlobalDistanceField;
 	OutViewRelevance.bUsesWorldPositionOffset = bUsesWorldPositionOffset;
 	OutViewRelevance.bDecal = bDecal;
+	OutViewRelevance.bTranslucentSurfaceLighting = bTranslucentSurfaceLighting;
 }
 
 //////////////////////////////////////////////////////////////////////////
@@ -91,6 +92,7 @@ FMaterialRelevance UMaterialInterface::GetRelevance_Internal(const UMaterial* Ma
 			MaterialRelevance.bOutputsVelocityInBasePass = Material->bOutputVelocityOnBasePass;	
 			MaterialRelevance.bUsesGlobalDistanceField = MaterialResource->UsesGlobalDistanceField_GameThread();
 			MaterialRelevance.bUsesWorldPositionOffset = MaterialResource->UsesWorldPositionOffset_GameThread();
+			MaterialRelevance.bTranslucentSurfaceLighting = bIsTranslucent && (MaterialResource->GetTranslucencyLightingMode() == TLM_SurfacePerPixelLighting);
 		}
 		return MaterialRelevance;
 	}
diff --git a/Engine/Source/Runtime/Engine/Private/Particles/ParticleGpuSimulation.cpp b/Engine/Source/Runtime/Engine/Private/Particles/ParticleGpuSimulation.cpp
--- a/Engine/Source/Runtime/Engine/Private/Particles/ParticleGpuSimulation.cpp
+++ b/Engine/Source/Runtime/Engine/Private/Particles/ParticleGpuSimulation.cpp
@@ -3729,7 +3729,7 @@ private:
 	 */
 	void InitLocalVectorField()
 	{
-		LocalVectorFieldRotation = FMath::Lerp(
+		LocalVectorFieldRotation = FMath::LerpRange(
 			EmitterInfo.LocalVectorField.MinInitialRotation,
 			EmitterInfo.LocalVectorField.MaxInitialRotation,
 			RandomStream.GetFraction() );
diff --git a/Engine/Source/Runtime/Engine/Private/Particles/ParticleResources.h b/Engine/Source/Runtime/Engine/Private/Particles/ParticleResources.h
--- a/Engine/Source/Runtime/Engine/Private/Particles/ParticleResources.h
+++ b/Engine/Source/Runtime/Engine/Private/Particles/ParticleResources.h
@@ -52,7 +52,7 @@ public:
 };
 
 /** Global particle index buffer. */
-extern TGlobalResource<FParticleIndexBuffer> GParticleIndexBuffer;
+ENGINE_API extern TGlobalResource<FParticleIndexBuffer> GParticleIndexBuffer;
 extern TGlobalResource<FSixTriangleParticleIndexBuffer> GSixTriangleParticleIndexBuffer;
 
 typedef FShaderResourceViewRHIParamRef FParticleShaderParamRef;
diff --git a/Engine/Source/Runtime/Engine/Private/PointLightSceneProxy.h b/Engine/Source/Runtime/Engine/Private/PointLightSceneProxy.h
--- a/Engine/Source/Runtime/Engine/Private/PointLightSceneProxy.h
+++ b/Engine/Source/Runtime/Engine/Private/PointLightSceneProxy.h
@@ -109,7 +109,7 @@ public:
 		FVector LightVector = SubjectBounds.Origin - LightPosition;
 		float LightDistance = LightVector.Size();
 		float SilhouetteRadius = 1.0f;
-		const float SubjectRadius = SubjectBounds.SphereRadius;
+		const float SubjectRadius = SubjectBounds.BoxExtent.Size();
 		const float ShadowRadiusMultiplier = 1.1f;
 
 		if (LightDistance > SubjectRadius)
diff --git a/Engine/Source/Runtime/Engine/Private/SceneView.cpp b/Engine/Source/Runtime/Engine/Private/SceneView.cpp
--- a/Engine/Source/Runtime/Engine/Private/SceneView.cpp
+++ b/Engine/Source/Runtime/Engine/Private/SceneView.cpp
@@ -619,8 +619,8 @@ FSceneView::FSceneView(const FSceneViewInitOptions& InitOptions)
 
 	// OpenGL Gamma space output in GLSL flips Y when rendering directly to the back buffer (so not needed on PC, as we never render directly into the back buffer)
 	auto ShaderPlatform = GShaderPlatformForFeatureLevel[FeatureLevel];
-	bool bUsingForwardRenderer = FSceneInterface::ShouldUseDeferredRenderer(FeatureLevel) == false;
-	bool bPlatformRequiresReverseCulling = (IsOpenGLPlatform(ShaderPlatform) && bUsingForwardRenderer && !IsPCPlatform(ShaderPlatform) && !IsVulkanMobilePlatform(ShaderPlatform));
+	bool bUsingMobileRenderer = FSceneInterface::GetShadingPath(FeatureLevel) == EShadingPath::Mobile;
+	bool bPlatformRequiresReverseCulling = (IsOpenGLPlatform(ShaderPlatform) && bUsingMobileRenderer && !IsPCPlatform(ShaderPlatform) && !IsVulkanMobilePlatform(ShaderPlatform));
 	static auto* MobileHDRCvar = IConsoleManager::Get().FindTConsoleVariableDataInt(TEXT("r.MobileHDR"));
 	check(MobileHDRCvar);
 	bReverseCulling = (bPlatformRequiresReverseCulling && MobileHDRCvar->GetValueOnAnyThread() == 0) ? !bReverseCulling : bReverseCulling;
diff --git a/Engine/Source/Runtime/Engine/Private/ShaderCompiler/ShaderCompiler.cpp b/Engine/Source/Runtime/Engine/Private/ShaderCompiler/ShaderCompiler.cpp
--- a/Engine/Source/Runtime/Engine/Private/ShaderCompiler/ShaderCompiler.cpp
+++ b/Engine/Source/Runtime/Engine/Private/ShaderCompiler/ShaderCompiler.cpp
@@ -2385,6 +2385,10 @@ void GlobalBeginCompileShader(
 				Input.DumpDebugInfoPath.ReplaceInline(TEXT("PROPAGATE_MULTIPLE_BOUNCES"), TEXT("MULT_BOUNC"));
 				Input.DumpDebugInfoPath.ReplaceInline(TEXT("PostProcess"), TEXT("Post"));
 				Input.DumpDebugInfoPath.ReplaceInline(TEXT("AntiAliasing"), TEXT("AA"));
+				Input.DumpDebugInfoPath.ReplaceInline(TEXT("Mobile"), TEXT("Mob"));
+				Input.DumpDebugInfoPath.ReplaceInline(TEXT("Linear"), TEXT("Lin"));
+				Input.DumpDebugInfoPath.ReplaceInline(TEXT("INT32_MAX"), TEXT("IMAX"));
+				Input.DumpDebugInfoPath.ReplaceInline(TEXT("Policy"), TEXT("Pol"));
 			}
 		}
 		// Sanitize the name to be used as a path
@@ -2530,6 +2534,11 @@ void GlobalBeginCompileShader(
 		Input.Environment.SetDefine(TEXT("PROJECT_ALLOW_GLOBAL_CLIP_PLANE"), CVar ? (CVar->GetInt() != 0) : 0);
 	}
 
+	{
+		static IConsoleVariable* CVar = IConsoleManager::Get().FindConsoleVariable(TEXT("r.ForwardShading"));
+		Input.Environment.SetDefine(TEXT("FORWARD_SHADING"), CVar ? (CVar->GetInt() != 0) : 0);
+	}
+
 	NewJobs.Add(NewJob);
 }
 
diff --git a/Engine/Source/Runtime/Engine/Private/StaticMeshRender.cpp b/Engine/Source/Runtime/Engine/Private/StaticMeshRender.cpp
--- a/Engine/Source/Runtime/Engine/Private/StaticMeshRender.cpp
+++ b/Engine/Source/Runtime/Engine/Private/StaticMeshRender.cpp
@@ -700,7 +700,7 @@ void FStaticMeshSceneProxy::DrawStaticElements(FStaticPrimitiveDrawInterface* PD
 
 						// Depth pass is only used for deferred renderer. The other conditions are meant to match the logic in FStaticMesh::AddToDrawLists.
 						// Could not link to "GEarlyZPassMovable" so moveable are ignored.
-						bUseUnifiedMeshForDepth = ShouldUseAsOccluder() && GetScene().ShouldUseDeferredRenderer() && !IsMovable();
+						bUseUnifiedMeshForDepth = ShouldUseAsOccluder() && GetScene().GetShadingPath() == EShadingPath::Deferred && !IsMovable();
 
 						if (bUseUnifiedMeshForShadow || bUseUnifiedMeshForDepth)
 						{
diff --git a/Engine/Source/Runtime/Engine/Private/Streaming/AsyncTextureStreaming.cpp b/Engine/Source/Runtime/Engine/Private/Streaming/AsyncTextureStreaming.cpp
--- a/Engine/Source/Runtime/Engine/Private/Streaming/AsyncTextureStreaming.cpp
+++ b/Engine/Source/Runtime/Engine/Private/Streaming/AsyncTextureStreaming.cpp
@@ -21,20 +21,17 @@ void FAsyncTextureStreamingData::Init(TArray<FStreamingViewInfo> InViewInfos, fl
 	}
 }
 
-void FAsyncTextureStreamingData::UpdateBoundSizes_Async()
+void FAsyncTextureStreamingData::UpdateBoundSizes_Async(const FTextureStreamingSettings& Settings)
 {
-	const bool bUseNewMetrics = CVarStreamingUseNewMetrics.GetValueOnAnyThread() > 0;
-	const float MaxEffectiveScreenSize = CVarStreamingScreenSizeEffectiveMax.GetValueOnAnyThread();
-
 	for (FTextureInstanceAsyncView& StaticInstancesView : StaticInstancesViews)
 	{
-		StaticInstancesView.UpdateBoundSizes_Async(ViewInfos, LastUpdateTime, bUseNewMetrics, MaxEffectiveScreenSize);
+		StaticInstancesView.UpdateBoundSizes_Async(ViewInfos, LastUpdateTime, Settings);
 	}
-	DynamicInstancesView.UpdateBoundSizes_Async(ViewInfos, LastUpdateTime, bUseNewMetrics, MaxEffectiveScreenSize);
+	DynamicInstancesView.UpdateBoundSizes_Async(ViewInfos, LastUpdateTime, Settings);
 }
 
 
-void FAsyncTextureStreamingData::UpdatePerfectWantedMips_Async(FStreamingTexture& StreamingTexture, float MipBias, bool bOutputToLog) const
+void FAsyncTextureStreamingData::UpdatePerfectWantedMips_Async(FStreamingTexture& StreamingTexture, const FTextureStreamingSettings& Settings, bool bOutputToLog) const
 {
 #if UE_BUILD_SHIPPING
 	bOutputToLog = false;
@@ -50,7 +47,14 @@ void FAsyncTextureStreamingData::UpdatePerfectWantedMips_Async(FStreamingTexture
 
 	const float MaxAllowedSize = StreamingTexture.GetMaxAllowedSize();
 
-	if (StreamingTexture.MinAllowedMips == StreamingTexture.MaxAllowedMips)
+	if (Settings.bFullyLoadUsedTextures)
+	{
+		if (StreamingTexture.LastRenderTime < 300)
+		{
+			MaxSize_VisibleOnly = FLT_MAX;
+		}
+	}
+	else if (StreamingTexture.MinAllowedMips == StreamingTexture.MaxAllowedMips)
 	{
 		MaxSize_VisibleOnly = MaxSize = MaxAllowedSize;
 	}
@@ -124,7 +128,7 @@ void FAsyncTextureStreamingData::UpdatePerfectWantedMips_Async(FStreamingTexture
 		}
 	}
 
-	StreamingTexture.SetPerfectWantedMips_Async(MaxSize, MaxSize_VisibleOnly, MipBias, bLooksLowRes);
+	StreamingTexture.SetPerfectWantedMips_Async(MaxSize, MaxSize_VisibleOnly, bLooksLowRes, Settings);
 }
 
 void FAsyncTextureStreamingTask::UpdateBudgetedMips_Async(int64& MemoryUsed, int64& TempMemoryUsed)
@@ -134,6 +138,8 @@ void FAsyncTextureStreamingTask::UpdateBudgetedMips_Async(int64& MemoryUsed, int
 	//*************************************
 
 	TArray<FStreamingTexture>& StreamingTextures = StreamingManager.StreamingTextures;
+	const FTextureStreamingSettings& Settings = StreamingManager.Settings;
+
 	TArray<int32> PrioritizedTextures;
 
 	int64 MemoryBudgeted = 0;
@@ -181,6 +187,8 @@ void FAsyncTextureStreamingTask::UpdateBudgetedMips_Async(int64& MemoryUsed, int
 	// If the budget is taking too much, drop some mips.
 	if (MemoryBudgeted > MemoryBudget && !IsAborted())
 	{
+		const int32 MaxPerTextureMipBias = Settings.MaxPerTextureMipBias();
+
 		PrioritizedTextures.Empty(StreamingTextures.Num());
 
 		for (int32 TextureIndex = 0; TextureIndex < StreamingTextures.Num() && !IsAborted(); ++TextureIndex)
@@ -224,7 +232,7 @@ void FAsyncTextureStreamingTask::UpdateBudgetedMips_Async(int64& MemoryUsed, int
 					continue;
 				}
 
-				int64 GivenMemory = StreamingTexture.DropOneMip_Async();
+				int64 GivenMemory = StreamingTexture.DropOneMip_Async(MaxPerTextureMipBias);
 
 				if (GivenMemory > 0)
 				{
@@ -305,14 +313,13 @@ void FAsyncTextureStreamingTask::UpdateBudgetedMips_Async(int64& MemoryUsed, int
 	// Handle drop mips debug option
 	//*************************************
 #if !UE_BUILD_SHIPPING
-	int32 DropMipsSetting = CVarStreamingDropMips.GetValueOnAnyThread();
-	if (DropMipsSetting > 0)
+	if (Settings.DropMips > 0)
 	{
 		for (FStreamingTexture& StreamingTexture : StreamingTextures)
 		{
 			if (IsAborted()) break;
 
-			if (DropMipsSetting == 1)
+			if (Settings.DropMips == 1)
 			{
 				StreamingTexture.BudgetedMips = FMath::Min<int32>(StreamingTexture.BudgetedMips, StreamingTexture.GetPerfectWantedMips());
 			}
@@ -333,7 +340,6 @@ void FAsyncTextureStreamingTask::UpdateLoadAndCancelationRequests_Async(int64 Me
 	PrioritizedTextures.Empty(StreamingTextures.Num());
 	for (int32 TextureIndex = 0; TextureIndex < StreamingTextures.Num() && !IsAborted(); ++TextureIndex)
 	{
-
 		FStreamingTexture& StreamingTexture = StreamingTextures[TextureIndex];
 		if (StreamingTexture.UpdateLoadOrderPriority_Async())
 		{
@@ -351,19 +357,14 @@ void FAsyncTextureStreamingTask::UpdateLoadAndCancelationRequests_Async(int64 Me
 		int32 TextureIndex = PrioritizedTextures[PriorityIndex];
 		FStreamingTexture& StreamingTexture = StreamingTextures[TextureIndex];
 
-		if (StreamingTexture.bInFlight)
+		if (StreamingTexture.bInFlight && !StreamingTexture.bCancelRequestAttempted)
 		{
-			// If there is a pending load that attempts to load unrequired data, 
+			// If there is a pending load that attempts to load unrequired data (by at least 2 mips), 
 			// or if there is a pending unload that attempts to unload required data, try to cancel it.
-			if (StreamingTexture.RequestedMips > FMath::Max<int32>(StreamingTexture.ResidentMips, StreamingTexture.WantedMips ) || 
+			if (StreamingTexture.RequestedMips > FMath::Max<int32>(StreamingTexture.ResidentMips, StreamingTexture.WantedMips + 1 ) || 
 				StreamingTexture.RequestedMips < FMath::Min<int32>(StreamingTexture.ResidentMips, StreamingTexture.WantedMips ))
 			{
 				CancelationRequests.Add(TextureIndex);
-				StreamingTexture.WantedMips = StreamingTexture.ResidentMips;
-			}
-			else // Reset wanted mips for stats.
-			{
-				StreamingTexture.WantedMips = StreamingTexture.RequestedMips;
 			}
 		}
 		else if (StreamingTexture.WantedMips < StreamingTexture.ResidentMips && TempMemoryUsed < TempMemoryBudget)
@@ -378,10 +379,6 @@ void FAsyncTextureStreamingTask::UpdateLoadAndCancelationRequests_Async(int64 Me
 				MemoryUsed -= UsedMemoryRequired;
 				TempMemoryUsed += TempMemoryRequired;
 			}
-			else // Reset wanted mips for stats.
-			{
-				StreamingTexture.WantedMips = StreamingTexture.ResidentMips;
-			}
 		}
 		else if (StreamingTexture.WantedMips > StreamingTexture.ResidentMips && MemoryUsed < MemoryBudget && TempMemoryUsed < TempMemoryBudget)
 		{
@@ -395,14 +392,6 @@ void FAsyncTextureStreamingTask::UpdateLoadAndCancelationRequests_Async(int64 Me
 				MemoryUsed += UsedMemoryRequired;
 				TempMemoryUsed += TempMemoryRequired;
 			}
-			else // Reset wanted mips for stats.
-			{
-				StreamingTexture.WantedMips = StreamingTexture.ResidentMips;
-			}
-		}
-		else // Reset wanted mips for stats.
-		{
-			StreamingTexture.WantedMips = StreamingTexture.ResidentMips;
 		}
 	}
 }
@@ -415,15 +404,16 @@ void FAsyncTextureStreamingTask::DoWork()
 	// 2 things can happen : a texture can be removed, in which case the texture will be set to null
 	// or some members can be updated following calls to UpdateDynamicData().
 	TArray<FStreamingTexture>& StreamingTextures = StreamingManager.StreamingTextures;
+	const FTextureStreamingSettings& Settings = StreamingManager.Settings;
 
 	// Update the distance and size for each bounds.
-	StreamingData.UpdateBoundSizes_Async();
+	StreamingData.UpdateBoundSizes_Async(Settings);
 	
 	for (FStreamingTexture& StreamingTexture : StreamingTextures)
 	{
 		if (IsAborted()) break;
 
-		StreamingData.UpdatePerfectWantedMips_Async(StreamingTexture, StreamingManager.MipBias);
+		StreamingData.UpdatePerfectWantedMips_Async(StreamingTexture, Settings);
 		StreamingTexture.DynamicBoostFactor = 1.f; // Reset after every computation.
 	}
 
diff --git a/Engine/Source/Runtime/Engine/Private/Streaming/AsyncTextureStreaming.h b/Engine/Source/Runtime/Engine/Private/Streaming/AsyncTextureStreaming.h
--- a/Engine/Source/Runtime/Engine/Private/Streaming/AsyncTextureStreaming.h
+++ b/Engine/Source/Runtime/Engine/Private/Streaming/AsyncTextureStreaming.h
@@ -17,9 +17,9 @@ public:
 	void Init(TArray<FStreamingViewInfo> InViewInfos, float InWorldTime, TArray<FLevelTextureManager>& LevelTextureManagers, FDynamicComponentTextureManager& DynamicComponentManager);
 
 	/** Update everything internally so to allow calls to CalcWantedMips */
-	void UpdateBoundSizes_Async();
+	void UpdateBoundSizes_Async(const FTextureStreamingSettings& Settings);
 
-	void UpdatePerfectWantedMips_Async(FStreamingTexture& StreamingTexture, float MipBias, bool bOutputToLog = false) const;
+	void UpdatePerfectWantedMips_Async(FStreamingTexture& StreamingTexture, const FTextureStreamingSettings& Settings, bool bOutputToLog = false) const;
 
 	uint32 GetAllocatedSize() const { return ViewInfos.GetAllocatedSize() + StaticInstancesViews.GetAllocatedSize(); }
 
diff --git a/Engine/Source/Runtime/Engine/Private/Streaming/StreamingManagerTexture.cpp b/Engine/Source/Runtime/Engine/Private/Streaming/StreamingManagerTexture.cpp
--- a/Engine/Source/Runtime/Engine/Private/Streaming/StreamingManagerTexture.cpp
+++ b/Engine/Source/Runtime/Engine/Private/Streaming/StreamingManagerTexture.cpp
@@ -32,7 +32,6 @@ FStreamingManagerTexture::FStreamingManagerTexture()
 :	CurrentUpdateStreamingTextureIndex(0)
 ,	bTriggerDumpTextureGroupStats( false )
 ,	bDetailedDumpTextureGroupStats( false )
-,	MipBias(0)
 ,	AsyncWork( nullptr )
 ,	ProcessingStage( 0 )
 ,	NumTextureProcessingStages(5)
@@ -47,7 +46,7 @@ FStreamingManagerTexture::FStreamingManagerTexture()
 ,	PreviousPoolSizeTimestamp(0.0)
 ,	PreviousPoolSizeSetting(-1)
 ,	bPauseTextureStreaming(false)
-,	LastUpdateTime(GIsEditor ? -FLT_MAX : 0) // In editor, visibility is not taken into consideration.
+,	LastWorldUpdateTime(GIsEditor ? -FLT_MAX : 0) // In editor, visibility is not taken into consideration.
 {
 	// Read settings from ini file.
 	int32 TempInt;
@@ -136,6 +135,10 @@ FStreamingManagerTexture::~FStreamingManagerTexture()
 {
 	AsyncWork->EnsureCompletion();
 	delete AsyncWork;
+
+	// Clear the stats
+	DisplayedStats.Reset();
+	STAT(DisplayedStats.Apply();)
 }
 
 /**
@@ -143,8 +146,6 @@ FStreamingManagerTexture::~FStreamingManagerTexture()
  */
 void FStreamingManagerTexture::CancelForcedResources()
 {
-	float CurrentTime = float(FPlatformTime::Seconds() - GStartTime);
-
 	// Update textures that are Forced on a timer.
 	for ( int32 TextureIndex=0; TextureIndex < StreamingTextures.Num(); ++TextureIndex )
 	{
@@ -154,8 +155,8 @@ void FStreamingManagerTexture::CancelForcedResources()
 		if ( StreamingTexture.Texture )
 		{
 			// Remove any prestream requests from textures
-			float TimeLeft = StreamingTexture.Texture->ForceMipLevelsToBeResidentTimestamp - CurrentTime;
-			if ( TimeLeft > 0.0f )
+			float TimeLeft = (float)(StreamingTexture.Texture->ForceMipLevelsToBeResidentTimestamp - FApp::GetCurrentTime());
+			if ( TimeLeft >= 0.0f )
 			{
 				StreamingTexture.Texture->SetForceMipLevelsToBeResident( -1.0f );
 				StreamingTexture.InstanceRemovedTimestamp = -FLT_MAX;
@@ -194,7 +195,7 @@ void FStreamingManagerTexture::SetDisregardWorldResourcesForFrames( int32 NumFra
  **/
 bool FStreamingManagerTexture::StreamOutTextureData( int64 RequiredMemorySize )
 {
-	const int64 MaxTempMemoryAllowed = CVarStreamingMaxTempMemoryAllowed.GetValueOnGameThread() * 1024 * 1024;
+	const int64 MaxTempMemoryAllowed = Settings.MaxTempMemoryAllowed * 1024 * 1024;
 	RequiredMemorySize = FMath::Max<int64>(RequiredMemorySize, MinEvictSize);
 
 	// Array of candidates for reducing mip-levels.
@@ -202,8 +203,6 @@ bool FStreamingManagerTexture::StreamOutTextureData( int64 RequiredMemorySize )
 	CandidateTextures.Reserve( StreamingTextures.Num() );
 
 	// Don't stream out character textures (to begin with)
-	float CurrentTime = float(FPlatformTime::Seconds() - GStartTime);
-
 	volatile bool bSucceeded = false;
 	
 	//resizing textures actually creates a temp copy so we can only resize so many at a time without running out of memory during the eject itself.
@@ -236,9 +235,7 @@ bool FStreamingManagerTexture::StreamOutTextureData( int64 RequiredMemorySize )
 				int32 CurrentBaseMip = NumMips - Texture->ResidentMips;
 				if ( MipTailBaseIndex < 0 || CurrentBaseMip < MipTailBaseIndex )
 				{
-					// Figure out whether texture should be forced resident based on bools and forced resident time.
-					bool bForceMipLevelsToBeResident = (Texture->ShouldMipLevelsBeForcedResident() || Texture->ForceMipLevelsToBeResidentTimestamp >= CurrentTime);
-					if ( bForceMipLevelsToBeResident == false && Texture->Resource )
+					if ( !Texture->ShouldMipLevelsBeForcedResident() && Texture->Resource )
 					{
 						// Don't try to stream out if the texture is currently being busy being streamed in/out.
 						bool bSafeToStream = (Texture->UpdateStreamingStatus() == false);
@@ -340,9 +337,6 @@ void FStreamingManagerTexture::ProcessRemovedTextures()
 
 void FStreamingManagerTexture::ProcessAddedTextures()
 {
-	const bool bOnlyStreamIn = CVarOnlyStreamInTextures.GetValueOnAnyThread() != 0;
-	const int32 HLODStategy = CVarStreamingHLODStrategy.GetValueOnAnyThread();
-
 	// Add new textures.
 	StreamingTextures.Reserve(StreamingTextures.Num() + PendingStreamingTextures.Num());
 	for (int32 TextureIndex=0; TextureIndex < PendingStreamingTextures.Num(); ++TextureIndex)
@@ -351,7 +345,7 @@ void FStreamingManagerTexture::ProcessAddedTextures()
 		if (!Texture) continue; // Could be null if it was removed after being added.
 
 		Texture->StreamingIndex = StreamingTextures.Num();
-		new (StreamingTextures) FStreamingTexture(Texture, MipBias, NumStreamedMips, bOnlyStreamIn, HLODStategy);
+		new (StreamingTextures) FStreamingTexture(Texture, NumStreamedMips, Settings);
 	}
 	PendingStreamingTextures.Empty();
 }
@@ -360,22 +354,20 @@ void FStreamingManagerTexture::ConditionalUpdateStaticData()
 {
 	static float PreviousLightmapStreamingFactor = GLightmapStreamingFactor;
 	static float PreviousShadowmapStreamingFactor = GShadowmapStreamingFactor;
-	static int32 PreviousUseNewMetrics = CVarStreamingUseNewMetrics.GetValueOnAnyThread();
-	static float PreviousMipBias = MipBias;
+	static FTextureStreamingSettings PreviousSettings = Settings;
 
 	if (PreviousLightmapStreamingFactor != GLightmapStreamingFactor || 
 		PreviousShadowmapStreamingFactor != GShadowmapStreamingFactor || 
-		PreviousUseNewMetrics != CVarStreamingUseNewMetrics.GetValueOnAnyThread() ||
-		PreviousMipBias != MipBias)
+		PreviousSettings != Settings)
 	{
 		// Update each texture static data.
 		for (FStreamingTexture& StreamingTexture : StreamingTextures)
 		{
-			StreamingTexture.UpdateStaticData(MipBias);
+			StreamingTexture.UpdateStaticData(Settings);
 		}
 
-#if !(UE_BUILD_SHIPPING || UE_BUILD_TEST)
-		// Reinsert all levels to update static instance data.
+#if 0
+		// Reinsert all levels to update static instance data. Only for testing perfs.
 		TArray<ULevel*> Levels;
 		for (FLevelTextureManager& LevelManager : LevelTextureManagers)
 		{
@@ -390,8 +382,7 @@ void FStreamingManagerTexture::ConditionalUpdateStaticData()
 		// Update the cache variables.
 		PreviousLightmapStreamingFactor = GLightmapStreamingFactor;
 		PreviousShadowmapStreamingFactor = GShadowmapStreamingFactor;
-		PreviousUseNewMetrics = CVarStreamingUseNewMetrics.GetValueOnAnyThread();
-		PreviousMipBias = MipBias;
+		PreviousSettings = Settings;
 	}
 }
 
@@ -404,21 +395,23 @@ void FStreamingManagerTexture::UpdateThreadData(bool bProcessEverything)
 	ProcessRemovedTextures();
 	ProcessAddedTextures();
 
+	Settings.Update();
 	ConditionalUpdateStaticData();
 
 	// Fully complete all pending update.
 	IncrementalUpdate(1.f);
 
-	MipBias = FMath::Max(CVarStreamingMipBias.GetValueOnGameThread(), 0.0f);
 	// Update the thread data.
 
 	FAsyncTextureStreamingTask& AsyncTask = AsyncWork->GetTask();
 	FTextureMemoryStats Stats;
 	RHIGetTextureMemoryStats(Stats);
 
-	if (Stats.IsUsingLimitedPoolSize() && !bProcessEverything)
+	// When processing all textures, we need unlimited budget so that textures get all at their required states.
+	// Same when forcing stream-in, for which we want all used textures to be fully loaded 
+	if (Stats.IsUsingLimitedPoolSize() && !bProcessEverything && !Settings.bFullyLoadUsedTextures)
 	{
-		const int64 TempMemoryBudget = CVarStreamingMaxTempMemoryAllowed.GetValueOnAnyThread() * 1024 * 1024;
+		const int64 TempMemoryBudget = Settings.MaxTempMemoryAllowed * 1024 * 1024;
 		AsyncTask.Reset(Stats.AllocatedMemorySize, Stats.TexturePoolSize, TempMemoryBudget, MemoryMargin);
 	}
 	else
@@ -426,7 +419,7 @@ void FStreamingManagerTexture::UpdateThreadData(bool bProcessEverything)
 		// Temp must be smaller since membudget only updates if it has a least temp memory available.
 		AsyncTask.Reset(Stats.AllocatedMemorySize, MAX_int64, MAX_int64 / 2, 0);
 	}
-	AsyncTask.StreamingData.Init(CurrentViewInfos, LastUpdateTime, LevelTextureManagers, DynamicComponentManager);
+	AsyncTask.StreamingData.Init(CurrentViewInfos, LastWorldUpdateTime, LevelTextureManagers, DynamicComponentManager);
 }
 
 /**
@@ -716,9 +709,7 @@ void FStreamingManagerTexture::UpdateIndividualTexture( UTexture2D* Texture )
 	FStreamingTexture* StreamingTexture = GetStreamingTexture(Texture);
 	if (!StreamingTexture) return;
 
-	const bool bOnlyStreamIn = CVarOnlyStreamInTextures.GetValueOnAnyThread() != 0;
-	const int32 HLODStategy = CVarStreamingHLODStrategy.GetValueOnAnyThread();
-	StreamingTexture->UpdateDynamicData(NumStreamedMips, bOnlyStreamIn, FMath::FloorToInt(MipBias), HLODStategy);
+	StreamingTexture->UpdateDynamicData(NumStreamedMips, Settings);
 
 	if (StreamingTexture->bForceFullyLoad) // Somewhat expected at this point.
 	{
@@ -744,10 +735,6 @@ void FStreamingManagerTexture::UpdateStreamingTextures( int32 StageIndex, int32
 		InflightTextures.Reset();
 	}
 
-	const bool bOnlyStreamIn = CVarOnlyStreamInTextures.GetValueOnAnyThread() != 0;
-	const int32 GlobalMipBias = FMath::FloorToInt(MipBias);
-	const int32 HLODStategy = CVarStreamingHLODStrategy.GetValueOnAnyThread();
-
 	int32 StartIndex = CurrentUpdateStreamingTextureIndex;
 	int32 EndIndex = StreamingTextures.Num() * (StageIndex + 1) / NumUpdateStages;
 	for ( int32 Index=StartIndex; Index < EndIndex; ++Index )
@@ -760,7 +747,7 @@ void FStreamingManagerTexture::UpdateStreamingTextures( int32 StageIndex, int32
 
 		STAT(int32 PreviousResidentMips = StreamingTexture.ResidentMips;)
 
-		StreamingTexture.UpdateDynamicData(NumStreamedMips, bOnlyStreamIn, GlobalMipBias, HLODStategy);
+		StreamingTexture.UpdateDynamicData(NumStreamedMips, Settings);
 
 		// Make a list of each texture that can potentially require additional UpdateStreamingStatus
 		if (StreamingTexture.bInFlight)
@@ -867,7 +854,7 @@ void FStreamingManagerTexture::SetLastUpdateTime()
 			float WorldTime = LevelTextureManagers[LevelIndex].GetWorldTime();
 			if (WorldTime > 0)
 			{
-				LastUpdateTime = WorldTime - .5f;
+				LastWorldUpdateTime = WorldTime - .5f;
 				break;
 			}
 		}
@@ -1111,19 +1098,21 @@ bool FStreamingManagerTexture::HandleListStreamingTexturesCommand( const TCHAR*
 
 		if (StreamingTexture.LastRenderTime != MAX_FLT)
 		{
-			UE_LOG(LogContentStreaming, Log,  TEXT("    Current=%dx%d Wanted=%dx%d Max=%dx%d LastRenderTime=%.3f Group=%s"), 
+			UE_LOG(LogContentStreaming, Log,  TEXT("    Current=%dx%d Wanted=%dx%d Max=%dx%d LastRenderTime=%.3f BudgetBias=%d Group=%s"), 
 				Mips[CurrentMipIndex].SizeX, Mips[CurrentMipIndex].SizeY, 
 				Mips[WantedMipIndex].SizeX, Mips[WantedMipIndex].SizeY, 
 				Mips[MaxMipIndex].SizeX, Mips[MaxMipIndex].SizeY, 
 				StreamingTexture.LastRenderTime,
+				StreamingTexture.BudgetMipBias,
 				UTexture::GetTextureGroupString(StreamingTexture.LODGroup));
 		}
 		else
 		{
-			UE_LOG(LogContentStreaming, Log,  TEXT("    Current=%dx%d Wanted=%dx%d Max=%dx%d Group=%s"), 
+			UE_LOG(LogContentStreaming, Log,  TEXT("    Current=%dx%d Wanted=%dx%d Max=%dx%d BudgetBias=%d Group=%s"), 
 				Mips[CurrentMipIndex].SizeX, Mips[CurrentMipIndex].SizeY, 
 				Mips[WantedMipIndex].SizeX, Mips[WantedMipIndex].SizeY, 
 				Mips[MaxMipIndex].SizeX, Mips[MaxMipIndex].SizeY, 
+				StreamingTexture.BudgetMipBias,
 				UTexture::GetTextureGroupString(StreamingTexture.LODGroup));
 		}
 
@@ -1385,7 +1374,7 @@ bool FStreamingManagerTexture::HandleInvestigateTextureCommand( const TCHAR* Cmd
 		// Make sure the async task is idle (also implies Update_Async is finished and that the distances are valid).
 		AsyncWork->EnsureCompletion();
 		FAsyncTextureStreamingData& StreamingData = AsyncWork->GetTask().StreamingData;
-		StreamingData.UpdateBoundSizes_Async();
+		StreamingData.UpdateBoundSizes_Async(Settings);
 
 		for ( int32 TextureIndex=0; TextureIndex < StreamingTextures.Num(); ++TextureIndex )
 		{
@@ -1415,8 +1404,7 @@ bool FStreamingManagerTexture::HandleInvestigateTextureCommand( const TCHAR* Cmd
 				}
 				else if ( Texture2D->ShouldMipLevelsBeForcedResident() )
 				{
-					float CurrentTime = float(FPlatformTime::Seconds() - GStartTime);
-					float TimeLeft = CurrentTime - Texture2D->ForceMipLevelsToBeResidentTimestamp;
+					float TimeLeft = (float)(Texture2D->ForceMipLevelsToBeResidentTimestamp - FApp::GetCurrentTime());
 					UE_LOG(LogContentStreaming, Log,  TEXT("  Force all mips:  %.1f seconds left"), FMath::Max(TimeLeft,0.0f) );
 				}
 				else if ( StreamingTexture.MipCount == 1 )
@@ -1430,11 +1418,11 @@ bool FStreamingManagerTexture::HandleInvestigateTextureCommand( const TCHAR* Cmd
 				UE_LOG(LogContentStreaming, Log,  TEXT("  Retention Priority: %d"), StreamingTexture.RetentionPriority );
 				UE_LOG(LogContentStreaming, Log,  TEXT("  Boost factor:    %.1f"), StreamingTexture.BoostFactor );
 				UE_LOG(LogContentStreaming, Log,  TEXT("  Allowed mips:    %d-%d"), StreamingTexture.MinAllowedMips, StreamingTexture.MaxAllowedMips );
-				UE_LOG(LogContentStreaming, Log,  TEXT("  Global mip bias: %.1f"), MipBias );
+				UE_LOG(LogContentStreaming, Log,  TEXT("  Mip bias : Texture=%d Global=%.1f"), StreamingTexture.BudgetMipBias, Settings.GlobalMipBias() );
 
 				if (InWorld && !GIsEditor)
 				{
-					UE_LOG(LogContentStreaming, Log,  TEXT("  Time: World=%.3f LastUpdate=%.3f "), InWorld->GetTimeSeconds(), LastUpdateTime);
+					UE_LOG(LogContentStreaming, Log,  TEXT("  Time: World=%.3f LastUpdate=%.3f "), InWorld->GetTimeSeconds(), LastWorldUpdateTime);
 				}
 
 				for( int32 ViewIndex=0; ViewIndex < StreamingData.GetViewInfos().Num(); ViewIndex++ )
@@ -1444,7 +1432,7 @@ bool FStreamingManagerTexture::HandleInvestigateTextureCommand( const TCHAR* Cmd
 					UE_LOG(LogContentStreaming, Log,  TEXT("  View%d: Position=(%s) ScreenSize=%f Boost=%f"), ViewIndex, *ViewInfo.ViewOrigin.ToString(), ViewInfo.ScreenSize, ViewInfo.BoostFactor);
 				}
 
-				StreamingData.UpdatePerfectWantedMips_Async(StreamingTexture, MipBias, true);
+				StreamingData.UpdatePerfectWantedMips_Async(StreamingTexture, Settings, true);
 			}
 		}
 	}
diff --git a/Engine/Source/Runtime/Engine/Private/Streaming/StreamingManagerTexture.h b/Engine/Source/Runtime/Engine/Private/Streaming/StreamingManagerTexture.h
--- a/Engine/Source/Runtime/Engine/Private/Streaming/StreamingManagerTexture.h
+++ b/Engine/Source/Runtime/Engine/Private/Streaming/StreamingManagerTexture.h
@@ -247,7 +247,7 @@ protected:
 	/** Cached from the system settings. */
 	int32 NumStreamedMips[TEXTUREGROUP_MAX];
 
-	float MipBias;
+	FTextureStreamingSettings Settings;
 
 	/** Async work for calculating priorities for all textures. */
 	FAsyncTask<FAsyncTextureStreamingTask>*	AsyncWork;
@@ -300,7 +300,7 @@ protected:
 	bool bPauseTextureStreaming;
 
 	/** Last time all data were fully updated. Instances are considered visible if they were rendered between that last time and the current time. */
-	float LastUpdateTime;
+	float LastWorldUpdateTime;
 
 	FTextureStreamingStats DisplayedStats;
 	FTextureStreamingStats GatheredStats;
diff --git a/Engine/Source/Runtime/Engine/Private/Streaming/StreamingTexture.cpp b/Engine/Source/Runtime/Engine/Private/Streaming/StreamingTexture.cpp
--- a/Engine/Source/Runtime/Engine/Private/Streaming/StreamingTexture.cpp
+++ b/Engine/Source/Runtime/Engine/Private/Streaming/StreamingTexture.cpp
@@ -10,11 +10,11 @@ StreamingTexture.cpp: Definitions of classes used for texture.
 #include "Engine/LightMapTexture2D.h"
 #include "Engine/ShadowMapTexture2D.h"
 
-FStreamingTexture::FStreamingTexture(UTexture2D* InTexture, float GlobalMipBias, const int32 NumStreamedMips[TEXTUREGROUP_MAX], bool bOnlyStreamIn, int32 HLODStrategy)
+FStreamingTexture::FStreamingTexture(UTexture2D* InTexture, const int32 NumStreamedMips[TEXTUREGROUP_MAX], const FTextureStreamingSettings& Settings)
 : Texture(InTexture)
 {
-	UpdateStaticData(GlobalMipBias);
-	UpdateDynamicData(NumStreamedMips, bOnlyStreamIn, FMath::FloorToInt(GlobalMipBias), HLODStrategy);
+	UpdateStaticData(Settings);
+	UpdateDynamicData(NumStreamedMips, Settings);
 
 	InstanceRemovedTimestamp = -FLT_MAX;
 	LastRenderTimeRefCountTimestamp = -FLT_MAX;
@@ -32,14 +32,15 @@ FStreamingTexture::FStreamingTexture(UTexture2D* InTexture, float GlobalMipBias,
 	WantedMips = MinAllowedMips;
 }
 
-void FStreamingTexture::UpdateStaticData(float GlobalMipBias)
+void FStreamingTexture::UpdateStaticData(const FTextureStreamingSettings& Settings)
 {
 	if (Texture)
 	{
 		LODGroup = (TextureGroup)Texture->LODGroup;
 		NumNonStreamingMips = Texture->GetNumNonStreamingMips();
 		MipCount = Texture->GetNumMips();
-		BoostFactor = GetExtraBoost(LODGroup, GlobalMipBias);
+		BudgetMipBias = 0;
+		BoostFactor = GetExtraBoost(LODGroup, Settings);
 
 		bIsCharacterTexture = (LODGroup == TEXTUREGROUP_Character || LODGroup == TEXTUREGROUP_CharacterSpecular || LODGroup == TEXTUREGROUP_CharacterNormalMap);
 		bIsTerrainTexture = (LODGroup == TEXTUREGROUP_Terrain_Heightmap || LODGroup == TEXTUREGROUP_Terrain_Weightmap);
@@ -54,6 +55,7 @@ void FStreamingTexture::UpdateStaticData(float GlobalMipBias)
 		LODGroup = TEXTUREGROUP_World;
 		NumNonStreamingMips = 0;
 		MipCount = 0;
+		BudgetMipBias = 0;
 		BoostFactor = 1.f;
 
 		bIsCharacterTexture = false;
@@ -66,7 +68,7 @@ void FStreamingTexture::UpdateStaticData(float GlobalMipBias)
 	}
 }
 
-void FStreamingTexture::UpdateDynamicData(const int32 NumStreamedMips[TEXTUREGROUP_MAX], bool bOnlyStreamIn, int32 GlobalMipBias, int32 HLODStrategy)
+void FStreamingTexture::UpdateDynamicData(const int32 NumStreamedMips[TEXTUREGROUP_MAX], const FTextureStreamingSettings& Settings)
 {
 	// Note that those values are read from the async task and must not be assigned temporary values!!
 	if (Texture)
@@ -77,23 +79,28 @@ void FStreamingTexture::UpdateDynamicData(const int32 NumStreamedMips[TEXTUREGRO
 		const float LastRenderTimeForTexture = Texture->GetLastRenderTimeForStreaming();
 		LastRenderTime = (FApp::GetCurrentTime() > LastRenderTimeForTexture) ? float( FApp::GetCurrentTime() - LastRenderTimeForTexture ) : 0.0f;
 
-		bForceFullyLoad = Texture->ShouldMipLevelsBeForcedResident() || (bOnlyStreamIn && LastRenderTime < 300)|| LODGroup == TEXTUREGROUP_Skybox || (LODGroup == TEXTUREGROUP_HierarchicalLOD && HLODStrategy == 2);
+		bForceFullyLoad = Texture->ShouldMipLevelsBeForcedResident() || LODGroup == TEXTUREGROUP_Skybox || (LODGroup == TEXTUREGROUP_HierarchicalLOD && Settings.HLODStrategy == 2);
 
 		const int32 NumCinematicMipLevels = (bForceFullyLoad && Texture->bUseCinematicMipLevels) ? Texture->NumCinematicMipLevels : 0;
 
+		int32 LODBias = 0;
+		if (!Settings.bUseAllMips)
+		{
+			LODBias = FMath::Max<int32>(Texture->GetCachedLODBias() - NumCinematicMipLevels, 0);
 
-		int32 LODBias = FMath::Max<int32>(Texture->GetCachedLODBias() - NumCinematicMipLevels, 0);
+			// Reduce the max allowed resolution according to LODBias if the texture group allows it.
+			if (LODGroup != TEXTUREGROUP_HierarchicalLOD && !bIsTerrainTexture)
+			{
+				LODBias += Settings.GlobalMipBiasAsInt();
+			}
 
-		// Reduce the max allowed resolution according to LODBias if the texture group allows it.
-		if (LODGroup != TEXTUREGROUP_HierarchicalLOD && !bIsTerrainTexture)
-		{
-			LODBias += GlobalMipBias;
+			LODBias += BudgetMipBias;
 		}
 
 		// The max mip count is affected by the texture bias and cinematic bias settings.
 		MaxAllowedMips = FMath::Clamp<int32>(FMath::Min<int32>(MipCount - LODBias, GMaxTextureMipCount), NumNonStreamingMips, MipCount);
 
-		if (LODGroup == TEXTUREGROUP_HierarchicalLOD && HLODStrategy == 1)
+		if (LODGroup == TEXTUREGROUP_HierarchicalLOD && Settings.HLODStrategy == 1)
 		{
 			MinAllowedMips = FMath::Clamp<int32>(MaxAllowedMips - 1, NumNonStreamingMips, MaxAllowedMips);
 		}
@@ -110,6 +117,7 @@ void FStreamingTexture::UpdateDynamicData(const int32 NumStreamedMips[TEXTUREGRO
 	{
 		bReadyForStreaming = false;
 		bInFlight = false;
+		bCancelRequestAttempted = false;
 		bForceFullyLoad = false;
 		ResidentMips = 0;
 		RequestedMips = 0;
@@ -134,6 +142,13 @@ void FStreamingTexture::UpdateStreamingStatus()
 			bInFlight = false;
 		}
 
+		 // Allow a future cancel request only once the texture has returned to a fully ready state.
+		if (bReadyForStreaming && !bInFlight)
+		{
+			bCancelRequestAttempted = false;
+		}
+
+
 		// This must be updated after UpdateStreamingStatus
 		ResidentMips = Texture->ResidentMips;
 		RequestedMips = Texture->RequestedMips;
@@ -144,10 +159,11 @@ void FStreamingTexture::UpdateStreamingStatus()
 	{
 		bReadyForStreaming = false;
 		bInFlight = false;
+		bCancelRequestAttempted = false;
 	}
 }
 
-float FStreamingTexture::GetExtraBoost(TextureGroup	LODGroup, float GlobalBias)
+float FStreamingTexture::GetExtraBoost(TextureGroup	LODGroup, const FTextureStreamingSettings& Settings)
 {
 	// When accurate distance computation, we need to relax the distance otherwhise it get's too conservative. (ex 513 goes to 1024)
 	const bool bUseApproxDistance = CVarStreamingUseNewMetrics.GetValueOnAnyThread() == 0;
@@ -169,45 +185,36 @@ float FStreamingTexture::GetExtraBoost(TextureGroup	LODGroup, float GlobalBias)
 	}
 	else
 	{
-		return DistanceScale * FMath::Exp2(-GlobalBias);
+		return DistanceScale * Settings.GlobalMipBiasAsScale();
 	}
 }
 
-int32 FStreamingTexture::GetWantedMipsFromSize(float Size, int32 MipBias) const
+int32 FStreamingTexture::GetWantedMipsFromSize(float Size) const
 {
 	float WantedMipsFloat = 1.0f + FMath::Log2(FMath::Max(1.f, Size));
-	// Using the new metrics has very conservative distance computation. Rounding to mip here compensates.
 	int32 WantedMipsInt = FMath::CeilToInt(WantedMipsFloat);
-
-	if (LODGroup != TEXTUREGROUP_HierarchicalLOD  && LODGroup != TEXTUREGROUP_Terrain_Heightmap)
-	{
-		return FMath::Max<int32>(FMath::Min<int32>(WantedMipsInt, MaxAllowedMips - MipBias), MinAllowedMips);
-	}
-	else
-	{
-		return FMath::Clamp<int32>(WantedMipsInt, MinAllowedMips, MaxAllowedMips);
-	}
+	return FMath::Clamp<int32>(WantedMipsInt, MinAllowedMips, MaxAllowedMips);
 }
 
 /** Set the wanted mips from the async task data */
-void FStreamingTexture::SetPerfectWantedMips_Async(float MaxSize, float MaxSize_VisibleOnly, float MipBias, bool InLooksLowRes)
+void FStreamingTexture::SetPerfectWantedMips_Async(float MaxSize, float MaxSize_VisibleOnly, bool InLooksLowRes, const FTextureStreamingSettings& Settings)
 {
-	const float HiddenScale = CVarStreamingHiddenPrimitiveScale.GetValueOnAnyThread();
-
 	bForceFullyLoadHeuristic = (MaxSize == FLT_MAX || MaxSize_VisibleOnly == FLT_MAX);
-	VisibleWantedMips = GetWantedMipsFromSize(MaxSize_VisibleOnly, MipBias);
+	VisibleWantedMips = GetWantedMipsFromSize(MaxSize_VisibleOnly);
 	bLooksLowRes = InLooksLowRes; // Things like lightmaps, HLOD and close instances.
 
+	// Here BudgetMipBias is added to NumMissingMips to prevent the same texture from dropping resolution too early (as it is already included in MaxMip)
+
 	// Terrain, Forced Fully Load and Things that already look bad are not affected by hidden scale.
 	if (bIsTerrainTexture || bForceFullyLoadHeuristic || bLooksLowRes)
 	{
-		HiddenWantedMips = GetWantedMipsFromSize(MaxSize, MipBias);
-		NumMissingMips = 0; // No impact for terrains as there are not allowed to drop mips.
+		HiddenWantedMips = GetWantedMipsFromSize(MaxSize);
+		NumMissingMips = BudgetMipBias; // No impact for terrains as there are not allowed to drop mips.
 	}
 	else
 	{
-		HiddenWantedMips = GetWantedMipsFromSize(MaxSize * HiddenScale, MipBias);
-		NumMissingMips = FMath::Max<int32>(GetWantedMipsFromSize(MaxSize, MipBias) - FMath::Max<int32>(VisibleWantedMips, HiddenWantedMips), 0);
+		HiddenWantedMips = GetWantedMipsFromSize(MaxSize * Settings.HiddenPrimitiveScale);
+		NumMissingMips = BudgetMipBias + FMath::Max<int32>(GetWantedMipsFromSize(MaxSize) - FMath::Max<int32>(VisibleWantedMips, HiddenWantedMips), 0);
 	}
 }
 
@@ -243,10 +250,22 @@ int64 FStreamingTexture::UpdateRetentionPriority_Async()
 	}
 }
 
-int64 FStreamingTexture::DropOneMip_Async()
+int64 FStreamingTexture::DropOneMip_Async(int32 MaxPerTextureMipBias)
 {
 	if (Texture && BudgetedMips > MinAllowedMips)
 	{
+		// Every time we are required to DropOneMip from max resolution, this increase the BudgetMipBias.
+		if (BudgetedMips == MaxAllowedMips && VisibleWantedMips == MaxAllowedMips && BudgetMipBias < MaxPerTextureMipBias)
+		{
+			// Decrease MaxAllowedMips and increase BudgetMipBias (as it should include it)
+			--MaxAllowedMips;
+			++BudgetMipBias;
+
+			// Update the wanted mips as if they were computed with the new MaxAllowedMips. (for stats)
+			VisibleWantedMips = FMath::Min<int32>(VisibleWantedMips, MaxAllowedMips);
+			HiddenWantedMips = FMath::Min<int32>(HiddenWantedMips, MaxAllowedMips);
+		}
+
 		--BudgetedMips;
 		return GetSize(BudgetedMips + 1) - GetSize(BudgetedMips);
 	}
@@ -275,7 +294,7 @@ bool FStreamingTexture::UpdateLoadOrderPriority_Async()
 	WantedMips = BudgetedMips;
 
 	// If the entry is valid and we need to send a new request to load/drop the right mip.
-	if (Texture && WantedMips != RequestedMips)
+	if (bReadyForStreaming && Texture && WantedMips != RequestedMips)
 	{
 		const bool bIsVisible			= WantedMips <= VisibleWantedMips; // Otherwise it means we are loading mips that are only usefull for non visible primitives.
 		const bool bMustLoadFirst		= bForceFullyLoadHeuristic || bIsTerrainTexture || bIsCharacterTexture;
@@ -296,16 +315,11 @@ bool FStreamingTexture::UpdateLoadOrderPriority_Async()
 
 void FStreamingTexture::CancelPendingMipChangeRequest()
 {
-	if (Texture && bReadyForStreaming && bInFlight)
+	if (Texture && bReadyForStreaming && bInFlight && !bCancelRequestAttempted)
 	{
-		// Try to cancel, but don't try anything else before next update.
-		bReadyForStreaming = false;
-
-		if (Texture->CancelPendingMipChangeRequest() )
-		{
-			bInFlight = false;
-			RequestedMips = ResidentMips;
-		}
+		// Try to cancel at most once per request.
+		bCancelRequestAttempted = true;
+		Texture->CancelPendingMipChangeRequest();
 	}
 }
 
@@ -315,17 +329,15 @@ void FStreamingTexture::StreamWantedMips(FStreamingManagerTexture& Manager)
 	{
 		ensure(ResidentMips == RequestedMips);
 
-		// Try to load/unload, but don't try anything else before next update.
-		bReadyForStreaming = false;
-
 		if (Texture->PendingMipChangeRequestStatus.GetValue() == TexState_ReadyFor_Requests)
 		{
-			bInFlight = true;
-			RequestedMips = WantedMips;
 			Texture->RequestedMips = WantedMips;
 
 			FTexture2DResource* Texture2DResource = (FTexture2DResource*)Texture->Resource;
 			Texture2DResource->BeginUpdateMipCount( (bForceFullyLoadHeuristic || bIsTerrainTexture || bIsCharacterTexture) && WantedMips <= VisibleWantedMips );
+
+			UpdateStreamingStatus();
+
 			TrackTextureEvent( this, Texture, bForceFullyLoadHeuristic != 0, &Manager );
 		}
 		else
diff --git a/Engine/Source/Runtime/Engine/Private/Streaming/StreamingTexture.h b/Engine/Source/Runtime/Engine/Private/Streaming/StreamingTexture.h
--- a/Engine/Source/Runtime/Engine/Private/Streaming/StreamingTexture.h
+++ b/Engine/Source/Runtime/Engine/Private/Streaming/StreamingTexture.h
@@ -15,13 +15,13 @@ StreamingTexture.h: Definitions of classes used for texture streaming.
 /** Self-contained structure to manage a streaming texture, possibly on a separate thread. */
 struct FStreamingTexture
 {
-	FStreamingTexture(UTexture2D* InTexture, float GlobalMipBias, const int32 NumStreamedMips[TEXTUREGROUP_MAX], bool bOnlyStreamIn, int32 HLODStrategy);
+	FStreamingTexture(UTexture2D* InTexture, const int32 NumStreamedMips[TEXTUREGROUP_MAX], const FTextureStreamingSettings& Settings);
 
 	/** Update data that should not change unless changing settings. */
-	void UpdateStaticData(float GlobalMipBias);
+	void UpdateStaticData(const FTextureStreamingSettings& Settings);
 
 	/** Update data that the engine could change through gameplay. */
-	void UpdateDynamicData(const int32 NumStreamedMips[TEXTUREGROUP_MAX], bool bOnlyStreamIn, int32 GlobalMipBias, int32 HLODStrategy);
+	void UpdateDynamicData(const int32 NumStreamedMips[TEXTUREGROUP_MAX], const FTextureStreamingSettings& Settings);
 
 	/** Lightweight version of UpdateDynamicData. */
 	void UpdateStreamingStatus();
@@ -38,18 +38,18 @@ struct FStreamingTexture
 		return TextureSizes[ InMipCount - 1 ];
 	}
 
-	static float GetExtraBoost(TextureGroup	LODGroup, float GlobalBias);
+	static float GetExtraBoost(TextureGroup	LODGroup, const FTextureStreamingSettings& Settings);
 
-	FORCEINLINE int32 GetWantedMipsFromSize(float Size, int32 MipBias) const;
+	FORCEINLINE int32 GetWantedMipsFromSize(float Size) const;
 
 	/** Set the wanted mips from the async task data */
-	void SetPerfectWantedMips_Async(float MaxSize, float MaxSize_VisibleOnly, float MipBias, bool InLooksLowRes);
+	void SetPerfectWantedMips_Async(float MaxSize, float MaxSize_VisibleOnly, bool InLooksLowRes, const FTextureStreamingSettings& Settings);
 
 	/** Init BudgetedMip and update RetentionPriority. Returns the size that would be taken if all budgeted mips where loaded. */
 	int64 UpdateRetentionPriority_Async();
 
 	/** Reduce BudgetedMip by 1 and return the size freed by doing so. */
-	int64 DropOneMip_Async();
+	int64 DropOneMip_Async(int32 MaxPerTextureMipBias = 0);
 
 	/** Increase BudgetedMip by 1, up to resident mips, and return the size taken. */
 	int64 KeepOneMip_Async();
@@ -113,7 +113,7 @@ struct FStreamingTexture
 	/** Max mip to be requested by the streaming  */
 	int32			MaxAllowedMips;
 
-	/** Last time this texture was rendered, 0-based from GStartTime. */
+	/** How much game time has elapsed since the texture was bound for rendering. Based on FApp::GetCurrentTime(). */
 	float			LastRenderTime;
 
 	/*****************************************************************************************
@@ -121,7 +121,10 @@ struct FStreamingTexture
 	 *****************************************************************************************/
 
 	/** Whether the texture is currently being streamed in/out. */
-	bool			bInFlight;
+	uint32			bInFlight : 1;
+
+	/** Whether an attemp to cancel the inflight request has been attempted. */
+	uint32			bCancelRequestAttempted : 1;
 
 	/** If non-zero, the most recent time an instance location was removed for this texture. */
 	double			InstanceRemovedTimestamp;
@@ -163,4 +166,7 @@ struct FStreamingTexture
 	int32			LoadOrderPriority;	
 	/** The mip that will be requested. Note, that some texture are loaded using split requests, so that the first request can be smaller than budgeted mip. */
 	int32			WantedMips;
+
+	/** A persistent bias applied to this texture. Increase whenever the streamer needs to make sacrifices to fit in budget */
+	int32			BudgetMipBias;
 };
