Manually merged cl#3045159 from Dev-Blueprints

diff --git a/Engine/Source/Runtime/Engine/Public/Net/UnrealNetwork.h b/Engine/Source/Runtime/Engine/Public/Net/UnrealNetwork.h
--- a/Engine/Source/Runtime/Engine/Public/Net/UnrealNetwork.h
+++ b/Engine/Source/Runtime/Engine/Public/Net/UnrealNetwork.h
@@ -71,7 +71,7 @@ static UProperty* GetReplicatedProperty(UClass* CallingClass, UClass* PropClass,
 #endif
 	return TheProperty;
 }
-	
+
 #define DOREPLIFETIME(c,v) \
 { \
 	static UProperty* sp##v = GetReplicatedProperty(StaticClass(), c::StaticClass(),GET_MEMBER_NAME_CHECKED(c,v)); \
@@ -81,9 +81,16 @@ static UProperty* GetReplicatedProperty(UClass* CallingClass, UClass* PropClass,
 	}																		\
 }
 
+/** This macro is used by nativized code (DynamicClasses), so the Property may be recreated. */
 #define DOREPLIFETIME_DIFFNAMES(c,v, n) \
 { \
-	static UProperty* sp##v = GetReplicatedProperty(StaticClass(), c::StaticClass(), n); \
+	static TWeakObjectPtr<UProperty> __swp##v{};							\
+	const UProperty* sp##v = __swp##v.Get();								\
+	if (nullptr == sp##v)													\
+	{																		\
+		sp##v = GetReplicatedProperty(StaticClass(), c::StaticClass(), n);	\
+		__swp##v = sp##v;													\
+	}																		\
 	for ( int32 i = 0; i < sp##v->ArrayDim; i++ )							\
 	{																		\
 		OutLifetimeProps.AddUnique( FLifetimeProperty( sp##v->RepIndex + i ) );	\
