Copying //UE4/Dev-Networking to //UE4/Main (Source: //UE4/Dev-Networking [REDACTED_HANDLE])

==========================
MAJOR FEATURES + CHANGES
==========================

Change 3004652 on 2016/06/07 by [REDACTED_PERSON]

	Improve comments

Change 3014311 on 2016/06/15 by [REDACTED_PERSON]

	Paragon specific Oodle dictionary generation

Change 3017091 on 2016/06/16 by [REDACTED_PERSON]

	Updated Oodle scripts:
	- engine scripts used from Orion script now take command line parameters
	- can now submit p4 changelist optionally at end of Orion script

Change 3019461 on 2016/06/20 by [REDACTED_PERSON]

	Missed changes to ExtractAllArchives.bat in my last scripts checkin:
	- Only extract gzipped files if needed
	- Only download captures from a specific changelist

Change 3031917 on 2016/06/28 by [REDACTED_PERSON]

	FIX OR-24516: Hero and AI disappear in Replay after clicking on Death in Timeline

Change 3034309 on 2016/06/30 by [REDACTED_PERSON]

	Added "Pause replication" feature, which can reduce bandwidth and reduce cheating. When replicating an actor, the server will call
AActor::IsReplicationPausedForConnection. If this function returns true, the server will send a reliable bunch with the "bIsReplicationPaused" flag checked. The server will
then stop replicating on that actor's channel until AActor::IsReplicationPausedForConnection returns false. On the client, when receiving a bunch,
AActor::OnReplicationPausedChanged will be called, so the client can hide the actor's mesh or take custom action that's needed.

	#tests implemented in shooter game and tested with multiple players

Change 3039989 on 2016/07/06 by [REDACTED_PERSON]

	PR #2566: Correct filename case in #include statement (Contributed by [REDACTED_PERSON])

	#tests compiled

Change 3041369 on 2016/07/07 by [REDACTED_PERSON]

	Add some logging to print out whether a server is actually capturing with oodle or not based on the capture percentage chance

	#tests ran with oodle and checked log message

Change 3041474 on 2016/07/07 by [REDACTED_PERSON]

	PR #2268: Fixed spelling (Contributed by [REDACTED_PERSON])

Change 3043191 on 2016/07/08 by [REDACTED_PERSON]

	PR #1660: Force delete NetLoadOnClient=false actors in client worlds (Contributed by [REDACTED_PERSON])
	#tests created a new world with an actor with NetLoadOnClient=false, made sure was broken without changed, and fixed with it

Change 3045207 on 2016/07/11 by [REDACTED_PERSON]

	Fix assert in channel cleanup code that could occur if the connection was cleaned up, and there were KeepProcessingActorChannelBunchesMap in-flight still

Change 3045281 on 2016/07/11 by [REDACTED_PERSON]

	PR #1863: Improve 'Too many calls to RPC' log message informativeness (Contributed by [REDACTED_PERSON])

Change 3047946 on 2016/07/13 by [REDACTED_PERSON]

	Ammortize replay checkpoints

	* You can now set a threshold of time allowed to save a checkpoint each frame (demo.CheckpointSaveMaxMSPerFrame)
	* By default, this threshold is unlimited though
	* We stop recording checkpoints when we go beyond this threshold
	* We pause normal recording while a checkpoint is in flight
	* Fixed a few issues with checkpoints modifying the connection state
	* Added ability to override ack state of packagemap so checkpoint can evolve this without interfering with normal replication
	* Added ability to track which channel needa to save for checkpoint

Change 3050006 on 2016/07/14 by [REDACTED_PERSON]

	Enable pause replication feature for characters in shootergame.

	#tests ran shooter game with 1 server and client

Change 3050048 on 2016/07/14 by [REDACTED_PERSON]

	- Disable pause replication points always being drawn regardless of the value of NetVisualizeRelevancyTestPoints
	- Remove unused DrawCircle function

	#tests tested in debug editor with server/client

Change 3050247 on 2016/07/14 by [REDACTED_PERSON]

	VOIP fixes

diff --git a/Engine/Plugins/Online/OnlineSubsystemSteam/Source/Private/OnlineSubsystemSteam.cpp b/Engine/Plugins/Online/OnlineSubsystemSteam/Source/Private/OnlineSubsystemSteam.cpp
--- a/Engine/Plugins/Online/OnlineSubsystemSteam/Source/Private/OnlineSubsystemSteam.cpp
+++ b/Engine/Plugins/Online/OnlineSubsystemSteam/Source/Private/OnlineSubsystemSteam.cpp
@@ -552,7 +552,11 @@ bool FOnlineSubsystemSteam::InitSteamworksServer()
 		}
 	}
 
-	GConfig->GetInt(TEXT("URL"), TEXT("Port"), GameServerGamePort, GEngineIni);
+	if (FParse::Value(FCommandLine::Get(), TEXT("Port="), GameServerGamePort) == false)
+	{
+		GConfig->GetInt(TEXT("URL"), TEXT("Port"), GameServerGamePort, GEngineIni);
+	}
+	
 	GameServerSteamPort = GameServerGamePort + 1;
 
 	// Allow the command line to override the default query port
diff --git a/Engine/Plugins/Online/OnlineSubsystemUtils/Source/OnlineSubsystemUtils/Private/OnlineEngineInterfaceImpl.cpp b/Engine/Plugins/Online/OnlineSubsystemUtils/Source/OnlineSubsystemUtils/Private/OnlineEngineInterfaceImpl.cpp
--- a/Engine/Plugins/Online/OnlineSubsystemUtils/Source/OnlineSubsystemUtils/Private/OnlineEngineInterfaceImpl.cpp
+++ b/Engine/Plugins/Online/OnlineSubsystemUtils/Source/OnlineSubsystemUtils/Private/OnlineEngineInterfaceImpl.cpp
@@ -309,7 +309,7 @@ TSharedPtr<FVoicePacket> UOnlineEngineInterfaceImpl::SerializeRemotePacket(UWorl
 	IOnlineVoicePtr VoiceInt = Online::GetVoiceInterface(World);
 	if (VoiceInt.IsValid())
 	{
-		VoiceInt->SerializeRemotePacket(Ar);
+		return VoiceInt->SerializeRemotePacket(Ar);
 	}
 	return nullptr;
 }
@@ -346,7 +346,7 @@ bool UOnlineEngineInterfaceImpl::MuteRemoteTalker(UWorld* World, uint8 LocalUser
 	IOnlineVoicePtr VoiceInt = Online::GetVoiceInterface(World);
 	if (VoiceInt.IsValid())
 	{
-		VoiceInt->MuteRemoteTalker(LocalUserNum, PlayerId, bIsSystemWide);
+		return VoiceInt->MuteRemoteTalker(LocalUserNum, PlayerId, bIsSystemWide);
 	}
 	return false;
 }
@@ -356,7 +356,7 @@ bool UOnlineEngineInterfaceImpl::UnmuteRemoteTalker(UWorld* World, uint8 LocalUs
 	IOnlineVoicePtr VoiceInt = Online::GetVoiceInterface(World);
 	if (VoiceInt.IsValid())
 	{
-		VoiceInt->UnmuteRemoteTalker(LocalUserNum, PlayerId, bIsSystemWide);
+		return VoiceInt->UnmuteRemoteTalker(LocalUserNum, PlayerId, bIsSystemWide);
 	}
 	return false;
 }
diff --git a/Engine/Plugins/Runtime/PacketHandlers/CompressionComponents/Oodle/BuildDictionary_Generic.bat b/Engine/Plugins/Runtime/PacketHandlers/CompressionComponents/Oodle/BuildDictionary_Generic.bat
--- a/Engine/Plugins/Runtime/PacketHandlers/CompressionComponents/Oodle/BuildDictionary_Generic.bat
+++ b/Engine/Plugins/Runtime/PacketHandlers/CompressionComponents/Oodle/BuildDictionary_Generic.bat
@@ -22,22 +22,32 @@ echo Could not locate UE4Editor.exe
 goto End
 
 :GetGame
+set GameName=%1
+if not "%1" == "" goto GetDictionaryOutput
 set /p GameName=Type the name of the game you are working with: 
 echo.
 
 :GetDictionaryOutput
+set DictionaryOutput=%2
+if not "%2" == "" goto GetFilter
 set /p DictionaryOutput=Type the absolute path and full name of the resulting dictionary file, or Input or Output for the default dictionary location: 
 echo.
 
 :GetFilter
+set FileFilter=%3
+if not "%3" == "" goto GetChangelistFilter
 set /p FileFilter=Type a filename filter to filter by, or all for all files: 
 echo.
 
 :GetChangelistFilter
+set ChangelistFilter=%4
+if not "%4" == "" goto GetDirectory
 set /p ChangelistFilter=Type a changelist number to filter by, or all for all files: 
 echo.
 
 :GetDirectory
+set DirectoryRoot=%5
+if not "%5" == "" goto AutoGenDictionaries
 set /p DirectoryRoot=Type the root directory where the capture files are located: 
 echo.
 
@@ -62,8 +72,4 @@ pause
 
 :End
 echo Execution complete.
-pause
-
-
-REM Put nothing past here.
 
diff --git a/Engine/Plugins/Runtime/PacketHandlers/CompressionComponents/Oodle/ExtractAllArchives.bat b/Engine/Plugins/Runtime/PacketHandlers/CompressionComponents/Oodle/ExtractAllArchives.bat
--- a/Engine/Plugins/Runtime/PacketHandlers/CompressionComponents/Oodle/ExtractAllArchives.bat
+++ b/Engine/Plugins/Runtime/PacketHandlers/CompressionComponents/Oodle/ExtractAllArchives.bat
@@ -1,13 +1,23 @@
 @echo off
 
 echo.
-echo This batch file will unzip all downloaded packet captures from S3.
+echo Downloading packet captures from S3
 echo.
 
+if "%1" == "" goto getdirectory
+set Directory=%1
+goto getoptions
+
+:getdirectory
 set /p Directory=Type the root directory of the zipped captures: 
 echo.
 
-aws s3 sync s3://dedicated-server-ucap %Directory%
+:getoptions
+if "%2" == "" GOTO continue
+set Options= --exclude "*" --include "*%2*.ucap.gz"
+
+:continue
+aws s3 sync s3://dedicated-server-ucap %Directory% %Options%
 
 pushd %CD%
 cd %Directory%
@@ -15,7 +25,7 @@ FOR /D /r %%F in ("*") DO (
 	pushd %CD%
 	cd %%F
 		FOR %%X in (*.gz) DO (
-			"C:\Program Files\7-zip\7z.exe" x "%%X" -y
+			"C:\Program Files\7-zip\7z.exe" x "%%X" -y -aos
 		)
 	popd
 )
diff --git a/Engine/Plugins/Runtime/PacketHandlers/CompressionComponents/Oodle/Source/OodleHandlerComponent/Private/OodleHandlerComponent.cpp b/Engine/Plugins/Runtime/PacketHandlers/CompressionComponents/Oodle/Source/OodleHandlerComponent/Private/OodleHandlerComponent.cpp
--- a/Engine/Plugins/Runtime/PacketHandlers/CompressionComponents/Oodle/Source/OodleHandlerComponent/Private/OodleHandlerComponent.cpp
+++ b/Engine/Plugins/Runtime/PacketHandlers/CompressionComponents/Oodle/Source/OodleHandlerComponent/Private/OodleHandlerComponent.cpp
@@ -340,12 +340,12 @@ void OodleHandlerComponent::Initialize()
 
 		if (bCaptureMode)
 		{
-			UE_LOG(OodleHandlerComponentLog, Log, TEXT("Enabling Oodle capture mode."));
-
 			int32 CapturePercentage = 100;
 			FParse::Value(FCommandLine::Get(), TEXT("CapturePercentage="), CapturePercentage);
-			
-			if (FMath::RandRange(0, 100) <= CapturePercentage)
+
+			int32 RandNum = FMath::RandRange(0, 100);
+			UE_LOG(OodleHandlerComponentLog, Log, TEXT("Enabling Oodle capture mode. Random number is: %d, Capture Percentage is: %d, random number must be less than capture percentage to capture."), RandNum, CapturePercentage);
+			if (RandNum <= CapturePercentage)
 			{
 				InitializePacketLogs();
 			}
diff --git a/Engine/Source/Runtime/Core/Public/Misc/NetworkVersion.h b/Engine/Source/Runtime/Core/Public/Misc/NetworkVersion.h
--- a/Engine/Source/Runtime/Core/Public/Misc/NetworkVersion.h
+++ b/Engine/Source/Runtime/Core/Public/Misc/NetworkVersion.h
@@ -53,7 +53,7 @@ struct CORE_API FNetworkVersion
 	static bool IsNetworkCompatible( const uint32 LocalNetworkVersion, const uint32 RemoteNetworkVersion );
 
 	/**
-	* Generates a spexial struct that contains information to send to replay server
+	* Generates a special struct that contains information to send to replay server
 	*/
 	static FNetworkReplayVersion GetReplayVersion();
 
diff --git a/Engine/Source/Runtime/Engine/Classes/Engine/ActorChannel.h b/Engine/Source/Runtime/Engine/Classes/Engine/ActorChannel.h
--- a/Engine/Source/Runtime/Engine/Classes/Engine/ActorChannel.h
+++ b/Engine/Source/Runtime/Engine/Classes/Engine/ActorChannel.h
@@ -59,6 +59,7 @@ class ENGINE_API UActorChannel
 	uint32  bActorMustStayDirty:1;	// ActorDirty may not be cleared at end of this tick
 	uint32  bActorStillInitial:1;	// Not all properties sent while bNetInitial, so still bNetInitial next tick
 	uint32  bIsReplicatingActor:1;	// true when in this channel's ReplicateActor() to avoid recursion as that can cause invalid data to be sent
+	uint32  bPendingCheckpoint:1;	// true when this channel is currently waiting to replicate for a checkpoint
 	
 	/** whether we should nullptr references to this channel's Actor in other channels' Recent data when this channel is closed
 	 * set to false in cases where the Actor can't become relevant again (e.g. destruction) as it's unnecessary in that case
diff --git a/Engine/Source/Runtime/Engine/Classes/Engine/Channel.h b/Engine/Source/Runtime/Engine/Classes/Engine/Channel.h
--- a/Engine/Source/Runtime/Engine/Classes/Engine/Channel.h
+++ b/Engine/Source/Runtime/Engine/Classes/Engine/Channel.h
@@ -41,7 +41,8 @@ class ENGINE_API UChannel
 	// Variables.
 	uint32				OpenAcked:1;		// If OpenedLocally is true, this means we have acknowledged the packet we sent the bOpen bunch on. Otherwise, it means we have received the bOpen bunch from the server.
 	uint32				Closing:1;			// State of the channel.
-	uint32				Dormant:1;			// Channel is going dormant (it will close but the client will not destroy 
+	uint32				Dormant:1;			// Channel is going dormant (it will close but the client will not destroy
+	uint32				bIsReplicationPaused:1;	// Replication is being paused, but channel will not be closed
 	uint32				OpenTemporary:1;	// Opened temporarily.
 	uint32				Broken:1;			// Has encountered errors and is ignoring subsequent packets.
 	uint32				bTornOff:1;			// Actor associated with this channel was torn off
@@ -141,6 +142,12 @@ protected:
 	/** cleans up channel structures and nulls references to the channel */
 	virtual bool CleanUp( const bool bForDestroy );
 
+	/** Sets whether replication is currently paused on this channel or not */
+	virtual void SetReplicationPaused(bool InbIsReplicationPaused) { bIsReplicationPaused = InbIsReplicationPaused; }
+
+	/** Returns whether replication is currently paused on this channel */
+	virtual bool IsReplicationPaused() { return bIsReplicationPaused; }
+
 private:
 
 	/** Just sends the bunch out on the connection */
diff --git a/Engine/Source/Runtime/Engine/Classes/Engine/DemoNetConnection.h b/Engine/Source/Runtime/Engine/Classes/Engine/DemoNetConnection.h
--- a/Engine/Source/Runtime/Engine/Classes/Engine/DemoNetConnection.h
+++ b/Engine/Source/Runtime/Engine/Classes/Engine/DemoNetConnection.h
@@ -62,6 +62,7 @@ public:
 	}
 
 	TArray<FQueuedDemoPacket> QueuedDemoPackets;
+	TArray<FQueuedDemoPacket> QueuedCheckpointPackets;
 
 private:
 	void TrackSendForProfiler(const void* Data, int32 NumBytes);
diff --git a/Engine/Source/Runtime/Engine/Classes/Engine/DemoNetDriver.h b/Engine/Source/Runtime/Engine/Classes/Engine/DemoNetDriver.h
--- a/Engine/Source/Runtime/Engine/Classes/Engine/DemoNetDriver.h
+++ b/Engine/Source/Runtime/Engine/Classes/Engine/DemoNetDriver.h
@@ -4,6 +4,8 @@
 
 #include "NetDriver.h"
 #include "NetworkReplayStreaming.h"
+#include "PackageMapClient.h"
+#include "DemoNetConnection.h"
 #include "DemoNetDriver.generated.h"
 
 DECLARE_LOG_CATEGORY_EXTERN( LogDemo, Log, All );
@@ -169,10 +171,15 @@ class ENGINE_API UDemoNetDriver : public UNetDriver
 	double		MaxRecordTime;
 	int32		RecordCountSinceFlush;
 
-	bool		bSavingCheckpoint;
-	double		LastCheckpointTime;
+	/** Checkpoint state */
+	FPackageMapAckState CheckpointAckState;					// Current ack state of packagemap for the current checkpoint being saved
+	bool				bSavingCheckpoint;					// True if we're currently saving a checkpoint
+	double				TotalCheckpointSaveTimeSeconds;		// Total time it took to save checkpoint across all frames
+	int32				TotalCheckpointSaveFrames;			// Total number of frames used to save a checkpoint
+	double				LastCheckpointTime;					// Last time a checkpoint was saved
 
 	void		SaveCheckpoint();
+	void		TickCheckpoint();
 	void		LoadCheckpoint( FArchive* GotoCheckpointArchive, int64 GotoCheckpointSkipExtraTimeInMS );
 
 	void		SaveExternalData( FArchive& Ar );
@@ -314,7 +321,7 @@ public:
 	bool ConditionallyReadDemoFrameIntoPlaybackPackets( FArchive& Ar );
 	bool ProcessPacket( uint8* Data, int32 Count );
 
-	void WriteDemoFrameFromQueuedDemoPackets( FArchive& Ar, UDemoNetConnection* Connection );
+	void WriteDemoFrameFromQueuedDemoPackets( FArchive& Ar, TArray<FQueuedDemoPacket>& QueuedPackets );
 	void WritePacket( FArchive& Ar, uint8* Data, int32 Count );
 
 	void TickDemoPlayback( float DeltaSeconds );
diff --git a/Engine/Source/Runtime/Engine/Classes/Engine/PackageMapClient.h b/Engine/Source/Runtime/Engine/Classes/Engine/PackageMapClient.h
--- a/Engine/Source/Runtime/Engine/Classes/Engine/PackageMapClient.h
+++ b/Engine/Source/Runtime/Engine/Classes/Engine/PackageMapClient.h
@@ -213,6 +213,13 @@ public:
 	TMap< FNetworkGUID, int32 >	NetGUIDAckStatus;				// Map that represents the ack state of each net guid for this connection
 	TMap< uint32, bool >		NetFieldExportGroupPathAcked;	// Map that represents whether or not a net field export group has been ack'd by the client
 	TMap< uint64, bool >		NetFieldExportAcked;			// Map that represents whether or not a net field export has been ack'd by the client
+
+	void Reset()
+	{
+		NetGUIDAckStatus.Empty();
+		NetFieldExportGroupPathAcked.Empty();
+		NetFieldExportAcked.Empty();
+	}
 };
 
 UCLASS(transient)
@@ -228,6 +235,7 @@ public:
 		Connection = InConnection;
 		GuidCache = InNetGUIDCache;
 		ExportNetGUIDCount = 0;
+		OverrideAckState = &AckState;
 	}
 
 	virtual ~UPackageMapClient()
@@ -281,6 +289,7 @@ public:
 
 	void SavePackageMapExportAckStatus( FPackageMapAckState& OutState );
 	void RestorePackageMapExportAckStatus( const FPackageMapAckState& InState );
+	void OverridePackageMapExportAckStatus( FPackageMapAckState* NewState );
 
 	/** Functions to help with exporting/importing net field info */
 	TSharedPtr< FNetFieldExportGroup >	GetNetFieldExportGroup( const FString& PathName );
@@ -315,7 +324,8 @@ protected:
 
 	TArray< FNetworkGUID >				PendingAckGUIDs;					// Quick access to all GUID's that haven't been acked
 
-	FPackageMapAckState					AckState;
+	FPackageMapAckState					AckState;							// Current ack state of exported data
+	FPackageMapAckState*				OverrideAckState;					// This is a pointer that allows us to override the current ack state, it's never NULL (it will point to AckState by default)
 
 	// Bunches of NetGUID/path tables to send with the current content bunch
 	TArray<FOutBunch* >					ExportBunches;
diff --git a/Engine/Source/Runtime/Engine/Classes/GameFramework/Actor.h b/Engine/Source/Runtime/Engine/Classes/GameFramework/Actor.h
--- a/Engine/Source/Runtime/Engine/Classes/GameFramework/Actor.h
+++ b/Engine/Source/Runtime/Engine/Classes/GameFramework/Actor.h
@@ -15,6 +15,7 @@ struct FHitResult;
 class AActor;
 class FTimerManager; 
 class UNetDriver;
+struct FNetViewer;
 
 #include "Actor.generated.h"
 
@@ -296,6 +297,12 @@ public:
 	/** Dormancy setting for actor to take itself off of the replication list without being destroyed on clients. */
 	TEnumAsByte<enum ENetDormancy> NetDormancy;
 
+	/** Gives the actor a chance to pause replication to a player represented by the passed in actor - only called on server */
+	virtual bool IsReplicationPausedForConnection(const FNetViewer& ConnectionOwnerNetViewer);
+
+	/** Called on the client when the replication paused value is changed */
+	virtual void OnReplicationPausedChanged(bool bIsReplicationPaused);
+
 	/** Automatically registers this actor to receive input from a player. */
 	UPROPERTY(EditAnywhere, Category=Input)
 	TEnumAsByte<EAutoReceiveInput::Type> AutoReceiveInput;
diff --git a/Engine/Source/Runtime/Engine/Private/Actor.cpp b/Engine/Source/Runtime/Engine/Private/Actor.cpp
--- a/Engine/Source/Runtime/Engine/Private/Actor.cpp
+++ b/Engine/Source/Runtime/Engine/Private/Actor.cpp
@@ -1749,6 +1749,15 @@ void AActor::ForceNetUpdate()
 	SetNetUpdateTime(FMath::Min(NetUpdateTime, GetWorld()->TimeSeconds - 0.01f));
 }
 
+bool AActor::IsReplicationPausedForConnection(const FNetViewer& ConnectionOwnerNetViewer)
+{
+	return false;
+}
+
+void AActor::OnReplicationPausedChanged(bool bIsReplicationPaused)
+{
+}
+
 void AActor::SetNetDormancy(ENetDormancy NewDormancy)
 {
 	if (IsNetMode(NM_Client))
diff --git a/Engine/Source/Runtime/Engine/Private/DataBunch.cpp b/Engine/Source/Runtime/Engine/Private/DataBunch.cpp
--- a/Engine/Source/Runtime/Engine/Private/DataBunch.cpp
+++ b/Engine/Source/Runtime/Engine/Private/DataBunch.cpp
@@ -90,6 +90,7 @@ FOutBunch::FOutBunch( UChannel* InChannel, bool bInClose )
 ,	bOpen		( 0 )
 ,	bClose		( bInClose )
 ,	bDormant	( 0 )
+,	bIsReplicationPaused	( 0 )
 ,	bReliable	( 0 )
 ,	bPartial	( 0 )
 ,	bPartialInitial			( 0 )
@@ -123,6 +124,7 @@ FOutBunch::FOutBunch( UPackageMap *InPackageMap, int64 MaxBits )
 ,	bOpen		( 0 )
 ,	bClose		( 0 )
 ,	bDormant	( 0 )
+,   bIsReplicationPaused	( 0 )
 ,	bReliable	( 0 )
 ,	bPartial	( 0 )
 ,	bPartialInitial ( 0 )
diff --git a/Engine/Source/Runtime/Engine/Private/DataChannel.cpp b/Engine/Source/Runtime/Engine/Private/DataChannel.cpp
--- a/Engine/Source/Runtime/Engine/Private/DataChannel.cpp
+++ b/Engine/Source/Runtime/Engine/Private/DataChannel.cpp
@@ -513,6 +513,7 @@ bool UChannel::ReceivedNextBunch( FInBunch & Bunch, bool & bOutSkipAck )
 					InPartialBunch->bPartialFinal			= true;
 					InPartialBunch->bClose					= Bunch.bClose;
 					InPartialBunch->bDormant				= Bunch.bDormant;
+					InPartialBunch->bIsReplicationPaused	= Bunch.bIsReplicationPaused;
 					InPartialBunch->bHasMustBeMappedGUIDs	= Bunch.bHasMustBeMappedGUIDs;
 				}
 				else
@@ -673,32 +674,40 @@ void UActorChannel::AppendExportBunches( TArray<FOutBunch *>& OutExportBunches )
 {
 	Super::AppendExportBunches( OutExportBunches );
 
-	// Let the profiler know about exported GUID bunches
-	for (const FOutBunch* ExportBunch : QueuedExportBunches )
+	// We don't want to append QueuedExportBunches to these bunches, since these were for queued RPC's, and we don't want to record RPC's during bResendAllDataSinceOpen
+	if ( !Connection->bResendAllDataSinceOpen )
 	{
-		if (ExportBunch != nullptr)
+		// Let the profiler know about exported GUID bunches
+		for ( const FOutBunch* ExportBunch : QueuedExportBunches )
 		{
-			NETWORK_PROFILER(GNetworkProfiler.TrackExportBunch(ExportBunch->GetNumBits(), Connection));
+			if ( ExportBunch != nullptr )
+			{
+				NETWORK_PROFILER( GNetworkProfiler.TrackExportBunch( ExportBunch->GetNumBits(), Connection ) );
+			}
 		}
-	}
 
-	if ( QueuedExportBunches.Num() )
-	{
-		OutExportBunches.Append( QueuedExportBunches );
-		QueuedExportBunches.Empty();
+		if ( QueuedExportBunches.Num() )
+		{
+			OutExportBunches.Append( QueuedExportBunches );
+			QueuedExportBunches.Empty();
+		}
 	}
 }
 
 void UActorChannel::AppendMustBeMappedGuids( FOutBunch* Bunch )
 {
-	if ( QueuedMustBeMappedGuidsInLastBunch.Num() > 0 )
+	// We don't want to append QueuedMustBeMappedGuidsInLastBunch to these bunches, since these were for queued RPC's, and we don't want to record RPC's during bResendAllDataSinceOpen
+	if ( !Connection->bResendAllDataSinceOpen )
 	{
-		// Just add our list to the main list on package map so we can re-use the code in UChannel to add them all together
-		UPackageMapClient * PackageMapClient = CastChecked< UPackageMapClient >( Connection->PackageMap );
+		if ( QueuedMustBeMappedGuidsInLastBunch.Num() > 0 )
+		{
+			// Just add our list to the main list on package map so we can re-use the code in UChannel to add them all together
+			UPackageMapClient * PackageMapClient = CastChecked< UPackageMapClient >( Connection->PackageMap );
 
-		PackageMapClient->GetMustBeMappedGuidsInLastBunch().Append( QueuedMustBeMappedGuidsInLastBunch );
+			PackageMapClient->GetMustBeMappedGuidsInLastBunch().Append( QueuedMustBeMappedGuidsInLastBunch );
 
-		QueuedMustBeMappedGuidsInLastBunch.Empty();
+			QueuedMustBeMappedGuidsInLastBunch.Empty();
+		}
 	}
 
 	// Actually add them to the bunch
@@ -846,6 +855,7 @@ FPacketIdRange UChannel::SendBunch( FOutBunch* Bunch, bool Merge )
 		NextBunch->bOpen = Bunch->bOpen;
 		NextBunch->bClose = Bunch->bClose;
 		NextBunch->bDormant = Bunch->bDormant;
+		NextBunch->bIsReplicationPaused = Bunch->bIsReplicationPaused;
 		NextBunch->ChIndex = Bunch->ChIndex;
 		NextBunch->ChType = Bunch->ChType;
 
@@ -1964,8 +1974,6 @@ void UActorChannel::ProcessBunch( FInBunch & Bunch )
 		return;
 	}
 
-	const bool bIsServer = Connection->Driver->IsServer();
-
 	FReplicationFlags RepFlags;
 
 	// ------------------------------------------------------------
@@ -2011,6 +2019,13 @@ void UActorChannel::ProcessBunch( FInBunch & Bunch )
 		UE_LOG(LogNetTraffic, Log, TEXT("      Actor %s:"), *Actor->GetFullName() );
 	}
 
+	bool bLatestIsReplicationPaused = Bunch.bIsReplicationPaused != 0;
+	if (bLatestIsReplicationPaused != IsReplicationPaused())
+	{
+		Actor->OnReplicationPausedChanged(bLatestIsReplicationPaused);
+		SetReplicationPaused(bLatestIsReplicationPaused);
+	}
+
 	// Owned by connection's player?
 	UNetConnection* ActorConnection = Actor->GetNetConnection();
 	if (ActorConnection == Connection || (ActorConnection != NULL && ActorConnection->IsA(UChildConnection::StaticClass()) && ((UChildConnection*)ActorConnection)->Parent == Connection))
@@ -2161,6 +2176,28 @@ bool UActorChannel::ReplicateActor()
 	check(Connection);
 	check(Connection->PackageMap);
 
+	const UWorld* const ActorWorld = Actor->GetWorld();
+
+	bool bIsNewlyReplicationPaused = false;
+	bool bIsNewlyReplicationUnpaused = false;
+	if (OpenPacketId.First != INDEX_NONE)
+	{
+		const TArray<FNetViewer>& NetViewers = ActorWorld->GetWorldSettings()->ReplicationViewers;
+
+		for (int32 viewerIdx = 0; viewerIdx < NetViewers.Num(); viewerIdx++)
+		{
+			bool bIsReplicationPausedForConnection = Actor->IsReplicationPausedForConnection(NetViewers[viewerIdx]);
+			bool bOldIsReplicationPaused = IsReplicationPaused();
+			bIsNewlyReplicationPaused = bIsReplicationPausedForConnection && !bOldIsReplicationPaused;
+			bIsNewlyReplicationUnpaused = !bIsReplicationPausedForConnection && bOldIsReplicationPaused;
+			SetReplicationPaused(bIsReplicationPausedForConnection);
+			if (bIsReplicationPausedForConnection && bOldIsReplicationPaused)
+			{
+				return false;
+			}
+		}
+	}
+
 	// The package map shouldn't have any carry over guids
 	if ( CastChecked< UPackageMapClient >( Connection->PackageMap )->GetMustBeMappedGuidsInLastBunch().Num() != 0 )
 	{
@@ -2170,7 +2207,7 @@ bool UActorChannel::ReplicateActor()
 	// Time how long it takes to replicate this particular actor
 	STAT( FScopeCycleCounterUObject FunctionScope(Actor) );
 
-	bool WroteSomethingImportant = false;
+	bool WroteSomethingImportant = bIsNewlyReplicationUnpaused || bIsNewlyReplicationPaused;
 
 	// triggering replication of an Actor while already in the middle of replication can result in invalid data being sent and is therefore illegal
 	if (bIsReplicatingActor)
@@ -2188,6 +2225,13 @@ bool UActorChannel::ReplicateActor()
 		return false;
 	}
 
+	if (bIsNewlyReplicationPaused)
+	{
+		Bunch.bReliable = true;
+		Bunch.bIsReplicationPaused = true;
+
+	}
+
 #if !(UE_BUILD_SHIPPING || UE_BUILD_TEST)
 	if (CVarNetReliableDebug.GetValueOnGameThread() > 0)
 	{
@@ -2248,7 +2292,6 @@ bool UActorChannel::ReplicateActor()
 	RepFlags.bNetSimulated	= ( Actor->GetRemoteRole() == ROLE_SimulatedProxy );
 	RepFlags.bRepPhysics	= Actor->ReplicatedMovement.bRepPhysics;
 
-	const UWorld* const ActorWorld	= Actor->GetWorld();
 	RepFlags.bReplay				= ActorWorld && (ActorWorld->DemoNetDriver == Connection->GetDriver());
 
 	RepFlags.bNetInitial = RepFlags.bNetInitial || bActorStillInitial; // for replication purposes, bNetInitial stays true until all properties sent
@@ -2267,39 +2310,42 @@ bool UActorChannel::ReplicateActor()
 	const uint32 ActorReplicateStartTime = GNetworkProfiler.IsTrackingEnabled() ? FPlatformTime::Cycles() : 0;
 #endif
 
-	// The Actor
-	WroteSomethingImportant |= ActorReplicator->ReplicateProperties( Bunch, RepFlags );
+	if (!bIsNewlyReplicationPaused)
+	{
+		// The Actor
+		WroteSomethingImportant |= ActorReplicator->ReplicateProperties(Bunch, RepFlags);
 
-	// The SubObjects
-	WroteSomethingImportant |= Actor->ReplicateSubobjects(this, &Bunch, &RepFlags);
+		// The SubObjects
+		WroteSomethingImportant |= Actor->ReplicateSubobjects(this, &Bunch, &RepFlags);
 
-	if ( Connection->bResendAllDataSinceOpen )
-	{
-		if ( WroteSomethingImportant )
+		if (Connection->bResendAllDataSinceOpen)
 		{
-			SendBunch( &Bunch, 1 );
-		}
+			if (WroteSomethingImportant)
+			{
+				SendBunch(&Bunch, 1);
+			}
 
-		MemMark.Pop();
+			MemMark.Pop();
 
-		bIsReplicatingActor = false;
+			bIsReplicatingActor = false;
 
-		return WroteSomethingImportant;
-	}
+			return WroteSomethingImportant;
+		}
 
-	// Look for deleted subobjects
-	for (auto RepComp = ReplicationMap.CreateIterator(); RepComp; ++RepComp)
-	{
-		if (!RepComp.Key().IsValid())
+		// Look for deleted subobjects
+		for (auto RepComp = ReplicationMap.CreateIterator(); RepComp; ++RepComp)
 		{
-			// Write a deletion content header:
-			WriteContentBlockForSubObjectDelete( Bunch, RepComp.Value()->ObjectNetGUID );
+			if (!RepComp.Key().IsValid())
+			{
+				// Write a deletion content header:
+				WriteContentBlockForSubObjectDelete(Bunch, RepComp.Value()->ObjectNetGUID);
 
-			WroteSomethingImportant = true;
-			Bunch.bReliable = true;
+				WroteSomethingImportant = true;
+				Bunch.bReliable = true;
 
-			RepComp.Value()->CleanUp();
-			RepComp.RemoveCurrent();			
+				RepComp.Value()->CleanUp();
+				RepComp.RemoveCurrent();
+			}
 		}
 	}
 
@@ -2313,42 +2359,46 @@ bool UActorChannel::ReplicateActor()
 	if( WroteSomethingImportant )
 	{
 		FPacketIdRange PacketRange = SendBunch( &Bunch, 1 );
-		for (auto RepComp = ReplicationMap.CreateIterator(); RepComp; ++RepComp)
-		{
-			RepComp.Value()->PostSendBunch( PacketRange, Bunch.bReliable );
-		}
 
-		// If there were any subobject keys pending, add them to the NakMap
-		if (PendingObjKeys.Num() >0)
+		if (!bIsNewlyReplicationPaused)
 		{
-			// For the packet range we just sent over
-			for(int32 PacketId = PacketRange.First; PacketId <= PacketRange.Last; ++PacketId)
+			for (auto RepComp = ReplicationMap.CreateIterator(); RepComp; ++RepComp)
 			{
-				// Get the existing set (its possible we send multiple bunches back to back and they end up on the same packet)
-				FPacketRepKeyInfo &Info = SubobjectNakMap.FindOrAdd(PacketId % SubobjectRepKeyBufferSize);
-				if (Info.PacketID != PacketId)
-				{
-					UE_LOG(LogNetTraffic, Verbose, TEXT("ActorChannel[%d]: Clearing out PacketRepKeyInfo for new packet: %d"), ChIndex, PacketId);
-					Info.ObjKeys.Empty(Info.ObjKeys.Num());
-				}
-				Info.PacketID = PacketId;
-				Info.ObjKeys.Append(PendingObjKeys);
+				RepComp.Value()->PostSendBunch(PacketRange, Bunch.bReliable);
+			}
 
-				FString VerboseString;
-				for (auto KeyIt = PendingObjKeys.CreateIterator(); KeyIt; ++KeyIt)
+			// If there were any subobject keys pending, add them to the NakMap
+			if (PendingObjKeys.Num() > 0)
+			{
+				// For the packet range we just sent over
+				for (int32 PacketId = PacketRange.First; PacketId <= PacketRange.Last; ++PacketId)
 				{
-					VerboseString += FString::Printf(TEXT(" %d"), *KeyIt);
+					// Get the existing set (its possible we send multiple bunches back to back and they end up on the same packet)
+					FPacketRepKeyInfo &Info = SubobjectNakMap.FindOrAdd(PacketId % SubobjectRepKeyBufferSize);
+					if (Info.PacketID != PacketId)
+					{
+						UE_LOG(LogNetTraffic, Verbose, TEXT("ActorChannel[%d]: Clearing out PacketRepKeyInfo for new packet: %d"), ChIndex, PacketId);
+						Info.ObjKeys.Empty(Info.ObjKeys.Num());
+					}
+					Info.PacketID = PacketId;
+					Info.ObjKeys.Append(PendingObjKeys);
+
+					FString VerboseString;
+					for (auto KeyIt = PendingObjKeys.CreateIterator(); KeyIt; ++KeyIt)
+					{
+						VerboseString += FString::Printf(TEXT(" %d"), *KeyIt);
+					}
+
+					UE_LOG(LogNetTraffic, Verbose, TEXT("ActorChannel[%d]: Sending ObjKeys: %s"), ChIndex, *VerboseString);
 				}
-				
-				UE_LOG(LogNetTraffic, Verbose, TEXT("ActorChannel[%d]: Sending ObjKeys: %s"), ChIndex, *VerboseString);
+			}
+
+			if (Actor->bNetTemporary)
+			{
+				Connection->SentTemporaries.Add(Actor);
 			}
 		}
-		
 		SentBunch = true;
-		if( Actor->bNetTemporary )
-		{
-			Connection->SentTemporaries.Add( Actor );
-		}
 	}
 
 	PendingObjKeys.Empty();
diff --git a/Engine/Source/Runtime/Engine/Private/DataReplication.cpp b/Engine/Source/Runtime/Engine/Private/DataReplication.cpp
--- a/Engine/Source/Runtime/Engine/Private/DataReplication.cpp
+++ b/Engine/Source/Runtime/Engine/Private/DataReplication.cpp
@@ -1063,8 +1063,8 @@ void FObjectReplicator::QueueRemoteFunctionBunch( UFunction* Func, FOutBunch &Bu
 	
 	if (++RemoteFuncInfo[InfoIdx].Calls > CVarMaxRPCPerNetUpdate.GetValueOnGameThread())
 	{
-		UE_LOG(LogRep, Verbose, TEXT("Too many calls to RPC %s within a single netupdate. Skipping. %s.  LastCallTime: %.2f. CurrentTime: %.2f. LastRelevantTime: %.2f. LastUpdateTime: %.2f "), 
-			*Func->GetName(), *GetObject()->GetName(), RemoteFuncInfo[InfoIdx].LastCallTime, OwningChannel->Connection->Driver->Time, OwningChannel->RelevantTime, OwningChannel->LastUpdateTime );
+		UE_LOG(LogRep, Verbose, TEXT("Too many calls (%d) to RPC %s within a single netupdate. Skipping. %s.  LastCallTime: %.2f. CurrentTime: %.2f. LastRelevantTime: %.2f. LastUpdateTime: %.2f "),
+			RemoteFuncInfo[InfoIdx].Calls, *Func->GetName(), *GetPathNameSafe(GetObject()), RemoteFuncInfo[InfoIdx].LastCallTime, OwningChannel->Connection->Driver->Time, OwningChannel->RelevantTime, OwningChannel->LastUpdateTime);
 		return;
 	}
 	
diff --git a/Engine/Source/Runtime/Engine/Private/DemoNetDriver.cpp b/Engine/Source/Runtime/Engine/Private/DemoNetDriver.cpp
--- a/Engine/Source/Runtime/Engine/Private/DemoNetDriver.cpp
+++ b/Engine/Source/Runtime/Engine/Private/DemoNetDriver.cpp
@@ -45,6 +45,7 @@ static TAutoConsoleVariable<int32> CVarDemoQueueCheckpointChannels( TEXT( "demo.
 static TAutoConsoleVariable<int32> CVarUseAdaptiveReplayUpdateFrequency( TEXT( "demo.UseAdaptiveReplayUpdateFrequency" ), 0, TEXT( "If 1, NetUpdateFrequency will be calculated based on how often actors actually write something when recording to a replay" ) );
 static TAutoConsoleVariable<int32> CVarDemoAsyncLoadWorld( TEXT( "demo.AsyncLoadWorld" ), 0, TEXT( "If 1, we will use seamless server travel to load the replay world asynchronously" ) );
 static TAutoConsoleVariable<float> CVarCheckpointUploadDelayInSeconds( TEXT( "demo.CheckpointUploadDelayInSeconds" ), 30.0f, TEXT( "" ) );
+static TAutoConsoleVariable<float> CVarCheckpointSaveMaxMSPerFrame( TEXT( "demo.CheckpointSaveMaxMSPerFrame" ), 0.0f, TEXT( "Maximum time allowed each frame to spend on saving a checkpoint. If 0, it will save the checkpoint in a single frame, regardless of how long it takes." ) );
 
 #if !(UE_BUILD_SHIPPING || UE_BUILD_TEST)
 static TAutoConsoleVariable<int32> CVarDemoForceFailure( TEXT( "demo.ForceFailure" ), 0, TEXT( "" ) );
@@ -1009,100 +1010,158 @@ void UDemoNetDriver::SaveCheckpoint()
 
 	check( CheckpointArchive->TotalSize() == 0 );
 
-	const double StartCheckpointTime = FPlatformTime::Seconds();
+	check( ClientConnections[0]->SendBuffer.GetNumBits() == 0 );
 
-	// Save the replicated server time so we can restore it after the checkpoint has been serialized.
-	// This preserves the existing behavior and prevents clients from receiving updated server time
-	// more often than the normal update rate.
-	float SavedReplicatedServerTimeSeconds = -1.0f;
+	// Mark all existing actor channels for pending checkpoint save
+	for ( TSharedPtr<FNetworkObjectInfo>& ObjectInfo : GetNetworkObjectList().GetObjects() )
+	{
+		AActor* Actor = ObjectInfo.Get()->Actor;
 
-	AGameState* const GameState = World != nullptr ? World->GetGameState() : nullptr;
+		UActorChannel* ActorChannel = ClientConnections[0]->ActorChannels.FindRef( Actor );
 
-	// Normally AGameState::ReplicatedWorldTimeSeconds is only updated periodically,
-	// but we want to make sure it's accurate for the checkpoint.
-	if (GameState != nullptr)
-	{
-		SavedReplicatedServerTimeSeconds = GameState->ReplicatedWorldTimeSeconds;
-		GameState->UpdateServerTimeSeconds();
+		if ( ActorChannel )
+		{
+			ActorChannel->bPendingCheckpoint = true;
+		}
 	}
 
-	// First, save the current guid cache
-	SerializeGuidCache( GuidCache, CheckpointArchive );
+	// So we know to tick checkpoints until it's all saved out (we also can't do normal replication while we're in the process of saving a checkpoint)
+	bSavingCheckpoint = true;
 
 	UPackageMapClient* PackageMapClient = ( ( UPackageMapClient* )ClientConnections[0]->PackageMap );
 
-	// Save the compatible rep layout map
-	PackageMapClient->SerializeNetFieldExportGroupMap( *CheckpointArchive );
+	PackageMapClient->SavePackageMapExportAckStatus( CheckpointAckState );
+
+	TotalCheckpointSaveTimeSeconds = 0;
+	TotalCheckpointSaveFrames = 0;
+
+	UE_LOG( LogDemo, Log, TEXT( "Starting checkpoint. Actors: %i" ), GetNetworkObjectList().GetObjects().Num() );
+
+	// Do the first checkpoint tick now
+	TickCheckpoint();
+}
+
+void UDemoNetDriver::TickCheckpoint()
+{
+	if ( !bSavingCheckpoint )
+	{
+		return;
+	}
+
+	DECLARE_SCOPE_CYCLE_COUNTER( TEXT( "SaveCheckpoint time" ), STAT_ReplayCheckpointSaveTime, STATGROUP_Net );
 
-	const uint32 GuidCacheSize = CheckpointArchive->TotalSize();
+	FArchive* CheckpointArchive = ReplayStreamer->GetCheckpointArchive();
 
-	FURL CheckpointURL;
-	CheckpointURL.Map = TEXT( "Checkpoint" );
+	if ( !ensure( CheckpointArchive != nullptr ) )
+	{
+		return;
+	}
+
+	const double StartCheckpointTime = FPlatformTime::Seconds();
 
-	const double SerializeCacheTime = FPlatformTime::Seconds();
+	TotalCheckpointSaveFrames++;
 
+	ClientConnections[0]->FlushNet();
 	check( ClientConnections[0]->SendBuffer.GetNumBits() == 0 );
 
+	UPackageMapClient* PackageMapClient = ( ( UPackageMapClient* )ClientConnections[0]->PackageMap );
+
+	// Save package map ack status in case we export stuff during the checkpoint (so we can restore the connection back to what it was before we saved the checkpoint)
+	PackageMapClient->OverridePackageMapExportAckStatus( &CheckpointAckState );
+
+	// Save the replicated server time so we can restore it after the checkpoint has been serialized.
+	// This preserves the existing behavior and prevents clients from receiving updated server time
+	// more often than the normal update rate.
+	AGameState* const GameState = World != nullptr ? World->GetGameState() : nullptr;
+
+	const float SavedReplicatedServerTimeSeconds = GameState ? GameState->ReplicatedWorldTimeSeconds : -1.0f;
+
+	// Normally AGameState::ReplicatedWorldTimeSeconds is only updated periodically,
+	// but we want to make sure it's accurate for the checkpoint.
+	if ( GameState )
+	{
+		GameState->UpdateServerTimeSeconds();
+	}
+
 	// Re-use the existing connection to record all properties that have changed since channels were first opened
 	// Set bResendAllDataSinceOpen to true to signify that we want to do this
 	ClientConnections[0]->bResendAllDataSinceOpen = true;
-	bSavingCheckpoint = true;
 
-	FPackageMapAckState SavedAckState;
+	bool bCheckpointFinished = true;
 
-	// Save package map ack status in case we export stuff during the checkpoint (so we can restore the connection back to what it was before we saved the checkpoint)
-	PackageMapClient->SavePackageMapExportAckStatus( SavedAckState );
+	const double CheckpointMaxUploadTimePerFrame = ( double )CVarCheckpointSaveMaxMSPerFrame.GetValueOnGameThread() / 1000;
 
 	for ( TSharedPtr<FNetworkObjectInfo>& ObjectInfo : GetNetworkObjectList().GetObjects() )
 	{
 		AActor* Actor = ObjectInfo.Get()->Actor;
 
-		if ( ClientConnections[0]->ActorChannels.Contains( Actor ) )
+		UActorChannel* ActorChannel = ClientConnections[0]->ActorChannels.FindRef( Actor );
+
+		if ( ActorChannel && ActorChannel->bPendingCheckpoint )
 		{
 			Actor->CallPreReplication( this );
 			DemoReplicateActor( Actor, ClientConnections[0], SpectatorController, true );
+
+			ActorChannel->bPendingCheckpoint = false;
+
+			const double CheckpointTime = FPlatformTime::Seconds();
+
+			if ( CheckpointMaxUploadTimePerFrame > 0 && CheckpointTime - StartCheckpointTime > CheckpointMaxUploadTimePerFrame )
+			{
+				bCheckpointFinished = false;
+				break;
+			}
 		}
 	}
 
-	const double AfterReplicationTime = FPlatformTime::Seconds();
+	if ( GameState )
+	{
+		// Restore the game state's replicated world time
+		GameState->ReplicatedWorldTimeSeconds = SavedReplicatedServerTimeSeconds;
+	}
 
+	// Make sure to flush the connection (fill up QueuedCheckpointPackets)
+	// This also frees up the connection to be used for normal streaming again
 	ClientConnections[0]->FlushNet();
 	check( ClientConnections[0]->SendBuffer.GetNumBits() == 0 );
 
-	WriteDemoFrameFromQueuedDemoPackets( *CheckpointArchive, CastChecked< UDemoNetConnection>( ClientConnections[0] ) );
-
-	const double AfterWritePacketsTime = FPlatformTime::Seconds();
-
-	// Restore the connection back to what it was before we saved the checkpoint
-	PackageMapClient->RestorePackageMapExportAckStatus( SavedAckState );
+	PackageMapClient->OverridePackageMapExportAckStatus( nullptr );
 
 	ClientConnections[0]->bResendAllDataSinceOpen = false;
-	bSavingCheckpoint = false;
 
-	const uint32 CheckpointSize = CheckpointArchive->TotalSize() - GuidCacheSize;
+	const double EndCheckpointTime = FPlatformTime::Seconds();
 
-	const int32 TotalSize = CheckpointArchive->TotalSize();
+	TotalCheckpointSaveTimeSeconds += ( EndCheckpointTime - StartCheckpointTime );
 
-	if ( CheckpointArchive->TotalSize() > 0 )
+	if ( bCheckpointFinished )
 	{
-		ReplayStreamer->FlushCheckpoint( GetDemoCurrentTimeInMS() );
-	}
+		// We're done saving this checkpoint
+		bSavingCheckpoint = false;
 
-	// Restore the game state's replicated world time
-	if (GameState != nullptr)
-	{
-		GameState->ReplicatedWorldTimeSeconds = SavedReplicatedServerTimeSeconds;
-	}
+		// First, save the current guid cache
+		SerializeGuidCache( GuidCache, CheckpointArchive );
 
-	const double EndCheckpointTime = FPlatformTime::Seconds();
+		// Save the compatible rep layout map
+		PackageMapClient->SerializeNetFieldExportGroupMap( *CheckpointArchive );
+
+		// Get the size of the guid data saved
+		const uint32 GuidCacheSize = CheckpointArchive->TotalSize();
+
+		// Write out all of the queued up packets generated while saving the checkpoint
+		WriteDemoFrameFromQueuedDemoPackets( *CheckpointArchive, CastChecked< UDemoNetConnection>( ClientConnections[0] )->QueuedCheckpointPackets );
+
+		// Get the total checkpoint size
+		const int32 TotalCheckpointSize = CheckpointArchive->TotalSize();
+
+		if ( CheckpointArchive->TotalSize() > 0 )
+		{
+			ReplayStreamer->FlushCheckpoint( GetDemoCurrentTimeInMS() );
+		}
 
-	const float CheckpointTimeInMS			= ( EndCheckpointTime - StartCheckpointTime ) * 1000.0;
-	const float SerializeCacheTimeMS		= ( SerializeCacheTime - StartCheckpointTime ) * 1000.0;
-	const float AfterReplicationTimeMS		= ( AfterReplicationTime - SerializeCacheTime ) * 1000.0;
-	const float AfterWritePacketsTimeMS		= ( AfterWritePacketsTime - AfterReplicationTime ) * 1000.0;
-	const float MiscTimeMS					= ( EndCheckpointTime - AfterWritePacketsTime ) * 1000.0;
+		const float TotalCheckpointTimeInMS = TotalCheckpointSaveTimeSeconds * 1000.0;
 
-	UE_LOG( LogDemo, Log, TEXT( "Checkpoint. Actors: %i, TotalSize: %i, RepSize: %i, PackageMap: %u, TotalTime: %2.2f, SerializeCacheTimeMS: %2.2f, AfterReplicationTimeMS: %2.2f, AfterWritePacketsTimeMS: %2.2f, Misc: %2.2f" ), GetNetworkObjectList().GetObjects().Num(), TotalSize, CheckpointSize, GuidCacheSize, CheckpointTimeInMS, SerializeCacheTimeMS, AfterReplicationTimeMS, AfterWritePacketsTimeMS, MiscTimeMS );
+		UE_LOG( LogDemo, Log, TEXT( "Finished checkpoint. Actors: %i, GuidCacheSize: %i, TotalSize: %i, TotalCheckpointSaveFrames: %i, TotalCheckpointTimeInMS: %2.2f" ), GetNetworkObjectList().GetObjects().Num(), GuidCacheSize, TotalCheckpointSize, TotalCheckpointSaveFrames, TotalCheckpointTimeInMS );
+	}
 }
 
 void UDemoNetDriver::SaveExternalData( FArchive& Ar )
@@ -1202,6 +1261,13 @@ void UDemoNetDriver::TickDemoRecord( float DeltaSeconds )
 		return;
 	}
 
+	if ( bSavingCheckpoint )
+	{
+		// If we're in the middle of saving a checkpoint, then update that now and return
+		TickCheckpoint();
+		return;
+	}
+
 	FArchive* FileAr = ReplayStreamer->GetStreamingArchive();
 
 	if ( FileAr == NULL )
@@ -1426,11 +1492,13 @@ void UDemoNetDriver::TickDemoRecord( float DeltaSeconds )
 	ClientConnections[0]->FlushNet();
 	check( ClientConnections[0]->SendBuffer.GetNumBits() == 0 );
 
-	WriteDemoFrameFromQueuedDemoPackets( *FileAr, CastChecked< UDemoNetConnection >( ClientConnections[0] ) );
+	WriteDemoFrameFromQueuedDemoPackets( *FileAr, CastChecked< UDemoNetConnection >( ClientConnections[0] )->QueuedDemoPackets );
 
 	// Save a checkpoint if it's time
 	if ( CVarEnableCheckpoints.GetValueOnGameThread() == 1 )
 	{
+		check( !bSavingCheckpoint );		// We early out above, so this shouldn't be possible
+
 		const double CHECKPOINT_DELAY = CVarCheckpointUploadDelayInSeconds.GetValueOnGameThread();
 
 		if ( DemoCurrentTime - LastCheckpointTime > CHECKPOINT_DELAY )
@@ -1749,10 +1817,8 @@ bool UDemoNetDriver::ProcessPacket( uint8* Data, int32 Count )
 	return true;
 }
 
-void UDemoNetDriver::WriteDemoFrameFromQueuedDemoPackets( FArchive& Ar, UDemoNetConnection* Connection )
+void UDemoNetDriver::WriteDemoFrameFromQueuedDemoPackets( FArchive& Ar, TArray<FQueuedDemoPacket>& QueuedPackets )
 {
-	check( Connection->SendBuffer.GetNumBits() == 0 );
-
 	// Save total absolute demo time in seconds
 	Ar << DemoCurrentTime;
 
@@ -1777,12 +1843,12 @@ void UDemoNetDriver::WriteDemoFrameFromQueuedDemoPackets( FArchive& Ar, UDemoNet
 	// Save external data
 	SaveExternalData( Ar );
 
-	for ( int32 i = 0; i < Connection->QueuedDemoPackets.Num(); i++ )
+	for ( int32 i = 0; i < QueuedPackets.Num(); i++ )
 	{
-		WritePacket( Ar, Connection->QueuedDemoPackets[i].Data.GetData(), Connection->QueuedDemoPackets[i].Data.Num() );
+		WritePacket( Ar, QueuedPackets[i].Data.GetData(), QueuedPackets[i].Data.Num() );
 	}
 
-	Connection->QueuedDemoPackets.Empty();
+	QueuedPackets.Empty();
 
 	// Write a count of 0 to signal the end of the frame
 	int32 EndCount = 0;
@@ -2524,6 +2590,13 @@ void UDemoNetConnection::LowLevelSend(void* Data, int32 CountBytes, int32 CountB
 
 	TrackSendForProfiler( Data, CountBytes );
 
+	if ( bResendAllDataSinceOpen )
+	{
+		// This path is only active for a checkpoint saving out, we need to queue in separate list
+		new( QueuedCheckpointPackets )FQueuedDemoPacket( ( uint8* )Data, CountBytes, CountBits );
+		return;
+	}
+
 	new(QueuedDemoPackets)FQueuedDemoPacket((uint8*)Data, CountBytes, CountBits);
 }
 
diff --git a/Engine/Source/Runtime/Engine/Private/Level.cpp b/Engine/Source/Runtime/Engine/Private/Level.cpp
--- a/Engine/Source/Runtime/Engine/Private/Level.cpp
+++ b/Engine/Source/Runtime/Engine/Private/Level.cpp
@@ -1555,7 +1555,7 @@ void ULevel::InitializeNetworkActors()
 				{
 					if (!Actor->bNetLoadOnClient)
 					{
-						Actor->Destroy();
+						Actor->Destroy(true);
 					}
 					else
 					{
diff --git a/Engine/Source/Runtime/Engine/Private/NetConnection.cpp b/Engine/Source/Runtime/Engine/Private/NetConnection.cpp
--- a/Engine/Source/Runtime/Engine/Private/NetConnection.cpp
+++ b/Engine/Source/Runtime/Engine/Private/NetConnection.cpp
@@ -355,10 +355,21 @@ void UNetConnection::CleanUp()
 		UChannel* OpenChannel = OpenChannels[i];
 		if (OpenChannel != NULL)
 		{
-			OpenChannel->ConditionalCleanUp();
+			OpenChannel->ConditionalCleanUp(true);
 		}
 	}
 
+	// Cleanup any straggler KeepProcessingActorChannelBunchesMap channels
+	for (const TPair<FNetworkGUID, TArray<UActorChannel*>>& MapKeyValuePair : KeepProcessingActorChannelBunchesMap)
+	{
+		for (UActorChannel* CurChannel : MapKeyValuePair.Value)
+		{
+			CurChannel->ConditionalCleanUp(true);
+		}
+	}
+
+	KeepProcessingActorChannelBunchesMap.Empty();
+
 	PackageMap = NULL;
 
 	if (GIsRunning)
@@ -980,6 +991,7 @@ void UNetConnection::ReceivedPacket( FBitReader& Reader )
 			Bunch.bOpen					= bControl ? Reader.ReadBit() : 0;
 			Bunch.bClose				= bControl ? Reader.ReadBit() : 0;
 			Bunch.bDormant				= Bunch.bClose ? Reader.ReadBit() : 0;
+			Bunch.bIsReplicationPaused  = Reader.ReadBit();
 			Bunch.bReliable				= Reader.ReadBit();
 			Bunch.ChIndex				= Reader.ReadInt( MAX_CHANNELS );
 			Bunch.bHasPackageMapExports	= Reader.ReadBit();
@@ -1371,6 +1383,7 @@ int32 UNetConnection::SendRawBunch( FOutBunch& Bunch, bool InAllowMerge )
 			Header.WriteBit( Bunch.bDormant );
 		}
 	}
+	Header.WriteBit( Bunch.bIsReplicationPaused );
 	Header.WriteBit( Bunch.bReliable );
 	Header.WriteIntWrapped(Bunch.ChIndex, MAX_CHANNELS);
 	Header.WriteBit( Bunch.bHasPackageMapExports );
diff --git a/Engine/Source/Runtime/Engine/Private/PackageMapClient.cpp b/Engine/Source/Runtime/Engine/Private/PackageMapClient.cpp
--- a/Engine/Source/Runtime/Engine/Private/PackageMapClient.cpp
+++ b/Engine/Source/Runtime/Engine/Private/PackageMapClient.cpp
@@ -1066,7 +1066,7 @@ void UPackageMapClient::TrackNetFieldExport( FNetFieldExportGroup* NetFieldExpor
 	const uint64 CmdHandle = ( ( uint64 )NetFieldExportGroup->PathNameIndex ) << 32 | ( uint64 )NetFieldExportHandle;
 
 	// If this cmd hasn't been confirmed as exported, we need to export it for this bunch
-	if ( !AckState.NetFieldExportAcked.Contains( CmdHandle ) )
+	if ( !OverrideAckState->NetFieldExportAcked.Contains( CmdHandle ) )
 	{
 		NetFieldExports.Add( CmdHandle );		// NOTE - This is a set, so it will only add once
 	}
@@ -1173,7 +1173,7 @@ void UPackageMapClient::AppendNetFieldExports( TArray<FOutBunch *>& OutgoingBunc
 			ExportBunch->SerializeIntPacked( PathNameIndex );
 
 			// Export the path if we need to
-			if ( !AckState.NetFieldExportGroupPathAcked.Contains( PathNameIndex ) && !ExportedPathInThisBunchAlready.Contains( PathNameIndex ) )
+			if ( !OverrideAckState->NetFieldExportGroupPathAcked.Contains( PathNameIndex ) && !ExportedPathInThisBunchAlready.Contains( PathNameIndex ) )
 			{
 				ExportBunch->WriteBit( 1 );
 				*ExportBunch << PathName;
@@ -1343,6 +1343,11 @@ void UPackageMapClient::RestorePackageMapExportAckStatus( const FPackageMapAckSt
 	AckState = InState;
 }
 
+void UPackageMapClient::OverridePackageMapExportAckStatus( FPackageMapAckState* NewState )
+{
+	OverrideAckState = NewState ? NewState : &AckState;
+}
+
 //--------------------------------------------------------------------
 //
 //	Network - ACKing
@@ -1360,8 +1365,8 @@ void UPackageMapClient::NotifyBunchCommit( const int32 OutPacketId, const FOutBu
 	// For this to work with normal connections, we'll need to do real ack logic here
 	for ( int32 i = 0; i < OutBunch->NetFieldExports.Num(); i++ )
 	{
-		AckState.NetFieldExportGroupPathAcked.Add( OutBunch->NetFieldExports[i] >> 32, true );
-		AckState.NetFieldExportAcked.Add( OutBunch->NetFieldExports[i], true );
+		OverrideAckState->NetFieldExportGroupPathAcked.Add( OutBunch->NetFieldExports[i] >> 32, true );
+		OverrideAckState->NetFieldExportAcked.Add( OutBunch->NetFieldExports[i], true );
 	}
 
 	const TArray< FNetworkGUID >& ExportNetGUIDs = OutBunch->ExportNetGUIDs;
@@ -1375,12 +1380,12 @@ void UPackageMapClient::NotifyBunchCommit( const int32 OutPacketId, const FOutBu
 
 	for ( int32 i = 0; i < ExportNetGUIDs.Num(); i++ )
 	{
-		if ( !AckState.NetGUIDAckStatus.Contains( ExportNetGUIDs[i] ) )
+		if ( !OverrideAckState->NetGUIDAckStatus.Contains( ExportNetGUIDs[i] ) )
 		{
-			AckState.NetGUIDAckStatus.Add( ExportNetGUIDs[i], GUID_PACKET_NOT_ACKED );
+			OverrideAckState->NetGUIDAckStatus.Add( ExportNetGUIDs[i], GUID_PACKET_NOT_ACKED );
 		}
 
-		int32& ExpectedPacketIdRef = AckState.NetGUIDAckStatus.FindChecked( ExportNetGUIDs[i] );
+		int32& ExpectedPacketIdRef = OverrideAckState->NetGUIDAckStatus.FindChecked( ExportNetGUIDs[i] );
 
 		// Only update expected sequence if this guid was previously nak'd
 		// If we always update to the latest packet id, we risk prolonging the ack for no good reason
@@ -1409,7 +1414,7 @@ void UPackageMapClient::ReceivedAck( const int32 AckPacketId )
 {
 	for ( int32 i = PendingAckGUIDs.Num() - 1; i >= 0; i-- )
 	{
-		int32& ExpectedPacketIdRef = AckState.NetGUIDAckStatus.FindChecked( PendingAckGUIDs[i] );
+		int32& ExpectedPacketIdRef = OverrideAckState->NetGUIDAckStatus.FindChecked( PendingAckGUIDs[i] );
 
 		check( ExpectedPacketIdRef > GUID_PACKET_ACKED );		// Make sure we really are pending, since we're on the list
 
@@ -1429,7 +1434,7 @@ void UPackageMapClient::ReceivedNak( const int32 NakPacketId )
 {
 	for ( int32 i = PendingAckGUIDs.Num() - 1; i >= 0; i-- )
 	{
-		int32& ExpectedPacketIdRef = AckState.NetGUIDAckStatus.FindChecked( PendingAckGUIDs[i] );
+		int32& ExpectedPacketIdRef = OverrideAckState->NetGUIDAckStatus.FindChecked( PendingAckGUIDs[i] );
 
 		check( ExpectedPacketIdRef > GUID_PACKET_ACKED );		// Make sure we aren't acked, since we're on the list
 
@@ -1467,12 +1472,12 @@ bool UPackageMapClient::NetGUIDHasBeenAckd(FNetworkGUID NetGUID)
 	}
 
 	// If brand new, add it to map with GUID_PACKET_NOT_ACKED
-	if ( !AckState.NetGUIDAckStatus.Contains( NetGUID ) )
+	if ( !OverrideAckState->NetGUIDAckStatus.Contains( NetGUID ) )
 	{
-		AckState.NetGUIDAckStatus.Add( NetGUID, GUID_PACKET_NOT_ACKED );
+		OverrideAckState->NetGUIDAckStatus.Add( NetGUID, GUID_PACKET_NOT_ACKED );
 	}
 
-	int32& AckPacketId = AckState.NetGUIDAckStatus.FindChecked( NetGUID );
+	int32& AckPacketId = OverrideAckState->NetGUIDAckStatus.FindChecked( NetGUID );
 
 	if ( AckPacketId == GUID_PACKET_ACKED )
 	{
@@ -1556,9 +1561,9 @@ void UPackageMapClient::LogDebugInfo( FOutputDevice & Ar )
 		FNetworkGUID NetGUID = It.Value();
 
 		FString Status = TEXT("Unused");
-		if ( AckState.NetGUIDAckStatus.Contains( NetGUID ) )
+		if ( OverrideAckState->NetGUIDAckStatus.Contains( NetGUID ) )
 		{
-			const int32 PacketId = AckState.NetGUIDAckStatus.FindRef(NetGUID);
+			const int32 PacketId = OverrideAckState->NetGUIDAckStatus.FindRef(NetGUID);
 			if ( PacketId == GUID_PACKET_NOT_ACKED )
 			{
 				Status = TEXT("UnAckd");
@@ -1623,7 +1628,7 @@ bool UPackageMapClient::IsNetGUIDAuthority() const
 void UPackageMapClient::GetNetGUIDStats(int32 &AckCount, int32 &UnAckCount, int32 &PendingCount)
 {
 	AckCount = UnAckCount = PendingCount = 0;
-	for ( auto It = AckState.NetGUIDAckStatus.CreateIterator(); It; ++It )
+	for ( auto It = OverrideAckState->NetGUIDAckStatus.CreateIterator(); It; ++It )
 	{
 		// Sanity check that we're in sync
 		check( ( It.Value() > GUID_PACKET_ACKED ) == PendingAckGUIDs.Contains( It.Key() ) );
@@ -2272,6 +2277,7 @@ UObject* FNetGUIDCache::GetObjectFromNetGUID( const FNetworkGUID& NetGUID, const
 			//	1. We are actually a package
 			//	2. We aren't already pending
 			//	3. We're actually suppose to load (levels don't load here for example)
+			//		(Refer to CanClientLoadObject, which is where we protect clients from trying to load levels)
 			check( !PendingAsyncPackages.Contains( CacheObjectPtr->PathName ) );
 
 			if ( CVarAllowAsyncLoading.GetValueOnGameThread() > 0 )
diff --git a/Engine/Source/Runtime/Engine/Public/Net/DataBunch.h b/Engine/Source/Runtime/Engine/Public/Net/DataBunch.h
--- a/Engine/Source/Runtime/Engine/Public/Net/DataBunch.h
+++ b/Engine/Source/Runtime/Engine/Public/Net/DataBunch.h
@@ -23,6 +23,7 @@ public:
 	uint8					bOpen;
 	uint8					bClose;
 	uint8					bDormant;
+	uint8					bIsReplicationPaused;   // Replication on this channel is being paused by the server
 	uint8					bReliable;
 	uint8					bPartial;				// Not a complete bunch
 	uint8					bPartialInitial;		// The first bunch of a partial bunch
@@ -72,6 +73,7 @@ public:
 		Str += FString::Printf(TEXT("bOpen: %d "), bOpen);
 		Str += FString::Printf(TEXT("bClose: %d "), bClose);
 		Str += FString::Printf(TEXT("bDormant: %d "), bDormant);
+		Str += FString::Printf(TEXT("bIsReplicationPaused: %d "), bIsReplicationPaused);
 		Str += FString::Printf(TEXT("bReliable: %d "), bReliable);
 		Str += FString::Printf(TEXT("bPartial: %d//%d//%d "), bPartial, bPartialInitial, bPartialFinal);
 		Str += FString::Printf( TEXT( "bHasPackageMapExports: %d " ), bHasPackageMapExports );
@@ -99,6 +101,7 @@ public:
 	uint8				bOpen;
 	uint8				bClose;
 	uint8				bDormant;				// Close, but go dormant
+	uint8				bIsReplicationPaused;	// Replication on this channel is being paused by the server
 	uint8				bReliable;
 	uint8				bPartial;				// Not a complete bunch
 	uint8				bPartialInitial;		// The first bunch of a partial bunch
@@ -118,6 +121,7 @@ public:
 		Str += FString::Printf(TEXT("bOpen: %d "), bOpen);
 		Str += FString::Printf(TEXT("bClose: %d "), bClose);
 		Str += FString::Printf(TEXT("bDormant: %d "), bDormant);
+		Str += FString::Printf(TEXT("bIsReplicationPaused: %d "), bIsReplicationPaused);
 		Str += FString::Printf(TEXT("bReliable: %d "), bReliable);
 		Str += FString::Printf(TEXT("bPartial: %d//%d//%d "), bPartial, bPartialInitial, bPartialFinal);
 		Str += FString::Printf( TEXT( "bHasPackageMapExports: %d " ), bHasPackageMapExports );
diff --git a/Engine/Source/Runtime/Online/ICMP/Private/Icmp.cpp b/Engine/Source/Runtime/Online/ICMP/Private/Icmp.cpp
--- a/Engine/Source/Runtime/Online/ICMP/Private/Icmp.cpp
+++ b/Engine/Source/Runtime/Online/ICMP/Private/Icmp.cpp
@@ -2,7 +2,7 @@
 #include "IcmpPrivatePCH.h"
 #include "Icmp.h"
 #include "SocketSubsystem.h"
-#include "IpAddress.h"
+#include "IPAddress.h"
 
 FIcmpEchoResult IcmpEchoImpl(ISocketSubsystem* SocketSub, const FString& TargetAddress, float Timeout);
 
diff --git a/Engine/Source/Runtime/Online/ICMP/Private/IcmpPosix.cpp b/Engine/Source/Runtime/Online/ICMP/Private/IcmpPosix.cpp
--- a/Engine/Source/Runtime/Online/ICMP/Private/IcmpPosix.cpp
+++ b/Engine/Source/Runtime/Online/ICMP/Private/IcmpPosix.cpp
@@ -6,7 +6,7 @@
 #if PLATFORM_USES_POSIX_IMCP
 
 #include "SocketSubsystem.h"
-#include "IpAddress.h"
+#include "IPAddress.h"
 
 #include <arpa/inet.h>
 #include <netinet/in.h>
diff --git a/Engine/Source/Runtime/Online/ICMP/Private/UDPPing.cpp b/Engine/Source/Runtime/Online/ICMP/Private/UDPPing.cpp
--- a/Engine/Source/Runtime/Online/ICMP/Private/UDPPing.cpp
+++ b/Engine/Source/Runtime/Online/ICMP/Private/UDPPing.cpp
@@ -2,7 +2,7 @@
 #include "IcmpPrivatePCH.h"
 #include "Icmp.h"
 #include "SocketSubsystem.h"
-#include "IpAddress.h"
+#include "IPAddress.h"
 
 #include "Sockets.h"
 
