Copying //UE4/Dev-Sequencer to //UE4/Dev-Main (Source: //UE4/Dev-Sequencer @ 3057646)

diff --git a/Engine/Source/Runtime/Core/Private/Android/AndroidJavaMediaPlayer.cpp b/Engine/Source/Runtime/Core/Private/Android/AndroidJavaMediaPlayer.cpp
--- a/Engine/Source/Runtime/Core/Private/Android/AndroidJavaMediaPlayer.cpp
+++ b/Engine/Source/Runtime/Core/Private/Android/AndroidJavaMediaPlayer.cpp
@@ -4,6 +4,24 @@
 #include "AndroidJavaMediaPlayer.h"
 #include "AndroidApplication.h"
 
+#if UE_BUILD_SHIPPING
+// always clear any exceptions in SHipping
+#define CHECK_JNI_RESULT(Id) if (Id == 0) { JEnv->ExceptionClear(); }
+#else
+#define CHECK_JNI_RESULT(Id) \
+if (Id == 0) \
+{ \
+	if (bIsOptional) { JEnv->ExceptionClear(); } \
+	else { JEnv->ExceptionDescribe(); checkf(Id != 0, TEXT("Failed to find " #Id)); } \
+}
+#endif
+
+static jfieldID FindField(JNIEnv* JEnv, jclass Class, const ANSICHAR* FieldName, const ANSICHAR* FieldType, bool bIsOptional)
+{
+	jfieldID Field = Class == NULL ? NULL : JEnv->GetFieldID(Class, FieldName, FieldType);
+	CHECK_JNI_RESULT(Field);
+	return Field;
+}
 
 FJavaAndroidMediaPlayer::FJavaAndroidMediaPlayer(bool swizzlePixels)
 	: FJavaClassObject(GetClassName(), "(Z)V", swizzlePixels)
@@ -12,7 +30,7 @@ FJavaAndroidMediaPlayer::FJavaAndroidMediaPlayer(bool swizzlePixels)
 	, GetCurrentPositionMethod(GetClassMethod("getCurrentPosition", "()I"))
 	, IsLoopingMethod(GetClassMethod("isLooping", "()Z"))
 	, IsPlayingMethod(GetClassMethod("isPlaying", "()Z"))
-	, SetDataSourceURLMethod(GetClassMethod("setDataSource", "(Ljava/lang/String;)V"))
+	, SetDataSourceURLMethod(GetClassMethod("setDataSourceURL", "(Ljava/lang/String;)Z"))
 	, SetDataSourceFileMethod(GetClassMethod("setDataSource", "(Ljava/lang/String;JJ)Z"))
 	, SetDataSourceAssetMethod(GetClassMethod("setDataSource", "(Landroid/content/res/AssetManager;Ljava/lang/String;JJ)Z"))
 	, PrepareMethod(GetClassMethod("prepare", "()V"))
@@ -28,7 +46,52 @@ FJavaAndroidMediaPlayer::FJavaAndroidMediaPlayer(bool swizzlePixels)
 	, PauseMethod(GetClassMethod("pause", "()V"))
 	, StopMethod(GetClassMethod("stop", "()V"))
 	, GetVideoLastFrameMethod(GetClassMethod("getVideoLastFrame", "(I)Z"))
+	, GetAudioTracksMethod(GetClassMethod("GetAudioTracks", "()[Lcom/epicgames/ue4/MediaPlayer14$AudioTrackInfo;"))
+	, GetCaptionTracksMethod(GetClassMethod("GetCaptionTracks", "()[Lcom/epicgames/ue4/MediaPlayer14$CaptionTrackInfo;"))
+	, GetVideoTracksMethod(GetClassMethod("GetVideoTracks", "()[Lcom/epicgames/ue4/MediaPlayer14$VideoTrackInfo;"))
+	, DidResolutionChangeMethod(GetClassMethod("didResolutionChange", "()Z"))
 {
+	bTrackInfoSupported = FAndroidMisc::GetAndroidBuildVersion() >= 16;
+
+	if (bTrackInfoSupported)
+	{
+		SelectTrackMethod = GetClassMethod("selectTrack", "(I)V");
+	}
+
+	JNIEnv* JEnv = FAndroidApplication::GetJavaEnv();
+
+	// get field IDs for AudioTrackInfo class members
+	jclass localAudioTrackInfoClass = FAndroidApplication::FindJavaClass("com/epicgames/ue4/MediaPlayer14$AudioTrackInfo");
+	AudioTrackInfoClass = (jclass)JEnv->NewGlobalRef(localAudioTrackInfoClass);
+	JEnv->DeleteLocalRef(localAudioTrackInfoClass);
+	AudioTrackInfo_Index = FindField(JEnv, AudioTrackInfoClass, "Index", "I", false);
+	AudioTrackInfo_MimeType = FindField(JEnv, AudioTrackInfoClass, "MimeType", "Ljava/lang/String;", false);
+	AudioTrackInfo_DisplayName = FindField(JEnv, AudioTrackInfoClass, "DisplayName", "Ljava/lang/String;", false);
+	AudioTrackInfo_Language = FindField(JEnv, AudioTrackInfoClass, "Language", "Ljava/lang/String;", false);
+	AudioTrackInfo_Channels = FindField(JEnv, AudioTrackInfoClass, "Channels", "I", false);
+	AudioTrackInfo_SampleRate = FindField(JEnv, AudioTrackInfoClass, "SampleRate", "I", false);
+
+	// get field IDs for CaptionTrackInfo class members
+	jclass localCaptionTrackInfoClass = FAndroidApplication::FindJavaClass("com/epicgames/ue4/MediaPlayer14$CaptionTrackInfo");
+	CaptionTrackInfoClass = (jclass)JEnv->NewGlobalRef(localCaptionTrackInfoClass);
+	JEnv->DeleteLocalRef(localCaptionTrackInfoClass);
+	CaptionTrackInfo_Index = FindField(JEnv, CaptionTrackInfoClass, "Index", "I", false);
+	CaptionTrackInfo_MimeType = FindField(JEnv, CaptionTrackInfoClass, "MimeType", "Ljava/lang/String;", false);
+	CaptionTrackInfo_DisplayName = FindField(JEnv, CaptionTrackInfoClass, "DisplayName", "Ljava/lang/String;", false);
+	CaptionTrackInfo_Language = FindField(JEnv, CaptionTrackInfoClass, "Language", "Ljava/lang/String;", false);
+
+	// get field IDs for VideoTrackInfo class members
+	jclass localVideoTrackInfoClass = FAndroidApplication::FindJavaClass("com/epicgames/ue4/MediaPlayer14$VideoTrackInfo");
+	VideoTrackInfoClass = (jclass)JEnv->NewGlobalRef(localVideoTrackInfoClass);
+	JEnv->DeleteLocalRef(localVideoTrackInfoClass);
+	VideoTrackInfo_Index = FindField(JEnv, VideoTrackInfoClass, "Index", "I", false);
+	VideoTrackInfo_MimeType = FindField(JEnv, VideoTrackInfoClass, "MimeType", "Ljava/lang/String;", false);
+	VideoTrackInfo_DisplayName = FindField(JEnv, VideoTrackInfoClass, "DisplayName", "Ljava/lang/String;", false);
+	VideoTrackInfo_Language = FindField(JEnv, VideoTrackInfoClass, "Language", "Ljava/lang/String;", false);
+	VideoTrackInfo_BitRate = FindField(JEnv, VideoTrackInfoClass, "BitRate", "I", false);
+	VideoTrackInfo_Width = FindField(JEnv, VideoTrackInfoClass, "Width", "I", false);
+	VideoTrackInfo_Height = FindField(JEnv, VideoTrackInfoClass, "Height", "I", false);
+	VideoTrackInfo_FrameRate = FindField(JEnv, VideoTrackInfoClass, "FrameRate", "F", false);
 }
 
 int32 FJavaAndroidMediaPlayer::GetDuration()
@@ -62,9 +125,9 @@ bool FJavaAndroidMediaPlayer::IsPlaying()
 	return CallMethod<bool>(IsPlayingMethod);
 }
 
-void FJavaAndroidMediaPlayer::SetDataSource(const FString & Url)
+bool FJavaAndroidMediaPlayer::SetDataSource(const FString & Url)
 {
-	CallMethod<void>(SetDataSourceURLMethod, GetJString(Url));
+	return CallMethod<bool>(SetDataSourceURLMethod, GetJString(Url));
 }
 
 bool FJavaAndroidMediaPlayer::SetDataSource(const FString& MoviePathOnDevice, int64 offset, int64 size)
@@ -153,6 +216,11 @@ void FJavaAndroidMediaPlayer::Pause()
 	CallMethod<void>(PauseMethod);
 }
 
+bool FJavaAndroidMediaPlayer::DidResolutionChange()
+{
+	return CallMethod<bool>(DidResolutionChangeMethod);
+}
+
 bool FJavaAndroidMediaPlayer::GetVideoLastFrame(int32 destTexture)
 {
 	// This can return an exception in some cases
@@ -178,3 +246,174 @@ FName FJavaAndroidMediaPlayer::GetClassName()
 		return FName("");
 	}
 }
+
+bool FJavaAndroidMediaPlayer::SelectTrack(int32 index)
+{
+	if (!bTrackInfoSupported)
+	{
+		// Just assume it worked
+		return true;
+	}
+
+	// This can return an exception in some cases
+	JNIEnv*	JEnv = FAndroidApplication::GetJavaEnv();
+	JEnv->CallVoidMethod(Object, SelectTrackMethod.Method, index);
+	if (JEnv->ExceptionCheck())
+	{
+		JEnv->ExceptionDescribe();
+		JEnv->ExceptionClear();
+		return false;
+	}
+	return true;
+}
+
+bool FJavaAndroidMediaPlayer::GetAudioTracks(TArray<FAudioTrack>& AudioTracks)
+{
+	AudioTracks.Empty();
+
+	jobjectArray TrackArray = CallMethod<jobjectArray>(GetAudioTracksMethod);
+	if (nullptr != TrackArray)
+	{
+		bool bIsOptional = false;
+		JNIEnv* JEnv = FAndroidApplication::GetJavaEnv();
+		jsize ElementCount = JEnv->GetArrayLength(TrackArray);
+
+		for (int Index = 0; Index < ElementCount; ++Index)
+		{
+			jobject Track = JEnv->GetObjectArrayElement(TrackArray, Index);
+
+			int32 AudioTrackIndex = AudioTracks.AddDefaulted();
+			FAudioTrack& AudioTrack = AudioTracks[AudioTrackIndex];
+
+			AudioTrack.Index = (int32)JEnv->GetIntField(Track, AudioTrackInfo_Index);
+
+			jstring jsMimeType = (jstring)JEnv->GetObjectField(Track, AudioTrackInfo_MimeType);
+			CHECK_JNI_RESULT(jsMimeType);
+			const char * nativeMimeType = JEnv->GetStringUTFChars(jsMimeType, 0);
+			AudioTrack.MimeType = FString(nativeMimeType);
+			JEnv->ReleaseStringUTFChars(jsMimeType, nativeMimeType);
+			JEnv->DeleteLocalRef(jsMimeType);
+
+			jstring jsDisplayName = (jstring)JEnv->GetObjectField(Track, AudioTrackInfo_DisplayName);
+			CHECK_JNI_RESULT(jsDisplayName);
+			const char * nativeDisplayName = JEnv->GetStringUTFChars(jsDisplayName, 0);
+			AudioTrack.DisplayName = FString(nativeDisplayName);
+			JEnv->ReleaseStringUTFChars(jsDisplayName, nativeDisplayName);
+			JEnv->DeleteLocalRef(jsDisplayName);
+
+			jstring jsLanguage = (jstring)JEnv->GetObjectField(Track, AudioTrackInfo_Language);
+			CHECK_JNI_RESULT(jsLanguage);
+			const char * nativeLanguage = JEnv->GetStringUTFChars(jsLanguage, 0);
+			AudioTrack.Language = FString(nativeLanguage);
+			JEnv->ReleaseStringUTFChars(jsLanguage, nativeLanguage);
+			JEnv->DeleteLocalRef(jsLanguage);
+
+			AudioTrack.Channels = (int32)JEnv->GetIntField(Track, AudioTrackInfo_Channels);
+			AudioTrack.SampleRate = (int32)JEnv->GetIntField(Track, AudioTrackInfo_SampleRate);
+		}
+		JEnv->DeleteGlobalRef(TrackArray);
+
+		return true;
+	}
+	return false;
+}
+
+bool FJavaAndroidMediaPlayer::GetCaptionTracks(TArray<FCaptionTrack>& CaptionTracks)
+{
+	CaptionTracks.Empty();
+
+	jobjectArray TrackArray = CallMethod<jobjectArray>(GetCaptionTracksMethod);
+	if (nullptr != TrackArray)
+	{
+		bool bIsOptional = false;
+		JNIEnv* JEnv = FAndroidApplication::GetJavaEnv();
+		jsize ElementCount = JEnv->GetArrayLength(TrackArray);
+
+		for (int Index = 0; Index < ElementCount; ++Index)
+		{
+			jobject Track = JEnv->GetObjectArrayElement(TrackArray, Index);
+
+			int32 CaptionTrackIndex = CaptionTracks.AddDefaulted();
+			FCaptionTrack& CaptionTrack = CaptionTracks[CaptionTrackIndex];
+
+			CaptionTrack.Index = (int32)JEnv->GetIntField(Track, CaptionTrackInfo_Index);
+
+			jstring jsMimeType = (jstring)JEnv->GetObjectField(Track, CaptionTrackInfo_MimeType);
+			CHECK_JNI_RESULT(jsMimeType);
+			const char * nativeMimeType = JEnv->GetStringUTFChars(jsMimeType, 0);
+			CaptionTrack.MimeType = FString(nativeMimeType);
+			JEnv->ReleaseStringUTFChars(jsMimeType, nativeMimeType);
+			JEnv->DeleteLocalRef(jsMimeType);
+
+			jstring jsDisplayName = (jstring)JEnv->GetObjectField(Track, CaptionTrackInfo_DisplayName);
+			CHECK_JNI_RESULT(jsDisplayName);
+			const char * nativeDisplayName = JEnv->GetStringUTFChars(jsDisplayName, 0);
+			CaptionTrack.DisplayName = FString(nativeDisplayName);
+			JEnv->ReleaseStringUTFChars(jsDisplayName, nativeDisplayName);
+			JEnv->DeleteLocalRef(jsDisplayName);
+
+			jstring jsLanguage = (jstring)JEnv->GetObjectField(Track, CaptionTrackInfo_Language);
+			CHECK_JNI_RESULT(jsLanguage);
+			const char * nativeLanguage = JEnv->GetStringUTFChars(jsLanguage, 0);
+			CaptionTrack.Language = FString(nativeLanguage);
+			JEnv->ReleaseStringUTFChars(jsLanguage, nativeLanguage);
+			JEnv->DeleteLocalRef(jsLanguage);
+		}
+		JEnv->DeleteGlobalRef(TrackArray);
+
+		return true;
+	}
+	return false;
+}
+
+bool FJavaAndroidMediaPlayer::GetVideoTracks(TArray<FVideoTrack>& VideoTracks)
+{
+	VideoTracks.Empty();
+
+	jobjectArray TrackArray = CallMethod<jobjectArray>(GetVideoTracksMethod);
+	if (nullptr != TrackArray)
+	{
+		bool bIsOptional = false;
+		JNIEnv* JEnv = FAndroidApplication::GetJavaEnv();
+		jsize ElementCount = JEnv->GetArrayLength(TrackArray);
+
+		for (int Index = 0; Index < ElementCount; ++Index)
+		{
+			jobject Track = JEnv->GetObjectArrayElement(TrackArray, Index);
+
+			int32 VideoTrackIndex = VideoTracks.AddDefaulted();
+			FVideoTrack& VideoTrack = VideoTracks[VideoTrackIndex];
+
+			VideoTrack.Index = (int32)JEnv->GetIntField(Track, VideoTrackInfo_Index);
+
+			jstring jsMimeType = (jstring)JEnv->GetObjectField(Track, VideoTrackInfo_MimeType);
+			CHECK_JNI_RESULT(jsMimeType);
+			const char * nativeMimeType = JEnv->GetStringUTFChars(jsMimeType, 0);
+			VideoTrack.MimeType = FString(nativeMimeType);
+			JEnv->ReleaseStringUTFChars(jsMimeType, nativeMimeType);
+			JEnv->DeleteLocalRef(jsMimeType);
+
+			jstring jsDisplayName = (jstring)JEnv->GetObjectField(Track, VideoTrackInfo_DisplayName);
+			CHECK_JNI_RESULT(jsDisplayName);
+			const char * nativeDisplayName = JEnv->GetStringUTFChars(jsDisplayName, 0);
+			VideoTrack.DisplayName = FString(nativeDisplayName);
+			JEnv->ReleaseStringUTFChars(jsDisplayName, nativeDisplayName);
+			JEnv->DeleteLocalRef(jsDisplayName);
+
+			jstring jsLanguage = (jstring)JEnv->GetObjectField(Track, VideoTrackInfo_Language);
+			CHECK_JNI_RESULT(jsLanguage);
+			const char * nativeLanguage = JEnv->GetStringUTFChars(jsLanguage, 0);
+			VideoTrack.Language = FString(nativeLanguage);
+			JEnv->ReleaseStringUTFChars(jsLanguage, nativeLanguage);
+			JEnv->DeleteLocalRef(jsLanguage);
+
+			VideoTrack.BitRate = (int32)JEnv->GetIntField(Track, VideoTrackInfo_BitRate);
+			VideoTrack.Dimensions = FIntPoint((int32)JEnv->GetIntField(Track, VideoTrackInfo_Width), (int32)JEnv->GetIntField(Track, VideoTrackInfo_Height));
+			VideoTrack.FrameRate = JEnv->GetFloatField(Track, VideoTrackInfo_FrameRate);
+		}
+		JEnv->DeleteGlobalRef(TrackArray);
+
+		return true;
+	}
+	return false;
+}
diff --git a/Engine/Source/Runtime/Core/Private/Tests/Misc/QueueTest.cpp b/Engine/Source/Runtime/Core/Private/Tests/Misc/QueueTest.cpp
--- a/Engine/Source/Runtime/Core/Private/Tests/Misc/QueueTest.cpp
+++ b/Engine/Source/Runtime/Core/Private/Tests/Misc/QueueTest.cpp
@@ -5,52 +5,65 @@
 
 #if WITH_DEV_AUTOMATION_TESTS
 
+
 IMPLEMENT_SIMPLE_AUTOMATION_TEST(FQueueTest, "System.Core.Misc.Queue", EAutomationTestFlags::ApplicationContextMask | EAutomationTestFlags::SmokeFilter)
 
+
 bool FQueueTest::RunTest( const FString& Parameters )
 {
 	// empty queues
-	TQueue<int32> q1_1;
-	int32 i1_1 = 666;
+	{
+		TQueue<int32> Queue;
+		int32 OutItem = 0;
 
-	TestTrue(TEXT("A new queue must be empty"), q1_1.IsEmpty());
+		TestTrue(TEXT("A new queue must be empty"), Queue.IsEmpty());
 
-	TestFalse(TEXT("A new queue must not dequeue anything"), q1_1.Dequeue(i1_1));
-	TestFalse(TEXT("A new queue must not peek anything"), q1_1.Peek(i1_1));
+		TestFalse(TEXT("A new queue must not dequeue anything"), Queue.Dequeue(OutItem));
+		TestFalse(TEXT("A new queue must not peek anything"), Queue.Peek(OutItem));
+	}
 
 	// insertion & removal
-	TQueue<int32> q2_1;
-	int32 i2_1 = 1;
-	int32 i2_2 = 2;
-	int32 i2_p = 0;
+	{
+		TQueue<int32> Queue;
+		int32 Item1 = 1;
+		int32 Item2 = 2;
+		int32 Item3 = 3;
+		int32 OutItem = 0;
 
-	TestTrue(TEXT("Inserting into a new queue must succeed"), q2_1.Enqueue(i2_1));
-	TestTrue(TEXT("Peek must succeed on a queue with one item"), q2_1.Peek(i2_p));
-	TestEqual(TEXT("Peek must return the first value"), i2_p, i2_1);
+		TestTrue(TEXT("Inserting into a new queue must succeed"), Queue.Enqueue(Item1));
+		TestTrue(TEXT("Peek must succeed on a queue with one item"), Queue.Peek(OutItem));
+		TestEqual(TEXT("Peek must return the first value"), OutItem, Item1);
 
-	TestTrue(TEXT("Inserting into a non-empty queue must succeed"), q2_1.Enqueue(i2_2));
-	TestTrue(TEXT("Peek must succeed on a queue with two items"), q2_1.Peek(i2_p));
-	TestEqual(TEXT("Peek must return the first value"), i2_p, i2_1);
+		TestTrue(TEXT("Inserting into a non-empty queue must succeed"), Queue.Enqueue(Item2));
+		TestTrue(TEXT("Peek must succeed on a queue with two items"), Queue.Peek(OutItem));
+		TestEqual(TEXT("Peek must return the first item"), OutItem, Item1);
 
-	TestTrue(TEXT("Dequeue must succeed on a queue with two items"), q2_1.Dequeue(i2_p));
-	TestEqual(TEXT("Dequeue must return the first value"), i2_p, i2_1);
-	TestTrue(TEXT("Dequeue must succeed on a queue with one item"), q2_1.Dequeue(i2_p));
-	TestEqual(TEXT("Dequeue must return the second value"), i2_p, i2_2);
+		Queue.Enqueue(Item3);
 
-	TestTrue(TEXT("After removing all items, the queue must be empty"), q2_1.IsEmpty());
+		TestTrue(TEXT("Dequeue must succeed on a queue with three items"), Queue.Dequeue(OutItem));
+		TestEqual(TEXT("Dequeue must return the first item"), OutItem, Item1);
+		TestTrue(TEXT("Dequeue must succeed on a queue with two items"), Queue.Dequeue(OutItem));
+		TestEqual(TEXT("Dequeue must return the second item"), OutItem, Item2);
+		TestTrue(TEXT("Dequeue must succeed on a queue with one item"), Queue.Dequeue(OutItem));
+		TestEqual(TEXT("Dequeue must return the third item"), OutItem, Item3);
+
+		TestTrue(TEXT("After removing all items, the queue must be empty"), Queue.IsEmpty());
+	}
 
 	// emptying
-	TQueue<int32> q3_1;
 	{
-		q3_1.Enqueue(1);
-		q3_1.Enqueue(2);
-		q3_1.Enqueue(3);
-		q3_1.Empty();
-	}
+		TQueue<int32> Queue;
 
-	TestTrue(TEXT("An emptied queue must be empty"), q3_1.IsEmpty());
+		Queue.Enqueue(1);
+		Queue.Enqueue(2);
+		Queue.Enqueue(3);
+		Queue.Empty();
+
+		TestTrue(TEXT("An emptied queue must be empty"), Queue.IsEmpty());
+	}
 
 	return true;
 }
 
-#endif //WITH_DEV_AUTOMATION_TESTS
\ No newline at end of file
+
+#endif //WITH_DEV_AUTOMATION_TESTS
diff --git a/Engine/Source/Runtime/Core/Private/Tests/Misc/TripleBufferTest.cpp b/Engine/Source/Runtime/Core/Private/Tests/Misc/TripleBufferTest.cpp
--- a/Engine/Source/Runtime/Core/Private/Tests/Misc/TripleBufferTest.cpp
+++ b/Engine/Source/Runtime/Core/Private/Tests/Misc/TripleBufferTest.cpp
@@ -11,60 +11,69 @@ IMPLEMENT_SIMPLE_AUTOMATION_TEST(FTripleBufferTest, "System.Core.Misc.TripleBuff
 bool FTripleBufferTest::RunTest(const FString& Parameters)
 {
 	// uninitialized buffer
-	TTripleBuffer<int32> b1_1(NoInit);
+	{
+		TTripleBuffer<int32> Buffer(NoInit);
 
-	TestFalse(TEXT("Uninitialized triple buffer must not be dirty"), b1_1.IsDirty());
+		TestFalse(TEXT("Uninitialized triple buffer must not be dirty"), Buffer.IsDirty());
+	}
 
 	// initialized buffer
-	TTripleBuffer<int32> b2_1(1);
+	{
+		TTripleBuffer<int32> Buffer(1);
 
-	TestFalse(TEXT("Initialized triple buffer must not be dirty"), b2_1.IsDirty());
-	TestEqual(TEXT("Initialized triple buffer must have correct read buffer value"), b2_1.Read(), 1);
+		TestFalse(TEXT("Initialized triple buffer must not be dirty"), Buffer.IsDirty());
+		TestEqual(TEXT("Initialized triple buffer must have correct read buffer value"), Buffer.Read(), 1);
 
-	b2_1.SwapReadBuffers();
+		Buffer.SwapReadBuffers();
 
-	TestEqual(TEXT("Initialized triple buffer must have correct temp buffer value"), b2_1.Read(), 1);
+		TestEqual(TEXT("Initialized triple buffer must have correct temp buffer value"), Buffer.Read(), 1);
 
-	b2_1.SwapWriteBuffers();
+		Buffer.SwapWriteBuffers();
 
-	TestTrue(TEXT("Write buffer swap must set dirty flag"), b2_1.IsDirty());
+		TestTrue(TEXT("Write buffer swap must set dirty flag"), Buffer.IsDirty());
 
-	b2_1.SwapReadBuffers();
+		Buffer.SwapReadBuffers();
 
-	TestFalse(TEXT("Read buffer swap must clear dirty flag"), b2_1.IsDirty());
-	TestEqual(TEXT("Initialized triple buffer must have correct temp buffer value"), b2_1.Read(), 1);
+		TestFalse(TEXT("Read buffer swap must clear dirty flag"), Buffer.IsDirty());
+		TestEqual(TEXT("Initialized triple buffer must have correct temp buffer value"), Buffer.Read(), 1);
+	}
 
-	// Pre-set buffer
-	TTripleBuffer<int32> b3_1(1, 2, 3);
+	// pre-set buffer
+	{
+		int32 Array[3] = { 1, 2, 3 };
+		TTripleBuffer<int32> Buffer(Array);
 
-	int32 Read = b3_1.Read();
-	TestEqual(TEXT("Pre-set triple buffer must have correct read buffer value"), Read, 1);
+		int32 Read = Buffer.Read();
+		TestEqual(TEXT("Pre-set triple buffer must have correct Read buffer value"), Read, 3);
 
-	b3_1.SwapReadBuffers();
+		Buffer.SwapReadBuffers();
 
-	int32 Temp = b3_1.Read();
-	TestEqual(TEXT("Pre-set triple buffer must have correct temp buffer value"), Temp, 3);
+		int32 Temp = Buffer.Read();
+		TestEqual(TEXT("Pre-set triple buffer must have correct Temp buffer value"), Temp, 1);
 
-	b3_1.SwapWriteBuffers();
-	b3_1.SwapReadBuffers();
+		Buffer.SwapWriteBuffers();
+		Buffer.SwapReadBuffers();
 
-	int32 Write = b3_1.Read();
-	TestEqual(TEXT("Pre-set triple buffer must have correct write buffer value"), Write, 2);
+		int32 Write = Buffer.Read();
+		TestEqual(TEXT("Pre-set triple buffer must have correct Write buffer value"), Write, 2);
+	}
 
 	// operations
-	TTripleBuffer<int32> b4_1;
+	{
+		TTripleBuffer<int32> Buffer;
 
-	b4_1.WriteAndSwap(1);
-	b4_1.SwapReadBuffers();
-	b4_1.WriteAndSwap(2);
-	b4_1.Write(3);
+		Buffer.WriteAndSwap(1);
+		Buffer.SwapReadBuffers();
+		Buffer.WriteAndSwap(2);
+		Buffer.Write(3);
 
-	TestEqual(TEXT("Triple buffer must read correct value (1)"), b4_1.Read(), 1);
-	TestEqual(TEXT("Triple buffer must read correct value (2)"), b4_1.SwapAndRead(), 2);
+		TestEqual(TEXT("Triple buffer must read correct value (1)"), Buffer.Read(), 1);
+		TestEqual(TEXT("Triple buffer must read correct value (2)"), Buffer.SwapAndRead(), 2);
 
-	b4_1.SwapWriteBuffers();
+		Buffer.SwapWriteBuffers();
 
-	TestEqual(TEXT("Triple buffer must read correct value (3)"), b4_1.SwapAndRead(), 3);
+		TestEqual(TEXT("Triple buffer must read correct value (3)"), Buffer.SwapAndRead(), 3);
+	}
 
 	return true;
 }
diff --git a/Engine/Source/Runtime/Core/Public/Android/AndroidJava.h b/Engine/Source/Runtime/Core/Public/Android/AndroidJava.h
--- a/Engine/Source/Runtime/Core/Public/Android/AndroidJava.h
+++ b/Engine/Source/Runtime/Core/Public/Android/AndroidJava.h
@@ -61,6 +61,9 @@ int FJavaClassObject::CallMethod<int>(FJavaClassMethod Method, ...);
 template<>
 jobject FJavaClassObject::CallMethod<jobject>(FJavaClassMethod Method, ...);
 
+template<>
+jobjectArray FJavaClassObject::CallMethod<jobjectArray>(FJavaClassMethod Method, ...);
+
 template<>
 int64 FJavaClassObject::CallMethod<int64>(FJavaClassMethod Method, ...);
 
diff --git a/Engine/Source/Runtime/Core/Public/Android/AndroidJavaMediaPlayer.h b/Engine/Source/Runtime/Core/Public/Android/AndroidJavaMediaPlayer.h
--- a/Engine/Source/Runtime/Core/Public/Android/AndroidJavaMediaPlayer.h
+++ b/Engine/Source/Runtime/Core/Public/Android/AndroidJavaMediaPlayer.h
@@ -8,6 +8,39 @@
 // Wrapper for com/epicgames/ue4/MediaPlayer*.java.
 class FJavaAndroidMediaPlayer : public FJavaClassObject
 {
+public:
+	struct FAudioTrack
+	{
+		int32 Index;
+		FString MimeType;
+		FString DisplayName;
+		FString Language;
+		FString Name;
+		uint32 Channels;
+		uint32 SampleRate;
+	};
+
+	struct FCaptionTrack
+	{
+		int32 Index;
+		FString MimeType;
+		FString DisplayName;
+		FString Language;
+		FString Name;
+	};
+
+	struct FVideoTrack
+	{
+		int32 Index;
+		FString MimeType;
+		FString DisplayName;
+		FString Language;
+		FString Name;
+		uint32 BitRate;
+		FIntPoint Dimensions;
+		float FrameRate;
+	};
+
 public:
 	FJavaAndroidMediaPlayer(bool swizzlePixels = true);
 	int32 GetDuration();
@@ -16,7 +49,7 @@ public:
 	int32 GetCurrentPosition();
 	bool IsLooping();
 	bool IsPlaying();
-	void SetDataSource(const FString & Url);
+	bool SetDataSource(const FString & Url);
 	bool SetDataSource(const FString& MoviePathOnDevice, int64 offset, int64 size);
 	bool SetDataSource(jobject AssetMgr, const FString& AssetPath, int64 offset, int64 size);
 	bool Prepare();
@@ -31,10 +64,17 @@ public:
 	void Start();
 	void Pause();
 	bool GetVideoLastFrame(int32 destTexture);
+	bool SelectTrack(int32 index);
+	bool GetAudioTracks(TArray<FAudioTrack>& AudioTracks);
+	bool GetCaptionTracks(TArray<FCaptionTrack>& CaptionTracks);
+	bool GetVideoTracks(TArray<FVideoTrack>& VideoTracks);
+	bool DidResolutionChange();
 
 private:
 	static FName GetClassName();
 
+	bool bTrackInfoSupported;
+
 	FJavaClassMethod GetDurationMethod;
 	FJavaClassMethod ResetMethod;
 	FJavaClassMethod StopMethod;
@@ -56,4 +96,36 @@ private:
 	FJavaClassMethod StartMethod;
 	FJavaClassMethod PauseMethod;
 	FJavaClassMethod GetVideoLastFrameMethod;
+	FJavaClassMethod SelectTrackMethod;
+	FJavaClassMethod GetAudioTracksMethod;
+	FJavaClassMethod GetCaptionTracksMethod;
+	FJavaClassMethod GetVideoTracksMethod;
+	FJavaClassMethod DidResolutionChangeMethod;
+
+	// AudioDeviceInfo member field ids
+	jclass AudioTrackInfoClass;
+	jfieldID AudioTrackInfo_Index;
+	jfieldID AudioTrackInfo_MimeType;
+	jfieldID AudioTrackInfo_DisplayName;
+	jfieldID AudioTrackInfo_Language;
+	jfieldID AudioTrackInfo_Channels;
+	jfieldID AudioTrackInfo_SampleRate;
+
+	// AudioDeviceInfo member field ids
+	jclass CaptionTrackInfoClass;
+	jfieldID CaptionTrackInfo_Index;
+	jfieldID CaptionTrackInfo_MimeType;
+	jfieldID CaptionTrackInfo_DisplayName;
+	jfieldID CaptionTrackInfo_Language;
+
+	// AudioDeviceInfo member field ids
+	jclass VideoTrackInfoClass;
+	jfieldID VideoTrackInfo_Index;
+	jfieldID VideoTrackInfo_MimeType;
+	jfieldID VideoTrackInfo_DisplayName;
+	jfieldID VideoTrackInfo_Language;
+	jfieldID VideoTrackInfo_BitRate;
+	jfieldID VideoTrackInfo_Width;
+	jfieldID VideoTrackInfo_Height;
+	jfieldID VideoTrackInfo_FrameRate;
 };
diff --git a/Engine/Source/Runtime/Core/Public/Async/Async.h b/Engine/Source/Runtime/Core/Public/Async/Async.h
--- a/Engine/Source/Runtime/Core/Public/Async/Async.h
+++ b/Engine/Source/Runtime/Core/Public/Async/Async.h
@@ -289,7 +289,7 @@ struct FAsyncThreadIndex
  * @param Execution The execution method to use, i.e. on Task Graph or in a separate thread.
  * @param Function The function to execute.
  * @param CompletionCallback An optional callback function that is executed when the function completed execution.
- * @result A TFuture object that will receive the return value from the function.
+ * @return A TFuture object that will receive the return value from the function.
  */
 template<typename ResultType>
 TFuture<ResultType> Async(EAsyncExecution Execution, TFunction<ResultType()> Function, TFunction<void()> CompletionCallback = TFunction<void()>())
diff --git a/Engine/Source/Runtime/Core/Public/Async/Future.h b/Engine/Source/Runtime/Core/Public/Async/Future.h
--- a/Engine/Source/Runtime/Core/Public/Async/Future.h
+++ b/Engine/Source/Runtime/Core/Public/Async/Future.h
@@ -54,7 +54,7 @@ public:
 	 * Blocks the calling thread until the future result is available.
 	 *
 	 * @param Duration The maximum time span to wait for the future result.
-	 * @result true if the result is available, false otherwise.
+	 * @return true if the result is available, false otherwise.
 	 * @see IsComplete
 	 */
 	bool WaitFor(const FTimespan& Duration) const
@@ -217,7 +217,7 @@ public:
 	 * Blocks the calling thread until the future result is available or the specified duration is exceeded.
 	 *
 	 * @param Duration The maximum time span to wait for the future result.
-	 * @result true if the result is available, false otherwise.
+	 * @return true if the result is available, false otherwise.
 	 * @see Wait, WaitUntil
 	 */
 	bool WaitFor(const FTimespan& Duration) const
@@ -229,7 +229,7 @@ public:
 	 * Blocks the calling thread until the future result is available or the specified time is hit.
 	 *
 	 * @param Time The time until to wait for the future result (in UTC).
-	 * @result true if the result is available, false otherwise.
+	 * @return true if the result is available, false otherwise.
 	 * @see Wait, WaitUntil
 	 */
 	bool WaitUntil(const FDateTime& Time) const
diff --git a/Engine/Source/Runtime/Core/Public/Containers/TripleBuffer.h b/Engine/Source/Runtime/Core/Public/Containers/TripleBuffer.h
--- a/Engine/Source/Runtime/Core/Public/Containers/TripleBuffer.h
+++ b/Engine/Source/Runtime/Core/Public/Containers/TripleBuffer.h
@@ -29,8 +29,10 @@
  * This class is thread-safe in single-producer, single-consumer scenarios.
  *
  * Based on ideas and C code in "Triple Buffering as a Concurrency Mechanism" (Reddit.com)
+ *
+ * @param BufferType The type of buffered items.
  */
-template<typename ItemType>
+template<typename BufferType>
 class TTripleBuffer
 {
 	/** Enumerates human readable bit values for accessing the Flags field. */
@@ -63,40 +65,52 @@ public:
 	/** Default constructor. */
 	TTripleBuffer()
 	{
-		Buffers[0] = Buffers[1] = Buffers[2] = ItemType();
-		Flags = EBufferFlag::Initial;
+		Initialize();
+		Buffers[0] = Buffers[1] = Buffers[2] = BufferType();
 	}
 
 	/** Default constructor (no initialization). */
 	explicit TTripleBuffer(ENoInit)
 	{
-		Flags = EBufferFlag::Initial;
+		Initialize();
 	}
 
 	/**
-	 * Create and initialize a new instance.
+	 * Create and initialize a new instance with a given buffer value.
 	 *
 	 * @param InValue The initial value of all three buffers.
 	 */
-	explicit TTripleBuffer(const ItemType& InValue)
+	explicit TTripleBuffer(const BufferType& InValue)
 	{
+		Initialize();
 		Buffers[0] = Buffers[1] = Buffers[2] = InValue;
-		Flags = EBufferFlag::Initial;
 	}
 
 	/**
-	 * Create and initialize a new instance.
+	 * Create and initialize a new instance using provided buffers.
+	 *
+	 * The elements of the provided items array are expected to have
+	 * the following initial contents:
+	 *     0 = Temp
+	 *     1 = Write
+	 *     2 = Read
 	 *
-	 * @param ReadValue The read buffer's initial value.
-	 * @param WriteValue The write buffer's initial value.
-	 * @param TempValue The temp buffer's initial value.
+	 * @param InBuffers The buffer memory to use.
 	 */
-	TTripleBuffer(const ItemType& ReadValue, const ItemType& WriteValue, const ItemType& TempValue)
+	TTripleBuffer(BufferType (&InBuffers)[3])
 	{
-		Buffers[0] = TempValue;
-		Buffers[1] = WriteValue;
-		Buffers[2] = ReadValue;
+		Buffers = &InBuffers[0];
 		Flags = EBufferFlag::Initial | EBufferFlag::Dirty;
+		OwnsMemory = false;
+	}
+
+	/** Destructor. */
+	~TTripleBuffer()
+	{
+		if (OwnsMemory)
+		{
+			delete[] Buffers;
+		}
 	}
 
 public:
@@ -117,7 +131,7 @@ public:
 	 * @return Reference to the read buffer's value.
 	 * @see SwapRead, Write
 	 */
-	ItemType& Read()
+	BufferType& Read()
 	{
 		return Buffers[Flags & EBufferFlag::ReaderMask];
 	}
@@ -136,13 +150,13 @@ public:
 			return;
 		}
 
-		int32 OldFlags = 0;
+		int32 CurrentFlags = 0;
 
 		do
 		{
-			OldFlags = Flags;
+			CurrentFlags = Flags;
 		}
-		while (FPlatformAtomics::InterlockedCompareExchange(&Flags, SwapReadWithTempFlags(OldFlags), OldFlags) != OldFlags);
+		while (FPlatformAtomics::InterlockedCompareExchange(&Flags, SwapReadWithTempFlags(CurrentFlags), CurrentFlags) != CurrentFlags);
 	}
 
 public:
@@ -152,7 +166,7 @@ public:
 	 *
 	 * @return Reference to write buffer.
 	 */
-	ItemType& GetWriteBuffer()
+	BufferType& GetWriteBuffer()
 	{
 		return Buffers[(Flags & EBufferFlag::WriterMask) >> EBufferFlag::WriterShift];
 	}
@@ -164,13 +178,13 @@ public:
 	 */
 	void SwapWriteBuffers()
 	{
-		int32 OldFlags = 0;
+		int32 CurrentFlags = 0;
 
 		do
 		{
-			OldFlags = Flags;
+			CurrentFlags = Flags;
 		}
-		while (FPlatformAtomics::InterlockedCompareExchange(&Flags, SwapWriteWithTempFlags(OldFlags), OldFlags) != OldFlags);
+		while (FPlatformAtomics::InterlockedCompareExchange(&Flags, SwapWriteWithTempFlags(CurrentFlags), CurrentFlags) != CurrentFlags);
 	}
 
 	/**
@@ -179,21 +193,18 @@ public:
 	 * @param Value The value to write.
 	 * @see SwapWrite, Read
 	 */
-	void Write(const ItemType Value)
+	void Write(const BufferType Value)
 	{
 		Buffers[(Flags & EBufferFlag::WriterMask) >> EBufferFlag::WriterShift] = Value;
 	}
 
 public:
 
-	/**
-	 * Get read-only access to the underlying buffers.
-	 *
-	 * @return Pointer to the first element of the buffer array,.
-	 */
-	ItemType* GetBuffers()
+	/** Reset the buffer. */
+	void Reset()
 	{
-		return Buffers;
+		Flags = EBufferFlag::Initial;
+		FPlatformMisc::MemoryBarrier();
 	}
 
 	/**
@@ -202,7 +213,7 @@ public:
 	 * @return Reference to the buffer.
 	 * @see SwapRead, Read, WriteAndSwap
 	 */
-	const ItemType& SwapAndRead()
+	const BufferType& SwapAndRead()
 	{
 		SwapReadBuffers();
 		return Read();
@@ -214,12 +225,22 @@ public:
 	 * @param Value The value to write into the buffer.
 	 * @see SwapAndRead, SwapWrite, Write
 	 */
-	void WriteAndSwap(const ItemType Value)
+	void WriteAndSwap(const BufferType Value)
 	{
 		Write(Value);
 		SwapWriteBuffers();
 	}
 
+protected:
+
+	/** Initialize the triple buffer. */
+	void Initialize()
+	{
+		Buffers = new BufferType[3];
+		Flags = EBufferFlag::Initial;
+		OwnsMemory = true;
+	}
+
 private:
 
 	/** Swaps the read and temp buffer indices in the Flags field. */
@@ -245,8 +266,11 @@ private:
 private:
 
 	/** The three buffers. */
-	ItemType Buffers[3];
+	BufferType* Buffers;
 
 	/** Buffer access flags. */
 	MS_ALIGN(16) int32 volatile Flags GCC_ALIGN(16);
+
+	/** Whether this instance owns the buffer memory. */
+	bool OwnsMemory;
 };
diff --git a/Engine/Source/Runtime/Core/Public/Containers/UnrealString.h b/Engine/Source/Runtime/Core/Public/Containers/UnrealString.h
--- a/Engine/Source/Runtime/Core/Public/Containers/UnrealString.h
+++ b/Engine/Source/Runtime/Core/Public/Containers/UnrealString.h
@@ -1509,7 +1509,7 @@ public:
 	/**
 	 * Returns a copy of this string, with the characters in reverse order
 	 */
-	FORCEINLINE FString Reverse() const;
+	FString Reverse() const;
 
 	/**
 	 * Reverses the order of characters in this string
diff --git a/Engine/Source/Runtime/Core/Public/Math/Color.h b/Engine/Source/Runtime/Core/Public/Math/Color.h
--- a/Engine/Source/Runtime/Core/Public/Math/Color.h
+++ b/Engine/Source/Runtime/Core/Public/Math/Color.h
@@ -11,6 +11,7 @@ struct FColor;
 struct FVector;
 class FFloat16Color;
 
+
 /**
  * Enum for the different kinds of gamma spaces we expect to need to convert from/to.
  */
@@ -291,7 +292,6 @@ struct FLinearColor
 	 * @param	From		The color and alpha to interpolate from as linear RGBA
 	 * @param	To			The color and alpha to interpolate to as linear RGBA
 	 * @param	Progress	Scalar interpolation amount (usually between 0.0 and 1.0 inclusive)
-	 * 
 	 * @return	The interpolated color in linear RGB space along with the interpolated alpha value
 	 */
 	static CORE_API FLinearColor LerpUsingHSV( const FLinearColor& From, const FLinearColor& To, const float Progress );
@@ -316,7 +316,7 @@ struct FLinearColor
 	/**
 	 * Returns the maximum value in this color structure
 	 *
-	 * @return	The maximum color channel value
+	 * @return The maximum color channel value
 	 */
 	FORCEINLINE float GetMax() const
 	{
@@ -332,7 +332,7 @@ struct FLinearColor
 	/**
 	 * Returns the minimum value in this color structure
 	 *
-	 * @return	The minimum color channel value
+	 * @return The minimum color channel value
 	 */
 	FORCEINLINE float GetMin() const
 	{
@@ -353,8 +353,7 @@ struct FLinearColor
 	 * Initialize this Color based on an FString. The String is expected to contain R=, G=, B=, A=.
 	 * The FLinearColor will be bogus when InitFromString returns false.
 	 *
-	 * @param	InSourceString	FString containing the color values.
-	 *
+	 * @param InSourceString FString containing the color values.
 	 * @return true if the R,G,B values were read successfully; false otherwise.
 	 */
 	bool InitFromString( const FString& InSourceString )
@@ -472,9 +471,7 @@ public:
 	 * Supported formats are: RGB, RRGGBB, RRGGBBAA, #RGB, #RRGGBB, #RRGGBBAA
 	 *
 	 * @param HexString - The hexadecimal string.
-	 *
 	 * @return The corresponding color value.
-	 *
 	 * @see ToHex
 	 */
 	static CORE_API FColor FromHex( const FString& HexString );
@@ -518,10 +515,8 @@ public:
 	 *
 	 * The format of the string is RRGGBBAA.
 	 *
-	 * @result Hexadecimal string.
-	 *
-	 * @see FromHex
-	 * @see ToString
+	 * @return Hexadecimal string.
+	 * @see FromHex, ToString
 	 */
 	FORCEINLINE FString ToHex() const
 	{
@@ -531,8 +526,7 @@ public:
 	/**
 	 * Converts this color value to a string.
 	 *
-	 * @result The string representation.
-	 *
+	 * @return The string representation.
 	 * @see ToHex
 	 */
 	FORCEINLINE FString ToString() const
@@ -545,7 +539,6 @@ public:
 	 * The FColor will be bogus when InitFromString returns false.
 	 *
 	 * @param	InSourceString	FString containing the color values.
-	 *
 	 * @return true if the R,G,B values were read successfully; false otherwise.
 	 */
 	bool InitFromString( const FString& InSourceString )
@@ -621,11 +614,13 @@ private:
 	explicit FColor(const FLinearColor& LinearColor);
 };
 
+
 FORCEINLINE uint32 GetTypeHash( const FColor& Color )
 {
 	return Color.DWColor();
 }
 
+
 FORCEINLINE uint32 GetTypeHash( const FLinearColor& LinearColor )
 {
 	// Note: this assumes there's no padding in FLinearColor that could contain uncompared data.
@@ -648,12 +643,15 @@ struct FDXTColor565
 {
 	/** Blue component, 5 bit. */
 	uint16 B:5;
+
 	/** Green component, 6 bit. */
 	uint16 G:6;
+
 	/** Red component, 5 bit */
 	uint16 R:5;
 };
 
+
 /**
  * Helper struct for a 16 bit 565 color of a DXT1/3/5 block.
  */
@@ -668,6 +666,7 @@ struct FDXTColor16
 	};
 };
 
+
 /**
  * Structure encompassing single DXT1 block.
  */
@@ -683,6 +682,7 @@ struct FDXT1
 	uint32 Indices;
 };
 
+
 /**
  * Structure encompassing single DXT5 block
  */
@@ -694,9 +694,9 @@ struct FDXT5
 	FDXT1	DXT1;
 };
 
+
 // Make DXT helpers act like PODs
 template <> struct TIsPODType<FDXT1> { enum { Value = true }; };
 template <> struct TIsPODType<FDXT5> { enum { Value = true }; };
 template <> struct TIsPODType<FDXTColor16> { enum { Value = true }; };
 template <> struct TIsPODType<FDXTColor565> { enum { Value = true }; };
-
diff --git a/Engine/Source/Runtime/Engine/Classes/Animation/MorphTarget.h b/Engine/Source/Runtime/Engine/Classes/Animation/MorphTarget.h
--- a/Engine/Source/Runtime/Engine/Classes/Animation/MorphTarget.h
+++ b/Engine/Source/Runtime/Engine/Classes/Animation/MorphTarget.h
@@ -1,9 +1,15 @@
 // Copyright 1998-2016 Epic Games, Inc. All Rights Reserved.
 
 #pragma once
+
+#include "PackedNormal.h"
 #include "MorphTarget.generated.h"
 
+
+class FStaticLODModel;
 class USkeletalMesh;
+class UStaticMesh;
+
 
 /**
 * Morph mesh vertex data used for comparisons and importing
@@ -11,9 +17,10 @@ class USkeletalMesh;
 struct FMorphMeshVertexRaw
 {
 	FVector Position;
-	FVector TanX,TanY,TanZ;
+	FVector TanX, TanY, TanZ;
 };
 
+
 /**
 * Converts a mesh to raw vertex data used to generate a morph target mesh
 */
@@ -22,18 +29,17 @@ class FMorphMeshRawSource
 public:
 	/** vertex data used for comparisons */
 	TArray<FMorphMeshVertexRaw> Vertices;
+
 	/** index buffer used for comparison */
 	TArray<uint32> Indices;
+
 	/** indices to original imported wedge points */
 	TArray<uint32> WedgePointIndices;
 
 	/** Constructor (default) */
-	ENGINE_API FMorphMeshRawSource() 
-	{
-	}
-
-	ENGINE_API FMorphMeshRawSource( USkeletalMesh* SrcMesh, int32 LODIndex=0 );
-	ENGINE_API FMorphMeshRawSource( UStaticMesh* SrcMesh, int32 LODIndex=0 );
+	ENGINE_API FMorphMeshRawSource() { }
+	ENGINE_API FMorphMeshRawSource( USkeletalMesh* SrcMesh, int32 LODIndex = 0 );
+	ENGINE_API FMorphMeshRawSource( UStaticMesh* SrcMesh, int32 LODIndex = 0 );
 	ENGINE_API FMorphMeshRawSource( FStaticLODModel& LODModel );
 
 	ENGINE_API bool IsValidTarget( const FMorphMeshRawSource& Target ) const;
@@ -42,16 +48,18 @@ private:
 	void Initialize(FStaticLODModel& LODModel);
 };
 
+
 /** Morph mesh vertex data used for rendering */
 struct FMorphTargetDelta
 {
 	/** change in position */
-	FVector			PositionDelta;
+	FVector PositionDelta;
+
 	/** Tangent basis normal */
-	FVector			TangentZDelta;
+	FVector TangentZDelta;
 
 	/** index of source vertex to apply deltas to */
-	uint32			SourceIdx;
+	uint32 SourceIdx;
 
 	/** pipe operator */
 	friend FArchive& operator<<(FArchive& Ar, FMorphTargetDelta& V)
@@ -59,7 +67,8 @@ struct FMorphTargetDelta
 		if (Ar.UE4Ver() < VER_UE4_MORPHTARGET_CPU_TANGENTZDELTA_FORMATCHANGE)
 		{
 			/** old format of change in tangent basis normal */
-			FPackedNormal	TangentZDelta_DEPRECATED;
+			FPackedNormal TangentZDelta_DEPRECATED;
+
 			if (Ar.IsSaving())
 			{
 				TangentZDelta_DEPRECATED = FPackedNormal(V.TangentZDelta);
@@ -80,6 +89,7 @@ struct FMorphTargetDelta
 	}
 };
 
+
 /**
 * Mesh data for a single LOD model of a morph target
 */
@@ -87,14 +97,16 @@ struct FMorphTargetLODModel
 {
 	/** vertex data for a single LOD morph mesh */
 	TArray<FMorphTargetDelta> Vertices;
+
 	/** number of original verts in the base mesh */
 	int32 NumBaseMeshVerts;
+	
 	/** Get Resource Size */
 	SIZE_T GetResourceSize() const;
 
 	FMorphTargetLODModel()
 		: NumBaseMeshVerts(0)
-	{}
+	{ }
 
 	/** pipe operator */
 	friend FArchive& operator<<(FArchive& Ar, FMorphTargetLODModel& M)
@@ -104,23 +116,23 @@ struct FMorphTargetLODModel
 	}
 };
 
+
 UCLASS(hidecategories=Object, MinimalAPI)
-class UMorphTarget : public UObject
+class UMorphTarget
+	: public UObject
 {
 	GENERATED_UCLASS_BODY()
 
 public:
+
 	/** USkeletalMesh that this vertex animation works on. */
 	UPROPERTY(AssetRegistrySearchable)
 	class USkeletalMesh* BaseSkelMesh;
 
 	/** morph mesh vertex data for each LOD */
-	TArray<FMorphTargetLODModel>	MorphLODModels;
+	TArray<FMorphTargetLODModel> MorphLODModels;
 
-	//~ Begin UObject Interface.
-	virtual void Serialize( FArchive& Ar ) override;
-	virtual SIZE_T GetResourceSize(EResourceSizeMode::Type Mode) override;
-	//~ Begin UObject Interface.
+public:
 
 	/** Post process after importing **/
 	ENGINE_API void PostProcess( USkeletalMesh * NewMesh, const FMorphMeshRawSource& BaseSource, const FMorphMeshRawSource& TargetSource, int32 LODIndex, bool bCompareNormal );
@@ -131,16 +143,23 @@ public:
 	FMorphTargetDelta* GetMorphTargetDelta(int32 LODIndex, int32& OutNumDeltas);
 	bool HasDataForLOD(int32 LODIndex);
 
+public:
+
+	//~ UObject interface
+
+	virtual void Serialize(FArchive& Ar) override;
+	virtual SIZE_T GetResourceSize(EResourceSizeMode::Type Mode) override;
+
 private:
+
 	/**
 	* Generate the streams for this morph target mesh using
 	* a base mesh and a target mesh to find the positon differences
 	* and other vertex attributes.
 	*
-	* @param	BaseSource - source mesh for comparing position differences
-	* @param	TargetSource - final target vertex positions/attributes 
-	* @param	LODIndex - level of detail to use for the geometry
+	* @param BaseSource Source mesh for comparing position differences
+	* @param TargetSource Final target vertex positions/attributes 
+	* @param LODIndex Level of detail to use for the geometry
 	*/
 	void CreateMorphMeshStreams( const FMorphMeshRawSource& BaseSource, const FMorphMeshRawSource& TargetSource, int32 LODIndex, bool bCompareNormal );
 };
-
diff --git a/Engine/Source/Runtime/Engine/Classes/Curves/CurveBase.h b/Engine/Source/Runtime/Engine/Classes/Curves/CurveBase.h
--- a/Engine/Source/Runtime/Engine/Classes/Curves/CurveBase.h
+++ b/Engine/Source/Runtime/Engine/Classes/Curves/CurveBase.h
@@ -2,761 +2,11 @@
 
 #pragma once
 
+#include "Curves/CurveOwnerInterface.h"
+#include "Curves/RichCurve.h"
 #include "CurveBase.generated.h"
 
 
-/* Basic types
- *****************************************************************************/
-
-/**
- * Key handles are used to keep a handle to a key. They are completely transient.
- */
-struct FKeyHandle
-{
-	ENGINE_API FKeyHandle();
-
-	bool operator ==(const FKeyHandle& Other) const
-	{
-		return Index == Other.Index;
-	}
-	
-	bool operator !=(const FKeyHandle& Other) const
-	{
-		return Index != Other.Index;
-	}
-	
-	friend uint32 GetTypeHash(const FKeyHandle& Handle)
-	{
-		return GetTypeHash(Handle.Index);
-	}
-
-	friend FArchive& operator<<(FArchive& Ar,FKeyHandle& Handle)
-	{
-		Ar << Handle.Index;
-		return Ar;
-	}
-
-private:
-
-	uint32 Index;
-};
-
-
-/**
- * Represents a mapping of key handles to key index that may be serialized 
- */
-USTRUCT()
-struct FKeyHandleMap
-{
-	GENERATED_USTRUCT_BODY()
-public:
-	FKeyHandleMap() {}
-
-	// This struct is not copyable.  This must be public or because derived classes are allowed to be copied
-	FKeyHandleMap( const FKeyHandleMap& Other ) {}
-	void operator=(const FKeyHandleMap& Other) {}
-
-	/** TMap functionality */
-	void Add( const FKeyHandle& InHandle, int32 InIndex );
-	void Empty();
-	void Remove( const FKeyHandle& InHandle );
-	const int32* Find( const FKeyHandle& InHandle ) const;
-	const FKeyHandle* FindKey( int32 KeyIndex ) const;
-	int32 Num() const;
-	TMap<FKeyHandle, int32>::TConstIterator CreateConstIterator() const;
-	TMap<FKeyHandle, int32>::TIterator CreateIterator();
-
-	/** ICPPStructOps implementation */
-	bool Serialize(FArchive& Ar);
-	bool operator==(const FKeyHandleMap& Other) const;
-	bool operator!=(const FKeyHandleMap& Other) const;
-
-	friend FArchive& operator<<(FArchive& Ar,FKeyHandleMap& P)
-	{
-		P.Serialize(Ar);
-		return Ar;
-	}
-
-private:
-
-	TMap<FKeyHandle, int32> KeyHandlesToIndices;
-};
-
-
-template<>
-struct TStructOpsTypeTraits<FKeyHandleMap>
-	: public TStructOpsTypeTraitsBase
-{
-	enum 
-	{
-		WithSerializer = true,
-		WithCopy = false,
-		WithIdenticalViaEquality = true,
-	};
-};
-
-
-/* Index curve
- *****************************************************************************/
-
-/**
- * A curve base class which enables key handles to index lookups.
- *
- * @todo sequencer: Some heavy refactoring can be done here. Much more stuff can go in this base class.
- */
-USTRUCT()
-struct ENGINE_API FIndexedCurve
-{
-	GENERATED_USTRUCT_BODY()
-
-public:
-
-	FIndexedCurve() { }
-
-	/** Get number of keys in curve. */
-	virtual int32 GetNumKeys() const PURE_VIRTUAL(FIndexedCurve::GetNumKeys, return 0;);
-
-	/** Const iterator for the handles. */
-	TMap<FKeyHandle, int32>::TConstIterator GetKeyHandleIterator() const;
-	
-	/** Gets the index of a handle, checks if the key handle is valid first. */
-	int32 GetIndexSafe(FKeyHandle KeyHandle) const;
-
-	/** Checks to see if the key handle is valid for this curve. */
-	virtual bool IsKeyHandleValid(FKeyHandle KeyHandle) const;
-
-protected:
-
-	/** Internal tool to get a handle from an index. */
-	FKeyHandle GetKeyHandle(int32 KeyIndex) const;
-	
-	/** Gets the index of a handle. */
-	int32 GetIndex(FKeyHandle KeyHandle) const;
-
-	/** Makes sure our handles are all valid and correct. */
-	void EnsureIndexHasAHandle(int32 KeyIndex) const;
-	void EnsureAllIndicesHaveHandles() const;
-
-protected:
-
-	/** Map of which key handles go to which indices. */
-	UPROPERTY(transient)
-	mutable FKeyHandleMap KeyHandlesToIndices;
-};
-
-
-/* Name curves
- *****************************************************************************/
-
-/**
- * One key in a curve of FNames.
- */
-USTRUCT()
-struct ENGINE_API FNameCurveKey
-{
-	GENERATED_USTRUCT_BODY()
-
-	/** Time at this key */
-	UPROPERTY(EditAnywhere, Category="Key")
-	float Time;
-
-	/** Value at this key */
-	UPROPERTY(EditAnywhere, Category="Key")
-	FName Value;
-
-	/** Default constructor. */
-	FNameCurveKey()
-		: Time(0.0f)
-		, Value(NAME_None)
-	{ }
-
-	/** Creates and initializes a new instance. */
-	FNameCurveKey(float InTime, const FName& InValue)
-		: Time(InTime)
-		, Value(InValue)
-	{ }
-
-public:
-
-	// TStructOpsTypeTraits interface
-
-	bool operator==(const FNameCurveKey& Other) const;
-	bool operator!=(const FNameCurveKey& Other) const;
-	bool Serialize(FArchive& Ar);
-
-	/** Serializes a name curve key from or into an archive. */
-	friend FArchive& operator<<(FArchive& Ar, FNameCurveKey& Key)
-	{
-		Key.Serialize(Ar);
-		return Ar;
-	}
-};
-
-
-template<>
-struct TIsPODType<FNameCurveKey>
-{
-	enum { Value = true };
-};
-
-
-template<>
-struct TStructOpsTypeTraits<FNameCurveKey>
-	: public TStructOpsTypeTraitsBase
-{
-	enum
-	{
-		WithSerializer = true,
-		WithCopy = false,
-		WithIdenticalViaEquality = true,
-	};
-};
-
-
-/**
- * Implements a curve of FNames.
- */
-USTRUCT()
-struct ENGINE_API FNameCurve
-	: public FIndexedCurve
-{
-	GENERATED_USTRUCT_BODY()
-
-	/** Virtual destructor. */
-	virtual ~FNameCurve() { }
-
-public:
-
-	/**
-	  * Add a new key to the curve with the supplied Time and Value. Returns the handle of the new key.
-	  * 
-	  * @param InTime The time at which to add the key.
-	  * @param InValue The value of the key.
-	  * @param KeyHandle Optional handle for the new key.
-	  */
-	FKeyHandle AddKey(float InTime, const FName& InValue, FKeyHandle KeyHandle = FKeyHandle());
-
-	/**
-	 * Remove the specified key from the curve.
-	 *
-	 * @param KeyHandle Handle to the key to remove.
-	 */
-	void DeleteKey(FKeyHandle KeyHandle);
-
-	/**
-	 * Finds a key a the specified time.
-	 *
-	 * @param KeyTime The time at which to find the key.
-	 * @param KeyTimeTolerance The key time tolerance to use for equality.
-	 * @return A handle to the key, or invalid key handle if not found.
-	 */
-	FKeyHandle FindKey(float KeyTime, float KeyTimeTolerance = KINDA_SMALL_NUMBER) const;
-
-	/**
-	 * Get a key.
-	 *
-	 * @param KeyHandle The handle of the key to get.
-	 * @return The key.
-	 */
-	FNameCurveKey& GetKey(FKeyHandle KeyHandle);
-	FNameCurveKey GetKey(FKeyHandle KeyHandle) const;
-
-	/**
-	 * Read-only access to the key collection.
-	 *
-	 * @return Collection of keys.
-	 */
-	const TArray<FNameCurveKey>& GetKeys() const
-	{
-		return Keys;
-	}
-
-	/**
-	 * Get the time for the Key with the specified index.
-	 *
-	 * @param KeyHandle Handle to the key whose time to get.
-	 * @return The key's time.
-	 */
-	float GetKeyTime(FKeyHandle KeyHandle) const;
-
-	/**
-	 * Move a key to a new time.
-	 *
-	 * This may change the index of the key, so the new key index is returned.
-	 *
-	 * @param KeyHandle The handle of the key to change.
-	 * @param NewTime The new time to set on the key.
-	 */
-	FKeyHandle SetKeyTime(FKeyHandle KeyHandle, float NewTime);
-
-	/**
-	 * Finds the key at InTime, and updates its value. If it can't find the key within the KeyTimeTolerance, it adds one at that time.
-	 *
-	 * @param InTime The time at which the key should be added or updated.
-	 * @param InValue The value of the key.
-	 * @param KeyTimeTolerance The tolerance used for key time equality.
-	 */
-	FKeyHandle UpdateOrAddKey(float InTime, const FName& InValue, float KeyTimeTolerance = KINDA_SMALL_NUMBER);
-
-public:
-
-	/** Shifts all keys forwards or backwards in time by an even amount, preserving order. */
-	void ShiftCurve(float DeltaTime);
-	void ShiftCurve(float DeltaTime, TSet<FKeyHandle>& KeyHandles);
-	
-	/** Scales all keys about an origin, preserving order. */
-	void ScaleCurve(float ScaleOrigin, float ScaleFactor);
-	void ScaleCurve(float ScaleOrigin, float ScaleFactor, TSet<FKeyHandle>& KeyHandles);
-
-public:
-
-	// FIndexedCurve interface
-
-	virtual int32 GetNumKeys() const override;
-	virtual bool IsKeyHandleValid(FKeyHandle KeyHandle) const override;
-
-public:
-
-	/** Sorted array of keys */
-	UPROPERTY(EditAnywhere, EditFixedSize, Category="Curve")
-	TArray<FNameCurveKey> Keys;
-};
-
-
-/* Rich curves
- *****************************************************************************/
-
-/** Method of interpolation between this key and the next. */
-UENUM()
-enum ERichCurveInterpMode
-{
-	RCIM_Linear,
-	RCIM_Constant,
-	RCIM_Cubic,
-	RCIM_None
-};
-
-
-/** If using RCIM_Cubic, this enum describes how the tangents should be controlled in editor. */
-UENUM()
-enum ERichCurveTangentMode
-{
-	RCTM_Auto,
-	RCTM_User,
-	RCTM_Break,
-	RCTM_None
-};
-
-
-/** Enumerates tangent weight modes. */
-UENUM()
-enum ERichCurveTangentWeightMode
-{
-	RCTWM_WeightedNone,
-	RCTWM_WeightedArrive,
-	RCTWM_WeightedLeave,
-	RCTWM_WeightedBoth
-};
-
-
-/** Enumerates extrapolation options. */
-UENUM()
-enum ERichCurveExtrapolation
-{
-	RCCE_Cycle,
-	RCCE_CycleWithOffset,
-	RCCE_Oscillate,
-	RCCE_Linear,
-	RCCE_Constant,
-	RCCE_None
-};
-
-
-/** One key in a rich, editable float curve */
-USTRUCT()
-struct ENGINE_API FRichCurveKey
-{
-	GENERATED_USTRUCT_BODY()
-
-	/** Interpolation mode between this key and the next */
-	UPROPERTY()
-	TEnumAsByte<ERichCurveInterpMode>			InterpMode;
-
-	/** Mode for tangents at this key */
-	UPROPERTY()
-	TEnumAsByte<ERichCurveTangentMode>			TangentMode;
-
-	/** If either tangent at this key is 'weighted' */
-	UPROPERTY()
-	TEnumAsByte<ERichCurveTangentWeightMode>	TangentWeightMode;
-
-	/** Time at this key */
-	UPROPERTY(EditAnywhere, Category="Key")
-	float Time;
-
-	/** Value at this key */
-	UPROPERTY(EditAnywhere, Category="Key")
-	float Value;
-
-	/** If RCIM_Cubic, the arriving tangent at this key */
-	UPROPERTY()
-	float ArriveTangent;
-
-	/** If RCTWM_WeightedArrive or RCTWM_WeightedBoth, the weight of the left tangent */
-	UPROPERTY()
-	float ArriveTangentWeight;
-
-	/** If RCIM_Cubic, the leaving tangent at this key */
-	UPROPERTY()
-	float LeaveTangent;
-
-	/** If RCTWM_WeightedLeave or RCTWM_WeightedBoth, the weight of the right tangent */
-	UPROPERTY()
-	float LeaveTangentWeight;
-
-	FRichCurveKey()
-		: InterpMode(RCIM_Linear)
-		, TangentMode(RCTM_Auto)
-		, TangentWeightMode(RCTWM_WeightedNone)
-		, Time(0.f)
-		, Value(0.f)
-		, ArriveTangent(0.f)
-		, ArriveTangentWeight(0.f)
-		, LeaveTangent(0.f)
-		, LeaveTangentWeight(0.f)
-	{ }
-
-	FRichCurveKey(float InTime, float InValue)
-		: InterpMode(RCIM_Linear)
-		, TangentMode(RCTM_Auto)
-		, TangentWeightMode(RCTWM_WeightedNone)
-		, Time(InTime)
-		, Value(InValue)
-		, ArriveTangent(0.f)
-		, ArriveTangentWeight(0.f)
-		, LeaveTangent(0.f)
-		, LeaveTangentWeight(0.f)
-	{ }
-
-	FRichCurveKey(float InTime, float InValue, float InArriveTangent, const float InLeaveTangent, ERichCurveInterpMode InInterpMode)
-		: InterpMode(InInterpMode)
-		, TangentMode(RCTM_Auto)
-		, TangentWeightMode(RCTWM_WeightedNone)
-		, Time(InTime)
-		, Value(InValue)
-		, ArriveTangent(InArriveTangent)
-		, ArriveTangentWeight(0.f)
-		, LeaveTangent(InLeaveTangent)
-		, LeaveTangentWeight(0.f)
-	{ }
-
-	/** Conversion constructor */
-	FRichCurveKey(const FInterpCurvePoint<float>& InPoint);
-	FRichCurveKey(const FInterpCurvePoint<FVector>& InPoint, int32 ComponentIndex);
-
-	/** ICPPStructOps interface */
-	bool Serialize(FArchive& Ar);
-	bool operator==(const FRichCurveKey& Other) const;
-	bool operator!=(const FRichCurveKey& Other) const;
-
-	friend FArchive& operator<<(FArchive& Ar, FRichCurveKey& P)
-	{
-		P.Serialize(Ar);
-		return Ar;
-	}
-};
-
-
-template<>
-struct TIsPODType<FRichCurveKey>
-{
-	enum { Value = true };
-};
-
-
-template<>
-struct TStructOpsTypeTraits<FRichCurveKey>
-	: public TStructOpsTypeTraitsBase
-{
-	enum
-	{
-		WithSerializer = true,
-		WithCopy = false,
-		WithIdenticalViaEquality = true,
-	};
-};
-
-
-/** A rich, editable float curve */
-USTRUCT()
-struct ENGINE_API FRichCurve
-	: public FIndexedCurve
-{
-	GENERATED_USTRUCT_BODY()
-
-	FRichCurve() 
-		: FIndexedCurve()
-		, PreInfinityExtrap(RCCE_Constant)
-		, PostInfinityExtrap(RCCE_Constant)
-		, DefaultValue(MAX_flt)
-	{ }
-
-	/** Virtual destructor. */
-	virtual ~FRichCurve() { }
-
-public:
-
-	/** Gets a copy of the keys, so indices and handles can't be meddled with */
-	TArray<FRichCurveKey> GetCopyOfKeys() const;
-
-	/** Gets a const reference of the keys, so indices and handles can't be meddled with */
-	const TArray<FRichCurveKey>& GetConstRefOfKeys() const;
-
-	/** Const iterator for the keys, so the indices and handles stay valid */
-	TArray<FRichCurveKey>::TConstIterator GetKeyIterator() const;
-	
-	/** Functions for getting keys based on handles */
-	FRichCurveKey& GetKey(FKeyHandle KeyHandle);
-	FRichCurveKey GetKey(FKeyHandle KeyHandle) const;
-	
-	/** Quick accessors for the first and last keys */
-	FRichCurveKey GetFirstKey() const;
-	FRichCurveKey GetLastKey() const;
-
-	/** Get the first key that matches any of the given key handles. */
-	FRichCurveKey* GetFirstMatchingKey(const TArray<FKeyHandle>& KeyHandles);
-
-	/** Get the next or previous key given the key handle */
-	FKeyHandle GetNextKey(FKeyHandle KeyHandle) const;
-	FKeyHandle GetPreviousKey(FKeyHandle KeyHandle) const;
-
-	/**
-	  * Add a new key to the curve with the supplied Time and Value. Returns the handle of the new key.
-	  * 
-	  * @param	bUnwindRotation		When true, the value will be treated like a rotation value in degrees, and will automatically be unwound to prevent flipping 360 degrees from the previous key 
-	  * @param  KeyHandle			Optionally can specify what handle this new key should have, otherwise, it'll make a new one
-	  */
-	FKeyHandle AddKey(float InTime, float InValue, const bool bUnwindRotation = false, FKeyHandle KeyHandle = FKeyHandle());
-
-	/**
-	 * Sets the keys with the supplied times and values. Expects that the times and values are already sorted.
-	 */
-	void SetKeys(const TArray<float>& InTimes, const TArray<float>& InValues);
-
-	/** Remove the specified key from the curve.*/
-	void DeleteKey(FKeyHandle KeyHandle);
-
-	/** Finds the key at InTime, and updates its value. If it can't find the key within the KeyTimeTolerance, it adds one at that time */
-	FKeyHandle UpdateOrAddKey(float InTime, float InValue, const bool bUnwindRotation = false, float KeyTimeTolerance = KINDA_SMALL_NUMBER);
-
-	/** Move a key to a new time. This may change the index of the key, so the new key index is returned. */
-	FKeyHandle SetKeyTime(FKeyHandle KeyHandle, float NewTime);
-
-	/** Get the time for the Key with the specified index. */
-	float GetKeyTime(FKeyHandle KeyHandle) const;
-
-	/** Finds a key a the specified time */
-	FKeyHandle FindKey( float KeyTime, float KeyTimeTolerance = KINDA_SMALL_NUMBER ) const;
-
-	/** Set the value of the specified key */
-	void SetKeyValue(FKeyHandle KeyHandle, float NewValue, bool bAutoSetTangents=true);
-
-	/** Returns the value of the specified key */
-	float GetKeyValue(FKeyHandle KeyHandle) const;
-
-	/** Set the default value of the curve */
-	void SetDefaultValue(float InDefaultValue) { DefaultValue = InDefaultValue; }
-
-	/** Get the default value for the curve */
-	float GetDefaultValue() const { return DefaultValue; }
-
-	/** Shifts all keys forwards or backwards in time by an even amount, preserving order */
-	void ShiftCurve(float DeltaTime);
-	void ShiftCurve(float DeltaTime, TSet<FKeyHandle>& KeyHandles);
-	
-	/** Scales all keys about an origin, preserving order */
-	void ScaleCurve(float ScaleOrigin, float ScaleFactor);
-	void ScaleCurve(float ScaleOrigin, float ScaleFactor, TSet<FKeyHandle>& KeyHandles);
-
-	/** Set the interp mode of the specified key */
-	void SetKeyInterpMode(FKeyHandle KeyHandle, ERichCurveInterpMode NewInterpMode);
-
-	/** Set the tangent mode of the specified key */
-	void SetKeyTangentMode(FKeyHandle KeyHandle, ERichCurveTangentMode NewTangentMode);
-
-	/** Set the tangent weight mode of the specified key */
-	void SetKeyTangentWeightMode(FKeyHandle KeyHandle, ERichCurveTangentWeightMode NewTangentWeightMode);
-
-	/** Get the interp mode of the specified key */
-	ERichCurveInterpMode GetKeyInterpMode(FKeyHandle KeyHandle) const;
-
-	/** Get the tangent mode of the specified key */
-	ERichCurveTangentMode GetKeyTangentMode(FKeyHandle KeyHandle) const;
-
-	/** Get range of input time values. Outside this region curve continues constantly the start/end values. */
-	void GetTimeRange(float& MinTime, float& MaxTime) const;
-
-	/** Get range of output values. */
-	void GetValueRange(float& MinValue, float& MaxValue) const;
-
-	/** Clear all keys. */
-	void Reset();
-
-	/** Remap InTime based on pre and post infinity extrapolation values */
-	void RemapTimeValue(float& InTime, float& CycleValueOffset) const;
-
-	/** Evaluate this rich curve at the specified time */
-	float Eval(float InTime, float InDefaultValue = 0.0f) const;
-
-	/** Auto set tangents for any 'auto' keys in curve */
-	void AutoSetTangents(float Tension = 0.f);
-
-	/** Resize curve length to the [MinTimeRange, MaxTimeRange] */
-	void ReadjustTimeRange(float NewMinTimeRange, float NewMaxTimeRange, bool bInsert/* whether insert or remove*/, float OldStartTime, float OldEndTime);
-
-	/** Determine if two RichCurves are the same */
-	bool operator == (const FRichCurve& Curve) const;
-
-	/** Bake curve given the sample rate */
-	void BakeCurve(float SampleRate);
-	void BakeCurve(float SampleRate, float FirstKeyTime, float LastKeyTime);
-
-	/** Remove redundant keys, comparing against Tolerance */
-	void RemoveRedundantKeys(float Tolerance);
-	void RemoveRedundantKeys(float Tolerance, float FirstKeyTime, float LastKeyTime);
-
-public:
-
-	// FIndexedCurve interface
-
-	virtual int32 GetNumKeys() const override;
-	virtual bool IsKeyHandleValid(FKeyHandle KeyHandle) const override;
-
-public:
-
-	/** Pre-infinity extrapolation state */
-	UPROPERTY()
-	TEnumAsByte<ERichCurveExtrapolation> PreInfinityExtrap;
-
-	/** Post-infinity extrapolation state */
-	UPROPERTY()
-	TEnumAsByte<ERichCurveExtrapolation> PostInfinityExtrap;
-
-	/** Sorted array of keys */
-	UPROPERTY(EditAnywhere, EditFixedSize, Category="Curve")
-	TArray<FRichCurveKey> Keys;
-
-	/** Default value */
-	UPROPERTY(EditAnywhere, Category="Curve")
-	float DefaultValue;
-};
-
-
-/* Curve editors
- *****************************************************************************/
-
-/**
- * Info about a curve to be edited.
- */
-template<class T>
-struct FRichCurveEditInfoTemplate
-{
-	/** Name of curve, used when displaying in editor. Can include comma's to allow tree expansion in editor */
-	FName CurveName;
-
-	/** Pointer to curves to be edited */
-	T CurveToEdit;
-
-	FRichCurveEditInfoTemplate()
-		: CurveName(NAME_None)
-		, CurveToEdit(nullptr)
-	{ }
-
-	FRichCurveEditInfoTemplate(T InCurveToEdit)
-		: CurveName(NAME_None)
-		, CurveToEdit(InCurveToEdit)
-	{ }
-
-	FRichCurveEditInfoTemplate(T InCurveToEdit, FName InCurveName)
-		: CurveName(InCurveName)
-		, CurveToEdit(InCurveToEdit)
-	{ }
-
-	inline bool operator==(const FRichCurveEditInfoTemplate<T>& Other) const
-	{
-		return Other.CurveName.IsEqual(CurveName) && Other.CurveToEdit == CurveToEdit;
-	}
-
-	uint32 GetTypeHash() const
-	{
-		return HashCombine(::GetTypeHash(CurveName), PointerHash(CurveToEdit));
-	}
-};
-
-
-template<class T>
-inline uint32 GetTypeHash( const FRichCurveEditInfoTemplate<T>& RichCurveEditInfo )
-{
-	return RichCurveEditInfo.GetTypeHash();
-}
-
-
-typedef FRichCurveEditInfoTemplate<FRichCurve*>			FRichCurveEditInfo;
-typedef FRichCurveEditInfoTemplate<const FRichCurve*>	FRichCurveEditInfoConst;
-
-
-/**
- * Interface you implement if you want the CurveEditor to be able to edit curves on you.
- */
-class ENGINE_API FCurveOwnerInterface
-{
-public:
-
-	virtual ~FCurveOwnerInterface() { }
-
-	/** Returns set of curves to edit. Must not release the curves while being edited. */
-	virtual TArray<FRichCurveEditInfoConst> GetCurves() const = 0;
-
-	/** Returns set of curves to query. Must not release the curves while being edited. */
-	virtual TArray<FRichCurveEditInfo> GetCurves() = 0;
-
-	/** Called to modify the owner of the curve */
-	virtual void ModifyOwner() = 0;
-
-	/** Called to make curve owner transactional */
-	virtual void MakeTransactional() = 0;
-
-	/** Called when any of the curves have been changed */
-	virtual void OnCurveChanged(const TArray<FRichCurveEditInfo>& ChangedCurveEditInfos) = 0;
-
-	/** Whether the curve represents a linear color */
-	virtual bool IsLinearColorCurve() const
-	{
-		return false;
-	}
-
-	/** Evaluate this color curve at the specified time */
-	virtual FLinearColor GetLinearColorValue(float InTime) const
-	{
-		return FLinearColor::Black;
-	}
-
-	/** @return True if the curve has any alpha keys */
-	virtual bool HasAnyAlphaKeys() const
-	{
-		return false;
-	}
-
-	/** Validates that a previously retrieved curve is still valid for editing. */
-	virtual bool IsValidCurve(FRichCurveEditInfo CurveInfo) = 0;
-
-	/** @return Color for this curve */
-	virtual FLinearColor GetCurveColor(FRichCurveEditInfo CurveInfo) const;
-};
-
-
-//////////////////////////////////////////////////////////////////////////
-
 /**
  * Defines a curve of interpolated points to evaluate over a given range
  */
@@ -788,7 +38,7 @@ class ENGINE_API UCurveBase
 
 public:
 
-	// FCurveOwnerInterface interface
+	//~ FCurveOwnerInterface interface
 
 	virtual TArray<FRichCurveEditInfoConst> GetCurves() const override
 	{
@@ -813,7 +63,7 @@ public:
 
 public:
 
-	// UObject overrides
+	//~ UObject interface
 
 #if WITH_EDITORONLY_DATA
 
@@ -825,129 +75,11 @@ public:
 	UPROPERTY(VisibleAnywhere, Instanced, Category=ImportSettings)
 	class UAssetImportData* AssetImportData;
 
+public:
+
 	/** The filename imported to create this object. Relative to this object's package, BaseDir() or absolute */
 	UPROPERTY()
 	FString ImportPath_DEPRECATED;
 
 #endif
 };
-
-
-//////////////////////////////////////////////////////////////////////////
-
-/** An integral key, which holds the key time and the key value */
-USTRUCT()
-struct FIntegralKey
-{
-	GENERATED_USTRUCT_BODY()
-public:
-
-	FIntegralKey(float InTime = 0.f, int32 InValue = 0)
-		: Time(InTime)
-		, Value(InValue)
-	{ }
-	
-	/** The keyed time */
-	UPROPERTY(EditAnywhere, Category=Key)
-	float Time;
-
-	/** The keyed integral value */
-	UPROPERTY(EditAnywhere, Category=Key)
-	int32 Value;
-};
-
-
-/** An integral curve, which holds the key time and the key value */
-USTRUCT()
-struct ENGINE_API FIntegralCurve
-	: public FIndexedCurve
-{
-	GENERATED_USTRUCT_BODY()
-
-public:
-
-	/** Default constructor. */
-	FIntegralCurve() 
-	: FIndexedCurve()
-	, DefaultValue(MAX_int32)
-	, bUseDefaultValueBeforeFirstKey(false)
-	{
-	}
-
-	/** Virtual destructor. */
-	virtual ~FIntegralCurve() { }
-	
-	/** Get number of keys in curve. */
-	virtual int32 GetNumKeys() const override;
-	
-	/** Checks to see if the key handle is valid for this curve */
-	virtual bool IsKeyHandleValid(FKeyHandle KeyHandle) const override;
-
-	/** Evaluates the value of an array of keys at a time */
-	int32 Evaluate(float Time, int32 InDefaultValue = 0) const;
-
-	/** Const iterator for the keys, so the indices and handles stay valid */
-	TArray<FIntegralKey>::TConstIterator GetKeyIterator() const;
-	
-	/**
-	  * Add a new key to the curve with the supplied Time and Value.
-	  * 
-	  * @param KeyHandle Optionally can specify what handle this new key should have, otherwise, it'll make a new one
-	  * @return The handle of the new key
-	  */
-	FKeyHandle AddKey( float InTime, int32 InValue, FKeyHandle KeyHandle = FKeyHandle() );
-	
-	/** Remove the specified key from the curve.*/
-	void DeleteKey(FKeyHandle KeyHandle);
-	
-	/** Finds the key at InTime, and updates its value. If it can't find the key within the KeyTimeTolerance, it adds one at that time */
-	FKeyHandle UpdateOrAddKey( float InTime, int32 Value, float KeyTimeTolerance = KINDA_SMALL_NUMBER );
-	
-	/** Move a key to a new time. This may change the index of the key, so the new key index is returned. */
-	FKeyHandle SetKeyTime(FKeyHandle KeyHandle, float NewTime);
-
-	/** Get the time for the Key with the specified index. */
-	float GetKeyTime(FKeyHandle KeyHandle) const;
-	
-	/** Set the default value for the curve */
-	void SetDefaultValue(int32 InDefaultValue) { DefaultValue = InDefaultValue; }
-
-	/** Get the default value for the curve */
-	int32 GetDefaultValue() const { return DefaultValue; }
-
-	/** Sets whether or not the default value should be used for evaluation for time values before the first key. */
-	void SetUseDefaultValueBeforeFirstKey(bool InbUseDefaultValueBeforeFirstKey) { bUseDefaultValueBeforeFirstKey = InbUseDefaultValueBeforeFirstKey; }
-	
-	/** Gets whether or not the default value should be used for evaluation for time values before the first key. */
-	bool GetUseDefaultValueBeforeFirstKey() const { return bUseDefaultValueBeforeFirstKey; }
-
-	/** Shifts all keys forwards or backwards in time by an even amount, preserving order */
-	void ShiftCurve(float DeltaTime);
-	void ShiftCurve(float DeltaTime, TSet<FKeyHandle>& KeyHandles);
-	
-	/** Scales all keys about an origin, preserving order */
-	void ScaleCurve(float ScaleOrigin, float ScaleFactor);
-	void ScaleCurve(float ScaleOrigin, float ScaleFactor, TSet<FKeyHandle>& KeyHandles);
-
-	/** Functions for getting keys based on handles */
-	FIntegralKey& GetKey(FKeyHandle KeyHandle);
-	FIntegralKey GetKey(FKeyHandle KeyHandle) const;
-
-	FKeyHandle FindKey(float KeyTime, float KeyTimeTolerance = KINDA_SMALL_NUMBER) const;
-
-	/** Gets the handle for the last key which is at or before the time requested.  If there are no keys at or before the requested time, an invalid handle is returned. */
-	FKeyHandle FindKeyBeforeOrAt(float KeyTime) const;
-
-private:
-
-	/** The keys, ordered by time */
-	UPROPERTY(EditAnywhere, Category="Curve")
-	TArray<FIntegralKey> Keys;
-
-	/** Default value */
-	UPROPERTY(EditAnywhere, Category="Curve")
-	int32 DefaultValue;
-
-	UPROPERTY()
-	bool bUseDefaultValueBeforeFirstKey;
-};
diff --git a/Engine/Source/Runtime/Engine/Classes/Curves/IndexedCurve.h b/Engine/Source/Runtime/Engine/Classes/Curves/IndexedCurve.h
new file mode 100644
--- /dev/null
+++ b/Engine/Source/Runtime/Engine/Classes/Curves/IndexedCurve.h
@@ -0,0 +1,53 @@
+#pragma once
+
+#include "Curves/KeyHandle.h"
+#include "IndexedCurve.generated.h"
+
+
+/**
+ * A curve base class which enables key handles to index lookups.
+ *
+ * @todo sequencer: Some heavy refactoring can be done here. Much more stuff can go in this base class.
+ */
+USTRUCT()
+struct ENGINE_API FIndexedCurve
+{
+	GENERATED_USTRUCT_BODY()
+
+public:
+
+	/** Default constructor. */
+	FIndexedCurve() { }
+
+public:
+
+	/** Gets the index of a handle, checks if the key handle is valid first. */
+	int32 GetIndexSafe(FKeyHandle KeyHandle) const;
+
+	/** Const iterator for the handles. */
+	TMap<FKeyHandle, int32>::TConstIterator GetKeyHandleIterator() const;
+
+	/** Get number of keys in curve. */
+	virtual int32 GetNumKeys() const PURE_VIRTUAL(FIndexedCurve::GetNumKeys, return 0;);
+
+	/** Checks to see if the key handle is valid for this curve. */
+	virtual bool IsKeyHandleValid(FKeyHandle KeyHandle) const;
+
+protected:
+
+	/** Makes sure our handles are all valid and correct. */
+	void EnsureAllIndicesHaveHandles() const;
+	void EnsureIndexHasAHandle(int32 KeyIndex) const;
+
+	/** Gets the index of a handle. */
+	int32 GetIndex(FKeyHandle KeyHandle) const;
+
+	/** Internal tool to get a handle from an index. */
+	FKeyHandle GetKeyHandle(int32 KeyIndex) const;
+	
+protected:
+
+	/** Map of which key handles go to which indices. */
+	UPROPERTY(transient)
+	mutable FKeyHandleMap KeyHandlesToIndices;
+};
diff --git a/Engine/Source/Runtime/Engine/Classes/Curves/IntegralCurve.h b/Engine/Source/Runtime/Engine/Classes/Curves/IntegralCurve.h
new file mode 100644
--- /dev/null
+++ b/Engine/Source/Runtime/Engine/Classes/Curves/IntegralCurve.h
@@ -0,0 +1,122 @@
+#pragma once
+
+#include "Curves/IndexedCurve.h"
+#include "Curves/KeyHandle.h"
+#include "IntegralCurve.generated.h"
+
+
+/** An integral key, which holds the key time and the key value */
+USTRUCT()
+struct FIntegralKey
+{
+	GENERATED_USTRUCT_BODY()
+public:
+
+	FIntegralKey(float InTime = 0.f, int32 InValue = 0)
+		: Time(InTime)
+		, Value(InValue)
+	{ }
+	
+	/** The keyed time */
+	UPROPERTY(EditAnywhere, Category=Key)
+	float Time;
+
+	/** The keyed integral value */
+	UPROPERTY(EditAnywhere, Category=Key)
+	int32 Value;
+};
+
+
+/** An integral curve, which holds the key time and the key value */
+USTRUCT()
+struct ENGINE_API FIntegralCurve
+	: public FIndexedCurve
+{
+	GENERATED_USTRUCT_BODY()
+
+public:
+
+	/** Default constructor. */
+	FIntegralCurve() 
+		: FIndexedCurve()
+		, DefaultValue(MAX_int32)
+		, bUseDefaultValueBeforeFirstKey(false)
+	{ }
+
+	/** Virtual destructor. */
+	virtual ~FIntegralCurve() { }
+	
+	/** Get number of keys in curve. */
+	virtual int32 GetNumKeys() const override;
+	
+	/** Checks to see if the key handle is valid for this curve */
+	virtual bool IsKeyHandleValid(FKeyHandle KeyHandle) const override;
+
+	/** Evaluates the value of an array of keys at a time */
+	int32 Evaluate(float Time, int32 InDefaultValue = 0) const;
+
+	/** Const iterator for the keys, so the indices and handles stay valid */
+	TArray<FIntegralKey>::TConstIterator GetKeyIterator() const;
+	
+	/**
+	  * Add a new key to the curve with the supplied Time and Value.
+	  * 
+	  * @param KeyHandle Optionally can specify what handle this new key should have, otherwise, it'll make a new one
+	  * @return The handle of the new key
+	  */
+	FKeyHandle AddKey( float InTime, int32 InValue, FKeyHandle KeyHandle = FKeyHandle() );
+	
+	/** Remove the specified key from the curve.*/
+	void DeleteKey(FKeyHandle KeyHandle);
+	
+	/** Finds the key at InTime, and updates its value. If it can't find the key within the KeyTimeTolerance, it adds one at that time */
+	FKeyHandle UpdateOrAddKey( float InTime, int32 Value, float KeyTimeTolerance = KINDA_SMALL_NUMBER );
+	
+	/** Move a key to a new time. This may change the index of the key, so the new key index is returned. */
+	FKeyHandle SetKeyTime(FKeyHandle KeyHandle, float NewTime);
+
+	/** Get the time for the Key with the specified index. */
+	float GetKeyTime(FKeyHandle KeyHandle) const;
+	
+	/** Set the default value for the curve */
+	void SetDefaultValue(int32 InDefaultValue) { DefaultValue = InDefaultValue; }
+
+	/** Get the default value for the curve */
+	int32 GetDefaultValue() const { return DefaultValue; }
+
+	/** Sets whether or not the default value should be used for evaluation for time values before the first key. */
+	void SetUseDefaultValueBeforeFirstKey(bool InbUseDefaultValueBeforeFirstKey) { bUseDefaultValueBeforeFirstKey = InbUseDefaultValueBeforeFirstKey; }
+	
+	/** Gets whether or not the default value should be used for evaluation for time values before the first key. */
+	bool GetUseDefaultValueBeforeFirstKey() const { return bUseDefaultValueBeforeFirstKey; }
+
+	/** Shifts all keys forwards or backwards in time by an even amount, preserving order */
+	void ShiftCurve(float DeltaTime);
+	void ShiftCurve(float DeltaTime, TSet<FKeyHandle>& KeyHandles);
+	
+	/** Scales all keys about an origin, preserving order */
+	void ScaleCurve(float ScaleOrigin, float ScaleFactor);
+	void ScaleCurve(float ScaleOrigin, float ScaleFactor, TSet<FKeyHandle>& KeyHandles);
+
+	/** Functions for getting keys based on handles */
+	FIntegralKey& GetKey(FKeyHandle KeyHandle);
+	FIntegralKey GetKey(FKeyHandle KeyHandle) const;
+
+	FKeyHandle FindKey(float KeyTime, float KeyTimeTolerance = KINDA_SMALL_NUMBER) const;
+
+	/** Gets the handle for the last key which is at or before the time requested.  If there are no keys at or before the requested time, an invalid handle is returned. */
+	FKeyHandle FindKeyBeforeOrAt(float KeyTime) const;
+
+private:
+
+	/** The keys, ordered by time */
+	UPROPERTY(EditAnywhere, Category="Curve")
+	TArray<FIntegralKey> Keys;
+
+	/** Default value */
+	UPROPERTY(EditAnywhere, Category="Curve")
+	int32 DefaultValue;
+
+	UPROPERTY()
+	bool bUseDefaultValueBeforeFirstKey;
+};
diff --git a/Engine/Source/Runtime/Engine/Classes/Curves/NameCurve.h b/Engine/Source/Runtime/Engine/Classes/Curves/NameCurve.h
new file mode 100644
--- /dev/null
+++ b/Engine/Source/Runtime/Engine/Classes/Curves/NameCurve.h
@@ -0,0 +1,180 @@
+#pragma once
+
+#include "Curves/IndexedCurve.h"
+#include "Curves/KeyHandle.h"
+#include "NameCurve.generated.h"
+
+
+/**
+ * One key in a curve of FNames.
+ */
+USTRUCT()
+struct ENGINE_API FNameCurveKey
+{
+	GENERATED_USTRUCT_BODY()
+
+	/** Time at this key */
+	UPROPERTY(EditAnywhere, Category="Key")
+	float Time;
+
+	/** Value at this key */
+	UPROPERTY(EditAnywhere, Category="Key")
+	FName Value;
+
+	/** Default constructor. */
+	FNameCurveKey()
+		: Time(0.0f)
+		, Value(NAME_None)
+	{ }
+
+	/** Creates and initializes a new instance. */
+	FNameCurveKey(float InTime, const FName& InValue)
+		: Time(InTime)
+		, Value(InValue)
+	{ }
+
+public:
+
+	// TStructOpsTypeTraits interface
+
+	bool operator==(const FNameCurveKey& Other) const;
+	bool operator!=(const FNameCurveKey& Other) const;
+	bool Serialize(FArchive& Ar);
+
+	/** Serializes a name curve key from or into an archive. */
+	friend FArchive& operator<<(FArchive& Ar, FNameCurveKey& Key)
+	{
+		Key.Serialize(Ar);
+		return Ar;
+	}
+};
+
+
+template<>
+struct TIsPODType<FNameCurveKey>
+{
+	enum { Value = true };
+};
+
+
+template<>
+struct TStructOpsTypeTraits<FNameCurveKey>
+	: public TStructOpsTypeTraitsBase
+{
+	enum
+	{
+		WithSerializer = true,
+		WithCopy = false,
+		WithIdenticalViaEquality = true,
+	};
+};
+
+
+/**
+ * Implements a curve of FNames.
+ */
+USTRUCT()
+struct ENGINE_API FNameCurve
+	: public FIndexedCurve
+{
+	GENERATED_USTRUCT_BODY()
+
+	/** Virtual destructor. */
+	virtual ~FNameCurve() { }
+
+public:
+
+	/**
+	  * Add a new key to the curve with the supplied Time and Value. Returns the handle of the new key.
+	  * 
+	  * @param InTime The time at which to add the key.
+	  * @param InValue The value of the key.
+	  * @param KeyHandle Optional handle for the new key.
+	  */
+	FKeyHandle AddKey(float InTime, const FName& InValue, FKeyHandle KeyHandle = FKeyHandle());
+
+	/**
+	 * Remove the specified key from the curve.
+	 *
+	 * @param KeyHandle Handle to the key to remove.
+	 */
+	void DeleteKey(FKeyHandle KeyHandle);
+
+	/**
+	 * Finds a key a the specified time.
+	 *
+	 * @param KeyTime The time at which to find the key.
+	 * @param KeyTimeTolerance The key time tolerance to use for equality.
+	 * @return A handle to the key, or invalid key handle if not found.
+	 */
+	FKeyHandle FindKey(float KeyTime, float KeyTimeTolerance = KINDA_SMALL_NUMBER) const;
+
+	/**
+	 * Get a key.
+	 *
+	 * @param KeyHandle The handle of the key to get.
+	 * @return The key.
+	 */
+	FNameCurveKey& GetKey(FKeyHandle KeyHandle);
+	FNameCurveKey GetKey(FKeyHandle KeyHandle) const;
+
+	/**
+	 * Read-only access to the key collection.
+	 *
+	 * @return Collection of keys.
+	 */
+	const TArray<FNameCurveKey>& GetKeys() const
+	{
+		return Keys;
+	}
+
+	/**
+	 * Get the time for the Key with the specified index.
+	 *
+	 * @param KeyHandle Handle to the key whose time to get.
+	 * @return The key's time.
+	 */
+	float GetKeyTime(FKeyHandle KeyHandle) const;
+
+	/**
+	 * Move a key to a new time.
+	 *
+	 * This may change the index of the key, so the new key index is returned.
+	 *
+	 * @param KeyHandle The handle of the key to change.
+	 * @param NewTime The new time to set on the key.
+	 */
+	FKeyHandle SetKeyTime(FKeyHandle KeyHandle, float NewTime);
+
+	/**
+	 * Finds the key at InTime, and updates its value. If it can't find the key within the KeyTimeTolerance, it adds one at that time.
+	 *
+	 * @param InTime The time at which the key should be added or updated.
+	 * @param InValue The value of the key.
+	 * @param KeyTimeTolerance The tolerance used for key time equality.
+	 */
+	FKeyHandle UpdateOrAddKey(float InTime, const FName& InValue, float KeyTimeTolerance = KINDA_SMALL_NUMBER);
+
+public:
+
+	/** Shifts all keys forwards or backwards in time by an even amount, preserving order. */
+	void ShiftCurve(float DeltaTime);
+	void ShiftCurve(float DeltaTime, TSet<FKeyHandle>& KeyHandles);
+	
+	/** Scales all keys about an origin, preserving order. */
+	void ScaleCurve(float ScaleOrigin, float ScaleFactor);
+	void ScaleCurve(float ScaleOrigin, float ScaleFactor, TSet<FKeyHandle>& KeyHandles);
+
+public:
+
+	// FIndexedCurve interface
+
+	virtual int32 GetNumKeys() const override;
+	virtual bool IsKeyHandleValid(FKeyHandle KeyHandle) const override;
+
+public:
+
+	/** Sorted array of keys */
+	UPROPERTY(EditAnywhere, EditFixedSize, Category="Curve")
+	TArray<FNameCurveKey> Keys;
+};
diff --git a/Engine/Source/Runtime/Engine/Classes/Curves/RichCurve.h b/Engine/Source/Runtime/Engine/Classes/Curves/RichCurve.h
new file mode 100644
--- /dev/null
+++ b/Engine/Source/Runtime/Engine/Classes/Curves/RichCurve.h
@@ -0,0 +1,392 @@
+#pragma once
+
+#include "Curves/IndexedCurve.h"
+#include "Curves/KeyHandle.h"
+#include "RichCurve.generated.h"
+
+
+/** Method of interpolation between this key and the next. */
+UENUM()
+enum ERichCurveInterpMode
+{
+	RCIM_Linear,
+	RCIM_Constant,
+	RCIM_Cubic,
+	RCIM_None
+};
+
+
+/** If using RCIM_Cubic, this enum describes how the tangents should be controlled in editor. */
+UENUM()
+enum ERichCurveTangentMode
+{
+	RCTM_Auto,
+	RCTM_User,
+	RCTM_Break,
+	RCTM_None
+};
+
+
+/** Enumerates tangent weight modes. */
+UENUM()
+enum ERichCurveTangentWeightMode
+{
+	RCTWM_WeightedNone,
+	RCTWM_WeightedArrive,
+	RCTWM_WeightedLeave,
+	RCTWM_WeightedBoth
+};
+
+
+/** Enumerates extrapolation options. */
+UENUM()
+enum ERichCurveExtrapolation
+{
+	RCCE_Cycle,
+	RCCE_CycleWithOffset,
+	RCCE_Oscillate,
+	RCCE_Linear,
+	RCCE_Constant,
+	RCCE_None
+};
+
+
+/** One key in a rich, editable float curve */
+USTRUCT()
+struct ENGINE_API FRichCurveKey
+{
+	GENERATED_USTRUCT_BODY()
+
+	/** Interpolation mode between this key and the next */
+	UPROPERTY()
+	TEnumAsByte<ERichCurveInterpMode> InterpMode;
+
+	/** Mode for tangents at this key */
+	UPROPERTY()
+	TEnumAsByte<ERichCurveTangentMode> TangentMode;
+
+	/** If either tangent at this key is 'weighted' */
+	UPROPERTY()
+	TEnumAsByte<ERichCurveTangentWeightMode> TangentWeightMode;
+
+	/** Time at this key */
+	UPROPERTY(EditAnywhere, Category="Key")
+	float Time;
+
+	/** Value at this key */
+	UPROPERTY(EditAnywhere, Category="Key")
+	float Value;
+
+	/** If RCIM_Cubic, the arriving tangent at this key */
+	UPROPERTY()
+	float ArriveTangent;
+
+	/** If RCTWM_WeightedArrive or RCTWM_WeightedBoth, the weight of the left tangent */
+	UPROPERTY()
+	float ArriveTangentWeight;
+
+	/** If RCIM_Cubic, the leaving tangent at this key */
+	UPROPERTY()
+	float LeaveTangent;
+
+	/** If RCTWM_WeightedLeave or RCTWM_WeightedBoth, the weight of the right tangent */
+	UPROPERTY()
+	float LeaveTangentWeight;
+
+	FRichCurveKey()
+		: InterpMode(RCIM_Linear)
+		, TangentMode(RCTM_Auto)
+		, TangentWeightMode(RCTWM_WeightedNone)
+		, Time(0.f)
+		, Value(0.f)
+		, ArriveTangent(0.f)
+		, ArriveTangentWeight(0.f)
+		, LeaveTangent(0.f)
+		, LeaveTangentWeight(0.f)
+	{ }
+
+	FRichCurveKey(float InTime, float InValue)
+		: InterpMode(RCIM_Linear)
+		, TangentMode(RCTM_Auto)
+		, TangentWeightMode(RCTWM_WeightedNone)
+		, Time(InTime)
+		, Value(InValue)
+		, ArriveTangent(0.f)
+		, ArriveTangentWeight(0.f)
+		, LeaveTangent(0.f)
+		, LeaveTangentWeight(0.f)
+	{ }
+
+	FRichCurveKey(float InTime, float InValue, float InArriveTangent, const float InLeaveTangent, ERichCurveInterpMode InInterpMode)
+		: InterpMode(InInterpMode)
+		, TangentMode(RCTM_Auto)
+		, TangentWeightMode(RCTWM_WeightedNone)
+		, Time(InTime)
+		, Value(InValue)
+		, ArriveTangent(InArriveTangent)
+		, ArriveTangentWeight(0.f)
+		, LeaveTangent(InLeaveTangent)
+		, LeaveTangentWeight(0.f)
+	{ }
+
+	/** Conversion constructor */
+	FRichCurveKey(const FInterpCurvePoint<float>& InPoint);
+	FRichCurveKey(const FInterpCurvePoint<FVector>& InPoint, int32 ComponentIndex);
+
+	/** ICPPStructOps interface */
+	bool Serialize(FArchive& Ar);
+	bool operator==(const FRichCurveKey& Other) const;
+	bool operator!=(const FRichCurveKey& Other) const;
+
+	friend FArchive& operator<<(FArchive& Ar, FRichCurveKey& P)
+	{
+		P.Serialize(Ar);
+		return Ar;
+	}
+};
+
+
+template<>
+struct TIsPODType<FRichCurveKey>
+{
+	enum { Value = true };
+};
+
+
+template<>
+struct TStructOpsTypeTraits<FRichCurveKey>
+	: public TStructOpsTypeTraitsBase
+{
+	enum
+	{
+		WithSerializer = true,
+		WithCopy = false,
+		WithIdenticalViaEquality = true,
+	};
+};
+
+
+/** A rich, editable float curve */
+USTRUCT()
+struct ENGINE_API FRichCurve
+	: public FIndexedCurve
+{
+	GENERATED_USTRUCT_BODY()
+
+	FRichCurve() 
+		: FIndexedCurve()
+		, PreInfinityExtrap(RCCE_Constant)
+		, PostInfinityExtrap(RCCE_Constant)
+		, DefaultValue(MAX_flt)
+	{ }
+
+	/** Virtual destructor. */
+	virtual ~FRichCurve() { }
+
+public:
+
+	/** Gets a copy of the keys, so indices and handles can't be meddled with */
+	TArray<FRichCurveKey> GetCopyOfKeys() const;
+
+	/** Gets a const reference of the keys, so indices and handles can't be meddled with */
+	const TArray<FRichCurveKey>& GetConstRefOfKeys() const;
+
+	/** Const iterator for the keys, so the indices and handles stay valid */
+	TArray<FRichCurveKey>::TConstIterator GetKeyIterator() const;
+	
+	/** Functions for getting keys based on handles */
+	FRichCurveKey& GetKey(FKeyHandle KeyHandle);
+	FRichCurveKey GetKey(FKeyHandle KeyHandle) const;
+	
+	/** Quick accessors for the first and last keys */
+	FRichCurveKey GetFirstKey() const;
+	FRichCurveKey GetLastKey() const;
+
+	/** Get the first key that matches any of the given key handles. */
+	FRichCurveKey* GetFirstMatchingKey(const TArray<FKeyHandle>& KeyHandles);
+
+	/** Get the next or previous key given the key handle */
+	FKeyHandle GetNextKey(FKeyHandle KeyHandle) const;
+	FKeyHandle GetPreviousKey(FKeyHandle KeyHandle) const;
+
+	/**
+	  * Add a new key to the curve with the supplied Time and Value. Returns the handle of the new key.
+	  * 
+	  * @param	bUnwindRotation		When true, the value will be treated like a rotation value in degrees, and will automatically be unwound to prevent flipping 360 degrees from the previous key 
+	  * @param  KeyHandle			Optionally can specify what handle this new key should have, otherwise, it'll make a new one
+	  */
+	FKeyHandle AddKey(float InTime, float InValue, const bool bUnwindRotation = false, FKeyHandle KeyHandle = FKeyHandle());
+
+	/**
+	 * Sets the keys with the supplied times and values.
+	 *
+	 * Expects that the times and values are already sorted.
+	 *
+	 * @param InTimes The times at which to insert the keys.
+	 * @param InValues The key values to insert.
+	 * @see AddKey, DeleteKey
+	 */
+	void SetKeys(const TArray<float>& InTimes, const TArray<float>& InValues);
+
+	/**
+	 *  Remove the specified key from the curve.
+	 *
+	 * @param KeyHandle The handle of the key to remove.
+	 * @see AddKey, SetKeys
+	 */
+	void DeleteKey(FKeyHandle KeyHandle);
+
+	/** Finds the key at InTime, and updates its value. If it can't find the key within the KeyTimeTolerance, it adds one at that time */
+	FKeyHandle UpdateOrAddKey(float InTime, float InValue, const bool bUnwindRotation = false, float KeyTimeTolerance = KINDA_SMALL_NUMBER);
+
+	/** Move a key to a new time. This may change the index of the key, so the new key index is returned. */
+	FKeyHandle SetKeyTime(FKeyHandle KeyHandle, float NewTime);
+
+	/** Get the time for the Key with the specified index. */
+	float GetKeyTime(FKeyHandle KeyHandle) const;
+
+	/** Finds a key a the specified time */
+	FKeyHandle FindKey(float KeyTime, float KeyTimeTolerance = KINDA_SMALL_NUMBER) const;
+
+	/** Set the value of the specified key */
+	void SetKeyValue(FKeyHandle KeyHandle, float NewValue, bool bAutoSetTangents=true);
+
+	/** Returns the value of the specified key */
+	float GetKeyValue(FKeyHandle KeyHandle) const;
+
+	/** Set the default value of the curve */
+	void SetDefaultValue(float InDefaultValue) { DefaultValue = InDefaultValue; }
+
+	/** Get the default value for the curve */
+	float GetDefaultValue() const { return DefaultValue; }
+
+	/** Shifts all keys forwards or backwards in time by an even amount, preserving order */
+	void ShiftCurve(float DeltaTime);
+	void ShiftCurve(float DeltaTime, TSet<FKeyHandle>& KeyHandles);
+	
+	/** Scales all keys about an origin, preserving order */
+	void ScaleCurve(float ScaleOrigin, float ScaleFactor);
+	void ScaleCurve(float ScaleOrigin, float ScaleFactor, TSet<FKeyHandle>& KeyHandles);
+
+	/** Set the interp mode of the specified key */
+	void SetKeyInterpMode(FKeyHandle KeyHandle, ERichCurveInterpMode NewInterpMode);
+
+	/** Set the tangent mode of the specified key */
+	void SetKeyTangentMode(FKeyHandle KeyHandle, ERichCurveTangentMode NewTangentMode);
+
+	/** Set the tangent weight mode of the specified key */
+	void SetKeyTangentWeightMode(FKeyHandle KeyHandle, ERichCurveTangentWeightMode NewTangentWeightMode);
+
+	/** Get the interp mode of the specified key */
+	ERichCurveInterpMode GetKeyInterpMode(FKeyHandle KeyHandle) const;
+
+	/** Get the tangent mode of the specified key */
+	ERichCurveTangentMode GetKeyTangentMode(FKeyHandle KeyHandle) const;
+
+	/** Get range of input time values. Outside this region curve continues constantly the start/end values. */
+	void GetTimeRange(float& MinTime, float& MaxTime) const;
+
+	/** Get range of output values. */
+	void GetValueRange(float& MinValue, float& MaxValue) const;
+
+	/** Clear all keys. */
+	void Reset();
+
+	/** Remap InTime based on pre and post infinity extrapolation values */
+	void RemapTimeValue(float& InTime, float& CycleValueOffset) const;
+
+	/** Evaluate this rich curve at the specified time */
+	float Eval(float InTime, float InDefaultValue = 0.0f) const;
+
+	/** Auto set tangents for any 'auto' keys in curve */
+	void AutoSetTangents(float Tension = 0.f);
+
+	/** Resize curve length to the [MinTimeRange, MaxTimeRange] */
+	void ReadjustTimeRange(float NewMinTimeRange, float NewMaxTimeRange, bool bInsert/* whether insert or remove*/, float OldStartTime, float OldEndTime);
+
+	/** Determine if two RichCurves are the same */
+	bool operator == (const FRichCurve& Curve) const;
+
+	/** Bake curve given the sample rate */
+	void BakeCurve(float SampleRate);
+	void BakeCurve(float SampleRate, float FirstKeyTime, float LastKeyTime);
+
+	/** Remove redundant keys, comparing against Tolerance */
+	void RemoveRedundantKeys(float Tolerance);
+	void RemoveRedundantKeys(float Tolerance, float FirstKeyTime, float LastKeyTime);
+
+public:
+
+	// FIndexedCurve interface
+
+	virtual int32 GetNumKeys() const override;
+	virtual bool IsKeyHandleValid(FKeyHandle KeyHandle) const override;
+
+public:
+
+	/** Pre-infinity extrapolation state */
+	UPROPERTY()
+	TEnumAsByte<ERichCurveExtrapolation> PreInfinityExtrap;
+
+	/** Post-infinity extrapolation state */
+	UPROPERTY()
+	TEnumAsByte<ERichCurveExtrapolation> PostInfinityExtrap;
+
+	/** Sorted array of keys */
+	UPROPERTY(EditAnywhere, EditFixedSize, Category="Curve")
+	TArray<FRichCurveKey> Keys;
+
+	/** Default value */
+	UPROPERTY(EditAnywhere, Category="Curve")
+	float DefaultValue;
+};
+
+
+/**
+ * Info about a curve to be edited.
+ */
+template<class T>
+struct FRichCurveEditInfoTemplate
+{
+	/** Name of curve, used when displaying in editor. Can include comma's to allow tree expansion in editor */
+	FName CurveName;
+
+	/** Pointer to curves to be edited */
+	T CurveToEdit;
+
+	FRichCurveEditInfoTemplate()
+		: CurveName(NAME_None)
+		, CurveToEdit(nullptr)
+	{ }
+
+	FRichCurveEditInfoTemplate(T InCurveToEdit)
+		: CurveName(NAME_None)
+		, CurveToEdit(InCurveToEdit)
+	{ }
+
+	FRichCurveEditInfoTemplate(T InCurveToEdit, FName InCurveName)
+		: CurveName(InCurveName)
+		, CurveToEdit(InCurveToEdit)
+	{ }
+
+	inline bool operator==(const FRichCurveEditInfoTemplate<T>& Other) const
+	{
+		return Other.CurveName.IsEqual(CurveName) && Other.CurveToEdit == CurveToEdit;
+	}
+
+	uint32 GetTypeHash() const
+	{
+		return HashCombine(::GetTypeHash(CurveName), PointerHash(CurveToEdit));
+	}
+};
+
+
+template<class T>
+inline uint32 GetTypeHash(const FRichCurveEditInfoTemplate<T>& RichCurveEditInfo)
+{
+	return RichCurveEditInfo.GetTypeHash();
+}
+
+
+typedef FRichCurveEditInfoTemplate<FRichCurve*>			FRichCurveEditInfo;
+typedef FRichCurveEditInfoTemplate<const FRichCurve*>	FRichCurveEditInfoConst;
diff --git a/Engine/Source/Runtime/Engine/Classes/Curves/StringCurve.h b/Engine/Source/Runtime/Engine/Classes/Curves/StringCurve.h
new file mode 100644
--- /dev/null
+++ b/Engine/Source/Runtime/Engine/Classes/Curves/StringCurve.h
@@ -0,0 +1,208 @@
+#pragma once
+
+#include "Curves/IndexedCurve.h"
+#include "Curves/KeyHandle.h"
+#include "StringCurve.generated.h"
+
+
+/**
+ * One key in a curve of FStrings.
+ */
+USTRUCT()
+struct ENGINE_API FStringCurveKey
+{
+	GENERATED_USTRUCT_BODY()
+
+	/** Time at this key */
+	UPROPERTY(EditAnywhere, Category="Key")
+	float Time;
+
+	/** Value at this key */
+	UPROPERTY(EditAnywhere, Category="Key")
+	FString Value;
+
+	/** Default constructor. */
+	FStringCurveKey()
+		: Time(0.0f)
+	{ }
+
+	/** Creates and initializes a new instance. */
+	FStringCurveKey(float InTime, const FString& InValue)
+		: Time(InTime)
+		, Value(InValue)
+	{ }
+
+public:
+
+	//~ TStructOpsTypeTraits interface
+
+	bool operator==(const FStringCurveKey& Other) const;
+	bool operator!=(const FStringCurveKey& Other) const;
+	bool Serialize(FArchive& Ar);
+
+	/** Serializes a name curve key from or into an archive. */
+	friend FArchive& operator<<(FArchive& Ar, FStringCurveKey& Key)
+	{
+		Key.Serialize(Ar);
+		return Ar;
+	}
+};
+
+
+template<>
+struct TIsPODType<FStringCurveKey>
+{
+	enum { Value = true };
+};
+
+
+template<>
+struct TStructOpsTypeTraits<FStringCurveKey>
+	: public TStructOpsTypeTraitsBase
+{
+	enum
+	{
+		WithSerializer = true,
+		WithCopy = false,
+		WithIdenticalViaEquality = true,
+	};
+};
+
+
+/**
+ * Implements a curve of FStrings.
+ */
+USTRUCT()
+struct ENGINE_API FStringCurve
+	: public FIndexedCurve
+{
+	GENERATED_USTRUCT_BODY()
+
+	/** Virtual destructor. */
+	virtual ~FStringCurve() { }
+
+public:
+
+	/**
+	  * Add a new key to the curve with the supplied Time and Value. Returns the handle of the new key.
+	  * 
+	  * @param InTime The time at which to add the key.
+	  * @param InValue The value of the key.
+	  * @param KeyHandle Optional handle for the new key.
+	  */
+	FKeyHandle AddKey(float InTime, const FString& InValue, FKeyHandle KeyHandle = FKeyHandle());
+
+	/**
+	 * Remove the specified key from the curve.
+	 *
+	 * @param KeyHandle Handle to the key to remove.
+	 */
+	void DeleteKey(FKeyHandle KeyHandle);
+
+	/**
+	 * Evaluate the curve at the specified time.
+	 *
+	 * @param Time The time to evaluate.
+	 * @param InDefaultValue The default value if no key exists.
+	 * @return The evaluated string.
+	 */
+	FString Eval(float Time, const FString& InDefaultValue) const;
+
+	/**
+	 * Finds a key a the specified time.
+	 *
+	 * @param KeyTime The time at which to find the key.
+	 * @param KeyTimeTolerance The key time tolerance to use for equality.
+	 * @return A handle to the key, or invalid key handle if not found.
+	 */
+	FKeyHandle FindKey(float KeyTime, float KeyTimeTolerance = KINDA_SMALL_NUMBER) const;
+
+	/** Get the default value for the curve */
+	FString GetDefaultValue() const
+	{
+		return DefaultValue;
+	}
+
+	/**
+	 * Get a key.
+	 *
+	 * @param KeyHandle The handle of the key to get.
+	 * @return The key.
+	 */
+	FStringCurveKey& GetKey(FKeyHandle KeyHandle);
+	FStringCurveKey GetKey(FKeyHandle KeyHandle) const;
+
+	/**
+	 * Read-only access to the key collection.
+	 *
+	 * @return Collection of keys.
+	 */
+	const TArray<FStringCurveKey>& GetKeys() const
+	{
+		return Keys;
+	}
+
+	/**
+	 * Get the time for the Key with the specified index.
+	 *
+	 * @param KeyHandle Handle to the key whose time to get.
+	 * @return The key's time.
+	 */
+	float GetKeyTime(FKeyHandle KeyHandle) const;
+
+	/**
+	 * Set the default value of the curve.
+	 *
+	 * @param InDefaultValue The value to set.
+	 */
+	void SetDefaultValue(const FString& InDefaultValue)
+	{
+		DefaultValue = InDefaultValue;
+	}
+
+	/**
+	 * Move a key to a new time.
+	 *
+	 * This may change the index of the key, so the new key index is returned.
+	 *
+	 * @param KeyHandle The handle of the key to change.
+	 * @param NewTime The new time to set on the key.
+	 */
+	FKeyHandle SetKeyTime(FKeyHandle KeyHandle, float NewTime);
+
+	/**
+	 * Finds the key at InTime, and updates its value. If it can't find the key within the KeyTimeTolerance, it adds one at that time.
+	 *
+	 * @param InTime The time at which the key should be added or updated.
+	 * @param InValue The value of the key.
+	 * @param KeyTimeTolerance The tolerance used for key time equality.
+	 */
+	FKeyHandle UpdateOrAddKey(float InTime, const FString& InValue, float KeyTimeTolerance = KINDA_SMALL_NUMBER);
+
+public:
+
+	/** Shifts all keys forwards or backwards in time by an even amount, preserving order. */
+	void ShiftCurve(float DeltaTime);
+	void ShiftCurve(float DeltaTime, TSet<FKeyHandle>& KeyHandles);
+	
+	/** Scales all keys about an origin, preserving order. */
+	void ScaleCurve(float ScaleOrigin, float ScaleFactor);
+	void ScaleCurve(float ScaleOrigin, float ScaleFactor, TSet<FKeyHandle>& KeyHandles);
+
+public:
+
+	//~ FIndexedCurve interface
+
+	virtual int32 GetNumKeys() const override;
+	virtual bool IsKeyHandleValid(FKeyHandle KeyHandle) const override;
+
+public:
+
+	/** Default value */
+	UPROPERTY(EditAnywhere, Category="Curve")
+	FString DefaultValue;
+
+	/** Sorted array of keys */
+	UPROPERTY(EditAnywhere, EditFixedSize, Category="Curve")
+	TArray<FStringCurveKey> Keys;
+};
diff --git a/Engine/Source/Runtime/Engine/Classes/GameFramework/HUD.h b/Engine/Source/Runtime/Engine/Classes/GameFramework/HUD.h
--- a/Engine/Source/Runtime/Engine/Classes/GameFramework/HUD.h
+++ b/Engine/Source/Runtime/Engine/Classes/GameFramework/HUD.h
@@ -452,9 +452,9 @@ public:
 	virtual void AddPostRenderedActor(AActor* A);
 
 	/**
-	 * check if we should be display debug information for particular types of debug messages
-	 * @param DebugType - type of debug message
-	 * @result bool - true if it should be displayed
+	 * check if we should be display debug information for particular types of debug messages.
+	 * @param DebugType - type of debug message.
+	 * @return true if it should be displayed, false otherwise.
 	 */
 	virtual bool ShouldDisplayDebug(const FName& DebugType) const;
 
diff --git a/Engine/Source/Runtime/Engine/Private/Curves/CurveOwnerInterface.cpp b/Engine/Source/Runtime/Engine/Private/Curves/CurveOwnerInterface.cpp
new file mode 100644
--- /dev/null
+++ b/Engine/Source/Runtime/Engine/Private/Curves/CurveOwnerInterface.cpp
@@ -0,0 +1,30 @@
+// Copyright 1998-2016 Epic Games, Inc. All Rights Reserved.
+
+#include "EnginePrivate.h"
+#include "Curves/CurveOwnerInterface.h"
+
+
+/* FCurveOwnerInterface
+ *****************************************************************************/
+
+FLinearColor FCurveOwnerInterface::GetCurveColor(FRichCurveEditInfo CurveInfo) const
+{
+	FString CurveName = CurveInfo.CurveName.ToString();
+
+	if (CurveName == TEXT("X") || CurveName == TEXT("R"))
+	{
+		return FLinearColor(1.0f, 0.0f, 0.0f);
+	}
+	
+	if (CurveName == TEXT("Y") || CurveName == TEXT("G"))
+	{
+		return FLinearColor(0.0f, 1.0f, 0.0f);
+	}
+	
+	if (CurveName == TEXT("Z") || CurveName == TEXT("B"))
+	{
+		return FLinearColor(0.05f, 0.05f, 1.0f);
+	}
+
+	return FLinearColor::Gray;
+}
diff --git a/Engine/Source/Runtime/Engine/Private/Curves/IndexedCurve.cpp b/Engine/Source/Runtime/Engine/Private/Curves/IndexedCurve.cpp
new file mode 100644
--- /dev/null
+++ b/Engine/Source/Runtime/Engine/Private/Curves/IndexedCurve.cpp
@@ -0,0 +1,68 @@
+// Copyright 1998-2016 Epic Games, Inc. All Rights Reserved.
+
+#include "EnginePrivate.h"
+#include "Curves/IndexedCurve.h"
+
+
+/* FIndexedCurve interface
+ *****************************************************************************/
+
+int32 FIndexedCurve::GetIndexSafe(FKeyHandle KeyHandle) const
+{
+	return IsKeyHandleValid(KeyHandle) ? *KeyHandlesToIndices.Find(KeyHandle) : INDEX_NONE;
+}
+
+
+TMap<FKeyHandle, int32>::TConstIterator FIndexedCurve::GetKeyHandleIterator() const
+{
+	EnsureAllIndicesHaveHandles();
+	return KeyHandlesToIndices.CreateConstIterator();
+}
+
+
+bool FIndexedCurve::IsKeyHandleValid(FKeyHandle KeyHandle) const
+{
+	EnsureAllIndicesHaveHandles();
+	return KeyHandlesToIndices.Find(KeyHandle) != NULL;
+}
+
+
+/* FIndexedCurve implementation
+ *****************************************************************************/
+
+void FIndexedCurve::EnsureAllIndicesHaveHandles() const
+{
+	if (KeyHandlesToIndices.Num() != GetNumKeys())
+	{
+		for (int32 i = 0; i < GetNumKeys(); ++i)
+		{
+			EnsureIndexHasAHandle(i);
+		}
+	}
+}
+
+
+void FIndexedCurve::EnsureIndexHasAHandle(int32 KeyIndex) const
+{
+	const FKeyHandle* KeyHandle = KeyHandlesToIndices.FindKey(KeyIndex);
+	if (!KeyHandle)
+	{
+		FKeyHandle OutKeyHandle = FKeyHandle();
+		KeyHandlesToIndices.Add(OutKeyHandle, KeyIndex);
+	}
+}
+
+
+int32 FIndexedCurve::GetIndex(FKeyHandle KeyHandle) const
+{
+	return *KeyHandlesToIndices.Find(KeyHandle);
+}
+
+
+FKeyHandle FIndexedCurve::GetKeyHandle(int32 KeyIndex) const
+{
+	check(KeyIndex >= 0 && KeyIndex < GetNumKeys());
+	EnsureIndexHasAHandle(KeyIndex);
+
+	return *KeyHandlesToIndices.FindKey(KeyIndex);
+}
