Copying //UE4/Dev-Sequencer to //UE4/Dev-Main (Source: //UE4/Dev-Sequencer @ 3057646)

diff --git a/Engine/Source/Runtime/Engine/Private/Curves/IntegralCurve.cpp b/Engine/Source/Runtime/Engine/Private/Curves/IntegralCurve.cpp
new file mode 100644
--- /dev/null
+++ b/Engine/Source/Runtime/Engine/Private/Curves/IntegralCurve.cpp
@@ -0,0 +1,296 @@
+// Copyright 1998-2016 Epic Games, Inc. All Rights Reserved.
+
+#include "EnginePrivate.h"
+#include "Curves/IntegralCurve.h"
+
+
+int32 FIntegralCurve::GetNumKeys() const
+{
+	return Keys.Num();
+}
+
+
+bool FIntegralCurve::IsKeyHandleValid(FKeyHandle KeyHandle) const
+{
+	bool bValid = false;
+	if (FIndexedCurve::IsKeyHandleValid(KeyHandle))
+	{
+		bValid = Keys.IsValidIndex(GetIndex(KeyHandle));
+	}
+	return bValid;
+}
+
+
+int32 FIntegralCurve::Evaluate(float Time, int32 InDefaultValue) const
+{
+	// If the default value hasn't been initialized, use the incoming default value
+	int32 ReturnVal = DefaultValue == MAX_int32 ? InDefaultValue : DefaultValue;
+
+	if (Keys.Num() == 0 || (bUseDefaultValueBeforeFirstKey && Time < Keys[0].Time))
+	{
+		// If no keys in curve, or bUseDefaultValueBeforeFirstKey is set and the time is before the first key, return the Default value.
+	}
+	else if (Keys.Num() < 2 || Time < Keys[0].Time)
+	{
+		// There is only one key or the time is before the first value. Return the first value
+		ReturnVal = Keys[0].Value;
+	}
+	else if (Time < Keys[Keys.Num() - 1].Time)
+	{
+		// The key is in the range of Key[0] to Keys[Keys.Num()-1].  Find it by searching
+		for (int32 i = 0; i < Keys.Num(); ++i)
+		{
+			if (Time < Keys[i].Time)
+			{
+				ReturnVal = Keys[FMath::Max(0, i - 1)].Value;
+				break;
+			}
+		}
+	}
+	else
+	{
+		// Key is beyon the last point in the curve.  Return it's value
+		ReturnVal = Keys[Keys.Num() - 1].Value;
+	}
+
+	return ReturnVal;
+}
+
+
+TArray<FIntegralKey>::TConstIterator FIntegralCurve::GetKeyIterator() const
+{
+	return Keys.CreateConstIterator();
+}
+
+
+FKeyHandle FIntegralCurve::AddKey(float InTime, int32 InValue, FKeyHandle InKeyHandle)
+{
+	int32 Index = 0;
+	for (; Index < Keys.Num() && Keys[Index].Time < InTime; ++Index);
+	Keys.Insert(FIntegralKey(InTime, InValue), Index);
+
+	for (auto It = KeyHandlesToIndices.CreateIterator(); It; ++It)
+	{
+		int32& KeyIndex = It.Value();
+		if (KeyIndex >= Index)
+		{
+			++KeyIndex;
+		}
+	}
+
+	KeyHandlesToIndices.Add(InKeyHandle, Index);
+
+	return GetKeyHandle(Index);
+}
+
+
+void FIntegralCurve::DeleteKey(FKeyHandle InKeyHandle)
+{
+	int32 Index = GetIndex(InKeyHandle);
+
+	Keys.RemoveAt(Index);
+
+	KeyHandlesToIndices.Remove(InKeyHandle);
+
+	for (auto It = KeyHandlesToIndices.CreateIterator(); It; ++It)
+	{
+		int32& KeyIndex = It.Value();
+		if (KeyIndex >= Index)
+		{
+			--KeyIndex;
+		}
+	}
+}
+
+
+FKeyHandle FIntegralCurve::UpdateOrAddKey(float InTime, int32 Value, float KeyTimeTolerance)
+{
+	for (int32 KeyIndex = 0; KeyIndex < Keys.Num(); ++KeyIndex)
+	{
+		float KeyTime = Keys[KeyIndex].Time;
+
+		if (FMath::IsNearlyEqual(KeyTime, InTime, KeyTimeTolerance))
+		{
+			Keys[KeyIndex].Value = Value;
+			return GetKeyHandle(KeyIndex);
+		}
+
+		if (KeyTime > InTime)
+		{
+			// All the rest of the keys exist after the key we want to add
+			// so there is no point in searching
+			break;
+		}
+	}
+
+	// A key wasnt found, add it now
+	return AddKey(InTime, Value);
+}
+
+
+FKeyHandle FIntegralCurve::SetKeyTime(FKeyHandle KeyHandle, float NewTime)
+{
+	if (!IsKeyHandleValid(KeyHandle))
+	{
+		return KeyHandle;
+	}
+
+	FIntegralKey OldKey = GetKey(KeyHandle);
+
+	DeleteKey(KeyHandle);
+	AddKey(NewTime, OldKey.Value, KeyHandle);
+
+	// Copy all properties from old key, but then fix time to be the new time
+	GetKey(KeyHandle) = OldKey;
+	GetKey(KeyHandle).Time = NewTime;
+
+	return KeyHandle;
+}
+
+
+float FIntegralCurve::GetKeyTime(FKeyHandle KeyHandle) const
+{
+	if (!IsKeyHandleValid(KeyHandle))
+	{
+		return 0.f;
+	}
+
+	return GetKey(KeyHandle).Time;
+}
+
+
+void FIntegralCurve::ShiftCurve(float DeltaTime)
+{
+	TSet<FKeyHandle> KeyHandles;
+	for (auto It = KeyHandlesToIndices.CreateIterator(); It; ++It)
+	{
+		FKeyHandle& KeyHandle = It.Key();
+		KeyHandles.Add(KeyHandle);
+	}
+
+	ShiftCurve(DeltaTime, KeyHandles);
+}
+
+
+void FIntegralCurve::ShiftCurve(float DeltaTime, TSet<FKeyHandle>& KeyHandles)
+{
+	for (auto It = KeyHandlesToIndices.CreateIterator(); It; ++It)
+	{
+		const FKeyHandle& KeyHandle = It.Key();
+		if (KeyHandles.Num() != 0 && KeyHandles.Contains(KeyHandle))
+		{
+			SetKeyTime(KeyHandle, GetKeyTime(KeyHandle) + DeltaTime);
+		}
+	}
+}
+
+
+void FIntegralCurve::ScaleCurve(float ScaleOrigin, float ScaleFactor)
+{
+	TSet<FKeyHandle> KeyHandles;
+	for (auto It = KeyHandlesToIndices.CreateIterator(); It; ++It)
+	{
+		FKeyHandle& KeyHandle = It.Key();
+		KeyHandles.Add(KeyHandle);
+	}
+
+	ScaleCurve(ScaleOrigin, ScaleFactor, KeyHandles);
+}
+
+
+void FIntegralCurve::ScaleCurve(float ScaleOrigin, float ScaleFactor, TSet<FKeyHandle>& KeyHandles)
+{
+	for (auto It = KeyHandlesToIndices.CreateIterator(); It; ++It)
+	{
+		const FKeyHandle& KeyHandle = It.Key();
+		if (KeyHandles.Num() != 0 && KeyHandles.Contains(KeyHandle))
+		{
+			SetKeyTime(KeyHandle, (GetKeyTime(KeyHandle) - ScaleOrigin) * ScaleFactor + ScaleOrigin);
+		}
+	}
+}
+
+
+FIntegralKey& FIntegralCurve::GetKey(FKeyHandle KeyHandle)
+{
+	EnsureAllIndicesHaveHandles();
+	return Keys[GetIndex(KeyHandle)];
+}
+
+
+FIntegralKey FIntegralCurve::GetKey(FKeyHandle KeyHandle) const
+{
+	EnsureAllIndicesHaveHandles();
+	return Keys[GetIndex(KeyHandle)];
+}
+
+
+FKeyHandle FIntegralCurve::FindKey(float KeyTime, float KeyTimeTolerance) const
+{
+	int32 Start = 0;
+	int32 End = Keys.Num() - 1;
+
+	// Binary search since the keys are in sorted order
+	while (Start <= End)
+	{
+		int32 TestPos = Start + (End - Start) / 2;
+		float TestKeyTime = Keys[TestPos].Time;
+
+		if (FMath::IsNearlyEqual(TestKeyTime, KeyTime, KeyTimeTolerance))
+		{
+			return GetKeyHandle(TestPos);
+		}
+		else if (TestKeyTime < KeyTime)
+		{
+			Start = TestPos + 1;
+		}
+		else
+		{
+			End = TestPos - 1;
+		}
+	}
+
+	return FKeyHandle();
+}
+
+FKeyHandle FIntegralCurve::FindKeyBeforeOrAt(float KeyTime) const
+{
+	// If there are no keys or the time is before the first key return an invalid handle.
+	if (Keys.Num() == 0 || KeyTime < Keys[0].Time)
+	{
+		return FKeyHandle();
+	}
+
+	// If the time is after or at the last key return the last key.
+	if (KeyTime >= Keys[Keys.Num() - 1].Time)
+	{
+		return GetKeyHandle(Keys.Num() - 1);
+	}
+
+	// Otherwise binary search to find the handle of the nearest key at or before the time.
+	int32 Start = 0;
+	int32 End = Keys.Num() - 1;
+	int32 FoundIndex = -1;
+	while (FoundIndex < 0)
+	{
+		int32 TestPos = (Start + End) / 2;
+		float TestKeyTime = Keys[TestPos].Time;
+		float NextTestKeyTime = Keys[TestPos + 1].Time;
+		if (TestKeyTime <= KeyTime)
+		{
+			if (NextTestKeyTime > KeyTime)
+			{
+				FoundIndex = TestPos;
+			}
+			else
+			{
+				Start = TestPos + 1;
+			}
+		}
+		else
+		{
+			End = TestPos;
+		}
+	}
+	return GetKeyHandle(FoundIndex);
+}
diff --git a/Engine/Source/Runtime/Engine/Private/Curves/KeyHandle.cpp b/Engine/Source/Runtime/Engine/Private/Curves/KeyHandle.cpp
new file mode 100644
--- /dev/null
+++ b/Engine/Source/Runtime/Engine/Private/Curves/KeyHandle.cpp
@@ -0,0 +1,92 @@
+// Copyright 1998-2016 Epic Games, Inc. All Rights Reserved.
+
+#include "EnginePrivate.h"
+
+
+/* FKeyHandleMap interface
+ *****************************************************************************/
+
+void FKeyHandleMap::Add( const FKeyHandle& InHandle, int32 InIndex )
+{
+	KeyHandlesToIndices.Add( InHandle, InIndex );
+}
+
+
+void FKeyHandleMap::Empty()
+{
+	KeyHandlesToIndices.Empty();
+}
+
+
+void FKeyHandleMap::Remove( const FKeyHandle& InHandle )
+{
+	KeyHandlesToIndices.Remove( InHandle );
+}
+
+
+const int32* FKeyHandleMap::Find( const FKeyHandle& InHandle ) const
+{
+	return KeyHandlesToIndices.Find( InHandle );
+}
+
+
+const FKeyHandle* FKeyHandleMap::FindKey( int32 KeyIndex ) const
+{
+	return KeyHandlesToIndices.FindKey( KeyIndex );
+}
+
+
+int32 FKeyHandleMap::Num() const
+{
+	return KeyHandlesToIndices.Num();
+}
+
+
+TMap<FKeyHandle, int32>::TConstIterator FKeyHandleMap::CreateConstIterator() const
+{
+	return KeyHandlesToIndices.CreateConstIterator();
+}
+
+
+TMap<FKeyHandle, int32>::TIterator FKeyHandleMap::CreateIterator() 
+{
+	return KeyHandlesToIndices.CreateIterator();
+}
+
+
+bool FKeyHandleMap::Serialize(FArchive& Ar)
+{
+	// only allow this map to be saved to the transaction buffer
+	if( Ar.IsTransacting() )
+	{
+		Ar << KeyHandlesToIndices;
+	}
+
+	return true;
+}
+
+
+bool FKeyHandleMap::operator==(const FKeyHandleMap& Other) const
+{
+	if (KeyHandlesToIndices.Num() != Other.KeyHandlesToIndices.Num())
+	{
+		return false;
+	}
+
+	for (TMap<FKeyHandle, int32>::TConstIterator It(KeyHandlesToIndices); It; ++It)
+	{
+		int32 const* OtherVal = Other.KeyHandlesToIndices.Find(It.Key());
+		if (!OtherVal || *OtherVal != It.Value() )
+		{
+			return false;
+		}
+	}
+
+	return true;
+}
+
+
+bool FKeyHandleMap::operator!=(const FKeyHandleMap& Other) const
+{
+	return !(*this==Other);
+}
diff --git a/Engine/Source/Runtime/Engine/Private/Curves/NameCurve.cpp b/Engine/Source/Runtime/Engine/Private/Curves/NameCurve.cpp
new file mode 100644
--- /dev/null
+++ b/Engine/Source/Runtime/Engine/Private/Curves/NameCurve.cpp
@@ -0,0 +1,243 @@
+// Copyright 1998-2016 Epic Games, Inc. All Rights Reserved.
+
+#include "EnginePrivate.h"
+#include "Curves/NameCurve.h"
+
+
+/* FNameCurveKey interface
+ *****************************************************************************/
+
+bool FNameCurveKey::operator==(const FNameCurveKey& Curve) const
+{
+	return ((Time == Curve.Time) && (Value == Curve.Value));
+}
+
+
+bool FNameCurveKey::operator!=(const FNameCurveKey& Other) const
+{
+	return !(*this == Other);
+}
+
+
+bool FNameCurveKey::Serialize(FArchive& Ar)
+{
+	Ar << Time << Value;
+	return true;
+}
+
+
+/* FNameCurve interface
+ *****************************************************************************/
+
+FKeyHandle FNameCurve::AddKey(float InTime, const FName& InValue, FKeyHandle KeyHandle)
+{
+	int32 Index = 0;
+
+	// insert key
+	for(; Index < Keys.Num() && Keys[Index].Time < InTime; ++Index);
+	Keys.Insert(FNameCurveKey(InTime, InValue), Index);
+
+	// update key indices
+	for ( auto It = KeyHandlesToIndices.CreateIterator(); It; ++It)
+	{
+		int32& KeyIndex = It.Value();
+		if (KeyIndex >= Index) {++KeyIndex;}
+	}
+
+	KeyHandlesToIndices.Add(KeyHandle, Index);
+
+	return GetKeyHandle(Index);
+}
+
+
+void FNameCurve::DeleteKey(FKeyHandle KeyHandle)
+{
+	// remove key
+	int32 Index = GetIndex(KeyHandle);
+	Keys.RemoveAt(Index);
+
+	KeyHandlesToIndices.Remove(KeyHandle);
+
+	// update key indices
+	for (auto It = KeyHandlesToIndices.CreateIterator(); It; ++It)
+	{
+		int32& KeyIndex = It.Value();
+		if (KeyIndex >= Index)
+		{
+			--KeyIndex;
+		}
+	}
+}
+
+
+FKeyHandle FNameCurve::FindKey(float KeyTime, float KeyTimeTolerance) const
+{
+	int32 Start = 0;
+	int32 End = Keys.Num()-1;
+
+	// Binary search since the keys are in sorted order
+	while (Start <= End)
+	{
+		int32 TestPos = Start + (End-Start) / 2;
+		float TestKeyTime = Keys[TestPos].Time;
+
+		if (FMath::IsNearlyEqual(TestKeyTime, KeyTime, KeyTimeTolerance))
+		{
+			return GetKeyHandle(TestPos);
+		}
+
+		if (TestKeyTime < KeyTime)
+		{
+			Start = TestPos+1;
+		}
+		else
+		{
+			End = TestPos-1;
+		}
+	}
+
+	return FKeyHandle();
+}
+
+
+FNameCurveKey& FNameCurve::GetKey(FKeyHandle KeyHandle)
+{
+	EnsureAllIndicesHaveHandles();
+	return Keys[GetIndex(KeyHandle)];
+}
+
+
+FNameCurveKey FNameCurve::GetKey(FKeyHandle KeyHandle) const
+{
+	EnsureAllIndicesHaveHandles();
+	return Keys[GetIndex(KeyHandle)];
+}
+
+
+float FNameCurve::GetKeyTime(FKeyHandle KeyHandle) const
+{
+	if (!IsKeyHandleValid(KeyHandle))
+	{
+		return 0.f;
+	}
+
+	return GetKey(KeyHandle).Time;
+}
+
+
+FKeyHandle FNameCurve::SetKeyTime(FKeyHandle KeyHandle, float NewTime)
+{
+	if (!IsKeyHandleValid(KeyHandle))
+	{
+		return KeyHandle;
+	}
+
+	FNameCurveKey OldKey = GetKey(KeyHandle);
+	
+	DeleteKey(KeyHandle);
+	AddKey(NewTime, OldKey.Value, KeyHandle);
+
+	// Copy all properties from old key, but then fix time to be the new time
+	GetKey(KeyHandle) = OldKey;
+	GetKey(KeyHandle).Time = NewTime;
+
+	return KeyHandle;
+}
+
+
+void FNameCurve::ShiftCurve(float DeltaTime)
+{
+	TSet<FKeyHandle> KeyHandles;
+	for (auto It = KeyHandlesToIndices.CreateIterator(); It; ++It)
+	{
+		FKeyHandle& KeyHandle = It.Key();
+		KeyHandles.Add(KeyHandle);
+	}
+
+	ShiftCurve(DeltaTime, KeyHandles);
+}
+
+
+void FNameCurve::ShiftCurve(float DeltaTime, TSet<FKeyHandle>& KeyHandles)
+{
+	for (auto It = KeyHandlesToIndices.CreateIterator(); It; ++It)
+	{
+		const FKeyHandle& KeyHandle = It.Key();
+		if (KeyHandles.Num() != 0 && KeyHandles.Contains(KeyHandle))
+		{
+			SetKeyTime(KeyHandle, GetKeyTime(KeyHandle) + DeltaTime);
+		}
+	}
+}
+
+
+void FNameCurve::ScaleCurve(float ScaleOrigin, float ScaleFactor)
+{
+	TSet<FKeyHandle> KeyHandles;
+	for (auto It = KeyHandlesToIndices.CreateIterator(); It; ++It)
+	{
+		FKeyHandle& KeyHandle = It.Key();
+		KeyHandles.Add(KeyHandle);
+	}
+
+	ScaleCurve(ScaleOrigin, ScaleFactor, KeyHandles);
+}
+
+
+void FNameCurve::ScaleCurve(float ScaleOrigin, float ScaleFactor, TSet<FKeyHandle>& KeyHandles)
+{
+	for (auto It = KeyHandlesToIndices.CreateIterator(); It; ++It)
+	{
+		const FKeyHandle& KeyHandle = It.Key();
+		if (KeyHandles.Num() != 0 && KeyHandles.Contains(KeyHandle))
+		{
+			SetKeyTime(KeyHandle, (GetKeyTime(KeyHandle) - ScaleOrigin) * ScaleFactor + ScaleOrigin);
+		}
+	}
+}
+
+
+FKeyHandle FNameCurve::UpdateOrAddKey(float InTime, const FName& InValue, float KeyTimeTolerance)
+{
+	// Search for a key that already exists at the time and if found, update its value
+	for (int32 KeyIndex = 0; KeyIndex < Keys.Num(); ++KeyIndex)
+	{
+		float KeyTime = Keys[KeyIndex].Time;
+
+		if (FMath::IsNearlyEqual(KeyTime, InTime, KeyTimeTolerance))
+		{
+			Keys[KeyIndex].Value = InValue;
+
+			return GetKeyHandle(KeyIndex);
+		}
+
+		if (KeyTime > InTime)
+		{
+			// All the rest of the keys exist after the key we want to add
+			// so there is no point in searching
+			break;
+		}
+	}
+
+	// A key wasnt found, add it now
+	return AddKey(InTime, InValue);
+}
+
+
+int32 FNameCurve::GetNumKeys() const
+{
+	return Keys.Num();
+}
+
+
+bool FNameCurve::IsKeyHandleValid(FKeyHandle KeyHandle) const
+{
+	bool bValid = false;
+
+	if (FIndexedCurve::IsKeyHandleValid(KeyHandle))
+	{
+		bValid = Keys.IsValidIndex(GetIndex(KeyHandle));
+	}
+
+	return bValid;
+}
\ No newline at end of file
diff --git a/Engine/Source/Runtime/Engine/Private/Curves/RichCurve.cpp b/Engine/Source/Runtime/Engine/Private/Curves/RichCurve.cpp
new file mode 100644
--- /dev/null
+++ b/Engine/Source/Runtime/Engine/Private/Curves/RichCurve.cpp
@@ -0,0 +1,1166 @@
+// Copyright 1998-2016 Epic Games, Inc. All Rights Reserved.
+
+#include "EnginePrivate.h"
+#include "Curves/RichCurve.h"
+
+
+DECLARE_CYCLE_STAT(TEXT("RichCurve Eval"), STAT_RichCurve_Eval, STATGROUP_Engine);
+
+
+/* FRichCurveKey interface
+ *****************************************************************************/
+
+static void SetModesFromLegacy(FRichCurveKey& InKey, EInterpCurveMode InterpMode)
+{
+	InKey.InterpMode = RCIM_Linear;
+	InKey.TangentWeightMode = RCTWM_WeightedNone;
+	InKey.TangentMode = RCTM_Auto;
+
+	if (InterpMode == CIM_Constant)
+	{
+		InKey.InterpMode = RCIM_Constant;
+	}
+	else if (InterpMode == CIM_Linear)
+	{
+		InKey.InterpMode = RCIM_Linear;
+	}
+	else
+	{
+		InKey.InterpMode = RCIM_Cubic;
+
+		if (InterpMode == CIM_CurveAuto || InterpMode == CIM_CurveAutoClamped)
+		{
+			InKey.TangentMode = RCTM_Auto;
+		}
+		else if (InterpMode == CIM_CurveBreak)
+		{
+			InKey.TangentMode = RCTM_Break;
+		}
+		else if (InterpMode == CIM_CurveUser)
+		{
+			InKey.TangentMode = RCTM_User;
+		}
+	}
+}
+
+
+FRichCurveKey::FRichCurveKey(const FInterpCurvePoint<float>& InPoint)
+{
+	SetModesFromLegacy(*this, InPoint.InterpMode);
+
+	Time = InPoint.InVal;
+	Value = InPoint.OutVal;
+
+	ArriveTangent = InPoint.ArriveTangent;
+	ArriveTangentWeight = 0.f;
+
+	LeaveTangent = InPoint.LeaveTangent;
+	LeaveTangentWeight = 0.f;
+}
+
+
+FRichCurveKey::FRichCurveKey(const FInterpCurvePoint<FVector>& InPoint, int32 ComponentIndex)
+{
+	SetModesFromLegacy(*this, InPoint.InterpMode);
+
+	Time = InPoint.InVal;
+
+	if (ComponentIndex == 0)
+	{
+		Value = InPoint.OutVal.X;
+		ArriveTangent = InPoint.ArriveTangent.X;
+		LeaveTangent = InPoint.LeaveTangent.X;
+	}
+	else if (ComponentIndex == 1)
+	{
+		Value = InPoint.OutVal.Y;
+		ArriveTangent = InPoint.ArriveTangent.Y;
+		LeaveTangent = InPoint.LeaveTangent.Y;
+	}
+	else
+	{
+		Value = InPoint.OutVal.Z;
+		ArriveTangent = InPoint.ArriveTangent.Z;
+		LeaveTangent = InPoint.LeaveTangent.Z;
+	}
+
+	ArriveTangentWeight = 0.f;
+	LeaveTangentWeight = 0.f;
+}
+
+
+bool FRichCurveKey::Serialize(FArchive& Ar)
+{
+	if (Ar.UE4Ver() < VER_UE4_SERIALIZE_RICH_CURVE_KEY)
+	{
+		return false;
+	}
+
+	// Serialization is handled manually to avoid the extra size overhead of UProperty tagging.
+	// Otherwise with many keys in a rich curve the size can become quite large.
+	Ar << InterpMode;
+	Ar << TangentMode;
+	Ar << TangentWeightMode;
+	Ar << Time;
+	Ar << Value;
+	Ar << ArriveTangent;
+	Ar << ArriveTangentWeight;
+	Ar << LeaveTangent;
+	Ar << LeaveTangentWeight;
+
+	return true;
+}
+
+
+bool FRichCurveKey::operator==( const FRichCurveKey& Curve ) const
+{
+	return (Time == Curve.Time) && (Value == Curve.Value) && (InterpMode == Curve.InterpMode) &&
+		   (TangentMode == Curve.TangentMode) && (TangentWeightMode == Curve.TangentWeightMode) &&
+		   ((InterpMode != RCIM_Cubic) || //also verify if it is cubic that tangents are the same
+		   ((ArriveTangent == Curve.ArriveTangent) && (LeaveTangent == Curve.LeaveTangent) ));
+
+}
+
+
+bool FRichCurveKey::operator!=(const FRichCurveKey& Other) const
+{
+	return !(*this == Other);
+}
+
+
+FKeyHandle::FKeyHandle()
+{
+	static uint32 LastKeyHandleIndex = 1;
+	check(LastKeyHandleIndex != 0); // check in the unlikely event that this overflows
+
+	Index = ++LastKeyHandleIndex;
+}
+
+
+/* FRichCurve interface
+ *****************************************************************************/
+
+TArray<FRichCurveKey> FRichCurve::GetCopyOfKeys() const
+{
+	return Keys;
+}
+
+const TArray<FRichCurveKey>& FRichCurve::GetConstRefOfKeys() const
+{
+	return Keys;
+}
+
+
+TArray<FRichCurveKey>::TConstIterator FRichCurve::GetKeyIterator() const
+{
+	return Keys.CreateConstIterator();
+}
+
+
+FRichCurveKey& FRichCurve::GetKey(FKeyHandle KeyHandle)
+{
+	EnsureAllIndicesHaveHandles();
+	return Keys[GetIndex(KeyHandle)];
+}
+
+
+FRichCurveKey FRichCurve::GetKey(FKeyHandle KeyHandle) const
+{
+	EnsureAllIndicesHaveHandles();
+	return Keys[GetIndex(KeyHandle)];
+}
+
+
+FRichCurveKey FRichCurve::GetFirstKey() const
+{
+	check(Keys.Num() > 0);
+	return Keys[0];
+}
+
+
+FRichCurveKey FRichCurve::GetLastKey() const
+{
+	check(Keys.Num() > 0);
+	return Keys[Keys.Num()-1];
+}
+
+
+FRichCurveKey* FRichCurve::GetFirstMatchingKey(const TArray<FKeyHandle>& KeyHandles)
+{
+	for (const auto& KeyHandle : KeyHandles)
+	{
+		if (IsKeyHandleValid(KeyHandle))
+		{
+			return &GetKey(KeyHandle);
+		}
+	}
+
+	return nullptr;
+}
+
+
+FKeyHandle FRichCurve::GetNextKey(FKeyHandle KeyHandle) const
+{
+	int32 KeyIndex = GetIndex(KeyHandle);
+	KeyIndex++;
+
+	if (Keys.IsValidIndex(KeyIndex))
+	{
+		return GetKeyHandle(KeyIndex);
+	}
+
+	return FKeyHandle();
+}
+
+
+FKeyHandle FRichCurve::GetPreviousKey(FKeyHandle KeyHandle) const
+{
+	int32 KeyIndex = GetIndex(KeyHandle);
+	KeyIndex--;
+
+	if (Keys.IsValidIndex(KeyIndex))
+	{
+		return GetKeyHandle(KeyIndex);
+	}
+
+	return FKeyHandle();
+}
+
+
+int32 FRichCurve::GetNumKeys() const
+{
+	return Keys.Num();
+}
+
+
+bool FRichCurve::IsKeyHandleValid(FKeyHandle KeyHandle) const
+{
+	bool bValid = false;
+
+	if (FIndexedCurve::IsKeyHandleValid(KeyHandle))
+	{
+		bValid = Keys.IsValidIndex( GetIndex(KeyHandle) );
+	}
+
+	return bValid;
+}
+
+
+FKeyHandle FRichCurve::AddKey( const float InTime, const float InValue, const bool bUnwindRotation, FKeyHandle NewHandle )
+{
+	int32 Index = 0;
+	for(; Index < Keys.Num() && Keys[Index].Time < InTime; ++Index);
+	Keys.Insert(FRichCurveKey(InTime, InValue), Index);
+
+	// If we were asked to treat this curve as a rotation value and to unwindow the rotation, then
+	// we'll look at the previous key and modify the key's value to use a rotation angle that is
+	// continuous with the previous key while retaining the exact same rotation angle, if at all necessary
+	if( Index > 0 && bUnwindRotation )
+	{
+		const float OldValue = Keys[ Index - 1 ].Value;
+		float NewValue = Keys[ Index ].Value;
+
+		while( NewValue - OldValue > 180.0f )
+		{
+			NewValue -= 360.0f;
+		}
+		while( NewValue - OldValue < -180.0f )
+		{
+			NewValue += 360.0f;
+		}
+
+		Keys[Index].Value = NewValue;
+	}
+	
+	for ( auto It = KeyHandlesToIndices.CreateIterator(); It; ++It)
+	{
+		const FKeyHandle& KeyHandle = It.Key();
+		int32& KeyIndex = It.Value();
+
+		if (KeyIndex >= Index) {++KeyIndex;}
+	}
+
+	KeyHandlesToIndices.Add(NewHandle, Index);
+
+	return GetKeyHandle(Index);
+}
+
+
+void FRichCurve::SetKeys(const TArray<float>& InTimes, const TArray<float>& InValues)
+{
+	Reset();
+
+	for (int32 Index = 0; Index < InTimes.Num(); ++Index)
+	{
+		Keys.Add(FRichCurveKey(InTimes[Index], InValues[Index]));
+		KeyHandlesToIndices.Add(FKeyHandle(), Index);
+	}
+
+	AutoSetTangents();
+}
+
+
+void FRichCurve::DeleteKey(FKeyHandle InKeyHandle)
+{
+	int32 Index = GetIndex(InKeyHandle);
+	
+	Keys.RemoveAt(Index);
+	AutoSetTangents();
+
+	KeyHandlesToIndices.Remove(InKeyHandle);
+
+	for (auto It = KeyHandlesToIndices.CreateIterator(); It; ++It)
+	{
+		const FKeyHandle& KeyHandle = It.Key();
+		int32& KeyIndex = It.Value();
+
+		if (KeyIndex >= Index)
+		{
+			--KeyIndex;
+		}
+	}
+}
+
+
+FKeyHandle FRichCurve::UpdateOrAddKey(float InTime, float InValue, const bool bUnwindRotation, float KeyTimeTolerance)
+{
+	// Search for a key that already exists at the time and if found, update its value
+	for (int32 KeyIndex = 0; KeyIndex < Keys.Num(); ++KeyIndex)
+	{
+		float KeyTime = Keys[KeyIndex].Time;
+
+		if (FMath::IsNearlyEqual(KeyTime, InTime, KeyTimeTolerance))
+		{
+			Keys[KeyIndex].Value = InValue;
+
+			return GetKeyHandle(KeyIndex);
+		}
+
+		if (KeyTime > InTime)
+		{
+			// All the rest of the keys exist after the key we want to add
+			// so there is no point in searching
+			break;
+		}
+	}
+
+	// A key wasnt found, add it now
+	return AddKey(InTime, InValue, bUnwindRotation);
+}
+
+
+FKeyHandle FRichCurve::SetKeyTime( FKeyHandle KeyHandle, float NewTime )
+{
+	if (!IsKeyHandleValid(KeyHandle))
+	{
+		return KeyHandle;
+	}
+
+	FRichCurveKey OldKey = GetKey(KeyHandle);
+	
+	DeleteKey(KeyHandle);
+	AddKey(NewTime, OldKey.Value, false, KeyHandle);
+
+	// Copy all properties from old key, but then fix time to be the new time
+	GetKey(KeyHandle) = OldKey;
+	GetKey(KeyHandle).Time = NewTime;
+
+	return KeyHandle;
+}
+
+
+float FRichCurve::GetKeyTime(FKeyHandle KeyHandle) const
+{
+	if (!IsKeyHandleValid(KeyHandle))
+	{
+		return 0.f;
+	}
+
+	return GetKey(KeyHandle).Time;
+}
+
+
+FKeyHandle FRichCurve::FindKey(float KeyTime, float KeyTimeTolerance) const
+{
+	int32 Start = 0;
+	int32 End = Keys.Num()-1;
+
+	// Binary search since the keys are in sorted order
+	while (Start <= End)
+	{
+		int32 TestPos = Start + (End-Start) / 2;
+		float TestKeyTime = Keys[TestPos].Time;
+
+		if (FMath::IsNearlyEqual(TestKeyTime, KeyTime, KeyTimeTolerance))
+		{
+			return GetKeyHandle(TestPos);
+		}
+		else if(TestKeyTime < KeyTime)
+		{
+			Start = TestPos+1;
+		}
+		else
+		{
+			End = TestPos-1;
+		}
+	}
+
+	return FKeyHandle();
+}
+
+
+void FRichCurve::SetKeyValue(FKeyHandle KeyHandle, float NewValue, bool bAutoSetTangents)
+{
+	if (!IsKeyHandleValid(KeyHandle))
+	{
+		return;
+	}
+
+	GetKey(KeyHandle).Value = NewValue;
+
+	if (bAutoSetTangents)
+	{
+		AutoSetTangents();
+	}
+}
+
+
+float FRichCurve::GetKeyValue(FKeyHandle KeyHandle) const
+{
+	if (!IsKeyHandleValid(KeyHandle))
+	{
+		return 0.f;
+	}
+
+	return GetKey(KeyHandle).Value;
+}
+
+
+void FRichCurve::ShiftCurve(float DeltaTime)
+{
+	TSet<FKeyHandle> KeyHandles;
+	for (auto It = KeyHandlesToIndices.CreateIterator(); It; ++It)
+	{
+		FKeyHandle& KeyHandle = It.Key();
+		KeyHandles.Add(KeyHandle);
+	}
+
+	ShiftCurve(DeltaTime, KeyHandles);
+}
+
+
+void FRichCurve::ShiftCurve(float DeltaTime, TSet<FKeyHandle>& KeyHandles)
+{
+	for (auto It = KeyHandlesToIndices.CreateIterator(); It; ++It)
+	{
+		const FKeyHandle& KeyHandle = It.Key();
+		if (KeyHandles.Num() != 0 && KeyHandles.Contains(KeyHandle))
+		{
+			SetKeyTime(KeyHandle, GetKeyTime(KeyHandle)+DeltaTime);
+		}
+	}
+}
+
+
+void FRichCurve::ScaleCurve(float ScaleOrigin, float ScaleFactor)
+{
+	TSet<FKeyHandle> KeyHandles;
+	for (auto It = KeyHandlesToIndices.CreateIterator(); It; ++It)
+	{
+		FKeyHandle& KeyHandle = It.Key();
+		KeyHandles.Add(KeyHandle);
+	}
+
+	ScaleCurve(ScaleOrigin, ScaleFactor, KeyHandles);
+}
+
+
+void FRichCurve::ScaleCurve(float ScaleOrigin, float ScaleFactor, TSet<FKeyHandle>& KeyHandles)
+{
+	for (auto It = KeyHandlesToIndices.CreateIterator(); It; ++It)
+	{
+		const FKeyHandle& KeyHandle = It.Key();
+		if (KeyHandles.Num() != 0 && KeyHandles.Contains(KeyHandle))
+		{
+			SetKeyTime(KeyHandle, (GetKeyTime(KeyHandle) - ScaleOrigin) * ScaleFactor + ScaleOrigin);
+		}
+	}
+}
+
+
+void FRichCurve::SetKeyInterpMode(FKeyHandle KeyHandle, ERichCurveInterpMode NewInterpMode)
+{
+	if (!IsKeyHandleValid(KeyHandle))
+	{
+		return;
+	}
+
+	GetKey(KeyHandle).InterpMode = NewInterpMode;
+	AutoSetTangents();
+}
+
+
+void FRichCurve::SetKeyTangentMode(FKeyHandle KeyHandle, ERichCurveTangentMode NewTangentMode)
+{
+	if (!IsKeyHandleValid(KeyHandle))
+	{
+		return;
+	}
+
+	GetKey(KeyHandle).TangentMode = NewTangentMode;
+	AutoSetTangents();
+}
+
+
+void FRichCurve::SetKeyTangentWeightMode(FKeyHandle KeyHandle, ERichCurveTangentWeightMode NewTangentWeightMode)
+{
+	if (!IsKeyHandleValid(KeyHandle))
+	{
+		return;
+	}
+
+	GetKey(KeyHandle).TangentWeightMode = NewTangentWeightMode;
+	AutoSetTangents();
+}
+
+
+ERichCurveInterpMode FRichCurve::GetKeyInterpMode(FKeyHandle KeyHandle) const
+{
+	if (!IsKeyHandleValid(KeyHandle))
+	{
+		return RCIM_Linear;
+	}
+
+	return GetKey(KeyHandle).InterpMode;
+}
+
+
+ERichCurveTangentMode FRichCurve::GetKeyTangentMode(FKeyHandle KeyHandle) const
+{
+	if (!IsKeyHandleValid(KeyHandle))
+	{
+		return RCTM_Auto;
+	}
+
+	return GetKey(KeyHandle).TangentMode;
+}
+
+
+void FRichCurve::GetTimeRange(float& MinTime, float& MaxTime) const
+{
+	if (Keys.Num() == 0)
+	{
+		MinTime = 0.f;
+		MaxTime = 0.f;
+	}
+	else
+	{
+		MinTime = Keys[0].Time;
+		MaxTime = Keys[Keys.Num()-1].Time;
+	}
+}
+
+
+/*	 Finds min/max for cubic curves:
+	Looks for feature points in the signal(determined by change in direction of local tangent), these locations are then re-examined in closer detail recursively */
+template<class T>
+void FeaturePointMethod(T& Function , float StartTime,  float EndTime, float StartValue,float Mu, int Depth, int MaxDepth, float& MaxV, float& MinVal)
+{
+	if (Depth >= MaxDepth)
+	{
+		return;
+	}
+
+	float PrevValue   = StartValue;
+	float PrevTangent = StartValue - Function.Eval(StartTime-Mu);
+	EndTime += Mu;
+
+	for (float f = StartTime + Mu;f < EndTime; f += Mu)
+	{
+		float Value		 = Function.Eval(f);
+
+		MaxV   = FMath::Max(Value, MaxV);
+		MinVal = FMath::Min(Value, MinVal);
+
+		float CurTangent = Value - PrevValue;
+		
+		//Change direction? Examine this area closer
+		if (FMath::Sign(CurTangent) != FMath::Sign(PrevTangent))
+		{
+			//feature point centered around the previous tangent
+			float FeaturePointTime = f-Mu*2.0f;
+			FeaturePointMethod(Function, FeaturePointTime, f, Function.Eval(FeaturePointTime), Mu*0.4f,Depth+1, MaxDepth, MaxV, MinVal);
+		}
+
+		PrevTangent = CurTangent;
+		PrevValue = Value;
+	}
+}
+
+
+void FRichCurve::GetValueRange(float& MinValue, float& MaxValue) const
+{
+	if (Keys.Num() == 0)
+	{
+		MinValue = MaxValue = 0.f;
+	}
+	else
+	{
+		int32 LastKeyIndex = Keys.Num()-1;
+		MinValue = MaxValue = Keys[0].Value;
+
+		for (int32 i = 0; i < Keys.Num(); i++)
+		{
+			const FRichCurveKey& Key = Keys[i];
+
+			MinValue = FMath::Min(MinValue, Key.Value);
+			MaxValue = FMath::Max(MaxValue, Key.Value);
+
+			if (Key.InterpMode == RCIM_Cubic && i != LastKeyIndex)
+			{
+				const FRichCurveKey& NextKey = Keys[i+1];
+				float TimeStep = (NextKey.Time - Key.Time) * 0.2f;
+
+				FeaturePointMethod(*this, Key.Time, NextKey.Time, Key.Value, TimeStep, 0, 3, MaxValue, MinValue);
+			}
+		}
+	}
+}
+
+
+void FRichCurve::Reset()
+{
+	Keys.Empty();
+	KeyHandlesToIndices.Empty();
+}
+
+
+void FRichCurve::AutoSetTangents(float Tension)
+{
+	// Iterate over all points in this InterpCurve
+	for (int32 KeyIndex = 0; KeyIndex<Keys.Num(); KeyIndex++)
+	{
+		FRichCurveKey& Key = Keys[KeyIndex];
+		float ArriveTangent = Key.ArriveTangent;
+		float LeaveTangent  = Key.LeaveTangent;
+
+		if (KeyIndex == 0)
+		{
+			if (KeyIndex < Keys.Num()-1) // Start point
+			{
+				// If first section is not a curve, or is a curve and first point has manual tangent setting.
+				if (Key.TangentMode == RCTM_Auto)
+				{
+					LeaveTangent = 0.0f;
+				}
+			}
+		}
+		else
+		{
+			
+			if (KeyIndex < Keys.Num() - 1) // Inner point
+			{
+				FRichCurveKey& PrevKey =  Keys[KeyIndex-1];
+
+				if (Key.InterpMode == RCIM_Cubic && (Key.TangentMode == RCTM_Auto))
+				{
+						FRichCurveKey& NextKey =  Keys[KeyIndex+1];
+						ComputeCurveTangent(
+							Keys[ KeyIndex - 1 ].Time,		// Previous time
+							Keys[ KeyIndex - 1 ].Value,	// Previous point
+							Keys[ KeyIndex ].Time,			// Current time
+							Keys[ KeyIndex ].Value,		// Current point
+							Keys[ KeyIndex + 1 ].Time,		// Next time
+							Keys[ KeyIndex + 1 ].Value,	// Next point
+							Tension,							// Tension
+							false,						// Want clamping?
+							ArriveTangent );					// Out
+
+						// In 'auto' mode, arrive and leave tangents are always the same
+						LeaveTangent = ArriveTangent;
+				}
+				else if ((PrevKey.InterpMode == RCIM_Constant) || (Key.InterpMode == RCIM_Constant))
+				{
+					if (Keys[ KeyIndex - 1 ].InterpMode != RCIM_Cubic)
+					{
+						ArriveTangent = 0.0f;
+					}
+
+					LeaveTangent  = 0.0f;
+				}
+				
+			}
+			else // End point
+			{
+				// If last section is not a curve, or is a curve and final point has manual tangent setting.
+				if (Key.InterpMode == RCIM_Cubic && Key.TangentMode == RCTM_Auto)
+				{
+					ArriveTangent = 0.0f;
+				}
+			}
+		}
+
+		Key.ArriveTangent = ArriveTangent;
+		Key.LeaveTangent = LeaveTangent;
+	}
+}
+
+
+void FRichCurve::ReadjustTimeRange(float NewMinTimeRange, float NewMaxTimeRange, bool bInsert/* whether insert or remove*/, float OldStartTime, float OldEndTime)
+{
+	// first readjust modified time keys
+	float ModifiedDuration = OldEndTime - OldStartTime;
+
+	if (bInsert)
+	{
+		for(int32 KeyIndex=0; KeyIndex<Keys.Num(); ++KeyIndex)
+		{
+			float& CurrentTime = Keys[KeyIndex].Time;
+			if (CurrentTime >= OldStartTime)
+			{
+				CurrentTime += ModifiedDuration;
+			}
+		}
+	}
+	else
+	{
+		// since we only allow one key at a given time, we will just cache the value that needs to be saved
+		// this is the key to be replaced when this section is gone
+		bool bAddNewKey = false; 
+		float NewValue = 0.f;
+		TArray<int32> KeysToDelete;
+
+		for(int32 KeyIndex=0; KeyIndex<Keys.Num(); ++KeyIndex)
+		{
+			float& CurrentTime = Keys[KeyIndex].Time;
+			// if this key exists between range of deleted
+			// we'll evaluate the value at the "OldStartTime"
+			// and re-add key, so that it keeps the previous value at the
+			// start time
+			// But that means if there are multiple keys, 
+			// since we don't want multiple values in the same time
+			// the last one will override the value
+			if( CurrentTime >= OldStartTime && CurrentTime <= OldEndTime)
+			{
+				// get new value and add new key on one of OldStartTime, OldEndTime;
+				// this is a bit complicated problem since we don't know if OldStartTime or OldEndTime is preferred. 
+				// generall we use OldEndTime unless OldStartTime == 0.f
+				// which means it's cut in the beginning. Otherwise it will always use the end time. 
+				bAddNewKey = true;
+				if (OldStartTime != 0.f)
+				{
+					NewValue = Eval(OldStartTime);
+				}
+				else
+				{
+					NewValue = Eval(OldEndTime);
+				}
+				// remove this key, but later because it might change eval result
+				KeysToDelete.Add(KeyIndex);
+			}
+			else if (CurrentTime > OldEndTime)
+			{
+				CurrentTime -= ModifiedDuration;
+			}
+		}
+
+		if (bAddNewKey)
+		{
+			for (int32 KeyIndex = KeysToDelete.Num()-1; KeyIndex >= 0; --KeyIndex)
+			{
+				const FKeyHandle* KeyHandle = KeyHandlesToIndices.FindKey(KeysToDelete[KeyIndex]);
+				if(KeyHandle)
+				{
+					DeleteKey(*KeyHandle);
+				}
+			}
+
+			UpdateOrAddKey(OldStartTime, NewValue);
+		}
+	}
+
+	// now remove all redundant key
+	TArray<FRichCurveKey> NewKeys;
+	Exchange(NewKeys, Keys);
+
+	for(int32 KeyIndex=0; KeyIndex<NewKeys.Num(); ++KeyIndex)
+	{
+		UpdateOrAddKey(NewKeys[KeyIndex].Time, NewKeys[KeyIndex].Value);
+	}
+
+	// now cull out all out of range 
+	float MinTime, MaxTime;
+	GetTimeRange(MinTime, MaxTime);
+
+	bool bNeedToDeleteKey=false;
+
+	// if there is key below min time, just add key at new min range, 
+	if (MinTime < NewMinTimeRange)
+	{
+		float NewValue = Eval(NewMinTimeRange);
+		UpdateOrAddKey(NewMinTimeRange, NewValue);
+
+		bNeedToDeleteKey = true;
+	}
+
+	// if there is key after max time, just add key at new max range, 
+	if(MaxTime > NewMaxTimeRange)
+	{
+		float NewValue = Eval(NewMaxTimeRange);
+		UpdateOrAddKey(NewMaxTimeRange, NewValue);
+
+		bNeedToDeleteKey = true;
+	}
+
+	// delete the keys outside of range
+	if (bNeedToDeleteKey)
+	{
+		for (int32 KeyIndex=0; KeyIndex<Keys.Num(); ++KeyIndex)
+		{
+			if (Keys[KeyIndex].Time < NewMinTimeRange || Keys[KeyIndex].Time > NewMaxTimeRange)
+			{
+				const FKeyHandle* KeyHandle = KeyHandlesToIndices.FindKey(KeyIndex);
+				if (KeyHandle)
+				{
+					DeleteKey(*KeyHandle);
+					--KeyIndex;
+				}
+			}
+		}
+	}
+}
+
+void FRichCurve::BakeCurve(float SampleRate)
+{
+	if (Keys.Num() == 0)
+	{
+		return;
+	}
+
+	float FirstKeyTime = Keys[0].Time;
+	float LastKeyTime = Keys[Keys.Num()-1].Time;
+
+	BakeCurve(SampleRate, FirstKeyTime, LastKeyTime);
+}
+
+void FRichCurve::BakeCurve(float SampleRate, float FirstKeyTime, float LastKeyTime)
+{
+	if (Keys.Num() == 0)
+	{
+		return;
+	}
+
+	for (float Time = FirstKeyTime + SampleRate; Time < LastKeyTime;  )
+	{
+		float Value = Eval(Time);
+		UpdateOrAddKey(Time, Value);
+		Time += SampleRate;
+	}
+}
+
+void FRichCurve::RemoveRedundantKeys(float Tolerance)
+{
+	if (Keys.Num() == 0)
+	{
+		return;
+	}
+
+	float FirstKeyTime = Keys[0].Time;
+	float LastKeyTime = Keys[Keys.Num()-1].Time;
+
+	RemoveRedundantKeys(Tolerance, FirstKeyTime, LastKeyTime);
+}
+
+void FRichCurve::RemoveRedundantKeys(float Tolerance, float FirstKeyTime, float LastKeyTime)
+{
+	for(int32 KeyIndex = 0; KeyIndex < Keys.Num(); ++KeyIndex)
+	{
+		// copy the key
+		FRichCurveKey OriginalKey = Keys[KeyIndex];
+		if (OriginalKey.Time < FirstKeyTime || OriginalKey.Time > LastKeyTime)
+		{
+			continue;
+		}
+
+		FKeyHandle KeyHandle = GetKeyHandle(KeyIndex);
+
+		// remove it
+		DeleteKey(KeyHandle);
+
+		// eval to check validity
+		float NewValue = Eval(OriginalKey.Time, DefaultValue);
+			
+		// outside tolerance? re-add the key.
+		if(FMath::Abs(NewValue - OriginalKey.Value) > Tolerance)
+		{
+			FKeyHandle NewKeyHandle = AddKey(OriginalKey.Time, OriginalKey.Value, false, KeyHandle);
+			FRichCurveKey& NewKey = GetKey(NewKeyHandle);
+			NewKey = OriginalKey;
+		}
+		else
+		{
+			KeyIndex--;
+		}
+	}
+}
+
+/** Util to find float value on bezier defined by 4 control points */ 
+static float BezierInterp(float P0, float P1, float P2, float P3, float Alpha)
+{
+	const float P01 = FMath::Lerp(P0, P1, Alpha);
+	const float P12 = FMath::Lerp(P1, P2, Alpha);
+	const float P23 = FMath::Lerp(P2, P3, Alpha);
+	const float P012 = FMath::Lerp(P01, P12, Alpha);
+	const float P123 = FMath::Lerp(P12, P23, Alpha);
+	const float P0123 = FMath::Lerp(P012, P123, Alpha);
+
+	return P0123;
+}
+
+
+static float BezierInterp2(float P0, float Y1, float Y2, float P3, float mu)
+{
+	float P1 = (-5.f * P0 + 18.f * Y1 - 9.f * Y2 + 2.f * P3) / 6.f;
+	float P2 = (2.f * P0 -  9.f * Y1 + 18.f * Y2 - 5.f * P3) / 6.f;
+	float A = P3 - 3.0f * P2 + 3.0f * P1 - P0;
+	float B = 3.0f * P2 - 6.0f * P1 + 3.0f * P0;
+	float C = 3.0f * P1 - 3.0f * P0;
+	float D = P0;
+	float Result = A * (mu * mu * mu) + B * (mu * mu) + C * mu + D;
+
+	return Result;
+}
+
+
+static void CycleTime(float MinTime, float MaxTime, float& InTime, int& CycleCount)
+{
+	float InitTime = InTime;
+	float Duration = MaxTime - MinTime;
+
+	if (InTime > MaxTime)
+	{
+		CycleCount = FMath::FloorToInt((MaxTime-InTime)/Duration);
+		InTime = InTime + Duration*CycleCount;
+	}
+	else if (InTime < MinTime)
+	{
+		CycleCount = FMath::FloorToInt((InTime-MinTime)/Duration);
+		InTime = InTime - Duration*CycleCount;
+	}
+
+	if (InTime == MaxTime && InitTime < MinTime)
+	{
+		InTime = MinTime;
+	}
+
+	if (InTime == MinTime && InitTime > MaxTime)
+	{
+		InTime = MaxTime;
+	}
+
+	CycleCount = FMath::Abs(CycleCount);
+}
+
+
+void FRichCurve::RemapTimeValue(float& InTime, float& CycleValueOffset) const
+{
+	const int32 NumKeys = Keys.Num();
+	
+	if (NumKeys < 2)
+	{
+		return;
+	} 
+
+	if (InTime <= Keys[0].Time)
+	{
+		if (PreInfinityExtrap != RCCE_Linear && PreInfinityExtrap != RCCE_Constant)
+		{
+			float MinTime = Keys[0].Time;
+			float MaxTime = Keys[NumKeys - 1].Time;
+
+			int CycleCount = 0;
+			CycleTime(MinTime, MaxTime, InTime, CycleCount);
+
+			if (PreInfinityExtrap == RCCE_CycleWithOffset)
+			{
+				float DV = Keys[0].Value - Keys[NumKeys - 1].Value;
+				CycleValueOffset = DV * CycleCount;
+			}
+			else if (PreInfinityExtrap == RCCE_Oscillate)
+			{
+				if (CycleCount % 2 == 1)
+				{
+					InTime = MinTime + (MaxTime - InTime);
+				}
+			}
+		}
+	}
+	else if (InTime >= Keys[NumKeys - 1].Time)
+	{
+		if (PostInfinityExtrap != RCCE_Linear && PostInfinityExtrap != RCCE_Constant)
+		{
+			float MinTime = Keys[0].Time;
+			float MaxTime = Keys[NumKeys - 1].Time;
+
+			int CycleCount = 0; 
+			CycleTime(MinTime, MaxTime, InTime, CycleCount);
+
+			if (PostInfinityExtrap == RCCE_CycleWithOffset)
+			{
+				float DV = Keys[NumKeys - 1].Value - Keys[0].Value;
+				CycleValueOffset = DV * CycleCount;
+			}
+			else if (PostInfinityExtrap == RCCE_Oscillate)
+			{
+				if (CycleCount % 2 == 1)
+				{
+					InTime = MinTime + (MaxTime - InTime);
+				}
+			}
+		}
+	}
+}
+
+
+float FRichCurve::Eval(float InTime, float InDefaultValue) const
+{
+	SCOPE_CYCLE_COUNTER(STAT_RichCurve_Eval);
+
+	// Remap time if extrapolation is present and compute offset value to use if cycling 
+	float CycleValueOffset = 0;
+	RemapTimeValue(InTime, CycleValueOffset);
+
+	const int32 NumKeys = Keys.Num();
+
+	// If the default value hasn't been initialized, use the incoming default value
+	float InterpVal = DefaultValue == MAX_flt ? InDefaultValue : DefaultValue;
+
+	if (NumKeys == 0)
+	{
+		// If no keys in curve, return the Default value.
+	} 
+	else if (NumKeys < 2 || (InTime <= Keys[0].Time))
+	{
+		if (PreInfinityExtrap == RCCE_Linear && NumKeys > 1)
+		{
+			float DT = Keys[1].Time - Keys[0].Time;
+			
+			if (FMath::IsNearlyZero(DT))
+			{
+				InterpVal = Keys[0].Value;
+			}
+			else
+			{
+				float DV = Keys[1].Value - Keys[0].Value;
+				float Slope = DV / DT;
+
+				InterpVal = Slope * (InTime - Keys[0].Time) + Keys[0].Value;
+			}
+		}
+		else
+		{
+			// Otherwise if constant or in a cycle or oscillate, always use the first key value
+			InterpVal = Keys[0].Value;
+		}
+	}
+	else if (InTime < Keys[NumKeys - 1].Time)
+	{
+		// perform a lower bound to get the second of the interpolation nodes
+		int32 first = 1;
+		int32 last = NumKeys - 1;
+		int32 count = last - first;
+
+		while (count > 0)
+		{
+			int32 step = count / 2;
+			int32 middle = first + step;
+
+			if (InTime >= Keys[middle].Time)
+			{
+				first = middle + 1;
+				count -= step + 1;
+			}
+			else
+			{
+				count = step;
+			}
+		}
+
+		int32 InterpNode = first;
+		const float Diff = Keys[InterpNode].Time - Keys[InterpNode - 1].Time;
+
+		if (Diff > 0.f && Keys[InterpNode - 1].InterpMode != RCIM_Constant)
+		{
+			const float Alpha = (InTime - Keys[InterpNode - 1].Time) / Diff;
+			const float P0 = Keys[InterpNode - 1].Value;
+			const float P3 = Keys[InterpNode].Value;
+
+			if (Keys[InterpNode - 1].InterpMode == RCIM_Linear)
+			{
+				InterpVal = FMath::Lerp(P0, P3, Alpha);
+			}
+			else
+			{
+				const float OneThird = 1.0f / 3.0f;
+				const float P1 = P0 + (Keys[InterpNode - 1].LeaveTangent * Diff*OneThird);
+				const float P2 = P3 - (Keys[InterpNode].ArriveTangent * Diff*OneThird);
+
+				InterpVal = BezierInterp(P0, P1, P2, P3, Alpha);
+			}
+		}
+		else
+		{
+			InterpVal = Keys[InterpNode - 1].Value;
+		}
+	}
+	else
+	{
+		if (PostInfinityExtrap == RCCE_Linear)
+		{
+			float DT = Keys[NumKeys - 2].Time - Keys[NumKeys - 1].Time;
+			
+			if (FMath::IsNearlyZero(DT))
+			{
+				InterpVal = Keys[NumKeys - 1].Value;
+			}
+			else
+			{
+				float DV = Keys[NumKeys - 2].Value - Keys[NumKeys - 1].Value;
+				float Slope = DV / DT;
+
+				InterpVal = Slope * (InTime - Keys[NumKeys - 1].Time) + Keys[NumKeys - 1].Value;
+			}
+		}
+		else
+		{
+			// Otherwise if constant or in a cycle or oscillate, always use the last key value
+			InterpVal = Keys[NumKeys - 1].Value;
+		}
+	}
+
+	return InterpVal+CycleValueOffset;
+}
+
+
+bool FRichCurve::operator==(const FRichCurve& Curve) const
+{
+	if(Keys.Num() != Curve.Keys.Num())
+	{
+		return false;
+	}
+
+	for(int32 i = 0;i<Keys.Num();++i)
+	{
+		if(!(Keys[i] == Curve.Keys[i]))
+		{
+			return false;
+		}
+	}
+
+	if (PreInfinityExtrap != Curve.PreInfinityExtrap || PostInfinityExtrap != Curve.PostInfinityExtrap)
+	{
+		return false;
+	}
+
+	return true;
+}
\ No newline at end of file
diff --git a/Engine/Source/Runtime/Engine/Private/Curves/StringCurve.cpp b/Engine/Source/Runtime/Engine/Private/Curves/StringCurve.cpp
new file mode 100644
--- /dev/null
+++ b/Engine/Source/Runtime/Engine/Private/Curves/StringCurve.cpp
@@ -0,0 +1,279 @@
+// Copyright 1998-2016 Epic Games, Inc. All Rights Reserved.
+
+#include "EnginePrivate.h"
+#include "Curves/StringCurve.h"
+
+
+/* FStringCurveKey interface
+ *****************************************************************************/
+
+bool FStringCurveKey::operator==(const FStringCurveKey& Curve) const
+{
+	return ((Time == Curve.Time) && (Value == Curve.Value));
+}
+
+
+bool FStringCurveKey::operator!=(const FStringCurveKey& Other) const
+{
+	return !(*this == Other);
+}
+
+
+bool FStringCurveKey::Serialize(FArchive& Ar)
+{
+	Ar << Time << Value;
+	return true;
+}
+
+
+/* FStringCurve interface
+ *****************************************************************************/
+
+FKeyHandle FStringCurve::AddKey(float InTime, const FString& InValue, FKeyHandle KeyHandle)
+{
+	int32 Index = 0;
+
+	// insert key
+	for(; Index < Keys.Num() && Keys[Index].Time < InTime; ++Index);
+	Keys.Insert(FStringCurveKey(InTime, InValue), Index);
+
+	// update key indices
+	for ( auto It = KeyHandlesToIndices.CreateIterator(); It; ++It)
+	{
+		int32& KeyIndex = It.Value();
+		if (KeyIndex >= Index) {++KeyIndex;}
+	}
+
+	KeyHandlesToIndices.Add(KeyHandle, Index);
+
+	return GetKeyHandle(Index);
+}
+
+
+void FStringCurve::DeleteKey(FKeyHandle KeyHandle)
+{
+	// remove key
+	int32 Index = GetIndex(KeyHandle);
+	Keys.RemoveAt(Index);
+
+	KeyHandlesToIndices.Remove(KeyHandle);
+
+	// update key indices
+	for (auto It = KeyHandlesToIndices.CreateIterator(); It; ++It)
+	{
+		int32& KeyIndex = It.Value();
+		if (KeyIndex >= Index)
+		{
+			--KeyIndex;
+		}
+	}
+}
+
+
+FString FStringCurve::Eval(float Time, const FString& InDefaultValue) const
+{
+	// If the default value is empty, use the incoming default value
+	FString ReturnVal = DefaultValue.IsEmpty() ? InDefaultValue : DefaultValue;
+
+	if (Keys.Num() == 0 || Time < Keys[0].Time)
+	{
+		// If no keys in curve, or the time is before the first key, return the Default value.
+	}
+	else if (Keys.Num() < 2 || Time < Keys[0].Time)
+	{
+		// There is only one key or the time is before the first value. Return the first value
+		ReturnVal = Keys[0].Value;
+	}
+	else if (Time < Keys[Keys.Num() - 1].Time)
+	{
+		// The key is in the range of Key[0] to Keys[Keys.Num()-1].  Find it by searching
+		for (int32 i = 0; i < Keys.Num(); ++i)
+		{
+			if (Time < Keys[i].Time)
+			{
+				ReturnVal = Keys[FMath::Max(0, i - 1)].Value;
+				break;
+			}
+		}
+	}
+	else
+	{
+		// Key is beyond the last point in the curve.  Return it's value
+		ReturnVal = Keys[Keys.Num() - 1].Value;
+	}
+
+	return ReturnVal;
+}
+
+
+FKeyHandle FStringCurve::FindKey(float KeyTime, float KeyTimeTolerance) const
+{
+	int32 Start = 0;
+	int32 End = Keys.Num()-1;
+
+	// Binary search since the keys are in sorted order
+	while (Start <= End)
+	{
+		int32 TestPos = Start + (End-Start) / 2;
+		float TestKeyTime = Keys[TestPos].Time;
+
+		if (FMath::IsNearlyEqual(TestKeyTime, KeyTime, KeyTimeTolerance))
+		{
+			return GetKeyHandle(TestPos);
+		}
+
+		if (TestKeyTime < KeyTime)
+		{
+			Start = TestPos+1;
+		}
+		else
+		{
+			End = TestPos-1;
+		}
+	}
+
+	return FKeyHandle();
+}
+
+
+FStringCurveKey& FStringCurve::GetKey(FKeyHandle KeyHandle)
+{
+	EnsureAllIndicesHaveHandles();
+	return Keys[GetIndex(KeyHandle)];
+}
+
+
+FStringCurveKey FStringCurve::GetKey(FKeyHandle KeyHandle) const
+{
+	EnsureAllIndicesHaveHandles();
+	return Keys[GetIndex(KeyHandle)];
+}
+
+
+float FStringCurve::GetKeyTime(FKeyHandle KeyHandle) const
+{
+	if (!IsKeyHandleValid(KeyHandle))
+	{
+		return 0.f;
+	}
+
+	return GetKey(KeyHandle).Time;
+}
+
+
+FKeyHandle FStringCurve::SetKeyTime(FKeyHandle KeyHandle, float NewTime)
+{
+	if (!IsKeyHandleValid(KeyHandle))
+	{
+		return KeyHandle;
+	}
+
+	FStringCurveKey OldKey = GetKey(KeyHandle);
+	
+	DeleteKey(KeyHandle);
+	AddKey(NewTime, OldKey.Value, KeyHandle);
+
+	// Copy all properties from old key, but then fix time to be the new time
+	GetKey(KeyHandle) = OldKey;
+	GetKey(KeyHandle).Time = NewTime;
+
+	return KeyHandle;
+}
+
+
+void FStringCurve::ShiftCurve(float DeltaTime)
+{
+	TSet<FKeyHandle> KeyHandles;
+	for (auto It = KeyHandlesToIndices.CreateIterator(); It; ++It)
+	{
+		FKeyHandle& KeyHandle = It.Key();
+		KeyHandles.Add(KeyHandle);
+	}
+
+	ShiftCurve(DeltaTime, KeyHandles);
+}
+
+
+void FStringCurve::ShiftCurve(float DeltaTime, TSet<FKeyHandle>& KeyHandles)
+{
+	for (auto It = KeyHandlesToIndices.CreateIterator(); It; ++It)
+	{
+		const FKeyHandle& KeyHandle = It.Key();
+		if (KeyHandles.Num() != 0 && KeyHandles.Contains(KeyHandle))
+		{
+			SetKeyTime(KeyHandle, GetKeyTime(KeyHandle) + DeltaTime);
+		}
+	}
+}
+
+
+void FStringCurve::ScaleCurve(float ScaleOrigin, float ScaleFactor)
+{
+	TSet<FKeyHandle> KeyHandles;
+	for (auto It = KeyHandlesToIndices.CreateIterator(); It; ++It)
+	{
+		FKeyHandle& KeyHandle = It.Key();
+		KeyHandles.Add(KeyHandle);
+	}
+
+	ScaleCurve(ScaleOrigin, ScaleFactor, KeyHandles);
+}
+
+
+void FStringCurve::ScaleCurve(float ScaleOrigin, float ScaleFactor, TSet<FKeyHandle>& KeyHandles)
+{
+	for (auto It = KeyHandlesToIndices.CreateIterator(); It; ++It)
+	{
+		const FKeyHandle& KeyHandle = It.Key();
+		if (KeyHandles.Num() != 0 && KeyHandles.Contains(KeyHandle))
+		{
+			SetKeyTime(KeyHandle, (GetKeyTime(KeyHandle) - ScaleOrigin) * ScaleFactor + ScaleOrigin);
+		}
+	}
+}
+
+
+FKeyHandle FStringCurve::UpdateOrAddKey(float InTime, const FString& InValue, float KeyTimeTolerance)
+{
+	// Search for a key that already exists at the time and if found, update its value
+	for (int32 KeyIndex = 0; KeyIndex < Keys.Num(); ++KeyIndex)
+	{
+		float KeyTime = Keys[KeyIndex].Time;
+
+		if (FMath::IsNearlyEqual(KeyTime, InTime, KeyTimeTolerance))
+		{
+			Keys[KeyIndex].Value = InValue;
+
+			return GetKeyHandle(KeyIndex);
+		}
+
+		if (KeyTime > InTime)
+		{
+			// All the rest of the keys exist after the key we want to add
+			// so there is no point in searching
+			break;
+		}
+	}
+
+	// a key wasn't found, add it now
+	return AddKey(InTime, InValue);
+}
+
+
+int32 FStringCurve::GetNumKeys() const
+{
+	return Keys.Num();
+}
+
+
+bool FStringCurve::IsKeyHandleValid(FKeyHandle KeyHandle) const
+{
+	bool bValid = false;
+
+	if (FIndexedCurve::IsKeyHandleValid(KeyHandle))
+	{
+		bValid = Keys.IsValidIndex(GetIndex(KeyHandle));
+	}
+
+	return bValid;
+}
diff --git a/Engine/Source/Runtime/Launch/Private/LaunchEngineLoop.cpp b/Engine/Source/Runtime/Launch/Private/LaunchEngineLoop.cpp
--- a/Engine/Source/Runtime/Launch/Private/LaunchEngineLoop.cpp
+++ b/Engine/Source/Runtime/Launch/Private/LaunchEngineLoop.cpp
@@ -2193,7 +2193,8 @@ bool FEngineLoop::LoadStartupCoreModules()
 	// Load runtime client modules (which are also needed at cook-time)
 	if( !IsRunningDedicatedServer() )
 	{
-		FModuleManager::Get().LoadModule( TEXT( "GameLiveStreaming" ) );
+		FModuleManager::Get().LoadModule(TEXT("GameLiveStreaming"));
+		FModuleManager::Get().LoadModule(TEXT("MediaAssets"));
 	}
 #endif
 
diff --git a/Engine/Source/Runtime/LevelSequence/Private/LevelSequenceActor.cpp b/Engine/Source/Runtime/LevelSequence/Private/LevelSequenceActor.cpp
--- a/Engine/Source/Runtime/LevelSequence/Private/LevelSequenceActor.cpp
+++ b/Engine/Source/Runtime/LevelSequence/Private/LevelSequenceActor.cpp
@@ -32,24 +32,27 @@ ALevelSequenceActor::ALevelSequenceActor(const FObjectInitializer& Init)
 			SetRootComponent(SpriteComponent);
 		}
 	}
-#endif // WITH_EDITORONLY_DATA
+#endif //WITH_EDITORONLY_DATA
 
 	BurnInOptions = Init.CreateDefaultSubobject<ULevelSequenceBurnInOptions>(this, "BurnInOptions");
 	PrimaryActorTick.bCanEverTick = true;
 	bAutoPlay = false;
 }
 
+
 void ALevelSequenceActor::BeginPlay()
 {
 	Super::BeginPlay();
 	InitializePlayer();
 }
 
+
 #if WITH_EDITOR
 
 bool ALevelSequenceActor::GetReferencedContentObjects(TArray<UObject*>& Objects) const
 {
-	ULevelSequence* LevelSequenceAsset = Cast<ULevelSequence>(LevelSequence.TryLoad());
+	ULevelSequence* LevelSequenceAsset = GetSequence(true);
+
 	if (LevelSequenceAsset)
 	{
 		Objects.Add(LevelSequenceAsset);
@@ -60,7 +63,8 @@ bool ALevelSequenceActor::GetReferencedContentObjects(TArray<UObject*>& Objects)
 	return true;
 }
 
-#endif // WITH_EDITOR
+#endif //WITH_EDITOR
+
 
 void ALevelSequenceActor::Tick(float DeltaSeconds)
 {
@@ -70,6 +74,13 @@ void ALevelSequenceActor::Tick(float DeltaSeconds)
 	}
 }
 
+
+ULevelSequence* ALevelSequenceActor::GetSequence(bool Load) const
+{
+	return Cast<ULevelSequence>(Load ? LevelSequence.TryLoad() : LevelSequence.ResolveObject());
+}
+
+
 void ALevelSequenceActor::SetSequence(ULevelSequence* InSequence)
 {
 	if (!SequencePlayer || !SequencePlayer->IsPlaying())
@@ -79,10 +90,12 @@ void ALevelSequenceActor::SetSequence(ULevelSequence* InSequence)
 	}
 }
 
+
 void ALevelSequenceActor::InitializePlayer()
 {
-	ULevelSequence* LevelSequenceAsset = Cast<ULevelSequence>(LevelSequence.TryLoad());
-	if (GetWorld()->IsGameWorld() && LevelSequenceAsset)
+	ULevelSequence* LevelSequenceAsset = GetSequence(true);
+
+	if (GetWorld()->IsGameWorld() && (LevelSequenceAsset != nullptr))
 	{
 		SequencePlayer = NewObject<ULevelSequencePlayer>(this, "AnimationPlayer");
 		SequencePlayer->Initialize(LevelSequenceAsset, GetWorld(), PlaybackSettings);
diff --git a/Engine/Source/Runtime/LevelSequence/Private/LevelSequencePlayer.cpp b/Engine/Source/Runtime/LevelSequence/Private/LevelSequencePlayer.cpp
--- a/Engine/Source/Runtime/LevelSequence/Private/LevelSequencePlayer.cpp
+++ b/Engine/Source/Runtime/LevelSequence/Private/LevelSequencePlayer.cpp
@@ -3,6 +3,7 @@
 #include "LevelSequencePCH.h"
 #include "LevelSequencePlayer.h"
 #include "MovieScene.h"
+#include "MovieSceneCommonHelpers.h"
 #include "MovieSceneSubSection.h"
 #include "MovieSceneSequence.h"
 #include "MovieSceneSequenceInstance.h"
@@ -459,11 +460,21 @@ void ULevelSequencePlayer::UpdateCameraCut(UObject* CameraObject, UObject* Unloc
 		LastViewTarget = ViewTarget;
 	}
 
+	UCameraComponent* CameraComponent = MovieSceneHelpers::CameraComponentFromRuntimeObject(CameraObject);
+
 	if (CameraObject == ViewTarget)
 	{
-		if ( bJumpCut && PC->PlayerCameraManager )
+		if ( bJumpCut )
 		{
-			PC->PlayerCameraManager->bGameCameraCutThisFrame = true;
+			if (PC->PlayerCameraManager)
+			{
+				PC->PlayerCameraManager->bGameCameraCutThisFrame = true;
+			}
+
+			if (CameraComponent)
+			{
+				CameraComponent->NotifyCameraCut();
+			}
 		}
 		return;
 	}
@@ -488,6 +499,11 @@ void ULevelSequencePlayer::UpdateCameraCut(UObject* CameraObject, UObject* Unloc
 	FViewTargetTransitionParams TransitionParams;
 	PC->SetViewTarget(CameraActor, TransitionParams);
 
+	if (CameraComponent)
+	{
+		CameraComponent->NotifyCameraCut();
+	}
+
 	if (PC->PlayerCameraManager)
 	{
 		PC->PlayerCameraManager->bClientSimulatingViewTarget = (CameraActor != nullptr);
diff --git a/Engine/Source/Runtime/LevelSequence/Public/LevelSequenceActor.h b/Engine/Source/Runtime/LevelSequence/Public/LevelSequenceActor.h
--- a/Engine/Source/Runtime/LevelSequence/Public/LevelSequenceActor.h
+++ b/Engine/Source/Runtime/LevelSequence/Public/LevelSequenceActor.h
@@ -75,6 +75,22 @@ public:
 
 public:
 
+	/**
+	 * Get the level sequence being played by this actor.
+	 *
+	 * @param Whether to load the sequence object if it is not already in memory.
+	 * @return Level sequence, or nullptr if not assigned or if it cannot be loaded.
+	 * @see SetSequence
+	 */
+	UFUNCTION(BlueprintCallable, Category="Game|Cinematic")
+	ULevelSequence* GetSequence(bool Load = false) const;
+
+	/**
+	 * Set the level sequence being played by this actor.
+	 *
+	 * @param InSequence The sequence object to set.
+	 * @see GetSequence
+	 */
 	UFUNCTION(BlueprintCallable, Category="Game|Cinematic")
 	void SetSequence(ULevelSequence* InSequence);
 
diff --git a/Engine/Source/Runtime/Media/Private/MediaModule.cpp b/Engine/Source/Runtime/Media/Private/MediaModule.cpp
--- a/Engine/Source/Runtime/Media/Private/MediaModule.cpp
+++ b/Engine/Source/Runtime/Media/Private/MediaModule.cpp
@@ -1,10 +1,8 @@
 // Copyright 1998-2016 Epic Games, Inc. All Rights Reserved.
 
-#include "MediaPrivatePCH.h"
+#include "MediaPCH.h"
 #include "IMediaModule.h"
 #include "IMediaPlayerFactory.h"
-#include "IMediaPlayer.h"
-#include "ModuleManager.h"
 
 
 /**
@@ -15,88 +13,54 @@ class FMediaModule
 {
 public:
 
-	// IModuleInterface interface
+	//~ IMediaModule interface
 
-	virtual void StartupModule( ) override { }
-
-	virtual void ShutdownModule( ) override
+	virtual const TArray<IMediaPlayerFactory*>& GetPlayerFactories() const override
 	{
-		MediaPlayerFactories.Reset();
+		return MediaPlayerFactories;
 	}
 
-	virtual bool SupportsDynamicReloading( ) override
+	virtual IMediaPlayerFactory* GetPlayerFactory(const FName& FactoryName) const override
 	{
-		return false;
-	}
-
-public:
-
-	// IMediaModuleInterface
-
-	virtual TSharedPtr<IMediaPlayer> CreatePlayer( const FString& Url ) override
-	{
-		TSharedPtr<IMediaPlayer> Player;
-
-		for (IMediaPlayerFactory* PlayerFactory : MediaPlayerFactories)
+		for (IMediaPlayerFactory* Factory : MediaPlayerFactories)
 		{
-			if (PlayerFactory->SupportsUrl(Url))
+			if (Factory->GetName() == FactoryName)
 			{
-				Player = PlayerFactory->CreatePlayer();
-
-				if (Player.IsValid())
-				{
-					break;
-				}
+				return Factory;
 			}
 		}
 
-		return Player;
+		return nullptr;
 	}
 
-	virtual int32 GetSupportedFileTypes( FMediaFileTypes& OutFileTypes ) override
+	virtual void RegisterPlayerFactory(IMediaPlayerFactory& Factory) override
 	{
-		OutFileTypes.Reset();
-
-		for (IMediaPlayerFactory* Factory : MediaPlayerFactories)
-		{
-			OutFileTypes.Append(Factory->GetSupportedFileTypes());
-		}
-	
-		return OutFileTypes.Num();
+		MediaPlayerFactories.AddUnique(&Factory);
 	}
 
-	DECLARE_DERIVED_EVENT(FMediaModule, IMediaModule::FOnFactoryAdded, FOnFactoryAdded);
-	virtual FOnFactoryAdded& OnFactoryAdded( ) override
+	virtual void UnregisterPlayerFactory(IMediaPlayerFactory& Factory) override
 	{
-		return FactoryAddedEvent;
+		MediaPlayerFactories.Remove(&Factory);
 	}
 
-	DECLARE_DERIVED_EVENT(FMediaModule, IMediaModule::FOnFactoryRemoved, FOnFactoryRemoved);
-	virtual FOnFactoryRemoved& OnFactoryRemoved( ) override
-	{
-		return FactoryRemovedEvent;
-	}
+public:
+
+	//~ IModuleInterface interface
 
-	virtual void RegisterPlayerFactory( IMediaPlayerFactory& Factory ) override
+	virtual void StartupModule() override { }
+
+	virtual void ShutdownModule() override
 	{
-		MediaPlayerFactories.AddUnique(&Factory);
-		FactoryAddedEvent.Broadcast();
+		MediaPlayerFactories.Reset();
 	}
 
-	virtual void UnregisterPlayerFactory( IMediaPlayerFactory& Factory ) override
+	virtual bool SupportsDynamicReloading() override
 	{
-		MediaPlayerFactories.Remove(&Factory);
-		FactoryRemovedEvent.Broadcast();
+		return false;
 	}
 
 private:
 
-	/** Holds an event delegate that is invoked after a video player factory has been added. */
-	FOnFactoryAdded FactoryAddedEvent;
-
-	/** Holds an event delegate that is invoked after a video player factory has been removed. */
-	FOnFactoryRemoved FactoryRemovedEvent;
-
 	/** Holds the registered video player factories. */
 	TArray<IMediaPlayerFactory*> MediaPlayerFactories;
 };
diff --git a/Engine/Source/Runtime/Media/Private/MediaPCH.h b/Engine/Source/Runtime/Media/Private/MediaPCH.h
new file mode 100644
--- /dev/null
+++ b/Engine/Source/Runtime/Media/Private/MediaPCH.h
@@ -0,0 +1,6 @@
+// Copyright 1998-2016 Epic Games, Inc. All Rights Reserved.
+
+#pragma once
+
+#include "Runtime/Core/Public/Core.h"
+#include "Runtime/Core/Public/Modules/ModuleManager.h"
diff --git a/Engine/Source/Runtime/Media/Public/IMediaAudioSink.h b/Engine/Source/Runtime/Media/Public/IMediaAudioSink.h
new file mode 100644
--- /dev/null
+++ b/Engine/Source/Runtime/Media/Public/IMediaAudioSink.h
@@ -0,0 +1,97 @@
+// Copyright 1998-2015 Epic Games, Inc. All Rights Reserved.
+
+#pragma once
+
+
+/**
+ * Available sample formats for media audio sinks.
+ */
+enum class EMediaAudioSinkFormat
+{
+	/** Uncompressed IEEE Floating Point samples. */
+	Float,
+
+	/** Uncompressed pulse code modulated samples.*/
+	Pcm,
+};
+
+
+/**
+ * Interface for media sinks that receive audio sample data.
+ *
+ * Implementations of this interface must be thread-safe as these
+ * methods will be called from some random media decoder thread.
+ */
+class IMediaAudioSink
+{
+public:
+
+	/**
+	 * Discard the sink's buffer and stop audio playback.
+	 *
+	 * @see PauseAudioSink, PlayAudioSink, ResumeAudioSink
+	 */
+	virtual void FlushAudioSink() = 0;
+
+	/**
+	 * Get the sink's current number of channels.
+	 *
+	 * @return Number of audio channels.
+	 * @see GetAudioSinkSampleRate, InitializeAudioSink
+	 */
+	virtual int32 GetAudioSinkChannels() const = 0;
+
+	/**
+	 * Get the sink's current sample rate.
+	 *
+	 * @return Samples per second.
+	 * @see GetAudioSinkChannels, InitializeAudioSink
+	 */
+	virtual int32 GetAudioSinkSampleRate() const = 0;
+
+	/**
+	 * Initialize the sink.
+	 *
+	 * @param Channels Number of audio channels.
+	 * @param SampleRate Number of samples per second.
+	 * @return true on success, false otherwise.
+	 * @see ShutdownAudioSink
+	 */
+	virtual bool InitializeAudioSink(uint32 Channels, uint32 SampleRate) = 0;
+
+	/**
+	 * Pause audio playback.
+	 *
+	 * @see FlushAudioSink, PlayAudioSink, ResumeAudioSink
+	 */
+	virtual void PauseAudioSink() = 0;
+
+	/**
+	 * Play the given samples.
+	 *
+	 * @param SampleBuffer The sample buffer.
+	 * @param BufferSize The number of bytes in the buffer.
+	 * @param Time The time of the sample (relative to the beginning of the media).
+	 * @see FlushAudioSink, PauseAudioSink, ResumeAudioSink
+	 */
+	virtual void PlayAudioSink(const uint8* SampleBuffer, uint32 BufferSize, FTimespan Time) = 0;
+
+	/**
+	 * Resume audio playback.
+	 *
+	 * @see FlushAudioSink, PauseAudioSink, PlayAudioSink
+	 */
+	virtual void ResumeAudioSink() = 0;
+
+	/**
+	 * Shut down the sink.
+	 *
+	 * @see InitializeAudioSink
+	 */
+	virtual void ShutdownAudioSink() = 0;
+
+public:
+
+	/** Virtual destructor. */
+	virtual ~IMediaAudioSink() { }
+};
diff --git a/Engine/Source/Runtime/Media/Public/IMediaControls.h b/Engine/Source/Runtime/Media/Public/IMediaControls.h
new file mode 100644
--- /dev/null
+++ b/Engine/Source/Runtime/Media/Public/IMediaControls.h
@@ -0,0 +1,246 @@
+// Copyright 1998-2016 Epic Games, Inc. All Rights Reserved.
+
+#pragma once
+
+
+/**
+ * Enumerates directions for playback rate information.
+ */
+enum class EMediaPlaybackDirections
+{
+	/** Forward playback rates. */
+	Forward,
+
+	/** Reverse playback rates. */
+	Reverse
+};
+
+
+/**
+ * Enumerates directions for seeking in media.
+ */
+enum class EMediaSeekDirection
+{
+	/** Seek backwards from current position. */
+	Backward,
+
+	/** Seek from the beginning of the media. */
+	Beginning,
+
+	/** Seek from the end of the media. */
+	End,
+
+	/** Seek forward from current position. */
+	Forward
+};
+
+
+/**
+ * Enumerates possible states of media playback.
+ */
+enum class EMediaState
+{
+	/** Media has been closed and cannot be played again. */
+	Closed,
+
+	/** Unrecoverable error occurred during playback. */
+	Error,
+
+	/** Playback has been paused, but can be resumed. */
+	Paused,
+
+	/** Media is currently playing. */
+	Playing,
+
+	/** Media is loading or buffering. */
+	Preparing,
+
+	/** Playback has been stopped, but can be restarted. */
+	Stopped
+};
+
+
+/**
+ * Interface for controlling media playback.
+ *
+ * @see IMediaOutput, IMediaTracks
+ */
+class IMediaControls
+{
+public:
+
+	/**
+	 * Get the media's duration.
+	 *
+	 * @return A time span representing the duration.
+	 * @see GetTime
+	 */
+	virtual FTimespan GetDuration() const = 0;
+
+	/**
+	 * Get the nominal playback rate, i.e. 1.0 for real time.
+	 *
+	 * @return Playback rate.
+	 * @see Pause, Play, SetRate
+	 */
+	virtual float GetRate() const = 0;
+
+	/**
+	 * Get the state of the media.
+	 *
+	 * @return Media state.
+	 */
+	virtual EMediaState GetState() const = 0;
+
+	/**
+	 * Get the range of supported playback rates in the specified playback direction.
+	 *
+	 * @param Direction The playback direction.
+	 * @param Unthinned Whether the rates are for unthinned playback.
+	 * @return The range of supported rates.
+	 * @see SupportsRate, SetRate
+	 */
+	virtual TRange<float> GetSupportedRates(EMediaPlaybackDirections Direction, bool Unthinned) const = 0;
+
+	/**
+	 * Get the player's current playback time.
+	 *
+	 * @return Playback time.
+	 * @see Seek
+	 */
+	virtual FTimespan GetTime() const = 0;
+
+	/**
+	 * Check whether playback is currently looping.
+	 *
+	 * @return true if playback is looping, false otherwise.
+	 * @see SetLooping
+	 */
+	virtual bool IsLooping() const = 0;
+
+	/**
+	 * Change the media's playback time.
+	 *
+	 * @param Time The playback time to set.
+	 * @return true on success, false otherwise.
+	 * @see GetTime
+	 */
+	virtual bool Seek(const FTimespan& Time) = 0;
+
+	/**
+	 * Set whether playback should be looping.
+	 *
+	 * @param Looping Enables or disables looping.
+	 * @see IsLooping
+	 */
+	virtual bool SetLooping(bool Looping) = 0;
+
+	/**
+	 * Set the current playback rate.
+	 *
+	 * A playback rate of 1.0 will play the media normally at real-time.
+	 * A rate of 0.0 corresponds to pausing playback. A negative rate, if
+	 * supported, plays the media in reverse, and a rate larger than 1.0
+	 * fast forwards playback.
+	 *
+	 * @param Rate The playback rate to set.
+	 * @return true on success, false otherwise.
+	 * @see GetRate, Pause, Play
+	 */
+	virtual bool SetRate(float Rate) = 0;
+
+	/**
+	 * Check whether the specified playback rate is supported.
+	 *
+	 * @param Rate The rate to check (can be negative for reverse play).
+	 * @param Unthinned Whether no frames should be dropped at the given rate.
+	 * @return true if the rate is supported, false otherwise.
+	 * @see GetSupportedRates, SetRate, SupportsScrubbing, SupportsSeeking
+	 */
+	virtual bool SupportsRate(float Rate, bool Unthinned) const = 0;
+
+	/**
+	 * Check whether scrubbing is supported.
+	 *
+	 * Scrubbing is the ability to decode video frames while seeking in a media item at a playback rate of 0.0.
+	 *
+	 * @return true if scrubbing is supported, false otherwise.
+	 * @see SupportsRate, SupportsSeeking
+	 */
+	virtual bool SupportsScrubbing() const = 0;
+
+	/**
+	 * Check whether the currently loaded media can jump to certain times.
+	 *
+	 * @return true if seeking is supported, false otherwise.
+	 * @see SupportsRate, SupportsScrubbing
+	 */
+	virtual bool SupportsSeeking() const = 0;
+
+public:
+
+	/**
+	 * Pause media playback.
+	 *
+	 * This is the same as setting the playback rate to 0.0.
+	 *
+	 * @return true if the media is being paused, false otherwise.
+	 * @see Play, Stop
+	 */
+	FORCEINLINE bool Pause()
+	{
+		return SetRate(0.0f);
+	}
+
+	/**
+	 * Start media playback at the default rate of 1.0.
+	 *
+	 * This is the same as setting the playback rate to 1.0.
+	 *
+	 * @return true if playback is starting, false otherwise.
+	 * @see Pause, Stop
+	 */
+	FORCEINLINE bool Play()
+	{
+		return SetRate(1.0f);
+	}
+
+	/**
+	 * Change the playback time of the media by a relative offset in the given direction.
+	 *
+	 * @param TimeOffset The offset to apply to the time.
+	 * @param Direction The direction to seek in.
+	 * @return true on success, false otherwise.
+	 * @see GetDuration, GetTime
+	 */
+	bool Seek(const FTimespan& TimeOffset, EMediaSeekDirection Direction)
+	{
+		FTimespan SeekTime;
+
+		switch (Direction)
+		{
+		case EMediaSeekDirection::Backward:
+			SeekTime = GetTime() - TimeOffset;
+			break;
+
+		case EMediaSeekDirection::Beginning:
+			SeekTime = TimeOffset;
+			break;
+
+		case EMediaSeekDirection::End:
+			SeekTime = GetDuration() - TimeOffset;
+			break;
+
+		case EMediaSeekDirection::Forward:
+			SeekTime = GetTime() + TimeOffset;
+			break;
+		}
+
+		return Seek(SeekTime);
+	}
+
+public:
+
+	/** Virtual destructor. */
+	virtual ~IMediaControls() { }
+};
diff --git a/Engine/Source/Runtime/Media/Public/IMediaModule.h b/Engine/Source/Runtime/Media/Public/IMediaModule.h
--- a/Engine/Source/Runtime/Media/Public/IMediaModule.h
+++ b/Engine/Source/Runtime/Media/Public/IMediaModule.h
@@ -3,16 +3,14 @@
 #pragma once
 
 #include "ModuleInterface.h"
-#include "ModuleManager.h"
 
 
-// forward declarations
 class IMediaPlayer;
 class IMediaPlayerFactory;
 
 
 /**
- * Interface for media modules.
+ * Interface for the Media module.
  */
 class IMediaModule
 	: public IModuleInterface
@@ -20,18 +18,21 @@ class IMediaModule
 public:
 
 	/**
-	 * Creates a media player for the specified media URL.
+	 * Get the list of installed media player factories.
 	 *
-	 * @param Url The URL to the media file container to create the player for.
+	 * @return Collection of media player factories.
+	 * @see GetPlayerInfo
 	 */
-	virtual TSharedPtr<IMediaPlayer> CreatePlayer(const FString& Url) = 0;
+	virtual const TArray<IMediaPlayerFactory*>& GetPlayerFactories() const = 0;
 
 	/**
-	 * Gets the collection of supported media file types.
+	 * Get a media player factory by name.
 	 *
-	 * @param OutFileTypes Will hold the supported file types.
+	 * @param FactoryName The name of the factory.
+	 * @return The factory, or nullptr if not found.
+	 * @see GetPlayerFactories
 	 */
-	virtual int32 GetSupportedFileTypes(TMap<FString, FText>& OutFileTypes) = 0;
+	virtual IMediaPlayerFactory* GetPlayerFactory(const FName& FactoryName) const = 0;
 
 	/**
 	 * Registers a media player factory.
@@ -47,45 +48,7 @@ public:
 	 * @param Factory The media player factory to unregister.
 	 * @see RegisterPlayerFactory
 	 */
-	virtual void UnregisterPlayerFactory( IMediaPlayerFactory& Factory ) = 0;
-
-public:
-
-	/**
-	 * Gets an event delegate that is invoked after a media player factory has been added.
-	 *
-	 * @return The event delegate.
-	 * @see FOnFactoryRemoved
-	 */
-	DECLARE_EVENT(IMediaModule, FOnFactoryAdded);
-	virtual FOnFactoryAdded& OnFactoryAdded() = 0;
-
-	/**
-	 * Gets an event delegate that is invoked after a media player factory has been removed.
-	 *
-	 * @return The event delegate.
-	 * @see FOnFactoryAdded
-	 */
-	DECLARE_EVENT(IMediaModule, FOnFactoryRemoved);
-	virtual FOnFactoryRemoved& OnFactoryRemoved() = 0;
-
-public:
-
-	/**
-	 * Gets a reference to the messaging module instance.
-	 *
-	 * @todo gmp: better implementation using dependency injection.
-	 * @return A reference to the Media module.
-	 */
-	static IMediaModule& Get()
-	{
-#if PLATFORM_IOS
-        static IMediaModule& MediaModule = FModuleManager::LoadModuleChecked<IMediaModule>("Media");
-        return MediaModule;
-#else
-        return FModuleManager::LoadModuleChecked<IMediaModule>("Media");
-#endif
-	}
+	virtual void UnregisterPlayerFactory(IMediaPlayerFactory& Factory) = 0;
 
 public:
 
diff --git a/Engine/Source/Runtime/Media/Public/IMediaOutput.h b/Engine/Source/Runtime/Media/Public/IMediaOutput.h
new file mode 100644
--- /dev/null
+++ b/Engine/Source/Runtime/Media/Public/IMediaOutput.h
@@ -0,0 +1,56 @@
+// Copyright 1998-2016 Epic Games, Inc. All Rights Reserved.
+
+#pragma once
+
+
+class IMediaAudioSink;
+class IMediaStringSink;
+class IMediaTextureSink;
+
+
+/**
+ * Interface for access to a media player's output.
+ *
+ * @see IMediaControls, IMediaTracks
+ */
+class IMediaOutput
+{
+public:
+
+	/**
+	 * Set the sink to receive data from the active audio track.
+	 *
+	 * @param Sink The sink to set, or nullptr to clear.
+	 * @see SetCaptionSink, SetImageSink, SetVideoSink
+	 */
+	virtual void SetAudioSink(IMediaAudioSink* Sink) = 0;
+
+	/**
+	 * Set the sink to receive data from the active caption track.
+	 *
+	 * @param Sink The sink to set, or nullptr to clear.
+	 * @see SetAudioSink, SetImageSink, SetVideoSink
+	 */
+	virtual void SetCaptionSink(IMediaStringSink* Sink) = 0;
+
+	/**
+	 * Set the sink to receive data from the active image track.
+	 *
+	 * @param Sink The sink to set, or nullptr to clear.
+	 * @see SetAudioSink, SetCaptionSink, SetVideoSink
+	 */
+	virtual void SetImageSink(IMediaTextureSink* Sink) = 0;
+
+	/**
+	 * Set the sink to receive data from the active video track.
+	 *
+	 * @param Sink The sink to set, or nullptr to clear.
+	 * @see SetAudioSink, SetCaptionSink, SetImageSink
+	 */
+	virtual void SetVideoSink(IMediaTextureSink* Sink) = 0;
+
+public:
+
+	/** Virtual destructor. */
+	virtual ~IMediaOutput() { }
+};
diff --git a/Engine/Source/Runtime/Media/Public/IMediaPlayer.h b/Engine/Source/Runtime/Media/Public/IMediaPlayer.h
--- a/Engine/Source/Runtime/Media/Public/IMediaPlayer.h
+++ b/Engine/Source/Runtime/Media/Public/IMediaPlayer.h
@@ -2,14 +2,10 @@
 
 #pragma once
 
-#include "IMediaInfo.h"
-#include "IMediaStream.h"
-
-
-class IMediaStream;
-class IMediaAudioTrack;
-class IMediaCaptionTrack;
-class IMediaVideoTrack;
+class IMediaControls;
+class IMediaOptions;
+class IMediaOutput;
+class IMediaTracks;
 
 
 /**
@@ -17,9 +13,6 @@ class IMediaVideoTrack;
  */
 enum class EMediaEvent
 {
-	/** Unknown event. */
-	Unknown,
-
 	/** The current media source has been closed. */
 	MediaClosed,
 
@@ -43,25 +36,6 @@ enum class EMediaEvent
 };
 
 
-/**
- * Enumerates directions for seeking in media.
- */
-enum class EMediaSeekDirection
-{
-	/** Seek backwards from current position. */
-	Backward,
-
-	/** Seek from the beginning of the media. */
-	Beginning,
-
-	/** Seek from the end of the media. */
-	End,
-
-	/** Seek forward from current position. */
-	Forward
-};
-
-
 /**
  * Interface for media players.
  *
@@ -72,241 +46,114 @@ class IMediaPlayer
 public:
 
 	/**
-	 * Closes a previously opened media.
+	 * Close a previously opened media.
+	 *
+	 * Call this method to free up all resources associated with an opened
+	 * media source. If no media is open, this function has no effect.
 	 *
-	 * If no media has been opened, this function has no effect.
+	 * The media may not necessarily be closed after this function succeeds,
+	 * because closing may happen asynchronously. Subscribe to the MediaClosed
+	 * event to detect when the media finished closing. This events is only
+	 * triggered if Close returns true.
 	 *
 	 * @see IsReady, Open
 	 */
 	virtual void Close() = 0;
 
 	/**
-	 * Get the media's audio tracks.
+	 * Get the media playback controls for this player.
 	 *
-	 * @return Collection of audio tracks.
-	 * @see GetCaptionTracks, GetVideoTracks
+	 * @return Playback controls.
+	 * @see GetOutput, GetTracks
 	 */
-	virtual const TArray<TSharedRef<IMediaAudioTrack, ESPMode::ThreadSafe>>& GetAudioTracks() const = 0;
+	virtual IMediaControls& GetControls() = 0;
 
 	/**
-	 * Get the media's caption tracks.
+	 * Get debug information about the player and currently opened media.
 	 *
-	 * @return Collection of caption tracks.
-	 * @see GetAudioTracks, GetVideoTracks
+	 * @return Information string.
+	 * @see GetStats
 	 */
-	virtual const TArray<TSharedRef<IMediaCaptionTrack, ESPMode::ThreadSafe>>& GetCaptionTracks() const = 0;
+	virtual FString GetInfo() const = 0;
 
 	/**
-	 * Gets the last error that occurred during media loading or playback.
+	 * Get access to the media player's output.
 	 *
-	 * @return The error string, or an empty string if no error occurred.
+	 * @return Media tracks interface.
+	 * @see GetControls, GetTracks
 	 */
-//	virtual FString GetLastError() const = 0;
+	virtual IMediaOutput& GetOutput() = 0;
 
 	/**
-	 * Gets information about the currently loaded media.
+	 * Get playback statistics information.
 	 *
-	 * @return Interface to media information.
+	 * @return Information string.
+	 * @see GetInfo
 	 */
-	virtual const IMediaInfo& GetMediaInfo() const = 0;
+	virtual FString GetStats() const = 0;
 
 	/**
-	 * Gets the nominal playback rate, i.e. 1.0 for real time.
+	 * Get access to the media player's tracks.
 	 *
-	 * @return Playback rate.
-	 * @see Pause, Play, SetRate
+	 * @return Media tracks interface.
+	 * @see GetControls, GetOutput
 	 */
-	virtual float GetRate() const = 0;
+	virtual IMediaTracks& GetTracks() = 0;
 
 	/**
-	 * Gets the player's current playback time.
+	 * Get the URL of the currently loaded media.
 	 *
-	 * @return Playback time.
-	 * @see Seek
+	 * @return Media URL.
 	 */
-	virtual FTimespan GetTime() const = 0;
+	virtual FString GetUrl() const = 0;
 
 	/**
-	 * Get the media's video tracks.
-	 *
-	 * @return Collection of video tracks.
-	 * @see GetAudioTracks, GetCaptionTracks
-	 */
-	virtual const TArray<TSharedRef<IMediaVideoTrack, ESPMode::ThreadSafe>>& GetVideoTracks() const = 0;
-
-	/**
-	 * Checks whether playback is currently looping.
-	 *
-	 * @return true if playback is looping, false otherwise.
-	 * @see SetLooping
-	 */
-	virtual bool IsLooping() const = 0;
-
-	/**
-	 * Checks whether media playback is currently paused.
-	 *
-	 * @return true if playback is paused, false otherwise.
-	 * @see IsPlaying, IsStopped, Pause, Play, Stop
-	 */
-	virtual bool IsPaused() const = 0;
-
-	/**
-	 * Checks whether media playback has currently in progress.
-	 *
-	 * @return true if playback is in progress, false otherwise.
-	 * @see IsPaused, IsStopped, Pause, Play, Stop
-	 */
-	virtual bool IsPlaying() const = 0;
-
-	/**
-	 * Checks whether this player is ready for playback.
-	 *
-	 * A media player is considered ready if some media has been opened
-	 * successfully using the Open method and no error occurred during
-	 * loading or playback.
-	 *
-	 * @return true if ready, false otherwise.
-	 * @see Close, Open
-	 */
-	virtual bool IsReady() const = 0;
-
-	/**
-	 * Opens a media from a URL, possibly asynchronously.
+	 * Open a media source from a URL with optional parameters.
 	 *
 	 * The media may not necessarily be opened after this function succeeds,
-	 * because opening may happen asynchronously. Subscribe to the OnOpened
-	 * OnOpenFailed events to detect when the media finished or failed to
-	 * open. These events are only triggered if Open returns true.
+	 * because opening may happen asynchronously. Subscribe to the MediaOpened
+	 * and MediaOpenFailed events to detect when the media finished or failed
+	 * to open. These events are only triggered if Open returns true.
+	 *
+	 * The optional parameters can be used to configure aspects of media playback
+	 * and are specific to the type of media source and the underlying player.
+	 * Check their documentation for available keys and values.
 	 *
 	 * @param Url The URL of the media to open (file name or web address).
+	 * @param Options Optional media parameters.
 	 * @return true if the media is being opened, false otherwise.
 	 * @see Close, IsReady, OnOpen, OnOpenFailed
 	 */
-	virtual bool Open(const FString& Url) = 0;
+	virtual bool Open(const FString& Url, const IMediaOptions& Options) = 0;
 
 	/**
-	 * Opens a media from a file or memory archive, possibly asynchronously.
+	 * Open a media source from a file or memory archive with optional parameters.
 	 *
 	 * The media may not necessarily be opened after this function succeeds,
-	 * because opening may happen asynchronously. Subscribe to the OnOpened
-	 * OnOpenFailed events to detect when the media finished or failed to
-	 * open. These events are only triggered if Open returns true.
+	 * because opening may happen asynchronously. Subscribe to the MediaOpened
+	 * and MediaOpenFailed events to detect when the media finished or failed
+	 * to open. These events are only triggered if Open returns true.
+	 *
+	 * The optional parameters can be used to configure aspects of media playback
+	 * and are specific to the type of media source and the underlying player.
+	 * Check their documentation for available keys and values.
 	 *
 	 * @param Archive The archive holding the media data.
 	 * @param OriginalUrl The original URL of the media that was loaded into the buffer.
+	 * @param Options Optional media parameters.
 	 * @return true if the media is being opened, false otherwise.
 	 * @see Close, IsReady, OnOpen, OnOpenFailed
 	 */
-	virtual bool Open(const TSharedRef<FArchive, ESPMode::ThreadSafe>& Archive, const FString& OriginalUrl) = 0;
-
-	/**
-	 * Changes the media's playback time.
-	 *
-	 * @param Time The playback time to set.
-	 * @return true on success, false otherwise.
-	 * @see GetTime
-	 */
-	virtual bool Seek(const FTimespan& Time) = 0;
-
-	/**
-	 * Sets whether playback should be looping.
-	 *
-	 * @param Looping Enables or disables looping.
-	 * @see IsLooping
-	 */
-	virtual bool SetLooping(bool Looping) = 0;
-
-	/**
-	 * Sets the current playback rate.
-	 *
-	 * A playback rate of 1.0 will play the media normally at real-time.
-	 * A rate of 0.0 corresponds to pausing playback. A negative rate, if
-	 * supported, plays the media in reverse, and a rate larger than 1.0
-	 * fast forwards playback.
-	 *
-	 * @param Rate The playback rate to set.
-	 * @return true on success, false otherwise.
-	 * @see GetRate, Pause, Play
-	 */
-	virtual bool SetRate( float Rate ) = 0;
+	virtual bool Open(const TSharedRef<FArchive, ESPMode::ThreadSafe>& Archive, const FString& OriginalUrl, const IMediaOptions& Options) = 0;
 
 public:
 
-	/** Gets an event delegate that is invoked when some interesting event occurred. */
+	/** Get an event delegate that is invoked when an event occurred. */
 	DECLARE_EVENT_OneParam(IMediaPlayer, FOnMediaEvent, EMediaEvent /*Event*/)
 	virtual FOnMediaEvent& OnMediaEvent() = 0;
 
-public:
-
-	/**
-	 * Pauses media playback.
-	 *
-	 * This is the same as setting the playback rate to 0.0.
-	 *
-	 * @return true if the media is being paused, false otherwise.
-	 * @see Play, Stop
-	 */
-	FORCEINLINE bool Pause()
-	{
-		return SetRate(0.0f);
-	}
-
-	/**
-	 * Starts media playback at the default rate of 1.0.
-	 *
-	 * This is the same as setting the playback rate to 1.0.
-	 *
-	 * @return true if playback is starting, false otherwise.
-	 * @see Pause, Stop
-	 */
-	FORCEINLINE bool Play()
-	{
-		return SetRate(1.0f);
-	}
-
-	/**
-	 * Changes the playback time of the media by a relative offset in the given direction.
-	 *
-	 * @param TimeOffset The offset to apply to the time.
-	 * @param Direction The direction to seek in.
-	 * @return true on success, false otherwise.
-	 * @see GetDuration, GetTime
-	 */
-	bool Seek(const FTimespan& TimeOffset, EMediaSeekDirection Direction)
-	{
-		FTimespan SeekTime;
-
-		switch (Direction)
-		{
-		case EMediaSeekDirection::Backward:
-			SeekTime = GetTime() - TimeOffset;
-			break;
-
-		case EMediaSeekDirection::Beginning:
-			SeekTime = TimeOffset;
-			break;
-
-		case EMediaSeekDirection::End:
-			SeekTime = GetMediaInfo().GetDuration() - TimeOffset;
-			break;
-
-		case EMediaSeekDirection::Forward:
-			SeekTime = GetTime() + TimeOffset;
-			break;
-		}
-
-		return Seek(SeekTime);
-	}
-
 public:
 
 	/** Virtual destructor. */
 	virtual ~IMediaPlayer() { }
 };
-
-
-/** Type definition for shared pointers to instances of IMediaPlayer. */
-typedef TSharedPtr<IMediaPlayer> IMediaPlayerPtr;
-
-/** Type definition for shared references to instances of IMediaPlayer. */
-typedef TSharedRef<IMediaPlayer> IMediaPlayerRef;
diff --git a/Engine/Source/Runtime/Media/Public/IMediaPlayerFactory.h b/Engine/Source/Runtime/Media/Public/IMediaPlayerFactory.h
--- a/Engine/Source/Runtime/Media/Public/IMediaPlayerFactory.h
+++ b/Engine/Source/Runtime/Media/Public/IMediaPlayerFactory.h
@@ -3,32 +3,33 @@
 #pragma once
 
 
+class IMediaOptions;
 class IMediaPlayer;
 
 
-/**
- * Type definition for media container file types.
- *
- * This type is a TMap that maps the file extensions of supported media file containers
- * formats to a human readable description text. The file extensions are stored without
- * a leading dot, i.e. {"avi", "Audio Video Interleave File"}.
- */
-typedef TMap<FString, FText> FMediaFileTypes;
-
-
 /**
  * Interface for media player factories.
  *
  * Media player factories are used to create instances of media player implementations.
  * Most media players will be implemented inside plug-ins, which will register their
- * factories on startup. The Media module will use the SupportsUrl() method on this
- * interface to determine which media player to instantiate for a given media URL, such
- * as a file path or a resource identifier on the internet.
+ * factories on startup. The Media module will use the CanPlayUrl() method on this
+ * interface to determine which media player to instantiate for a given media source.
  */
 class IMediaPlayerFactory
 {
 public:
 
+	/**
+	 * Whether the player can play the specified source URL.
+	 *
+	 * @param Url The media source URL to check.
+	 * @param Options Optional media player parameters.
+	 * @param OutWarnings Will contain warning messages (optional).
+	 * @param OutErrors will contain error messages (optional).
+	 * @return true if the source can be played, false otherwise.
+	 */
+	virtual bool CanPlayUrl(const FString& Url, const IMediaOptions& Options, TArray<FText>* OutWarnings, TArray<FText>* OutErrors) const = 0;
+
 	/**
 	 * Creates a media player.
 	 *
@@ -37,26 +38,57 @@ public:
 	virtual TSharedPtr<IMediaPlayer> CreatePlayer() = 0;
 
 	/**
-	 * Gets a read-only collection of supported media file types.
+	 * Get the human readable name of the player.
 	 *
-	 * This function is used for file based media containers that reside on the users
-	 * local hard drive, a shared network drive or a location on the internet. It does
-	 * not take into account support for extension-less streaming media URI schemes,
-	 * such as rtsp:// and udp://
+	 * @return Display name text.
+	 * @see GetName
+	 */
+	virtual FText GetDisplayName() const = 0;
+
+	/**
+	 * Get the unique name of the media player.
+	 *
+	 * @return Player name.
+	 * @see GetDisplayName
+	 */
+	virtual FName GetName() const = 0;
+
+	/**
+	 * Get the names of platforms that the media player supports.
+	 *
+	 * The returned platforms names must match the ones returned by
+	 * FPlatformProperties::IniPlatformName, i.e. "Windows", "Android", etc.
+	 *
+	 * @return Platform name collection.
+	 * @see GetOptionalParameters, GetRequiredParameters, GetSupportedFileExtensions, GetSupportedUriSchemes
+	 */
+	virtual const TArray<FString>& GetSupportedPlatforms() const = 0;
+
+public:
+
+	/**
+	 * Whether the player can play the specified source URL.
 	 *
-	 * @return A collection of supported file types.
-	 * @see SupportsUrl
+	 * @param Url The media source URL to check.
+	 * @param Options Optional media player parameters.
+	 * @return true if the source can be played, false otherwise.
 	 */
-	virtual const FMediaFileTypes& GetSupportedFileTypes() const = 0;
+	bool CanPlayUrl(const FString& Url, const IMediaOptions& Options) const
+	{
+		return CanPlayUrl(Url, Options, nullptr, nullptr);
+	}
 
 	/**
-	 * Checks whether this factory can create media players that support the given URL.
+	 * Whether the player works on the given platform.
 	 *
-	 * @param Url The URL to the media, i.e. a file path or URI.
-	 * @return true if the URL supported, false otherwise.
-	 * @see GetSupportedFileTypes
+	 * @param PlatformName The name of the platform to check.
+	 * @return true if the platform is supported, false otherwise.
+	 * @see GetSupportedPlatforms
 	 */
-	virtual bool SupportsUrl(const FString& Url) const = 0;
+	bool SupportsPlatform(const FString& PlatformName) const
+	{
+		return GetSupportedPlatforms().Contains(PlatformName);
+	}
 
 public:
 
diff --git a/Engine/Source/Runtime/Media/Public/IMediaStringSink.h b/Engine/Source/Runtime/Media/Public/IMediaStringSink.h
new file mode 100644
--- /dev/null
+++ b/Engine/Source/Runtime/Media/Public/IMediaStringSink.h
@@ -0,0 +1,43 @@
+// Copyright 1998-2015 Epic Games, Inc. All Rights Reserved.
+
+#pragma once
+
+
+/**
+ * Interface for media sinks that receive string data.
+ *
+ * Implementations of this interface must be thread-safe as these
+ * methods will be called from some random media decoder thread.
+ */
+class IMediaStringSink
+{
+public:
+
+	/**
+	 * Initialize the sink.
+	 *
+	 * @return true on success, false otherwise.
+	 * @see ShutdownSink
+	 */
+	virtual bool InitializeStringSink() = 0;
+
+	/**
+	 * Display the given string.
+	 *
+	 * @param String The string to display.
+	 * @param Time The time at which to display the string (relative to the beginning of the media).
+	 */
+	virtual void DisplayStringSinkString(const FString& String, FTimespan Time) = 0;
+
+	/**
+	 * Shut down the sink.
+	 *
+	 * @see InitializeSink
+	 */
+	virtual void ShutdownStringSink() = 0;
+
+public:
+
+	/** Virtual destructor. */
+	virtual ~IMediaStringSink() { }
+};
diff --git a/Engine/Source/Runtime/Media/Public/IMediaTextureSink.h b/Engine/Source/Runtime/Media/Public/IMediaTextureSink.h
new file mode 100644
--- /dev/null
+++ b/Engine/Source/Runtime/Media/Public/IMediaTextureSink.h
@@ -0,0 +1,286 @@
+// Copyright 1998-2015 Epic Games, Inc. All Rights Reserved.
+
+#pragma once
+
+
+class FRHITexture;
+class FRHITexture2D;
+
+
+/**
+ * Available pixel formats for media texture sinks.
+ */
+enum class EMediaTextureSinkFormat
+{
+	/** Four 8-bit unsigned integer components (AYUV packing) per pixel. */
+	CharAYUV,
+
+	/** Four 8-bit unsigned integer components (Blue, Green, Red, Alpha) per pixel. */
+	CharBGRA,
+
+	/** NV12 encoded monochrome texture with 8 bits per channel. */
+	CharNV12,
+
+	/** NV21 encoded monochrome texture with 8 bits per channel. */
+	CharNV21,
+
+	/** Four 8-bit unsigned integer components (UYVY packing) per pixel. */
+	CharUYVY,
+
+	/** Four 8-bit unsigned integer components (YUY2 packing) per pixel. */
+	CharYUY2,
+
+	/** Four 8-bit unsigned integer components (YVYU packing) per pixel. */
+	CharYVYU,
+
+	/** Three floating point components (Red, Green, Blue) per pixel. */
+	FloatRGB,
+
+	/** Four floating point components (Red, Green, Blue, Alpha) per pixel. */
+	FloatRGBA,
+};
+
+
+/**
+ * Available modes that texture sinks can operate in.
+ */
+enum class EMediaTextureSinkMode
+{
+	/** Use triple-buffered texture memory (for decoding on non-render threads). */
+	Buffered,
+
+	/**
+	 * Render to the render target's texture resource.
+	 *
+	 * Use IMediaTextureSink.GetTextureSinkTexture to get the sink's render target
+	 * texture resource. This mode is only available when compiling against the RHI.
+	 */
+	Unbuffered
+};
+
+
+/**
+ * Interface for media sinks that receive texture data.
+ *
+ * Implementors of this interface must be thread-safe as the methods may be called
+ * from any random media decoder thread.
+ *
+ * Texture sinks can operate in one of two modes: Unbuffered and Buffered. Unbuffered
+ * mode provides direct access to the texture resource, and it is the preferred method
+ * if the native media player is able to decode or return the decoded frame buffer on
+ * the render thread, or if it performs custom pixel conversions on the render thread.
+ *
+ * Unbuffered mode also allows for replacing the underlying resources directly. This
+ * is the fastest mechanism, but generally only useful for decoders that can decode
+ * directly on the GPU and hence provide their own texture resources.
+ *
+ * Note: Unbuffered mode is only exposed in the API when compiling against the Engine,
+ * because texture resources are Engine RHI types.
+ *
+ * Buffered mode provides access to low-level untyped texture memory in one or more
+ * intermediate buffers. It can be used in applications that are not compiled against
+ * the Engine, or in native media players that return decoded frame data on their own
+ * thread. It is possible to either decode directly into the buffers using the
+ * AcquireTextureSinkBuffer and ReleaseTextureSinkBuffer methods, or to copy an
+ * existing buffer using the UpdateTextureSinkBuffer method. Because buffered mode
+ * uses intermediate buffers to decouple the decoder from the render thread, it may
+ * use considerably more memory and CPU time.
+ *
+ * Depending on when and how the video frames are being made available by the decoder,
+ * native media players should chose one of the following four mechanisms to send new
+ * video frame data to the texture sinks:
+ *
+ * - GetTextureSinkTexture() [implies Unbuffered mode]
+ * - AcquireTextureSinkBuffer() + ReleaseTextureSinkBuffer() + DisplayTextureSinkBuffer()
+ * - UpdateTextureSinkBuffer()
+ * - UpdateTextureSinkResource()
+ *
+ * Texture buffers are initialized with one of the pixel formats in EMediaTextureSinkMode.
+ * If a texture format is not supported natively by a sink, it must perform any required
+ * pixel format conversion, which may add a cost on CPU time, GPU time, or both.
+ *
+ * Note: Most common video formats are not supported natively. Most decoders can output
+ * frames in a supported format, such as RGB or RGBA. However, it may still be a lot
+ * faster to have the sink perform the pixel format conversion on the GPU instead,
+ * because decoders often use a much slower conversion on the CPU.
+ */
+class IMediaTextureSink
+{
+public:
+
+	/**
+	 * Get the current dimensions of the texture sink's frame buffer.
+	 *
+	 * @return Width and height of the buffer (in pixels).
+	 * @see GetTextureSinkFormat, GetTextureSinkMode, InitializeTextureSink
+	 */
+	virtual FIntPoint GetTextureSinkDimensions() const = 0;
+
+	/**
+	 * Get the sink's current texture format.
+	 *
+	 * @return texture format.
+	 * @see GetTextureSinkDimensions, GetTextureSinkMode, InitializeTextureSink
+	 */
+	virtual EMediaTextureSinkFormat GetTextureSinkFormat() const = 0;
+
+	/**
+	 * Get the texture sink's current buffer mode.
+	 *
+	 * @return Buffer mode.
+	 * @see GetTextureSinkDimensions, GetTextureSinkFormat, InitializeTextureSink
+	 */
+	virtual EMediaTextureSinkMode GetTextureSinkMode() const = 0;
+
+	/**
+	 * Initialize the sink.
+	 *
+	 * @param Dimensions Width and height of the texture (in pixels).
+	 * @param Format The pixel format of the sink's render target texture.
+	 * @param Mode The mode to operate the sink in (buffered vs. unbuffered).
+	 * @return true on success, false if initialization is already in progress or failed.
+	 * @see GetTextureSinkDimensions, GetTextureSinkMode, ShutdownSink, SupportsTextureSinkFormat
+	 */
+	virtual bool InitializeTextureSink(FIntPoint Dimensions, EMediaTextureSinkFormat Format, EMediaTextureSinkMode Mode) = 0;
+
+	/**
+	 * Shut down the sink.
+	 *
+	 * @see InitializeSink
+	 */
+	virtual void ShutdownTextureSink() = 0;
+
+	/**
+	 * Whether the sink supports the specified pixel format.
+	 *
+	 * @param Format The format to check.
+	 * @return true if the format is supported, false otherwise.
+	 * @see InitializeTextureSink
+	 */
+	virtual bool SupportsTextureSinkFormat(EMediaTextureSinkFormat Format) const = 0;
+
+public:
+
+	//~ Method 1: Get the sink's render target texture and write to it
+
+#if WITH_ENGINE
+
+	/**
+	 * Get the render target texture to render to.
+	 *
+	 * If called on the rendering thread and when the sink has been initialized in
+	 * Unbuffered mode, this method returns a texture object that can be written to
+	 * directly on the render thread.
+	 *
+	 * @return The texture, or nullptr if called in Buffered mode or not on rendering thread.
+	 * @see InitializeTextureSink
+	 */
+	virtual FRHITexture* GetTextureSinkTexture() = 0;
+
+#endif
+
+public:
+
+	//~ Method 2: Get the render target's memory and write, release, and display it
+
+	/**
+	 * Acquire a pointer to locked texture memory buffer to render to.
+	 *
+	 * This method can be used to acquire a render target texture buffer that
+	 * the decoder can decode a frame into. In Unbuffered mode it must be called
+	 * on the render thread to have an effect. In Buffered mode it can be called
+	 * from on thread, but the sink may return an intermediate buffer that will
+	 * be copied when released.
+	 *
+	 * If a valid pointer is returned by this method, the caller must call
+	 * ReleaseTextureSinkBuffer to release the buffer. In Unbuffered mode,
+	 * the changes are committed immediately when the buffer is released.
+	 * In Buffered mode the caller must call DisplayTextureSinkBuffer to
+	 * commit the buffer to the GPU.
+	 *
+	 * If the decoder provides its own frame buffer, consider using the
+	 * UpdateTextureSinkBuffer instead.
+	 *
+	 * @return Pointer to buffer on success, nullptr otherwise.
+	 * @see DisplayTextureSinkBuffer, InitializeTextureSink, ReleaseTextureSinkBuffer
+	 */
+	virtual void* AcquireTextureSinkBuffer() = 0;
+
+	/**
+	 * Display the latest texture buffer.
+	 *
+	 * This method is used only if the sink has been initialized in Buffered
+	 * mode. It signals the sink that the previously released or updated texture
+	 * buffer is ready for display.
+	 *
+	 * @param Time The time of the frame relative to the beginning of the media.
+	 * @see AcquireTextureSinkBuffer, ReleaseTextureSinkBuffer, UpdateTextureSinkBuffer
+	 */
+	virtual void DisplayTextureSinkBuffer(FTimespan Time) = 0;
+
+	/**
+	 * Release a previously acquired texture buffer.
+	 *
+	 * This method is used to release a previously acquired render target texture
+	 * buffer from the sink. It indicates to the sink that writing to the texture
+	 * buffer has completed. In Unbuffered mode, the changes are committed
+	 * immediately. In Buffered mode the caller must call DisplayTextureSinkBuffer
+	 * to commit the buffer to the GPU.
+	 *
+	 * @see AcquireTextureSinkBuffer, DisplayTextureSinkBuffer
+	 */
+	virtual void ReleaseTextureSinkBuffer() = 0;
+
+public:
+
+	//~ Method 3: Update the render target texture from a provided buffer
+
+	/**
+	 * Update the render target texture with new frame data.
+	 *
+	 * This method can be used to update the sink's render target texture with data
+	 * from a frame buffer that was provided by the decoder. In Unbuffered mode it
+	 * must be called on the render thread to have an effect. In Buffered mode it
+	 * can be called on any thread, but the sink may perform extra buffer copies.
+	 *
+	 * In Unbuffered mode, the changes are committed immediately. In Buffered mode
+	 * the caller must call DisplayTextureSinkBuffer to commit the buffer to the GPU.
+	 *
+	 * If the decoder does not provide its own frame buffer, consider using the
+	 * AcquireTextureSinkBuffer and ReleaseTextureSinkBuffer instead.
+	 *
+	 * @param Data The frame data to copy into the render target.
+	 * @param Pitch Number of bytes per pixel row (0 = default).
+	 * @see DisplayTextureSinkBuffer, InitializeTextureSink
+	 */
+	virtual void UpdateTextureSinkBuffer(const uint8* Data, uint32 Pitch = 0) = 0;
+
+public:
+
+	//~ Method 4: Replace the sink's texture resource with your own
+
+#if WITH_ENGINE
+
+	/**
+	 * Update the sink's RHI resources.
+	 *
+	 * This method can be used to replace the sink's underlying texture resources
+	 * with different textures, which is the recommended way if the decoder provides
+	 * its own render targets. It can be called from any thread.
+	 *
+	 * This method implies Unbuffered mode and does not require any buffer copies,
+	 * but the sink may perform a pixel format conversion if the format is not
+	 * natively supported by the renderer.
+	 *
+	 * @param RenderTarget The new render target texture.
+	 * @param ShaderResource The new shader resource texture.
+	 */
+	virtual void UpdateTextureSinkResource(FRHITexture* RenderTarget, FRHITexture* ShaderResource) = 0;
+
+#endif
+
+public:
+
+	/** Virtual destructor. */
+	~IMediaTextureSink() { }
+};
