Fix CIS non-unity build

diff --git a/Engine/Source/Runtime/Engine/Private/Components/PrimitiveComponent.cpp b/Engine/Source/Runtime/Engine/Private/Components/PrimitiveComponent.cpp
--- a/Engine/Source/Runtime/Engine/Private/Components/PrimitiveComponent.cpp
+++ b/Engine/Source/Runtime/Engine/Private/Components/PrimitiveComponent.cpp
@@ -7,6 +7,7 @@
 #include "EnginePrivate.h"
 #include "GameFramework/PhysicsVolume.h"
 #include "PhysicsPublic.h"
+#include "PhysicsEngine/BodySetup.h"
 #include "LevelUtils.h"
 #if WITH_EDITOR
 #include "ShowFlags.h"
