Resolved remaining UE localization warnings

- Ensured that everything has a unique loc key.
- Downgraded a warning from asset gathering to a log (it's not that important).
- Tested gathering for Engine and Editor (which were showing warnings) and both now gather cleanly.

diff --git a/Engine/Source/Developer/BlueprintProfiler/Private/ScriptInstrumentationPlayback.cpp b/Engine/Source/Developer/BlueprintProfiler/Private/ScriptInstrumentationPlayback.cpp
--- a/Engine/Source/Developer/BlueprintProfiler/Private/ScriptInstrumentationPlayback.cpp
+++ b/Engine/Source/Developer/BlueprintProfiler/Private/ScriptInstrumentationPlayback.cpp
@@ -1790,7 +1790,7 @@ void FBlueprintTunnelInstanceContext::MapTunnelContext(TSharedPtr<FBlueprintFunc
 				TunnelEntryParams.ObservedObject = TunnelInstance;
 				TunnelEntryParams.ObservedPin = InstanceEntryPoint.Value;
 				TunnelEntryParams.DisplayName = TunnelInstance->GetNodeTitle(ENodeTitleType::ListView);
-				TunnelEntryParams.Tooltip = LOCTEXT("NavigateToNodeLocationHyperlink_ToolTip", "Navigate to the Tunnel Instance");
+				TunnelEntryParams.Tooltip = LOCTEXT("NavigateToTunnelInstanceHyperlink_ToolTip", "Navigate to the Tunnel Instance");
 				TunnelEntryParams.NodeFlags = EScriptExecutionNodeFlags::TunnelEntryPinInstance;
 				TunnelEntryParams.IconColor = FLinearColor(1.f, 1.f, 1.f, 0.8f);
 				TunnelEntryParams.Icon = TunnelInstanceParams.Icon;
diff --git a/Engine/Source/Editor/AnimGraph/Private/AnimGraphNode_ModifyBone.cpp b/Engine/Source/Editor/AnimGraph/Private/AnimGraphNode_ModifyBone.cpp
--- a/Engine/Source/Editor/AnimGraph/Private/AnimGraphNode_ModifyBone.cpp
+++ b/Engine/Source/Editor/AnimGraph/Private/AnimGraphNode_ModifyBone.cpp
@@ -21,14 +21,14 @@ void UAnimGraphNode_ModifyBone::ValidateAnimNodeDuringCompilation(USkeleton* For
 	{
 		if (Node.BoneToModify.BoneName == NAME_None)
 		{
-			MessageLog.Warning(*LOCTEXT("NoBoneToModify", "@@ - You must pick a bone to modify").ToString(), this);
+			MessageLog.Warning(*LOCTEXT("NoBoneSelectedToModify", "@@ - You must pick a bone to modify").ToString(), this);
 		}
 		else
 		{
 			FFormatNamedArguments Args;
 			Args.Add(TEXT("BoneName"), FText::FromName(Node.BoneToModify.BoneName));
 
-			FText Msg = FText::Format(LOCTEXT("NoBoneToModify", "@@ - Bone {BoneName} not found in Skeleton"), Args);
+			FText Msg = FText::Format(LOCTEXT("NoBoneFoundToModify", "@@ - Bone {BoneName} not found in Skeleton"), Args);
 
 			MessageLog.Warning(*Msg.ToString(), this);
 		}
diff --git a/Engine/Source/Editor/AnimGraph/Private/AnimGraphNode_PoseByName.cpp b/Engine/Source/Editor/AnimGraph/Private/AnimGraphNode_PoseByName.cpp
--- a/Engine/Source/Editor/AnimGraph/Private/AnimGraphNode_PoseByName.cpp
+++ b/Engine/Source/Editor/AnimGraph/Private/AnimGraphNode_PoseByName.cpp
@@ -69,7 +69,7 @@ FText UAnimGraphNode_PoseByName::GetNodeTitle(ENodeTitleType::Type TitleType) co
 		UEdGraphPin* PosePin = FindPin(GET_MEMBER_NAME_STRING_CHECKED(FAnimNode_PoseByName, PoseAsset));
 		if (PosePin && PosePin->LinkedTo.Num() > 0)
 		{
-			return LOCTEXT("EvaluateSequence_TitleVariable", "Pose");
+			return LOCTEXT("PoseByName_TitleVariable", "Pose");
 		}
 		else if (PosePin && PosePin->DefaultObject != nullptr)
 		{
diff --git a/Engine/Source/Editor/AnimGraph/Private/AnimGraphNode_SubInstance.cpp b/Engine/Source/Editor/AnimGraph/Private/AnimGraphNode_SubInstance.cpp
--- a/Engine/Source/Editor/AnimGraph/Private/AnimGraphNode_SubInstance.cpp
+++ b/Engine/Source/Editor/AnimGraph/Private/AnimGraphNode_SubInstance.cpp
@@ -41,7 +41,7 @@ FText UAnimGraphNode_SubInstance::GetNodeTitle(ENodeTitleType::Type TitleType) c
 	}
 	if(TitleType == ENodeTitleType::ListView)
 	{
-		return FText::Format(LOCTEXT("TitleFormat", "{NodeTitle} - Target Class: {TargetClass}"), Args);
+		return FText::Format(LOCTEXT("TitleListFormat", "{NodeTitle} - Target Class: {TargetClass}"), Args);
 	}
 	else
 	{
diff --git a/Engine/Source/Editor/LandscapeEditor/Private/LandscapeEdMode.cpp b/Engine/Source/Editor/LandscapeEditor/Private/LandscapeEdMode.cpp
--- a/Engine/Source/Editor/LandscapeEditor/Private/LandscapeEdMode.cpp
+++ b/Engine/Source/Editor/LandscapeEditor/Private/LandscapeEdMode.cpp
@@ -2540,7 +2540,7 @@ void FEdModeLandscape::ImportData(const FLandscapeTargetListInfo& TargetInfo, co
 					Args.Add(TEXT("LandscapeSizeY"), LandscapeResolution.Height);
 
 					FMessageDialog::Open(EAppMsgType::Ok,
-						FText::Format(NSLOCTEXT("LandscapeEditor.NewLandscape", "Import_LayerSizeMismatchRaw", "The layer file does not match the current Landscape extent ({LandscapeSizeX}\u00D7{LandscapeSizeY}), and its exact resolution could not be determined"), Args));
+						FText::Format(NSLOCTEXT("LandscapeEditor.NewLandscape", "Import_LayerSizeMismatch_ResNotDetermined", "The layer file does not match the current Landscape extent ({LandscapeSizeX}\u00D7{LandscapeSizeY}), and its exact resolution could not be determined"), Args));
 
 					return;
 				}
@@ -2575,7 +2575,7 @@ void FEdModeLandscape::ImportData(const FLandscapeTargetListInfo& TargetInfo, co
 					Args.Add(TEXT("LandscapeSizeY"), LandscapeResolution.Height);
 
 					auto Result = FMessageDialog::Open(EAppMsgType::OkCancel,
-						FText::Format(NSLOCTEXT("LandscapeEditor.NewLandscape", "Import_LayerSizeMismatch", "The layer file's size ({FileSizeX}\u00D7{FileSizeY}) does not match the current Landscape extent ({LandscapeSizeX}\u00D7{LandscapeSizeY}), if you continue it will be padded/clipped to fit"), Args));
+						FText::Format(NSLOCTEXT("LandscapeEditor.NewLandscape", "Import_LayerSizeMismatch_WillClamp", "The layer file's size ({FileSizeX}\u00D7{FileSizeY}) does not match the current Landscape extent ({LandscapeSizeX}\u00D7{LandscapeSizeY}), if you continue it will be padded/clipped to fit"), Args));
 
 					if (Result != EAppReturnType::Ok)
 					{
diff --git a/Engine/Source/Editor/MovieSceneTools/Private/MatineeImportTools.cpp b/Engine/Source/Editor/MovieSceneTools/Private/MatineeImportTools.cpp
--- a/Engine/Source/Editor/MovieSceneTools/Private/MatineeImportTools.cpp
+++ b/Engine/Source/Editor/MovieSceneTools/Private/MatineeImportTools.cpp
@@ -232,7 +232,7 @@ bool FMatineeImportTools::CopyInterpColorTrack( UInterpTrackColorProp* ColorProp
 
 bool FMatineeImportTools::CopyInterpLinearColorTrack( UInterpTrackLinearColorProp* LinearColorPropTrack, UMovieSceneColorTrack* ColorTrack )
 {
-	const FScopedTransaction Transaction( NSLOCTEXT( "Sequencer", "PasteMatineeColorTrack", "Paste Matinee Linear Color Track" ) );
+	const FScopedTransaction Transaction( NSLOCTEXT( "Sequencer", "PasteMatineeLinearColorTrack", "Paste Matinee Linear Color Track" ) );
 	bool bSectionCreated = false;
 
 	ColorTrack->Modify();
diff --git a/Engine/Source/Editor/UnrealEd/Private/Commandlets/GatherTextFromAssetsCommandlet.cpp b/Engine/Source/Editor/UnrealEd/Private/Commandlets/GatherTextFromAssetsCommandlet.cpp
--- a/Engine/Source/Editor/UnrealEd/Private/Commandlets/GatherTextFromAssetsCommandlet.cpp
+++ b/Engine/Source/Editor/UnrealEd/Private/Commandlets/GatherTextFromAssetsCommandlet.cpp
@@ -74,7 +74,7 @@ public:
 		{
 			static const FString LogIndentation = TEXT("    ");
 
-			UE_LOG(LogGatherTextFromAssetsCommandlet, Warning, TEXT("Package '%s' produced %d error(s) and %d warning(s) while loading. Please verify that your text has gathered correctly."), *PackageContext, ErrorCount, WarningCount);
+			UE_LOG(LogGatherTextFromAssetsCommandlet, Log, TEXT("Package '%s' produced %d error(s) and %d warning(s) while loading. Please verify that your text has gathered correctly."), *PackageContext, ErrorCount, WarningCount);
 			
 			GWarn->Log(NAME_None, ELogVerbosity::Log, FString::Printf(TEXT("The following errors and warnings were reported while loading '%s':"), *PackageContext));
 			for (const auto& FormattedOutput : FormattedErrorsAndWarningsList)
diff --git a/Engine/Source/Runtime/CoreUObject/Private/Misc/PackageName.cpp b/Engine/Source/Runtime/CoreUObject/Private/Misc/PackageName.cpp
--- a/Engine/Source/Runtime/CoreUObject/Private/Misc/PackageName.cpp
+++ b/Engine/Source/Runtime/CoreUObject/Private/Misc/PackageName.cpp
@@ -427,7 +427,7 @@ bool FPackageName::DoesPackageNameContainInvalidCharacters(const FString& InLong
 		{
 			FFormatNamedArguments Args;
 			Args.Add( TEXT("IllegalNameCharacters"), FText::FromString(MatchedInvalidChars) );
-			*OutReason = FText::Format( NSLOCTEXT("Core", "NameContainsInvalidCharacters", "Name may not contain the following characters: {IllegalNameCharacters}"), Args );
+			*OutReason = FText::Format( NSLOCTEXT("Core", "PackageNameContainsInvalidCharacters", "Name may not contain the following characters: {IllegalNameCharacters}"), Args );
 		}
 		return true;
 	}
diff --git a/Engine/Source/Runtime/Engine/Private/KismetRenderingLibrary.cpp b/Engine/Source/Runtime/Engine/Private/KismetRenderingLibrary.cpp
--- a/Engine/Source/Runtime/Engine/Private/KismetRenderingLibrary.cpp
+++ b/Engine/Source/Runtime/Engine/Private/KismetRenderingLibrary.cpp
@@ -73,11 +73,11 @@ void UKismetRenderingLibrary::DrawMaterialToRenderTarget(UObject* WorldContextOb
 	}
 	else if (!Material)
 	{
-		FMessageLog("Blueprint").Warning(LOCTEXT("InvalidMaterial", "DrawMaterialToRenderTarget: Material must be non-null."));
+		FMessageLog("Blueprint").Warning(LOCTEXT("DrawMaterialToRenderTarget_InvalidMaterial", "DrawMaterialToRenderTarget: Material must be non-null."));
 	}
 	else if (!TextureRenderTarget)
 	{
-		FMessageLog("Blueprint").Warning(LOCTEXT("InvalidTextureRenderTarget", "DrawMaterialToRenderTarget: TextureRenderTarget must be non-null."));
+		FMessageLog("Blueprint").Warning(LOCTEXT("DrawMaterialToRenderTarget_InvalidTextureRenderTarget", "DrawMaterialToRenderTarget: TextureRenderTarget must be non-null."));
 	}
 }
 
@@ -105,21 +105,21 @@ void UKismetRenderingLibrary::ExportRenderTarget(UObject* WorldContextObject, UT
 		}
 		else
 		{
-			FMessageLog("Blueprint").Warning(LOCTEXT("FileWriterFailedToCreate", "ExportRenderTarget: FileWrite failed to create."));
+			FMessageLog("Blueprint").Warning(LOCTEXT("ExportRenderTarget_FileWriterFailedToCreate", "ExportRenderTarget: FileWrite failed to create."));
 		}
 	}
 
     else if (!TextureRenderTarget)
 	{
-		FMessageLog("Blueprint").Warning(LOCTEXT("InvalidTextureRenderTarget", "ExportRenderTarget: TextureRenderTarget must be non-null."));
+		FMessageLog("Blueprint").Warning(LOCTEXT("ExportRenderTarget_InvalidTextureRenderTarget", "ExportRenderTarget: TextureRenderTarget must be non-null."));
 	}
 	if (!PathError.IsEmpty())
 	{
-		FMessageLog("Blueprint").Warning(FText::Format(LOCTEXT("InvalidFilePath", "ExportRenderTarget: Invalid file path provided: '{0}'"), PathError));
+		FMessageLog("Blueprint").Warning(FText::Format(LOCTEXT("ExportRenderTarget_InvalidFilePath", "ExportRenderTarget: Invalid file path provided: '{0}'"), PathError));
 	}
 	if (FileName.IsEmpty())
 	{
-		FMessageLog("Blueprint").Warning(LOCTEXT("InvalidFileName", "ExportRenderTarget: FileName must be non-empty."));
+		FMessageLog("Blueprint").Warning(LOCTEXT("ExportRenderTarget_InvalidFileName", "ExportRenderTarget: FileName must be non-empty."));
 	}
 }
 
@@ -147,21 +147,21 @@ void UKismetRenderingLibrary::ExportTexture2D(UObject* WorldContextObject, UText
 		}
 		else
 		{
-			FMessageLog("Blueprint").Warning(LOCTEXT("FileWriterFailedToCreate", "ExportTexture2D: FileWrite failed to create."));
+			FMessageLog("Blueprint").Warning(LOCTEXT("ExportTexture2D_FileWriterFailedToCreate", "ExportTexture2D: FileWrite failed to create."));
 		}
 	}
 
 	else if (!Texture)
 	{
-		FMessageLog("Blueprint").Warning(LOCTEXT("InvalidTextureRenderTarget", "ExportTexture2D: TextureRenderTarget must be non-null."));
+		FMessageLog("Blueprint").Warning(LOCTEXT("ExportTexture2D_InvalidTextureRenderTarget", "ExportTexture2D: TextureRenderTarget must be non-null."));
 	}
 	if (!PathError.IsEmpty())
 	{
-		FMessageLog("Blueprint").Warning(FText::Format(LOCTEXT("InvalidFilePath", "ExportTexture2D: Invalid file path provided: '{0}'"), PathError));
+		FMessageLog("Blueprint").Warning(FText::Format(LOCTEXT("ExportTexture2D_InvalidFilePath", "ExportTexture2D: Invalid file path provided: '{0}'"), PathError));
 	}
 	if (FileName.IsEmpty())
 	{
-		FMessageLog("Blueprint").Warning(LOCTEXT("InvalidFileName", "ExportTexture2D: FileName must be non-empty."));
+		FMessageLog("Blueprint").Warning(LOCTEXT("ExportTexture2D_InvalidFileName", "ExportTexture2D: FileName must be non-empty."));
 	}
 }
 
@@ -211,7 +211,7 @@ void UKismetRenderingLibrary::BeginDrawCanvasToRenderTarget(UObject* WorldContex
 	}
 	else if (!TextureRenderTarget)
 	{
-		FMessageLog("Blueprint").Warning(LOCTEXT("InvalidTextureRenderTarget", "BeginDrawCanvasToRenderTarget: TextureRenderTarget must be non-null."));
+		FMessageLog("Blueprint").Warning(LOCTEXT("BeginDrawCanvasToRenderTarget_InvalidTextureRenderTarget", "BeginDrawCanvasToRenderTarget: TextureRenderTarget must be non-null."));
 	}
 }
 
@@ -249,7 +249,7 @@ void UKismetRenderingLibrary::EndDrawCanvasToRenderTarget(UObject* WorldContextO
 		}
 		else if (!Context.RenderTarget)
 		{
-			FMessageLog("Blueprint").Warning(LOCTEXT("InvalidContext", "EndDrawCanvasToRenderTarget: Context must be valid."));
+			FMessageLog("Blueprint").Warning(LOCTEXT("EndDrawCanvasToRenderTarget_InvalidContext", "EndDrawCanvasToRenderTarget: Context must be valid."));
 		}
 	}
 }
diff --git a/Engine/Source/Runtime/Engine/Private/LevelActor.cpp b/Engine/Source/Runtime/Engine/Private/LevelActor.cpp
--- a/Engine/Source/Runtime/Engine/Private/LevelActor.cpp
+++ b/Engine/Source/Runtime/Engine/Private/LevelActor.cpp
@@ -1267,7 +1267,7 @@ void UWorld::IssueEditorLoadWarnings()
 		Arguments.Add(TEXT("LoadTime"), FText::FromString(FString::Printf(TEXT("%.1fs"), TotalLoadTimeFromFixups)));
 			
 		FMessageLog("MapCheck").Warning()
-			->AddToken(FTextToken::Create(FText::Format( LOCTEXT( "MapCheck_Message_RepairedPaintedVertexColors", "Spent {LoadTime} repairing painted vertex colors due to static mesh re-imports!  This will happen every load until the maps are resaved." ), Arguments ) ))
+			->AddToken(FTextToken::Create(FText::Format( LOCTEXT( "MapCheck_Message_SpentXRepairingPaintedVertexColors", "Spent {LoadTime} repairing painted vertex colors due to static mesh re-imports!  This will happen every load until the maps are resaved." ), Arguments ) ))
 			->AddToken(FMapErrorToken::Create(FMapErrors::RepairedPaintedVertexColors));
 	}
 }
diff --git a/Engine/Source/Runtime/UMG/Private/Components/Widget.cpp b/Engine/Source/Runtime/UMG/Private/Components/Widget.cpp
--- a/Engine/Source/Runtime/UMG/Private/Components/Widget.cpp
+++ b/Engine/Source/Runtime/UMG/Private/Components/Widget.cpp
@@ -739,6 +739,7 @@ void UWidget::Deselect()
 }
 
 #undef LOCTEXT_NAMESPACE
+#define LOCTEXT_NAMESPACE "UMG"
 #endif
 
 bool UWidget::Modify(bool bAlwaysMarkDirty)
