Fix for warning from GetBoneTransform when we call it before the master component has been registered.

diff --git a/Engine/Source/Runtime/Engine/Private/Components/SkinnedMeshComponent.cpp b/Engine/Source/Runtime/Engine/Private/Components/SkinnedMeshComponent.cpp
--- a/Engine/Source/Runtime/Engine/Private/Components/SkinnedMeshComponent.cpp
+++ b/Engine/Source/Runtime/Engine/Private/Components/SkinnedMeshComponent.cpp
@@ -832,6 +832,12 @@ FTransform USkinnedMeshComponent::GetBoneTransform(int32 BoneIdx, const FTransfo
 	const USkinnedMeshComponent* const MasterPoseComponentInst = MasterPoseComponent.Get();
 	if(MasterPoseComponentInst)
 	{
+		if (!MasterPoseComponentInst->IsRegistered())
+		{
+			// We aren't going to get anything valid from the master pose if it
+			// isn't valid so for now return identity
+			return FTransform::Identity;
+		}
 		if(BoneIdx < MasterBoneMap.Num())
 		{
 			int32 ParentBoneIndex = MasterBoneMap[BoneIdx];
