Copying //UE4/Dev-Rendering to //UE4/Dev-Main (Source: //UE4/Dev-Rendering @ 3072736)

diff --git a/Engine/Source/Runtime/Engine/Private/Streaming/StreamingManagerTexture.cpp b/Engine/Source/Runtime/Engine/Private/Streaming/StreamingManagerTexture.cpp
--- a/Engine/Source/Runtime/Engine/Private/Streaming/StreamingManagerTexture.cpp
+++ b/Engine/Source/Runtime/Engine/Private/Streaming/StreamingManagerTexture.cpp
@@ -1104,25 +1104,25 @@ bool FStreamingManagerTexture::HandleListStreamingTexturesCommand( const TCHAR*
 
 		int32 CurrentMipIndex = FMath::Max(Texture2D->GetNumMips() - StreamingTexture.ResidentMips, 0);
 		int32 WantedMipIndex = FMath::Max(Texture2D->GetNumMips() - StreamingTexture.GetPerfectWantedMips(), 0);
-		int32 MaxMipIndex = FMath::Max(Texture2D->GetNumMips() - StreamingTexture.MaxAllowedMips, 0);
+		int32 MaxAllowedMipIndex = FMath::Max(Texture2D->GetNumMips() - StreamingTexture.MaxAllowedMips, 0);
 		const TIndirectArray<struct FTexture2DMipMap>& Mips = Texture2D->PlatformData->Mips;
 
 		if (StreamingTexture.LastRenderTime != MAX_FLT)
 		{
-			UE_LOG(LogContentStreaming, Log,  TEXT("    Current=%dx%d Wanted=%dx%d Max=%dx%d LastRenderTime=%.3f BudgetBias=%d Group=%s"), 
+			UE_LOG(LogContentStreaming, Log,  TEXT("    Current=%dx%d Wanted=%dx%d MaxAllowed=%dx%d LastRenderTime=%.3f BudgetBias=%d Group=%s"), 
 				Mips[CurrentMipIndex].SizeX, Mips[CurrentMipIndex].SizeY, 
 				Mips[WantedMipIndex].SizeX, Mips[WantedMipIndex].SizeY, 
-				Mips[MaxMipIndex].SizeX, Mips[MaxMipIndex].SizeY, 
+				Mips[MaxAllowedMipIndex].SizeX, Mips[MaxAllowedMipIndex].SizeY,
 				StreamingTexture.LastRenderTime,
 				StreamingTexture.BudgetMipBias,
 				UTexture::GetTextureGroupString(StreamingTexture.LODGroup));
 		}
 		else
 		{
-			UE_LOG(LogContentStreaming, Log,  TEXT("    Current=%dx%d Wanted=%dx%d Max=%dx%d BudgetBias=%d Group=%s"), 
+			UE_LOG(LogContentStreaming, Log,  TEXT("    Current=%dx%d Wanted=%dx%d MaxAllowed=%dx%d BudgetBias=%d Group=%s"), 
 				Mips[CurrentMipIndex].SizeX, Mips[CurrentMipIndex].SizeY, 
 				Mips[WantedMipIndex].SizeX, Mips[WantedMipIndex].SizeY, 
-				Mips[MaxMipIndex].SizeX, Mips[MaxMipIndex].SizeY, 
+				Mips[MaxAllowedMipIndex].SizeX, Mips[MaxAllowedMipIndex].SizeY,
 				StreamingTexture.BudgetMipBias,
 				UTexture::GetTextureGroupString(StreamingTexture.LODGroup));
 		}
@@ -1424,7 +1424,7 @@ bool FStreamingManagerTexture::HandleInvestigateTextureCommand( const TCHAR* Cmd
 				}
 				UE_LOG(LogContentStreaming, Log,  TEXT("  Current size:    %dx%d"), Texture2D->PlatformData->Mips[CurrentMipIndex].SizeX, Texture2D->PlatformData->Mips[CurrentMipIndex].SizeY );
 				UE_LOG(LogContentStreaming, Log,  TEXT("  Wanted size:     %dx%d"), Texture2D->PlatformData->Mips[WantedMipIndex].SizeX, Texture2D->PlatformData->Mips[WantedMipIndex].SizeY );
-				UE_LOG(LogContentStreaming, Log,  TEXT("  Max size:        %dx%d"), Texture2D->PlatformData->Mips[MaxMipIndex].SizeX, Texture2D->PlatformData->Mips[MaxMipIndex].SizeY );
+				UE_LOG(LogContentStreaming, Log,  TEXT("  MaxAllowed size: %dx%d"), Texture2D->PlatformData->Mips[MaxMipIndex].SizeX, Texture2D->PlatformData->Mips[MaxMipIndex].SizeY );
 				UE_LOG(LogContentStreaming, Log,  TEXT("  LoadOrder Priority: %d"), StreamingTexture.LoadOrderPriority );
 				UE_LOG(LogContentStreaming, Log,  TEXT("  Retention Priority: %d"), StreamingTexture.RetentionPriority );
 				UE_LOG(LogContentStreaming, Log,  TEXT("  Boost factor:    %.1f"), StreamingTexture.BoostFactor );
diff --git a/Engine/Source/Runtime/Engine/Private/UnrealEngine.cpp b/Engine/Source/Runtime/Engine/Private/UnrealEngine.cpp
--- a/Engine/Source/Runtime/Engine/Private/UnrealEngine.cpp
+++ b/Engine/Source/Runtime/Engine/Private/UnrealEngine.cpp
@@ -305,6 +305,11 @@ void ScalabilityCVarsSinkCallback()
 		bool bRecreateRenderstate = false;
 		bool bCacheResourceShaders = false;
 
+		if (LocalScalabilityCVars.DetailMode != GCachedScalabilityCVars.DetailMode)
+		{
+			bRecreateRenderstate = true;
+		}
+
 		if (LocalScalabilityCVars.MaterialQualityLevel != GCachedScalabilityCVars.MaterialQualityLevel)
 		{
 			bCacheResourceShaders = true;
@@ -2267,32 +2272,28 @@ void UEngine::GetAllLocalPlayerControllers(TArray<APlayerController*> & PlayerLi
  */
 struct FSortedTexture 
 {
-	int32		OrigSizeX;
-	int32		OrigSizeY;
-	int32		CookedSizeX;
-	int32		CookedSizeY;
+	int32		MaxAllowedSizeX;	// This is the disk size when cooked.
+	int32		MaxAllowedSizeY;
 	EPixelFormat Format;
 	int32		CurSizeX;
 	int32		CurSizeY;
 	int32		LODBias;
-	int32		MaxSize;
+	int32		MaxAllowedSize;
 	int32		CurrentSize;
-	FString Name;
+	FString		Name;
 	int32		LODGroup;
-	bool	bIsStreaming;
+	bool		bIsStreaming;
 	int32		UsageCount;
 
 	/** Constructor, initializing every member variable with passed in values. */
-	FSortedTexture(	int32 InOrigSizeX, int32 InOrigSizeY, int32 InCookedSizeX, int32 InCookedSizeY, EPixelFormat InFormat, int32 InCurSizeX, int32 InCurSizeY, int32 InLODBias, int32 InMaxSize, int32 InCurrentSize, const FString& InName, int32 InLODGroup, bool bInIsStreaming, int32 InUsageCount )
-	:	OrigSizeX( InOrigSizeX )
-	,	OrigSizeY( InOrigSizeY )
-	,	CookedSizeX( InCookedSizeX )
-	,	CookedSizeY( InCookedSizeY )
+	FSortedTexture(	int32 InMaxAllowedSizeX, int32 InMaxAllowedSizeY, EPixelFormat InFormat, int32 InCurSizeX, int32 InCurSizeY, int32 InLODBias, int32 InMaxAllowedSize, int32 InCurrentSize, const FString& InName, int32 InLODGroup, bool bInIsStreaming, int32 InUsageCount )
+	:	MaxAllowedSizeX( InMaxAllowedSizeX )
+	,	MaxAllowedSizeY( InMaxAllowedSizeY )
 	,	Format( InFormat )
 	,	CurSizeX( InCurSizeX )
 	,	CurSizeY( InCurSizeY )
 	,	LODBias( InLODBias )
-	,	MaxSize( InMaxSize )
+	,	MaxAllowedSize( InMaxAllowedSize )
 	,	CurrentSize( InCurrentSize )
 	,	Name( InName )
 	,	LODGroup( InLODGroup )
@@ -3693,17 +3694,15 @@ bool UEngine::HandleListTexturesCommand( const TCHAR* Cmd, FOutputDevice& Ar )
 		int32				LODGroup			= Texture->LODGroup;
 		int32				LODBias				= Texture->GetCachedLODBias();
 		int32				NumMips				= Texture->GetNumMips();	
-		int32				MaxMips				= FMath::Max( 1, FMath::Min( NumMips - Texture->GetCachedLODBias(), GMaxTextureMipCount ) );
-		int32				OrigSizeX			= Texture->GetSizeX();
-		int32				OrigSizeY			= Texture->GetSizeY();
-		int32				CookedSizeX			= Texture->GetSizeX() >> LODBias;
-		int32				CookedSizeY			= Texture->GetSizeY() >> LODBias;
+		int32				MaxAllowedMips		= FMath::Max( 1, FMath::Min( NumMips - LODBias, GMaxTextureMipCount ) );
+		int32				MaxAllowedSizeX		= Texture->GetSizeX() >> LODBias;
+		int32				MaxAllowedSizeY		= Texture->GetSizeY() >> LODBias;
 		EPixelFormat		Format				= Texture->GetPixelFormat();
 		int32				DroppedMips			= Texture->GetNumMips() - Texture->ResidentMips;
 		int32				CurSizeX			= Texture->GetSizeX() >> DroppedMips;
 		int32				CurSizeY			= Texture->GetSizeY() >> DroppedMips;
 		bool			bIsStreamingTexture		= Texture->GetStreamingIndex() != INDEX_NONE;
-		int32				MaxSize				= Texture->CalcTextureMemorySizeEnum( TMC_AllMips );
+		int32				MaxAllowedSize		= Texture->CalcTextureMemorySizeEnum( TMC_AllMipsBiased );
 		int32				CurrentSize			= Texture->CalcTextureMemorySizeEnum( TMC_ResidentMips );
 		int32				UsageCount			= TextureToUsageMap.FindRef( Texture );
 		bool				bIsForced			= Texture->bForceMiplevelsToBeResident && bIsStreamingTexture;
@@ -3714,16 +3713,14 @@ bool UEngine::HandleListTexturesCommand( const TCHAR* Cmd, FOutputDevice& Ar )
 			(!bShouldOnlyListStreaming && !bShouldOnlyListNonStreaming && !bShouldOnlyListForced) )
 		{
 			new(SortedTextures) FSortedTexture( 
-				OrigSizeX, 
-				OrigSizeY, 
-				CookedSizeX,
-				CookedSizeY,
+				MaxAllowedSizeX,
+				MaxAllowedSizeY,
 				Format,
 				CurSizeX,
 				CurSizeY,
 				LODBias, 
-				MaxSize / 1024, 
-				CurrentSize / 1024, 
+				MaxAllowedSize,
+				CurrentSize,
 				Texture->GetPathName(), 
 				LODGroup, 
 				bIsStreamingTexture,
@@ -3738,71 +3735,88 @@ bool UEngine::HandleListTexturesCommand( const TCHAR* Cmd, FOutputDevice& Ar )
 	TArray<FString> TextureGroupNames = UTextureLODSettings::GetTextureGroupNames();
 
 	TArray<uint64> TextureGroupCurrentSizes;
-	TArray<uint64> TextureGroupMaxSizes;
+	TArray<uint64> TextureGroupMaxAllowedSizes;
 	
 	TArray<uint64> FormatCurrentSizes;
-	TArray<uint64> FormatMaxSizes;
+	TArray<uint64> FormatMaxAllowedSizes;
 
 	TextureGroupCurrentSizes.AddZeroed(TextureGroupNames.Num());
-	TextureGroupMaxSizes.AddZeroed(TextureGroupNames.Num());
+	TextureGroupMaxAllowedSizes.AddZeroed(TextureGroupNames.Num());
 
 	FormatCurrentSizes.AddZeroed(PF_MAX);
-	FormatMaxSizes.AddZeroed(PF_MAX);
+	FormatMaxAllowedSizes.AddZeroed(PF_MAX);
 
 	// Display.
-	int32 TotalMaxSize		= 0;
+	int32 TotalMaxAllowedSize = 0;
 	int32 TotalCurrentSize	= 0;
-	Ar.Logf( TEXT(",Authored Width,Authored Height,Cooked Width,Cooked Height,Format,Current Width,Current Height,Max Size,Current Size,LODBias,LODGroup,Name,Streaming,Usage Count") );
+
+	if (!FPlatformProperties::RequiresCookedData())
+	{
+		Ar.Logf(TEXT("MaxAllowedSize: Width x Height (Size in KB, Bias from Authored), Current/InMem: Width x Height (Size in KB), Format, LODGroup, Name, Streaming, Usage Count"));
+	}
+	else
+	{
+		Ar.Logf(TEXT("Cooked/OnDisk: Width x Height (Size in KB), Current/InMem: Width x Height (Size in KB), Format, LODGroup, Name, Streaming, Usage Count"));
+	}
+
 	for( int32 TextureIndex=0; TextureIndex<SortedTextures.Num(); TextureIndex++ )
 	{
 		const FSortedTexture& SortedTexture = SortedTextures[TextureIndex];
 		const bool bValidTextureGroup = TextureGroupNames.IsValidIndex(SortedTexture.LODGroup);
-		Ar.Logf( TEXT(",%i,%i,%i,%i,%s,%i,%i,%i,%i,%i,%s,%s,%s,%i"),
-			SortedTexture.OrigSizeX,
-			SortedTexture.OrigSizeY,
-			SortedTexture.CookedSizeX,
-			SortedTexture.CookedSizeY,
-			GetPixelFormatString(SortedTexture.Format),
-			SortedTexture.CurSizeX,
-			SortedTexture.CurSizeY,
-			SortedTexture.MaxSize,
-			SortedTexture.CurrentSize,
-			SortedTexture.LODBias,
-			bValidTextureGroup ? *TextureGroupNames[SortedTexture.LODGroup] : TEXT("INVALID"),
-			*SortedTexture.Name,
-			SortedTexture.bIsStreaming ? TEXT("YES") : TEXT("NO"),
-			SortedTexture.UsageCount );
+
+		if (!FPlatformProperties::RequiresCookedData())
+		{
+			Ar.Logf(TEXT("%ix%i (%i KB, %i), %ix%i (%i KB), %s, %s, %s, %s, %i"),
+				SortedTexture.MaxAllowedSizeX, SortedTexture.MaxAllowedSizeY, SortedTexture.MaxAllowedSize / 1024, SortedTexture.LODBias,
+				SortedTexture.CurSizeX, SortedTexture.CurSizeY, SortedTexture.CurrentSize / 1024,
+				GetPixelFormatString(SortedTexture.Format),
+				bValidTextureGroup ? *TextureGroupNames[SortedTexture.LODGroup] : TEXT("INVALID"),
+				*SortedTexture.Name,
+				SortedTexture.bIsStreaming ? TEXT("YES") : TEXT("NO"),
+				SortedTexture.UsageCount);
+		}
+		else
+		{
+			Ar.Logf(TEXT("%ix%i (%i KB), %ix%i (%i KB), %s, %s, %s, %s, %i"),
+				SortedTexture.MaxAllowedSizeX, SortedTexture.MaxAllowedSizeY, SortedTexture.MaxAllowedSize / 1024,
+				SortedTexture.CurSizeX, SortedTexture.CurSizeY, SortedTexture.CurrentSize / 1024,
+				GetPixelFormatString(SortedTexture.Format),
+				bValidTextureGroup ? *TextureGroupNames[SortedTexture.LODGroup] : TEXT("INVALID"),
+				*SortedTexture.Name,
+				SortedTexture.bIsStreaming ? TEXT("YES") : TEXT("NO"),
+				SortedTexture.UsageCount);
+		}
 
 		if (bValidTextureGroup)
 		{
 			TextureGroupCurrentSizes[SortedTexture.LODGroup] += SortedTexture.CurrentSize;
-			TextureGroupMaxSizes[SortedTexture.LODGroup] += SortedTexture.MaxSize;
+			TextureGroupMaxAllowedSizes[SortedTexture.LODGroup] += SortedTexture.MaxAllowedSize;
 		}
 
 		if (SortedTexture.Format >= 0 && SortedTexture.Format < PF_MAX)
 		{
 			FormatCurrentSizes[SortedTexture.Format] += SortedTexture.CurrentSize;
-			FormatMaxSizes[SortedTexture.Format] += SortedTexture.MaxSize;
+			FormatMaxAllowedSizes[SortedTexture.Format] += SortedTexture.MaxAllowedSize;
 		}
 
-		TotalMaxSize		+= SortedTexture.MaxSize;
+		TotalMaxAllowedSize	+= SortedTexture.MaxAllowedSize;
 		TotalCurrentSize	+= SortedTexture.CurrentSize;
 	}
 
-	Ar.Logf(TEXT("Total size: Current= %d KB  Max= %d KB  Count=%d"), TotalCurrentSize, TotalMaxSize, SortedTextures.Num() );
+	Ar.Logf(TEXT("Total size: InMem= %.2f MB  OnDisk= %.2f MB  Count=%d"), (double)TotalCurrentSize / 1024. / 1024., (double)TotalMaxAllowedSize / 1024. / 1024., SortedTextures.Num() );
 	for (int32 i = 0; i < PF_MAX; ++i)
 	{
-		if (FormatCurrentSizes[i] > 0 || FormatMaxSizes[i] > 0)
+		if (FormatCurrentSizes[i] > 0 || FormatMaxAllowedSizes[i] > 0)
 		{
-			Ar.Logf(TEXT("Total %s size: Current= %d MB  Max= %d MB "), GetPixelFormatString((EPixelFormat)i), FormatCurrentSizes[i] / 1024, FormatMaxSizes[i] / 1024);
+			Ar.Logf(TEXT("Total %s size: InMem= %.2f MB  OnDisk= %.2f MB "), GetPixelFormatString((EPixelFormat)i), (double)FormatCurrentSizes[i] / 1024. / 1024., (double)FormatMaxAllowedSizes[i] / 1024. / 1024.);
 		}
 	}
 
 	for (int32 i = 0; i < TextureGroupCurrentSizes.Num(); ++i)
 	{
-		if (TextureGroupCurrentSizes[i] > 0 || TextureGroupMaxSizes[i] > 0)
+		if (TextureGroupCurrentSizes[i] > 0 || TextureGroupMaxAllowedSizes[i] > 0)
 		{
-			Ar.Logf(TEXT("Total %s size: Current= %d MB  Max= %d MB "), *TextureGroupNames[i], TextureGroupCurrentSizes[i] / 1024, TextureGroupMaxSizes[i] / 1024);
+			Ar.Logf(TEXT("Total %s size: InMem= %.2f MB  OnDisk= %.2f MB "), *TextureGroupNames[i], (double)TextureGroupCurrentSizes[i] / 1024. / 1024., (double)TextureGroupMaxAllowedSizes[i] / 1024. / 1024.);
 		}
 	}
 	return true;
diff --git a/Engine/Source/Runtime/Engine/Public/MaterialShared.h b/Engine/Source/Runtime/Engine/Public/MaterialShared.h
--- a/Engine/Source/Runtime/Engine/Public/MaterialShared.h
+++ b/Engine/Source/Runtime/Engine/Public/MaterialShared.h
@@ -997,6 +997,7 @@ public:
 	virtual bool IsCrackFreeDisplacementEnabled() const { return false; }
 	virtual bool IsAdaptiveTessellationEnabled() const { return false; }
 	virtual bool IsFullyRough() const { return false; }
+	virtual bool UseNormalCurvatureToRoughness() const { return false; }
 	virtual bool IsUsingFullPrecision() const { return false; }
 	virtual bool IsUsingHQForwardReflections() const { return false; }
 	virtual bool IsUsingPlanarForwardReflections() const { return false; }
@@ -1624,6 +1625,7 @@ public:
 	ENGINE_API virtual bool IsCrackFreeDisplacementEnabled() const override;
 	ENGINE_API virtual bool IsAdaptiveTessellationEnabled() const override;
 	ENGINE_API virtual bool IsFullyRough() const override;
+	ENGINE_API virtual bool UseNormalCurvatureToRoughness() const override;
 	ENGINE_API virtual bool IsUsingFullPrecision() const override;
 	ENGINE_API virtual bool IsUsingHQForwardReflections() const override;
 	ENGINE_API virtual bool IsUsingPlanarForwardReflections() const override;
diff --git a/Engine/Source/Runtime/Engine/Public/SceneInterface.h b/Engine/Source/Runtime/Engine/Public/SceneInterface.h
--- a/Engine/Source/Runtime/Engine/Public/SceneInterface.h
+++ b/Engine/Source/Runtime/Engine/Public/SceneInterface.h
@@ -122,7 +122,7 @@ public:
 	 * Updates the contents of the given sky capture by rendering the scene. 
 	 * This must be called on the game thread.
 	 */
-	virtual void UpdateSkyCaptureContents(const USkyLightComponent* CaptureComponent, bool bCaptureEmissiveOnly, UTextureCube* SourceCubemap, FTexture* OutProcessedTexture, FSHVectorRGB3& OutIrradianceEnvironmentMap) {}
+	virtual void UpdateSkyCaptureContents(const USkyLightComponent* CaptureComponent, bool bCaptureEmissiveOnly, UTextureCube* SourceCubemap, FTexture* OutProcessedTexture, float& OutAverageBrightness, FSHVectorRGB3& OutIrradianceEnvironmentMap) {}
 
 	/** Runs a slow preculling operation on static meshes, removing triangles that are invisible or inside a precull volume. */
 	virtual void PreCullStaticMeshes(const TArray<UStaticMeshComponent*>& ComponentsToPreCull, const TArray<TArray<FPlane> >& CullVolumes) {}
diff --git a/Engine/Source/Runtime/Engine/Public/SceneManagement.h b/Engine/Source/Runtime/Engine/Public/SceneManagement.h
--- a/Engine/Source/Runtime/Engine/Public/SceneManagement.h
+++ b/Engine/Source/Runtime/Engine/Public/SceneManagement.h
@@ -804,7 +804,12 @@ public:
 	/** Initialization constructor. */
 	FSkyLightSceneProxy(const class USkyLightComponent* InLightComponent);
 
-	void Initialize(float InBlendFraction, const FSHVectorRGB3* InIrradianceEnvironmentMap, const FSHVectorRGB3* BlendDestinationIrradianceEnvironmentMap);
+	void Initialize(
+		float InBlendFraction, 
+		const FSHVectorRGB3* InIrradianceEnvironmentMap, 
+		const FSHVectorRGB3* BlendDestinationIrradianceEnvironmentMap,
+		const float* InAverageBrightness,
+		const float* BlendDestinationAverageBrightness);
 
 	const USkyLightComponent* LightComponent;
 	FTexture* ProcessedTexture;
@@ -817,6 +822,7 @@ public:
 	bool bHasStaticLighting;
 	FLinearColor LightColor;
 	FSHVectorRGB3 IrradianceEnvironmentMap;
+	float AverageBrightness;
 	float IndirectLightingIntensity;
 	float OcclusionMaxDistance;
 	float Contrast;
@@ -1246,6 +1252,8 @@ public:
 	/** Used in Feature level SM4 */
 	FTexture* SM4FullHDRCubemap;
 
+	float AverageBrightness;
+
 	/** Used in Feature level ES2 */
 	FTexture* EncodedHDRCubemap;
 
@@ -1269,6 +1277,8 @@ public:
 
 	FReflectionCaptureProxy(const class UReflectionCaptureComponent* InComponent);
 
+	void InitializeAverageBrightness(const float& AverageBrightness);
+
 	void SetTransform(const FMatrix& InTransform);
 };
 
diff --git a/Engine/Source/Runtime/Engine/Public/SceneView.h b/Engine/Source/Runtime/Engine/Public/SceneView.h
--- a/Engine/Source/Runtime/Engine/Public/SceneView.h
+++ b/Engine/Source/Runtime/Engine/Public/SceneView.h
@@ -415,7 +415,6 @@ enum ETranslucencyVolumeCascade
 	VIEW_UNIFORM_BUFFER_MEMBER(uint32, FrameNumber) \
 	VIEW_UNIFORM_BUFFER_MEMBER(uint32, StateFrameIndexMod8) \
 	VIEW_UNIFORM_BUFFER_MEMBER_EX(float, CameraCut, EShaderPrecisionModifier::Half) \
-	VIEW_UNIFORM_BUFFER_MEMBER_EX(float, UseLightmaps, EShaderPrecisionModifier::Half) \
 	VIEW_UNIFORM_BUFFER_MEMBER_EX(float, UnlitViewmodeMask, EShaderPrecisionModifier::Half) \
 	VIEW_UNIFORM_BUFFER_MEMBER_EX(FLinearColor, DirectionalLightColor, EShaderPrecisionModifier::Half) \
 	VIEW_UNIFORM_BUFFER_MEMBER_EX(FVector, DirectionalLightDirection, EShaderPrecisionModifier::Half) \
@@ -451,6 +450,8 @@ enum ETranslucencyVolumeCascade
 	VIEW_UNIFORM_BUFFER_MEMBER(uint32, AtmosphericFogRenderMask) \
 	VIEW_UNIFORM_BUFFER_MEMBER(uint32, AtmosphericFogInscatterAltitudeSampleNum) \
 	VIEW_UNIFORM_BUFFER_MEMBER(FLinearColor, AtmosphericFogSunColor) \
+	VIEW_UNIFORM_BUFFER_MEMBER(FVector2D, NormalCurvatureToRoughnessScaleBias) \
+	VIEW_UNIFORM_BUFFER_MEMBER(float, RenderingReflectionCaptureMask) \
 	VIEW_UNIFORM_BUFFER_MEMBER(FLinearColor, AmbientCubemapTint) \
 	VIEW_UNIFORM_BUFFER_MEMBER(float, AmbientCubemapIntensity) \
 	VIEW_UNIFORM_BUFFER_MEMBER(FVector2D, RenderTargetSize) \
@@ -463,6 +464,7 @@ enum ETranslucencyVolumeCascade
 	VIEW_UNIFORM_BUFFER_MEMBER_EX(float, ReflectionCubemapMaxMip, EShaderPrecisionModifier::Half) \
 	VIEW_UNIFORM_BUFFER_MEMBER(float, ShowDecalsMask) \
 	VIEW_UNIFORM_BUFFER_MEMBER(uint32, DistanceFieldAOSpecularOcclusionMode) \
+	VIEW_UNIFORM_BUFFER_MEMBER(FVector2D, ReflectionEnvironmentRoughnessMixingScaleBias) \
 	VIEW_UNIFORM_BUFFER_MEMBER(uint32, StereoPassIndex) \
 	VIEW_UNIFORM_BUFFER_MEMBER_ARRAY(FVector4, GlobalVolumeCenterAndExtent_UB, [GMaxGlobalDistanceFieldClipmaps]) \
 	VIEW_UNIFORM_BUFFER_MEMBER_ARRAY(FVector4, GlobalVolumeWorldToUVAddAndMul_UB, [GMaxGlobalDistanceFieldClipmaps]) \
@@ -753,6 +755,7 @@ public:
 	 * Setup by the main thread, passed to the render thread and never touched again by the main thread.
 	 */
 	FFinalPostProcessSettings FinalPostProcessSettings;
+	EAntiAliasingMethod AntiAliasingMethod;
 
 	/** Parameters for atmospheric fog. */
 	FTextureRHIRef AtmosphereTransmittanceTexture;
@@ -883,6 +886,8 @@ public:
 	/** applied global restrictions from show flags */
 	void EndFinalPostprocessSettings(const FSceneViewInitOptions& ViewInitOptions);
 
+	void SetupAntiAliasingMethod();
+
 	/** Configure post process settings for the buffer visualization system */
 	void ConfigureBufferVisualizationSettings();
 
diff --git a/Engine/Source/Runtime/Launch/Private/LaunchEngineLoop.cpp b/Engine/Source/Runtime/Launch/Private/LaunchEngineLoop.cpp
--- a/Engine/Source/Runtime/Launch/Private/LaunchEngineLoop.cpp
+++ b/Engine/Source/Runtime/Launch/Private/LaunchEngineLoop.cpp
@@ -823,6 +823,11 @@ int32 FEngineLoop::PreInit( const TCHAR* CmdLine )
 			}
 		}
 	}
+
+	if (FParse::Param(FCommandLine::Get(), TEXT("emitdrawevents")))
+	{
+		GEmitDrawEvents = true;
+	}	
 #endif // !UE_BUILD_SHIPPING
 
 	// Switch into executable's directory (may be required by some of the platform file overrides)
@@ -2522,22 +2527,24 @@ void FEngineLoop::ProcessLocalPlayerSlateOperations() const
 
 			if ( ViewportWidget.IsValid() )
 			{
-				for( FConstPlayerControllerIterator Iterator = CurWorld->GetPlayerControllerIterator(); Iterator; ++Iterator )
+				FWidgetPath PathToWidget;
+				SlateApp.GeneratePathToWidgetUnchecked(ViewportWidget.ToSharedRef(), PathToWidget);
+
+				if (PathToWidget.IsValid())
 				{
-					APlayerController* PlayerController = *Iterator;
-					if( PlayerController )
+					for (FConstPlayerControllerIterator Iterator = CurWorld->GetPlayerControllerIterator(); Iterator; ++Iterator)
 					{
-						ULocalPlayer* LocalPlayer = Cast< ULocalPlayer >( PlayerController->Player );
-						if( LocalPlayer )
+						APlayerController* PlayerController = *Iterator;
+						if (PlayerController)
 						{
-							FReply& TheReply = LocalPlayer->GetSlateOperations();
-
-							FWidgetPath PathToWidget;
-							SlateApp.GeneratePathToWidgetUnchecked( ViewportWidget.ToSharedRef(), PathToWidget );
-
-							SlateApp.ProcessReply( PathToWidget, TheReply, nullptr, nullptr, LocalPlayer->GetControllerId() );
+							ULocalPlayer* LocalPlayer = Cast< ULocalPlayer >(PlayerController->Player);
+							if (LocalPlayer)
+							{
+								FReply& TheReply = LocalPlayer->GetSlateOperations();
+								SlateApp.ProcessReply(PathToWidget, TheReply, nullptr, nullptr, LocalPlayer->GetControllerId());
 
-							TheReply = FReply::Unhandled();
+								TheReply = FReply::Unhandled();
+							}
 						}
 					}
 				}
diff --git a/Engine/Source/Runtime/RHI/Public/DynamicRHI.h b/Engine/Source/Runtime/RHI/Public/DynamicRHI.h
--- a/Engine/Source/Runtime/RHI/Public/DynamicRHI.h
+++ b/Engine/Source/Runtime/RHI/Public/DynamicRHI.h
@@ -1054,6 +1054,9 @@ public:
 	virtual void RHIEnqueueDecompress(uint8_t* SrcBuffer, uint8_t* DestBuffer, int CompressedSize, void* ErrorCodeBuffer) {}
 	virtual void RHIEnqueueCompress(uint8_t* SrcBuffer, uint8_t* DestBuffer, int UnCompressedSize, void* ErrorCodeBuffer) {}
 
+	// FlushType: Flush Immediate
+	virtual void RHIRecreateRecursiveBoundShaderStates() {}
+
 	/**
 	 *	Retrieve available screen resolutions.
 	 *
diff --git a/Engine/Source/Runtime/RHI/Public/RHICommandList.h b/Engine/Source/Runtime/RHI/Public/RHICommandList.h
--- a/Engine/Source/Runtime/RHI/Public/RHICommandList.h
+++ b/Engine/Source/Runtime/RHI/Public/RHICommandList.h
@@ -3376,6 +3376,11 @@ FORCEINLINE void* RHIGetNativeDevice()
 	return FRHICommandListExecutor::GetImmediateCommandList().GetNativeDevice();
 }
 
+FORCEINLINE void RHIRecreateRecursiveBoundShaderStates()
+{
+	FRHICommandListExecutor::GetImmediateCommandList(). ImmediateFlush(EImmediateFlushType::FlushRHIThread);
+	GDynamicRHI->RHIRecreateRecursiveBoundShaderStates();
+}
 
 
 
diff --git a/Engine/Source/Runtime/RHI/Public/RHIDefinitions.h b/Engine/Source/Runtime/RHI/Public/RHIDefinitions.h
--- a/Engine/Source/Runtime/RHI/Public/RHIDefinitions.h
+++ b/Engine/Source/Runtime/RHI/Public/RHIDefinitions.h
@@ -794,8 +794,8 @@ inline bool RHINeedsToSwitchVerticalAxis(EShaderPlatform Platform)
 
 inline bool RHISupportsSeparateMSAAAndResolveTextures(const EShaderPlatform Platform)
 {
-	// Metal needs to handle MSAA and resolve textures internally (unless RHICreateTexture2D was changed to take an optional resolve target)
-	return !IsMetalPlatform(Platform) && IsVulkanPlatform(Platform);
+	// Metal and Vulkan needs to handle MSAA and resolve textures internally (unless RHICreateTexture2D was changed to take an optional resolve target)
+	return !IsMetalPlatform(Platform) && !IsVulkanPlatform(Platform);
 }
 
 inline bool RHISupportsComputeShaders(const EShaderPlatform Platform)
diff --git a/Engine/Source/Runtime/Renderer/Private/AtmosphereRendering.cpp b/Engine/Source/Runtime/Renderer/Private/AtmosphereRendering.cpp
--- a/Engine/Source/Runtime/Renderer/Private/AtmosphereRendering.cpp
+++ b/Engine/Source/Runtime/Renderer/Private/AtmosphereRendering.cpp
@@ -438,9 +438,6 @@ void FDeferredShadingSceneRenderer::RenderAtmosphere(FRHICommandListImmediate& R
 				sizeof(Vertices[0])
 				);
 		}
-
-		//no need to resolve since we used alpha blending
-		SceneContext.FinishRenderingSceneColor(RHICmdList, false);
 	}
 }
 
diff --git a/Engine/Source/Runtime/Renderer/Private/BasePassRendering.cpp b/Engine/Source/Runtime/Renderer/Private/BasePassRendering.cpp
--- a/Engine/Source/Runtime/Renderer/Private/BasePassRendering.cpp
+++ b/Engine/Source/Runtime/Renderer/Private/BasePassRendering.cpp
@@ -24,6 +24,13 @@ static TAutoConsoleVariable<int32> CVarGlobalClipPlane(
 	TEXT("Enables mesh shaders to support a global clip plane, needed for planar reflections, which adds about 15% BasePass GPU cost on PS4."),
 	ECVF_ReadOnly | ECVF_RenderThreadSafe);
 
+// Changing this causes a full shader recompile
+static TAutoConsoleVariable<int32> CVarVertexFoggingForOpaque(
+	TEXT("r.VertexFoggingForOpaque"),
+	0,
+	TEXT("Causes opaque materials to use per-vertex fogging, which costs less and integrates properly with MSAA.  Only supported with forward shading."),
+	ECVF_ReadOnly | ECVF_RenderThreadSafe);
+
 static TAutoConsoleVariable<int32> CVarParallelBasePass(
 	TEXT("r.ParallelBasePass"),
 	1,
@@ -133,13 +140,15 @@ void FSkyLightReflectionParameters::GetSkyParametersFromScene(
 	float& OutApplySkyLightMask, 
 	float& OutSkyMipCount, 
 	bool& bOutSkyLightIsDynamic, 
-	float& OutBlendFraction)
+	float& OutBlendFraction,
+	float& OutSkyAverageBrightness)
 {
 	OutSkyLightTextureResource = GBlackTextureCube;
 	OutSkyLightBlendDestinationTextureResource = GBlackTextureCube;
 	OutApplySkyLightMask = 0;
 	bOutSkyLightIsDynamic = false;
 	OutBlendFraction = 0;
+	OutSkyAverageBrightness = 1.0f;
 
 	if (Scene
 		&& Scene->SkyLight 
@@ -165,6 +174,7 @@ void FSkyLightReflectionParameters::GetSkyParametersFromScene(
 		
 		OutApplySkyLightMask = 1;
 		bOutSkyLightIsDynamic = !SkyLight.bHasStaticLighting && !SkyLight.bWantsStaticShadowing;
+		OutSkyAverageBrightness = SkyLight.AverageBrightness;
 	}
 
 	OutSkyMipCount = 1;
@@ -203,7 +213,7 @@ void FBasePassReflectionParameters::SetMesh(FRHICommandList& RHICmdList, FPixelS
 	FMatrix BoxTransformVal = FMatrix::Identity;
 	FVector4 PositionAndRadius = FVector::ZeroVector;
 	FVector4 BoxScalesVal(0, 0, 0, 0);
-	FVector CaptureOffsetVal = FVector::ZeroVector;
+	FVector4 CaptureOffsetAndAverageBrightnessValue(0, 0, 0, 1);
 	EReflectionCaptureShape::Type CaptureShape = EReflectionCaptureShape::Box;
 	
 	if (PrimitiveSceneInfo && ReflectionProxy)
@@ -213,7 +223,7 @@ void FBasePassReflectionParameters::SetMesh(FRHICommandList& RHICmdList, FPixelS
 		CaptureShape = ReflectionProxy->Shape;
 		BoxTransformVal = ReflectionProxy->BoxTransform;
 		BoxScalesVal = FVector4(ReflectionProxy->BoxScales, ReflectionProxy->BoxTransitionDistance);
-		CaptureOffsetVal = ReflectionProxy->CaptureOffset;
+		CaptureOffsetAndAverageBrightnessValue = FVector4(ReflectionProxy->CaptureOffset, ReflectionProxy->AverageBrightness);
 	}
 
 	SetTextureParameter(
@@ -229,8 +239,7 @@ void FBasePassReflectionParameters::SetMesh(FRHICommandList& RHICmdList, FPixelS
 	SetShaderValue(RHICmdList, PixelShaderRHI, ReflectionShape, (float)CaptureShape);
 	SetShaderValue(RHICmdList, PixelShaderRHI, BoxTransform, BoxTransformVal);
 	SetShaderValue(RHICmdList, PixelShaderRHI, BoxScales, BoxScalesVal);
-	SetShaderValue(RHICmdList, PixelShaderRHI, CaptureOffset, CaptureOffsetVal);
-	
+	SetShaderValue(RHICmdList, PixelShaderRHI, CaptureOffsetAndAverageBrightness, CaptureOffsetAndAverageBrightnessValue);
 }
 
 void FTranslucentLightingParameters::Set(FRHICommandList& RHICmdList, FPixelShaderRHIParamRef PixelShaderRHI, const FViewInfo* View)
@@ -367,8 +376,8 @@ public:
 					Parameters.TextureMode,
 					bRenderSkylight,
 					bRenderAtmosphericFog,
-					/* DebugViewShaderMode = */ DVSM_None,
-					/* bInAllowGlobalFog = */ false,
+					DVSM_None,
+					Scene->ReadOnlyCVARCache.bEnableVertexFoggingForOpaque,
 					/* bInEnableEditorPrimitiveDepthTest = */ false,
 					/* bInEnableReceiveDecalOutput = */ true
 				),
@@ -490,7 +499,7 @@ public:
 			bRenderSkylight,
 			bRenderAtmosphericFog,
 			View.Family->GetDebugViewShaderMode(),
-			false,
+			Scene ? Scene->ReadOnlyCVARCache.bEnableVertexFoggingForOpaque : false,
 			Parameters.bEditorCompositeDepthTest,
 			/* bInEnableReceiveDecalOutput = */ Scene != nullptr
 			);
diff --git a/Engine/Source/Runtime/Renderer/Private/BasePassRendering.h b/Engine/Source/Runtime/Renderer/Private/BasePassRendering.h
--- a/Engine/Source/Runtime/Renderer/Private/BasePassRendering.h
+++ b/Engine/Source/Runtime/Renderer/Private/BasePassRendering.h
@@ -475,6 +475,7 @@ public:
 		SkyLightBlendDestinationCubemap.Bind(ParameterMap, TEXT("SkyLightBlendDestinationCubemap"));
 		SkyLightBlendDestinationCubemapSampler.Bind(ParameterMap, TEXT("SkyLightBlendDestinationCubemapSampler"));
 		SkyLightParameters.Bind(ParameterMap, TEXT("SkyLightParameters"));
+		SkyLightCubemapBrightness.Bind(ParameterMap, TEXT("SkyLightCubemapBrightness"));
 	}
 
 	template<typename TParamRef, typename TRHICmdList>
@@ -488,19 +489,21 @@ public:
 			float SkyMipCount = 1;
 			float BlendFraction = 0;
 			bool bSkyLightIsDynamic = false;
+			float SkyAverageBrightness = 1.0f;
 
-			GetSkyParametersFromScene(Scene, bApplySkyLight, SkyLightTextureResource, SkyLightBlendDestinationTextureResource, ApplySkyLightMask, SkyMipCount, bSkyLightIsDynamic, BlendFraction);
+			GetSkyParametersFromScene(Scene, bApplySkyLight, SkyLightTextureResource, SkyLightBlendDestinationTextureResource, ApplySkyLightMask, SkyMipCount, bSkyLightIsDynamic, BlendFraction, SkyAverageBrightness);
 
 			SetTextureParameter(RHICmdList, ShaderRHI, SkyLightCubemap, SkyLightCubemapSampler, SkyLightTextureResource);
 			SetTextureParameter(RHICmdList, ShaderRHI, SkyLightBlendDestinationCubemap, SkyLightBlendDestinationCubemapSampler, SkyLightBlendDestinationTextureResource);
 			const FVector4 SkyParametersValue(SkyMipCount - 1.0f, ApplySkyLightMask, bSkyLightIsDynamic ? 1.0f : 0.0f, BlendFraction);
 			SetShaderValue(RHICmdList, ShaderRHI, SkyLightParameters, SkyParametersValue);
+			SetShaderValue(RHICmdList, ShaderRHI, SkyLightCubemapBrightness, SkyAverageBrightness);
 		}
 	}
 
 	friend FArchive& operator<<(FArchive& Ar,FSkyLightReflectionParameters& P)
 	{
-		Ar << P.SkyLightCubemap << P.SkyLightCubemapSampler << P.SkyLightParameters << P.SkyLightBlendDestinationCubemap << P.SkyLightBlendDestinationCubemapSampler;
+		Ar << P.SkyLightCubemap << P.SkyLightCubemapSampler << P.SkyLightParameters << P.SkyLightBlendDestinationCubemap << P.SkyLightBlendDestinationCubemapSampler << P.SkyLightCubemapBrightness;
 		return Ar;
 	}
 
@@ -511,6 +514,7 @@ private:
 	FShaderResourceParameter SkyLightBlendDestinationCubemap;
 	FShaderResourceParameter SkyLightBlendDestinationCubemapSampler;
 	FShaderParameter SkyLightParameters;
+	FShaderParameter SkyLightCubemapBrightness;
 
 	void GetSkyParametersFromScene(
 		const FScene* Scene, 
@@ -520,7 +524,8 @@ private:
 		float& OutApplySkyLightMask, 
 		float& OutSkyMipCount, 
 		bool& bSkyLightIsDynamic, 
-		float& OutBlendFraction);
+		float& OutBlendFraction,
+		float& OutSkyAverageBrightness);
 };
 
 /** Parameters needed for reflections, shared by multiple shaders. */
@@ -538,7 +543,7 @@ public:
 		ReflectionShape.Bind(ParameterMap, TEXT("ReflectionShape"));
 		BoxTransform.Bind(ParameterMap, TEXT("BoxTransform"));
 		BoxScales.Bind(ParameterMap, TEXT("BoxScales"));		
-		CaptureOffset.Bind(ParameterMap, TEXT("CaptureOffset"));
+		CaptureOffsetAndAverageBrightness.Bind(ParameterMap, TEXT("CaptureOffsetAndAverageBrightness"));
 		SkyLightReflectionParameters.Bind(ParameterMap);
 	}
 
@@ -557,7 +562,7 @@ public:
 		Ar << P.ReflectionShape;
 		Ar << P.BoxTransform;
 		Ar << P.BoxScales;
-		Ar << P.CaptureOffset;
+		Ar << P.CaptureOffsetAndAverageBrightness;
 		Ar << P.SkyLightReflectionParameters;
 		return Ar;
 	}
@@ -573,7 +578,7 @@ private:
 	FShaderParameter ReflectionShape;
 	FShaderParameter BoxTransform;
 	FShaderParameter BoxScales;
-	FShaderParameter CaptureOffset;
+	FShaderParameter CaptureOffsetAndAverageBrightness;
 
 	FSkyLightReflectionParameters SkyLightReflectionParameters;
 };
diff --git a/Engine/Source/Runtime/Renderer/Private/CompositionLighting/PostProcessDeferredDecals.cpp b/Engine/Source/Runtime/Renderer/Private/CompositionLighting/PostProcessDeferredDecals.cpp
--- a/Engine/Source/Runtime/Renderer/Private/CompositionLighting/PostProcessDeferredDecals.cpp
+++ b/Engine/Source/Runtime/Renderer/Private/CompositionLighting/PostProcessDeferredDecals.cpp
@@ -623,8 +623,11 @@ void FRCPassPostProcessDeferredDecals::Process(FRenderingCompositePassContext& C
 
 		check(bDBuffer);
 
+		FPooledRenderTargetDesc GBufferADesc;
+		SceneContext.GetGBufferADesc(GBufferADesc);
+
 		// DBuffer: Decal buffer
-		FPooledRenderTargetDesc Desc(FPooledRenderTargetDesc::Create2DDesc(SceneContext.GBufferA->GetDesc().Extent,
+		FPooledRenderTargetDesc Desc(FPooledRenderTargetDesc::Create2DDesc(GBufferADesc.Extent,
 			PF_B8G8R8A8,
 			FClearValueBinding::None,
 			TexCreate_None,
@@ -636,8 +639,8 @@ void FRCPassPostProcessDeferredDecals::Process(FRenderingCompositePassContext& C
 
 		if (GSupportsRenderTargetWriteMask)
 		{
-			int32 W = FPlatformMath::CeilToInt(SceneContext.GBufferA->GetDesc().Extent.X / 8.0f);
-			int32 H = FPlatformMath::CeilToInt(SceneContext.GBufferA->GetDesc().Extent.Y / 8.0f);
+			int32 W = FPlatformMath::CeilToInt(GBufferADesc.Extent.X / 8.0f);
+			int32 H = FPlatformMath::CeilToInt(GBufferADesc.Extent.Y / 8.0f);
 			FIntPoint RTWriteMaskDims(W, H);
 			FPooledRenderTargetDesc MaskDesc(FPooledRenderTargetDesc::Create2DDesc(RTWriteMaskDims,
 				PF_R8_UINT,
@@ -682,7 +685,7 @@ void FRCPassPostProcessDeferredDecals::Process(FRenderingCompositePassContext& C
 			RenderTargets[1] = FRHIRenderTargetView(SceneContext.DBufferB->GetRenderTargetItem().TargetableTexture, 0, -1, ERenderTargetLoadAction::EClear, ERenderTargetStoreAction::EStore);
 			RenderTargets[2] = FRHIRenderTargetView(SceneContext.DBufferC->GetRenderTargetItem().TargetableTexture, 0, -1, ERenderTargetLoadAction::EClear, ERenderTargetStoreAction::EStore);
 
-			FRHIDepthRenderTargetView DepthView(SceneContext.GetSceneDepthSurface(), ERenderTargetLoadAction::ELoad, ERenderTargetStoreAction::ENoAction, ERenderTargetLoadAction::ELoad, ERenderTargetStoreAction::ENoAction, FExclusiveDepthStencil(FExclusiveDepthStencil::DepthRead_StencilWrite));
+			FRHIDepthRenderTargetView DepthView(SceneContext.GetSceneDepthTexture(), ERenderTargetLoadAction::ELoad, ERenderTargetStoreAction::ENoAction, ERenderTargetLoadAction::ELoad, ERenderTargetStoreAction::ENoAction, FExclusiveDepthStencil(FExclusiveDepthStencil::DepthRead_StencilWrite));
 
 			FRHISetRenderTargetsInfo Info(3, RenderTargets, DepthView);
 			RHICmdList.SetRenderTargetsAndClear(Info);
@@ -705,7 +708,7 @@ void FRCPassPostProcessDeferredDecals::Process(FRenderingCompositePassContext& C
 		//don't early return.  Resolves must be run for fast clears to work.
 		if (Scene.Decals.Num())
 		{
-			FDecalRenderTargetManager RenderTargetManager(RHICmdList, CurrentStage);
+			FDecalRenderTargetManager RenderTargetManager(RHICmdList, Context.GetShaderPlatform(), CurrentStage);
 
 			// Build a list of decals that need to be rendered for this view
 			FTransientDecalRenderDataList SortedDecals;
diff --git a/Engine/Source/Runtime/Renderer/Private/CompositionLighting/PostProcessDeferredDecals.h b/Engine/Source/Runtime/Renderer/Private/CompositionLighting/PostProcessDeferredDecals.h
--- a/Engine/Source/Runtime/Renderer/Private/CompositionLighting/PostProcessDeferredDecals.h
+++ b/Engine/Source/Runtime/Renderer/Private/CompositionLighting/PostProcessDeferredDecals.h
@@ -61,13 +61,13 @@ struct FDecalRenderTargetManager
 	bool bGufferADirty;
 
 	// constructor
-	FDecalRenderTargetManager(FRHICommandList& InRHICmdList, EDecalRenderStage CurrentStage)
+	FDecalRenderTargetManager(FRHICommandList& InRHICmdList, EShaderPlatform ShaderPlatform, EDecalRenderStage CurrentStage)
 		: RHICmdList(InRHICmdList)
-		, bGufferADirty(CurrentStage == DRS_AfterBasePass) // Normal buffer is already dirty at this point and needs resolve before being read from (irrelevant for DBuffer).
+		, bGufferADirty(false) 
 	{
 		FSceneRenderTargets& SceneContext = FSceneRenderTargets::Get(RHICmdList);
 
-		for(uint32 i = 0; i < ResolveBufferMax; ++i)
+		for (uint32 i = 0; i < ResolveBufferMax; ++i)
 		{
 			TargetsToTransitionWritable[i] = true;
 			TargetsToResolve[i] = nullptr;
@@ -81,10 +81,17 @@ struct FDecalRenderTargetManager
 		{
 			TargetsToResolve[DBufferBIndex] = SceneContext.DBufferB->GetRenderTargetItem().TargetableTexture;
 		}
+
 		if(SceneContext.DBufferC)
 		{
 			TargetsToResolve[DBufferCIndex] = SceneContext.DBufferC->GetRenderTargetItem().TargetableTexture;
 		}
+
+		if (!IsAnyForwardShadingEnabled(ShaderPlatform))
+		{
+			// Normal buffer is already dirty at this point and needs resolve before being read from (irrelevant for DBuffer).
+			bGufferADirty = CurrentStage == DRS_AfterBasePass;
+		}
 	}
 
 	// destructor
@@ -139,7 +146,7 @@ struct FDecalRenderTargetManager
 				TargetsToResolve[GBufferAIndex] = bHasNormal ? SceneContext.GBufferA->GetRenderTargetItem().TargetableTexture : (PLATFORM_MAC ? SceneContext.GBufferB->GetRenderTargetItem().TargetableTexture : nullptr); // @todo Workaround a Mac NV/Intel graphics driver bug that requires we pointlessly bind into RT1 even though we don't write to it, otherwise the writes to RT2 and RT3 go haywire. This isn't really possible to fix lower down the stack.
 				TargetsToResolve[GBufferBIndex] = SceneContext.GBufferB->GetRenderTargetItem().TargetableTexture;
 				TargetsToResolve[GBufferCIndex] = SceneContext.GBufferC->GetRenderTargetItem().TargetableTexture;
-				SetRenderTargets(RHICmdList, 4, TargetsToResolve, SceneContext.GetSceneDepthSurface(), ESimpleRenderTargetMode::EExistingColorAndDepth, FExclusiveDepthStencil::DepthRead_StencilWrite, TargetsToTransitionWritable[CurrentRenderTargetMode]);
+				SetRenderTargets(RHICmdList, 4, TargetsToResolve, SceneContext.GetSceneDepthTexture(), ESimpleRenderTargetMode::EExistingColorAndDepth, FExclusiveDepthStencil::DepthRead_StencilWrite, TargetsToTransitionWritable[CurrentRenderTargetMode]);
 				break;
 
 			case FDecalRenderingCommon::RTM_SceneColorAndGBufferDepthWriteWithNormal:
@@ -149,24 +156,24 @@ struct FDecalRenderTargetManager
 				TargetsToResolve[GBufferBIndex] = SceneContext.GBufferB->GetRenderTargetItem().TargetableTexture;
 				TargetsToResolve[GBufferCIndex] = SceneContext.GBufferC->GetRenderTargetItem().TargetableTexture;
 				TargetsToResolve[GBufferEIndex] = SceneContext.GBufferE->GetRenderTargetItem().TargetableTexture;
-				SetRenderTargets(RHICmdList, 5, TargetsToResolve, SceneContext.GetSceneDepthSurface(), ESimpleRenderTargetMode::EExistingColorAndDepth, FExclusiveDepthStencil::DepthWrite_StencilWrite, TargetsToTransitionWritable[CurrentRenderTargetMode]);
+				SetRenderTargets(RHICmdList, 5, TargetsToResolve, SceneContext.GetSceneDepthTexture(), ESimpleRenderTargetMode::EExistingColorAndDepth, FExclusiveDepthStencil::DepthWrite_StencilWrite, TargetsToTransitionWritable[CurrentRenderTargetMode]);
 				break;
 
 			case FDecalRenderingCommon::RTM_GBufferNormal:
 				TargetsToResolve[GBufferAIndex] = SceneContext.GBufferA->GetRenderTargetItem().TargetableTexture;
-				SetRenderTarget(RHICmdList, TargetsToResolve[GBufferAIndex], SceneContext.GetSceneDepthSurface(), ESimpleRenderTargetMode::EExistingColorAndDepth, FExclusiveDepthStencil::DepthRead_StencilWrite, TargetsToTransitionWritable[CurrentRenderTargetMode]);
+				SetRenderTarget(RHICmdList, TargetsToResolve[GBufferAIndex], SceneContext.GetSceneDepthTexture(), ESimpleRenderTargetMode::EExistingColorAndDepth, FExclusiveDepthStencil::DepthRead_StencilWrite, TargetsToTransitionWritable[CurrentRenderTargetMode]);
 				break;
 
 			case FDecalRenderingCommon::RTM_SceneColor:
 				TargetsToResolve[SceneColorIndex] = SceneContext.GetSceneColor()->GetRenderTargetItem().TargetableTexture;
-				SetRenderTarget(RHICmdList, TargetsToResolve[SceneColorIndex], SceneContext.GetSceneDepthSurface(), ESimpleRenderTargetMode::EExistingColorAndDepth, FExclusiveDepthStencil::DepthRead_StencilWrite, TargetsToTransitionWritable[CurrentRenderTargetMode]);
+				SetRenderTarget(RHICmdList, TargetsToResolve[SceneColorIndex], SceneContext.GetSceneDepthTexture(), ESimpleRenderTargetMode::EExistingColorAndDepth, FExclusiveDepthStencil::DepthRead_StencilWrite, TargetsToTransitionWritable[CurrentRenderTargetMode]);
 				break;
 
 			case FDecalRenderingCommon::RTM_DBuffer:
 				TargetsToResolve[DBufferAIndex] = SceneContext.DBufferA->GetRenderTargetItem().TargetableTexture;
 				TargetsToResolve[DBufferBIndex] = SceneContext.DBufferB->GetRenderTargetItem().TargetableTexture;
 				TargetsToResolve[DBufferCIndex] = SceneContext.DBufferC->GetRenderTargetItem().TargetableTexture;
-				SetRenderTargets(RHICmdList, 3, &TargetsToResolve[DBufferAIndex], SceneContext.GetSceneDepthSurface(), ESimpleRenderTargetMode::EExistingColorAndDepth, FExclusiveDepthStencil::DepthRead_StencilWrite, TargetsToTransitionWritable[CurrentRenderTargetMode]);
+				SetRenderTargets(RHICmdList, 3, &TargetsToResolve[DBufferAIndex], SceneContext.GetSceneDepthTexture(), ESimpleRenderTargetMode::EExistingColorAndDepth, FExclusiveDepthStencil::DepthRead_StencilWrite, TargetsToTransitionWritable[CurrentRenderTargetMode]);
 				break;
 
 			default:
diff --git a/Engine/Source/Runtime/Renderer/Private/CompositionLighting/PostProcessLpvIndirect.cpp b/Engine/Source/Runtime/Renderer/Private/CompositionLighting/PostProcessLpvIndirect.cpp
--- a/Engine/Source/Runtime/Renderer/Private/CompositionLighting/PostProcessLpvIndirect.cpp
+++ b/Engine/Source/Runtime/Renderer/Private/CompositionLighting/PostProcessLpvIndirect.cpp
@@ -434,27 +434,7 @@ void FRCPassPostProcessVisualizeLPV::Process(FRenderingCompositePassContext& Con
 	SetRenderTarget(Context.RHICmdList, DestRenderTarget.TargetableTexture, FTextureRHIRef());
 
 	{
-		// this is a helper class for FCanvas to be able to get screen size
-		class FRenderTargetTemp : public FRenderTarget
-		{
-		public:
-			const FSceneView& View;
-			const FTexture2DRHIRef Texture;
-
-			FRenderTargetTemp(const FSceneView& InView, const FTexture2DRHIRef InTexture)
-				: View(InView), Texture(InTexture)
-			{
-			}
-			virtual FIntPoint GetSizeXY() const
-			{
-				return View.ViewRect.Size();
-			};
-			virtual const FTexture2DRHIRef& GetRenderTargetTexture() const
-			{
-				return Texture;
-			}
-		} TempRenderTarget(View, (const FTexture2DRHIRef&)DestRenderTarget.TargetableTexture);
-
+		FRenderTargetTemp TempRenderTarget(View, (const FTexture2DRHIRef&)DestRenderTarget.TargetableTexture);
 		FCanvas Canvas(&TempRenderTarget, NULL, ViewFamily.CurrentRealTime, ViewFamily.CurrentWorldTime, ViewFamily.DeltaWorldTime, View.GetFeatureLevel());
 
 		float X = 30;
diff --git a/Engine/Source/Runtime/Renderer/Private/CompositionLighting/PostProcessMeshDecals.cpp b/Engine/Source/Runtime/Renderer/Private/CompositionLighting/PostProcessMeshDecals.cpp
--- a/Engine/Source/Runtime/Renderer/Private/CompositionLighting/PostProcessMeshDecals.cpp
+++ b/Engine/Source/Runtime/Renderer/Private/CompositionLighting/PostProcessMeshDecals.cpp
@@ -357,7 +357,7 @@ public:
 			, CurrentDecalStage(InCurrentDecalStage)
 			, LastRenderTargetMode(FDecalRenderingCommon::RTM_Unknown)
 			, LastDecalBlendMode(EDecalBlendMode::DBM_MAX)
-			, RenderTargetManager(Context.RHICmdList, CurrentDecalStage)
+			, RenderTargetManager(Context.RHICmdList, Context.GetShaderPlatform(), CurrentDecalStage)
 		{
 		}
 
diff --git a/Engine/Source/Runtime/Renderer/Private/CompositionLighting/PostProcessPassThrough.cpp b/Engine/Source/Runtime/Renderer/Private/CompositionLighting/PostProcessPassThrough.cpp
--- a/Engine/Source/Runtime/Renderer/Private/CompositionLighting/PostProcessPassThrough.cpp
+++ b/Engine/Source/Runtime/Renderer/Private/CompositionLighting/PostProcessPassThrough.cpp
@@ -11,46 +11,27 @@
 #include "PostProcessing.h"
 #include "SceneUtils.h"
 
-/** Encapsulates a simple copy pixel shader. */
-class FPostProcessPassThroughPS : public FGlobalShader
+FPostProcessPassThroughPS::FPostProcessPassThroughPS(const ShaderMetaType::CompiledShaderInitializerType& Initializer)
+	: FGlobalShader(Initializer)
 {
-	DECLARE_SHADER_TYPE(FPostProcessPassThroughPS, Global);
-
-	static bool ShouldCache(EShaderPlatform Platform)
-	{
-		return IsFeatureLevelSupported(Platform, ERHIFeatureLevel::SM4);
-	}
-
-	/** Default constructor. */
-	FPostProcessPassThroughPS() {}
-
-public:
-	FPostProcessPassParameters PostprocessParameter;
+	PostprocessParameter.Bind(Initializer.ParameterMap);
+}
 
-	/** Initialization constructor. */
-	FPostProcessPassThroughPS(const ShaderMetaType::CompiledShaderInitializerType& Initializer)
-		: FGlobalShader(Initializer)
-	{
-		PostprocessParameter.Bind(Initializer.ParameterMap);
-	}
+bool FPostProcessPassThroughPS::Serialize(FArchive& Ar)
+{
+	bool bShaderHasOutdatedParameters = FGlobalShader::Serialize(Ar);
+	Ar << PostprocessParameter;
+	return bShaderHasOutdatedParameters;
+}
 
-	// FShader interface.
-	virtual bool Serialize(FArchive& Ar) override
-	{
-		bool bShaderHasOutdatedParameters = FGlobalShader::Serialize(Ar);
-		Ar << PostprocessParameter;
-		return bShaderHasOutdatedParameters;
-	}
+void FPostProcessPassThroughPS::SetParameters(const FRenderingCompositePassContext& Context)
+{
+	const FPixelShaderRHIParamRef ShaderRHI = GetPixelShader();
 
-	void SetParameters(const FRenderingCompositePassContext& Context)
-	{
-		const FPixelShaderRHIParamRef ShaderRHI = GetPixelShader();
-		
-		FGlobalShader::SetParameters(Context.RHICmdList, ShaderRHI, Context.View);
+	FGlobalShader::SetParameters(Context.RHICmdList, ShaderRHI, Context.View);
 
-		PostprocessParameter.SetPS(ShaderRHI, Context, TStaticSamplerState<SF_Point,AM_Clamp,AM_Clamp,AM_Clamp>::GetRHI());
-	}
-};
+	PostprocessParameter.SetPS(ShaderRHI, Context, TStaticSamplerState<SF_Point, AM_Clamp, AM_Clamp, AM_Clamp>::GetRHI());
+}
 
 IMPLEMENT_SHADER_TYPE(,FPostProcessPassThroughPS,TEXT("PostProcessPassThrough"),TEXT("MainPS"),SF_Pixel);
 
@@ -140,6 +121,9 @@ void FRCPassPostProcessPassThrough::Process(FRenderingCompositePassContext& Cont
 		Context.HasHmdMesh(),
 		EDRF_UseTriangleOptimization);
 
+	// Draw custom data (like legends) for derived types
+	DrawCustom(Context);
+
 	Context.RHICmdList.CopyToResolveTarget(DestRenderTarget.TargetableTexture, DestRenderTarget.ShaderResourceTexture, false, FResolveParams());
 }
 
diff --git a/Engine/Source/Runtime/Renderer/Private/CompositionLighting/PostProcessPassThrough.h b/Engine/Source/Runtime/Renderer/Private/CompositionLighting/PostProcessPassThrough.h
--- a/Engine/Source/Runtime/Renderer/Private/CompositionLighting/PostProcessPassThrough.h
+++ b/Engine/Source/Runtime/Renderer/Private/CompositionLighting/PostProcessPassThrough.h
@@ -8,6 +8,31 @@
 
 #include "RenderingCompositionGraph.h"
 
+/** Encapsulates a simple copy pixel shader. */
+class FPostProcessPassThroughPS : public FGlobalShader
+{
+	DECLARE_SHADER_TYPE(FPostProcessPassThroughPS, Global);
+
+	static bool ShouldCache(EShaderPlatform Platform)
+	{
+		return IsFeatureLevelSupported(Platform, ERHIFeatureLevel::SM4);
+	}
+
+	/** Default constructor. */
+	FPostProcessPassThroughPS() {}
+
+public:
+	FPostProcessPassParameters PostprocessParameter;
+
+	/** Initialization constructor. */
+	FPostProcessPassThroughPS(const ShaderMetaType::CompiledShaderInitializerType& Initializer);
+
+	// FShader interface.
+	virtual bool Serialize(FArchive& Ar) override;
+
+	void SetParameters(const FRenderingCompositePassContext& Context);
+};
+
 // ePId_Input0: Input image
 // derives from TRenderingCompositePassBase<InputCount, OutputCount> 
 class FRCPassPostProcessPassThrough : public TRenderingCompositePassBase<1, 1>
@@ -25,6 +50,11 @@ public:
 	virtual void Release() override { delete this; }
 	virtual FPooledRenderTargetDesc ComputeOutputDesc(EPassOutputId InPassOutputId) const override;
 
+protected:
+
+	// Override this function in derived classes to draw custom UI like legends. This is called after the fullscreen copy.
+	virtual void DrawCustom(FRenderingCompositePassContext& Context) {}
+
 private:
 	// 0 if a new intermediate should be created
 	IPooledRenderTarget* Dest;
diff --git a/Engine/Source/Runtime/Renderer/Private/DebugViewModeRendering.h b/Engine/Source/Runtime/Renderer/Private/DebugViewModeRendering.h
--- a/Engine/Source/Runtime/Renderer/Private/DebugViewModeRendering.h
+++ b/Engine/Source/Runtime/Renderer/Private/DebugViewModeRendering.h
@@ -10,6 +10,7 @@ DebugViewModeRendering.h: Contains definitions for rendering debug viewmodes.
 #include "ShaderParameterUtils.h"
 
 static const int32 NumStreamingAccuracyColors = 5;
+static const float UndefinedStreamingAccuracyIntensity = .02f;
 
 /**
  * Vertex shader for quad overdraw. Required because overdraw shaders need to have SV_Position as first PS interpolant.
diff --git a/Engine/Source/Runtime/Renderer/Private/DeferredShadingRenderer.cpp b/Engine/Source/Runtime/Renderer/Private/DeferredShadingRenderer.cpp
--- a/Engine/Source/Runtime/Renderer/Private/DeferredShadingRenderer.cpp
+++ b/Engine/Source/Runtime/Renderer/Private/DeferredShadingRenderer.cpp
@@ -750,10 +750,13 @@ void FDeferredShadingSceneRenderer::Render(FRHICommandListImmediate& RHICmdList)
 	RHICmdList.SetCurrentStat(GET_STATID(STAT_CLM_AfterPrePass));
 	ServiceLocalQueue();
 
+	SceneContext.ResolveSceneDepthTexture(RHICmdList);
 
 	const bool bShouldRenderVelocities = ShouldRenderVelocities();
 	const bool bUseVelocityGBuffer = FVelocityRendering::OutputsToGBuffer();
 	const bool bUseSelectiveBasePassOutputs = UseSelectiveBasePassOutputs();
+
+	if (bUseGBuffer)
 	{
 		SCOPE_CYCLE_COUNTER(STAT_FDeferredShadingSceneRenderer_AllocGBufferTargets);
 		SceneContext.PreallocGBufferTargets(bUseVelocityGBuffer); // Even if !bShouldRenderVelocities, the velocity buffer must be bound because it's a compile time option for the shader.
@@ -807,7 +810,6 @@ void FDeferredShadingSceneRenderer::Render(FRHICommandListImmediate& RHICmdList)
 	if (bDBuffer)
 	{
 		SCOPE_CYCLE_COUNTER(STAT_FDeferredShadingSceneRenderer_DBuffer);
-		SceneContext.ResolveSceneDepthTexture(RHICmdList);
 		SceneContext.ResolveSceneDepthToAuxiliaryTexture(RHICmdList);
 
 		// e.g. DBuffer deferred decals
@@ -887,11 +889,8 @@ void FDeferredShadingSceneRenderer::Render(FRHICommandListImmediate& RHICmdList)
 		bRequiresFarZQuadClear = false;
 	}
 	
-	{
-		SCOPE_CYCLE_COUNTER(STAT_FDeferredShadingSceneRenderer_ResolveDepth_After_Basepass);
-		SceneContext.ResolveSceneDepthTexture(RHICmdList);
-		SceneContext.ResolveSceneDepthToAuxiliaryTexture(RHICmdList);
-	}
+	SceneContext.ResolveSceneDepthTexture(RHICmdList);
+	SceneContext.ResolveSceneDepthToAuxiliaryTexture(RHICmdList);
 
 	bool bOcclusionAfterBasePass = bIsOcclusionTesting && !bOcclusionBeforeBasePass;
 	bool bHZBAfterBasePass = !bHZBBeforeBasePass;
@@ -907,7 +906,11 @@ void FDeferredShadingSceneRenderer::Render(FRHICommandListImmediate& RHICmdList)
 	{
 		SCOPE_CYCLE_COUNTER(STAT_FDeferredShadingSceneRenderer_Resolve_After_Basepass);
 		SceneContext.ResolveSceneColor(RHICmdList, FResolveRect(0, 0, ViewFamily.FamilySizeX, ViewFamily.FamilySizeY));
-		SceneContext.FinishRenderingGBuffer(RHICmdList);
+
+		if (bUseGBuffer)
+		{
+			SceneContext.FinishRenderingGBuffer(RHICmdList);
+		}
 	}
 
 	if(GetCustomDepthPassLocation() == 1)
@@ -926,7 +929,7 @@ void FDeferredShadingSceneRenderer::Render(FRHICommandListImmediate& RHICmdList)
 			RHICmdList,
 			Views.GetData(),
 			SceneContext.GetSceneDepthTexture(),
-			SceneContext.GetGBufferATexture()
+			SceneContext.GBufferA ? SceneContext.GetGBufferATexture() : NULL
 			);
 		ServiceLocalQueue();
 	}
@@ -1047,7 +1050,7 @@ void FDeferredShadingSceneRenderer::Render(FRHICommandListImmediate& RHICmdList)
 		ServiceLocalQueue();
 
 		// SSS need the SceneColor finalized as an SRV.
-		SceneContext.FinishRenderingSceneColor(RHICmdList, true);
+		SceneContext.FinishRenderingSceneColor(RHICmdList);
 
 		// Render reflections that only operate on opaque pixels
 		RenderDeferredReflections(RHICmdList, DynamicBentNormalAO);
@@ -1127,7 +1130,7 @@ void FDeferredShadingSceneRenderer::Render(FRHICommandListImmediate& RHICmdList)
 			GetRendererModule().RenderPostOpaqueExtensions(View, RHICmdList, SceneContext);
 		}
 
-		SceneContext.FinishRenderingSceneColor(RHICmdList, true);
+		SceneContext.FinishRenderingSceneColor(RHICmdList);
 	}
 
 	// No longer needed, release
diff --git a/Engine/Source/Runtime/Renderer/Private/DistanceFieldObjectManagement.cpp b/Engine/Source/Runtime/Renderer/Private/DistanceFieldObjectManagement.cpp
--- a/Engine/Source/Runtime/Renderer/Private/DistanceFieldObjectManagement.cpp
+++ b/Engine/Source/Runtime/Renderer/Private/DistanceFieldObjectManagement.cpp
@@ -904,6 +904,8 @@ void ProcessPrimitiveUpdate(
 						const int32 InstanceIndex = PrimitiveSceneInfo->DistanceFieldInstanceIndices[TransformIndex];
 						if (InstanceIndex >= 0)
 						{
+							// For an update transform we have to dirty the previous bounds and the new bounds, in case of large movement (teleport)
+							DistanceFieldSceneData.PrimitiveModifiedBounds.Add(DistanceFieldSceneData.PrimitiveInstanceMapping[InstanceIndex].BoundingSphere);
 							DistanceFieldSceneData.PrimitiveInstanceMapping[InstanceIndex].BoundingSphere = ObjectBoundingSphere;
 						}
 					}
diff --git a/Engine/Source/Runtime/Renderer/Private/DistanceFieldShadowing.cpp b/Engine/Source/Runtime/Renderer/Private/DistanceFieldShadowing.cpp
--- a/Engine/Source/Runtime/Renderer/Private/DistanceFieldShadowing.cpp
+++ b/Engine/Source/Runtime/Renderer/Private/DistanceFieldShadowing.cpp
@@ -68,6 +68,63 @@ FIntPoint GetBufferSizeForDFShadows()
 
 TGlobalResource<FDistanceFieldObjectBufferResource> GShadowCulledObjectBuffers;
 
+/**  */
+class FClearObjectIndirectArguments : public FGlobalShader
+{
+	DECLARE_SHADER_TYPE(FClearObjectIndirectArguments,Global)
+public:
+
+	static bool ShouldCache(EShaderPlatform Platform)
+	{
+		return IsFeatureLevelSupported(Platform, ERHIFeatureLevel::SM5) && DoesPlatformSupportDistanceFieldShadowing(Platform);
+	}
+
+	FClearObjectIndirectArguments(const ShaderMetaType::CompiledShaderInitializerType& Initializer)
+		: FGlobalShader(Initializer)
+	{
+		ObjectIndirectArguments.Bind(Initializer.ParameterMap, TEXT("ObjectIndirectArguments"));
+	}
+
+	FClearObjectIndirectArguments()
+	{
+	}
+
+	void SetParameters(FRHICommandList& RHICmdList, const FSceneView& View)
+	{
+		FComputeShaderRHIParamRef ShaderRHI = GetComputeShader();
+
+		FGlobalShader::SetParameters(RHICmdList, ShaderRHI, View);
+
+		FUnorderedAccessViewRHIParamRef OutUAVs[1];
+		OutUAVs[0] = GShadowCulledObjectBuffers.Buffers.ObjectIndirectArguments.UAV;	
+		RHICmdList.TransitionResources(EResourceTransitionAccess::ERWBarrier, EResourceTransitionPipeline::EComputeToCompute, OutUAVs, ARRAY_COUNT(OutUAVs));
+
+		ObjectIndirectArguments.SetBuffer(RHICmdList, ShaderRHI, GShadowCulledObjectBuffers.Buffers.ObjectIndirectArguments);
+	}
+
+	void UnsetParameters(FRHICommandList& RHICmdList)
+	{
+		ObjectIndirectArguments.UnsetUAV(RHICmdList, GetComputeShader());
+
+		FUnorderedAccessViewRHIParamRef OutUAVs[1];
+		OutUAVs[0] = GShadowCulledObjectBuffers.Buffers.ObjectIndirectArguments.UAV;	
+		RHICmdList.TransitionResources(EResourceTransitionAccess::EReadable, EResourceTransitionPipeline::EComputeToCompute, OutUAVs, ARRAY_COUNT(OutUAVs));
+	}
+
+	virtual bool Serialize(FArchive& Ar) override
+	{		
+		bool bShaderHasOutdatedParameters = FGlobalShader::Serialize(Ar);
+		Ar << ObjectIndirectArguments;
+		return bShaderHasOutdatedParameters;
+	}
+
+private:
+
+	FRWShaderParameter ObjectIndirectArguments;
+};
+
+IMPLEMENT_SHADER_TYPE(,FClearObjectIndirectArguments,TEXT("DistanceFieldShadowing"),TEXT("ClearObjectIndirectArguments"),SF_Compute);
+
 class FCullObjectsForShadowCS : public FGlobalShader
 {
 	DECLARE_SHADER_TYPE(FCullObjectsForShadowCS,Global)
@@ -729,8 +786,23 @@ void CullDistanceFieldObjectsForLight(
 		}
 
 		{
-			uint32 ClearValues[4] = { 0 };
-			RHICmdList.ClearUAV(GShadowCulledObjectBuffers.Buffers.ObjectIndirectArguments.UAV, ClearValues);
+			// Manual compute shader clear to work around a bug on PS4 where a ClearUAV does not complete before the DrawIndexedPrimitiveIndirect read in the graphics pipe
+			bool bUseComputeShaderClear = true;
+
+			if (bUseComputeShaderClear)
+			{
+				TShaderMapRef<FClearObjectIndirectArguments> ComputeShader(View.ShaderMap);
+
+				RHICmdList.SetComputeShader(ComputeShader->GetComputeShader());
+				ComputeShader->SetParameters(RHICmdList, View);
+				DispatchComputeShader(RHICmdList, *ComputeShader, 1, 1, 1);
+				ComputeShader->UnsetParameters(RHICmdList);
+			}
+			else
+			{
+				uint32 ClearValues[4] = { 0 };
+				RHICmdList.ClearUAV(GShadowCulledObjectBuffers.Buffers.ObjectIndirectArguments.UAV, ClearValues);
+			}
 
 			TShaderMapRef<FCullObjectsForShadowCS> ComputeShader(GetGlobalShaderMap(Scene->GetFeatureLevel()));
 			RHICmdList.SetComputeShader(ComputeShader->GetComputeShader());
@@ -796,7 +868,7 @@ void CullDistanceFieldObjectsForLight(
 
 			TArray<FUnorderedAccessViewRHIParamRef> UAVs;
 			PixelShader->GetUAVs(View, TileIntersectionResources, UAVs);
-			RHICmdList.TransitionResources(EResourceTransitionAccess::ERWBarrier, EResourceTransitionPipeline::EComputeToCompute, UAVs.GetData(), UAVs.Num());
+			RHICmdList.TransitionResources(EResourceTransitionAccess::ERWBarrier, EResourceTransitionPipeline::EComputeToGfx, UAVs.GetData(), UAVs.Num());
 			RHICmdList.SetRenderTargets(0, (const FRHIRenderTargetView *)NULL, NULL, UAVs.Num(), UAVs.GetData());
 
 			RHICmdList.SetViewport(0, 0, 0.0f, LightTileDimensions.X, LightTileDimensions.Y, 1.0f);
@@ -822,7 +894,7 @@ void CullDistanceFieldObjectsForLight(
 				0);
 
 			SetRenderTarget(RHICmdList, NULL, NULL);
-			RHICmdList.TransitionResources(EResourceTransitionAccess::ERWBarrier, EResourceTransitionPipeline::EComputeToCompute, UAVs.GetData(), UAVs.Num());
+			RHICmdList.TransitionResources(EResourceTransitionAccess::EReadable, EResourceTransitionPipeline::EGfxToCompute, UAVs.GetData(), UAVs.Num());
 		}
 	}
 }
diff --git a/Engine/Source/Runtime/Renderer/Private/DistortionRendering.cpp b/Engine/Source/Runtime/Renderer/Private/DistortionRendering.cpp
--- a/Engine/Source/Runtime/Renderer/Private/DistortionRendering.cpp
+++ b/Engine/Source/Runtime/Renderer/Private/DistortionRendering.cpp
@@ -868,6 +868,7 @@ void FSceneRenderer::RenderDistortion(FRHICommandListImmediate& RHICmdList)
 		{
 			FPooledRenderTargetDesc Desc(FPooledRenderTargetDesc::Create2DDesc(SceneContext.GetBufferSizeXY(), PF_B8G8R8A8, FClearValueBinding::Transparent, TexCreate_None, TexCreate_RenderTargetable, false));
 			Desc.Flags |= TexCreate_FastVRAM;
+			Desc.NumSamples = SceneContext.SceneDepthZ->GetDesc().NumSamples;
 			GRenderTargetPool.FindFreeElement(RHICmdList, Desc, DistortionRT, TEXT("Distortion"));
 
 			// use RGBA8 light target for accumulating distortion offsets	
@@ -1017,7 +1018,7 @@ void FSceneRenderer::RenderDistortion(FRHICommandListImmediate& RHICmdList)
 		RHICmdList.SetDepthStencilState(TStaticDepthStencilState<false, CF_DepthNearOrEqual>::GetRHI());		
 			
 		// Distortions RT is no longer needed, buffer can be reused by the pool, see BeginRenderingDistortionAccumulation() call above
-		SceneContext.FinishRenderingSceneColor(RHICmdList, true);
+		SceneContext.FinishRenderingSceneColor(RHICmdList);
 	}
 }
 
diff --git a/Engine/Source/Runtime/Renderer/Private/FogRendering.cpp b/Engine/Source/Runtime/Renderer/Private/FogRendering.cpp
--- a/Engine/Source/Runtime/Renderer/Private/FogRendering.cpp
+++ b/Engine/Source/Runtime/Renderer/Private/FogRendering.cpp
@@ -300,7 +300,7 @@ void SetFogShaders(FRHICommandList& RHICmdList, FScene* Scene, const FViewInfo&
 
 bool FDeferredShadingSceneRenderer::RenderFog(FRHICommandListImmediate& RHICmdList, const FLightShaftsOutput& LightShaftsOutput)
 {
-	if (Scene->ExponentialFogs.Num() > 0)
+	if (Scene->ExponentialFogs.Num() > 0 && !Scene->ReadOnlyCVARCache.bEnableVertexFoggingForOpaque)
 	{
 		static const FVector2D Vertices[4] =
 		{
@@ -355,8 +355,6 @@ bool FDeferredShadingSceneRenderer::RenderFog(FRHICommandListImmediate& RHICmdLi
 				);
 		}
 
-		//no need to resolve since we used alpha blending
-		SceneContext.FinishRenderingSceneColor(RHICmdList, false);
 		return true;
 	}
 
diff --git a/Engine/Source/Runtime/Renderer/Private/GlobalDistanceField.cpp b/Engine/Source/Runtime/Renderer/Private/GlobalDistanceField.cpp
--- a/Engine/Source/Runtime/Renderer/Private/GlobalDistanceField.cpp
+++ b/Engine/Source/Runtime/Renderer/Private/GlobalDistanceField.cpp
@@ -887,6 +887,7 @@ void ComputeUpdateRegionsAndUpdateViewState(
 					int32 TotalTexelsBeingUpdated = 0;
 
 					// Trim fully contained update regions
+					//@todo - trim overlapping regions, especially for update transforms which create update regions for previous and current bounds which are typically overlapping
 					for (int32 UpdateRegionIndex = 0; UpdateRegionIndex < Clipmap.UpdateRegions.Num(); UpdateRegionIndex++)
 					{
 						const FVolumeUpdateRegion& UpdateRegion = Clipmap.UpdateRegions[UpdateRegionIndex];
diff --git a/Engine/Source/Runtime/Renderer/Private/LightRendering.cpp b/Engine/Source/Runtime/Renderer/Private/LightRendering.cpp
--- a/Engine/Source/Runtime/Renderer/Private/LightRendering.cpp
+++ b/Engine/Source/Runtime/Renderer/Private/LightRendering.cpp
@@ -590,9 +590,6 @@ void FDeferredShadingSceneRenderer::RenderLights(FRHICommandListImmediate& RHICm
 				GetLightNameForDrawEvent(LightSceneInfo.Proxy, LightNameWithLevel);
 				SCOPED_DRAW_EVENTF(RHICmdList, EventLightPass, *LightNameWithLevel);
 
-				// Do not resolve to scene color texture, this is done lazily
-				SceneContext.FinishRenderingSceneColor(RHICmdList, false);
-
 				if (bDrawShadows)
 				{
 					INC_DWORD_STAT(STAT_NumShadowedLights);
@@ -660,9 +657,6 @@ void FDeferredShadingSceneRenderer::RenderLights(FRHICommandListImmediate& RHICm
 				}
 			}
 
-			// Do not resolve to scene color texture, this is done lazily
-			SceneContext.FinishRenderingSceneColor(RHICmdList, false);
-
 			// Restore the default mode
 			SceneContext.SetLightAttenuationMode(true);
 
diff --git a/Engine/Source/Runtime/Renderer/Private/LightShaftRendering.cpp b/Engine/Source/Runtime/Renderer/Private/LightShaftRendering.cpp
--- a/Engine/Source/Runtime/Renderer/Private/LightShaftRendering.cpp
+++ b/Engine/Source/Runtime/Renderer/Private/LightShaftRendering.cpp
@@ -538,7 +538,7 @@ void ApplyTemporalAA(
 	/** Output of Temporal AA for the next step in the pipeline. */
 	TRefCountPtr<IPooledRenderTarget>& HistoryOutput)
 {
-	if (View.FinalPostProcessSettings.AntiAliasingMethod == AAM_TemporalAA
+	if (View.AntiAliasingMethod == AAM_TemporalAA
 		&& HistoryState)
 	{
 		if (*HistoryState && !View.bCameraCut)
@@ -901,7 +901,7 @@ void ApplyLightShaftBloom(FRHICommandListImmediate& RHICmdList, const FViewInfo&
 		*ScreenVertexShader,
 		EDRF_UseTriangleOptimization);
 
-	SceneContext.FinishRenderingSceneColor(RHICmdList, true);
+	SceneContext.FinishRenderingSceneColor(RHICmdList);
 }
 
 void FSceneViewState::TrimHistoryRenderTargets(const FScene* Scene)
diff --git a/Engine/Source/Runtime/Renderer/Private/MaterialTexCoordScalesRendering.h b/Engine/Source/Runtime/Renderer/Private/MaterialTexCoordScalesRendering.h
--- a/Engine/Source/Runtime/Renderer/Private/MaterialTexCoordScalesRendering.h
+++ b/Engine/Source/Runtime/Renderer/Private/MaterialTexCoordScalesRendering.h
@@ -51,6 +51,7 @@ public:
 
 	static void ModifyCompilationEnvironment(EShaderPlatform Platform, const FMaterial* Material, FShaderCompilerEnvironment& OutEnvironment)
 	{
+		OutEnvironment.SetDefine(TEXT("UNDEFINED_ACCURACY"), UndefinedStreamingAccuracyIntensity);
 		OutEnvironment.SetDefine(TEXT("MAX_NUM_TEX_COORD"), (uint32)FMaterialTexCoordBuildInfo::MAX_NUM_TEX_COORD);
 		OutEnvironment.SetDefine(TEXT("INITIAL_GPU_SCALE"), (uint32)FMaterialTexCoordBuildInfo::INITIAL_GPU_SCALE);
 		OutEnvironment.SetDefine(TEXT("TILE_RESOLUTION"), (uint32)FMaterialTexCoordBuildInfo::TILE_RESOLUTION);
diff --git a/Engine/Source/Runtime/Renderer/Private/MeshTexCoordSizeAccuracyRendering.h b/Engine/Source/Runtime/Renderer/Private/MeshTexCoordSizeAccuracyRendering.h
--- a/Engine/Source/Runtime/Renderer/Private/MeshTexCoordSizeAccuracyRendering.h
+++ b/Engine/Source/Runtime/Renderer/Private/MeshTexCoordSizeAccuracyRendering.h
@@ -41,6 +41,7 @@ public:
 
 	static void ModifyCompilationEnvironment(EShaderPlatform Platform, FShaderCompilerEnvironment& OutEnvironment)
 	{
+		OutEnvironment.SetDefine(TEXT("UNDEFINED_ACCURACY"), UndefinedStreamingAccuracyIntensity);
 	}
 
 	virtual void SetParameters(
diff --git a/Engine/Source/Runtime/Renderer/Private/MobileBasePassRendering.h b/Engine/Source/Runtime/Renderer/Private/MobileBasePassRendering.h
--- a/Engine/Source/Runtime/Renderer/Private/MobileBasePassRendering.h
+++ b/Engine/Source/Runtime/Renderer/Private/MobileBasePassRendering.h
@@ -203,6 +203,7 @@ public:
 		PixelParametersType::Bind(Initializer.ParameterMap);
 		ReflectionCubemap.Bind(Initializer.ParameterMap, TEXT("ReflectionCubemap"));
 		ReflectionSampler.Bind(Initializer.ParameterMap, TEXT("ReflectionCubemapSampler"));
+		InvReflectionCubemapAverageBrightness.Bind(Initializer.ParameterMap, TEXT("InvReflectionCubemapAverageBrightness"));
 		EditorCompositeParams.Bind(Initializer.ParameterMap);
 		LightPositionAndInvRadiusParameter.Bind(Initializer.ParameterMap, TEXT("LightPositionAndInvRadius"));
 		LightColorAndFalloffExponentParameter.Bind(Initializer.ParameterMap, TEXT("LightColorAndFalloffExponent"));
@@ -242,6 +243,7 @@ public:
 			const FShaderResourceParameter* ReflectionSamplerParameters[MaxNumReflections] = { &ReflectionSampler, &ReflectionSampler1, &ReflectionSampler2 };
 			FTexture* ReflectionCubemapTextures[MaxNumReflections] = { GBlackTextureCube, GBlackTextureCube, GBlackTextureCube };
 			FVector4 CapturePositions[MaxNumReflections] = { FVector4(0, 0, 0, 0), FVector4(0, 0, 0, 0), FVector4(0, 0, 0, 0) };
+			FVector AverageBrightness(1, 1, 1);
 
 			if (PrimitiveSceneInfo)
 			{
@@ -256,6 +258,7 @@ public:
 						{
 							ReflectionCubemapTextures[i] = PrimitiveSceneInfo->CachedReflectionCaptureProxies[i]->EncodedHDRCubemap;
 						}
+						AverageBrightness[i] = ReflectionProxy->AverageBrightness;
 					}
 				}
 			}
@@ -272,21 +275,26 @@ public:
 			{
 				SetShaderValueArray(RHICmdList, PixelShader, ReflectionPositionsAndRadii, CapturePositions, MaxNumReflections);
 			}
+
+			SetShaderValue(RHICmdList, PixelShader, InvReflectionCubemapAverageBrightness, FVector(1.0f / AverageBrightness.X, 1.0f / AverageBrightness.Y, 1.0f / AverageBrightness.Z));
 		}
 		else if (ReflectionCubemap.IsBound())
 		{
 			FTexture* ReflectionTexture = GBlackTextureCube;
+			float AverageBrightness = 1.0f;
 
 			if (PrimitiveSceneInfo 
 				&& PrimitiveSceneInfo->CachedReflectionCaptureProxy
 				&& PrimitiveSceneInfo->CachedReflectionCaptureProxy->EncodedHDRCubemap
 				&& PrimitiveSceneInfo->CachedReflectionCaptureProxy->EncodedHDRCubemap->IsInitialized())
 			{
+				AverageBrightness = PrimitiveSceneInfo->CachedReflectionCaptureProxy->AverageBrightness;
 				ReflectionTexture = PrimitiveSceneInfo->CachedReflectionCaptureProxy->EncodedHDRCubemap;
 			}
 
 			// Set the reflection cubemap
 			SetTextureParameter(RHICmdList, PixelShader, ReflectionCubemap, ReflectionSampler, ReflectionTexture);
+			SetShaderValue(RHICmdList, PixelShader, InvReflectionCubemapAverageBrightness, FVector(1.0f / AverageBrightness, 0, 0));
 		}
 
 		if (NumDynamicPointLights > 0)
@@ -315,6 +323,7 @@ public:
 		PixelParametersType::Serialize(Ar);
 		Ar << ReflectionCubemap;
 		Ar << ReflectionSampler;
+		Ar << InvReflectionCubemapAverageBrightness;
 		Ar << EditorCompositeParams;
 		Ar << LightPositionAndInvRadiusParameter;
 		Ar << LightColorAndFalloffExponentParameter;
@@ -339,6 +348,7 @@ private:
 
 	FShaderResourceParameter ReflectionCubemap;
 	FShaderResourceParameter ReflectionSampler;
+	FShaderParameter InvReflectionCubemapAverageBrightness;
 	FEditorCompositingParameters EditorCompositeParams;
 	FShaderParameter LightPositionAndInvRadiusParameter;
 	FShaderParameter LightColorAndFalloffExponentParameter;
diff --git a/Engine/Source/Runtime/Renderer/Private/PlanarReflectionRendering.cpp b/Engine/Source/Runtime/Renderer/Private/PlanarReflectionRendering.cpp
--- a/Engine/Source/Runtime/Renderer/Private/PlanarReflectionRendering.cpp
+++ b/Engine/Source/Runtime/Renderer/Private/PlanarReflectionRendering.cpp
@@ -292,16 +292,6 @@ void FScene::UpdatePlanarReflectionContents(UPlanarReflectionComponent* CaptureC
 			});
 		}
 
-		FPostProcessSettings PostProcessSettings;
-		PostProcessSettings.bOverride_AntiAliasingMethod = true;
-
-		static IConsoleVariable* CVarDefaultAntiAliasing = IConsoleManager::Get().FindConsoleVariable(TEXT("r.DefaultFeature.AntiAliasing"));
-		int32 Value = CVarDefaultAntiAliasing->GetInt();
-		if (Value >= 0 && Value < AAM_MAX)
-		{
-			PostProcessSettings.AntiAliasingMethod = (EAntiAliasingMethod)Value;
-		}
-
 		FMatrix ComponentTransform = CaptureComponent->ComponentToWorld.ToMatrixWithScale();
 		FPlane MirrorPlane = FPlane(ComponentTransform.TransformPosition(FVector::ZeroVector), ComponentTransform.TransformVector(FVector(0, 0, 1)));
 
@@ -314,6 +304,7 @@ void FScene::UpdatePlanarReflectionContents(UPlanarReflectionComponent* CaptureC
 
 		FMatrix ProjectionMatrix;
 		BuildProjectionMatrix(DesiredPlanarReflectionTextureSize, ECameraProjectionMode::Perspective, FOV + CaptureComponent->ExtraFOV * (float)PI / 180.0f, 1.0f, ProjectionMatrix);
+		FPostProcessSettings PostProcessSettings;
 
 		FSceneRenderer* SceneRenderer = CreateSceneRendererForSceneCapture(this, CaptureComponent, CaptureComponent->RenderTarget, DesiredPlanarReflectionTextureSize, ViewRotationMatrix, ViewLocation, ProjectionMatrix, CaptureComponent->MaxViewDistanceOverride, true, true, &PostProcessSettings, 1.0f);
 
diff --git a/Engine/Source/Runtime/Renderer/Private/PostProcess/PostProcessBokehDOF.cpp b/Engine/Source/Runtime/Renderer/Private/PostProcess/PostProcessBokehDOF.cpp
--- a/Engine/Source/Runtime/Renderer/Private/PostProcess/PostProcessBokehDOF.cpp
+++ b/Engine/Source/Runtime/Renderer/Private/PostProcess/PostProcessBokehDOF.cpp
@@ -224,27 +224,7 @@ void FRCPassPostProcessVisualizeDOF::Process(FRenderingCompositePassContext& Con
 
 
 	{
-		// this is a helper class for FCanvas to be able to get screen size
-		class FRenderTargetTemp : public FRenderTarget
-		{
-		public:
-			const FSceneView& View;
-			const FTexture2DRHIRef Texture;
-
-			FRenderTargetTemp(const FSceneView& InView, const FTexture2DRHIRef InTexture)
-				: View(InView), Texture(InTexture)
-			{
-			}
-			virtual FIntPoint GetSizeXY() const
-			{
-				return View.ViewRect.Size();
-			};
-			virtual const FTexture2DRHIRef& GetRenderTargetTexture() const
-			{
-				return Texture;
-			}
-		} TempRenderTarget(View, (const FTexture2DRHIRef&)DestRenderTarget.TargetableTexture);
-
+		FRenderTargetTemp TempRenderTarget(View, (const FTexture2DRHIRef&)DestRenderTarget.TargetableTexture);
 		FCanvas Canvas(&TempRenderTarget, NULL, ViewFamily.CurrentRealTime, ViewFamily.CurrentWorldTime, ViewFamily.DeltaWorldTime, Context.GetFeatureLevel());
 
 		float X = 30;
diff --git a/Engine/Source/Runtime/Renderer/Private/PostProcess/PostProcessBufferInspector.cpp b/Engine/Source/Runtime/Renderer/Private/PostProcess/PostProcessBufferInspector.cpp
--- a/Engine/Source/Runtime/Renderer/Private/PostProcess/PostProcessBufferInspector.cpp
+++ b/Engine/Source/Runtime/Renderer/Private/PostProcess/PostProcessBufferInspector.cpp
@@ -8,52 +8,10 @@
 #include "ScenePrivate.h"
 #include "SceneFilterRendering.h"
 #include "PostProcessBufferInspector.h"
+#include "PostProcessPassThrough.h"
 #include "PostProcessing.h"
 #include "SceneUtils.h"
 
-/** Encapsulates a simple copy pixel shader. */
-class FPostProcessBufferInspectorPS : public FGlobalShader
-{
-	DECLARE_SHADER_TYPE(FPostProcessBufferInspectorPS, Global);
-
-	static bool ShouldCache(EShaderPlatform Platform)
-	{
-		return IsFeatureLevelSupported(Platform, ERHIFeatureLevel::SM4);
-	}
-
-	/** Default constructor. */
-	FPostProcessBufferInspectorPS() {}
-
-public:
-	FPostProcessPassParameters PostprocessParameter;
-
-	/** Initialization constructor. */
-	FPostProcessBufferInspectorPS(const ShaderMetaType::CompiledShaderInitializerType& Initializer)
-		: FGlobalShader(Initializer)
-	{
-		PostprocessParameter.Bind(Initializer.ParameterMap);
-	}
-
-	// FShader interface.
-	virtual bool Serialize(FArchive& Ar) override
-	{
-		bool bShaderHasOutdatedParameters = FGlobalShader::Serialize(Ar);
-		Ar << PostprocessParameter;
-		return bShaderHasOutdatedParameters;
-	}
-
-	void SetParameters(const FRenderingCompositePassContext& Context)
-	{
-		const FPixelShaderRHIParamRef ShaderRHI = GetPixelShader();
-
-		FGlobalShader::SetParameters(Context.RHICmdList, ShaderRHI, Context.View);
-
-		PostprocessParameter.SetPS(ShaderRHI, Context, TStaticSamplerState<SF_Point, AM_Clamp, AM_Clamp, AM_Clamp>::GetRHI());
-	}
-};
-
-IMPLEMENT_SHADER_TYPE(, FPostProcessBufferInspectorPS, TEXT("PostProcessPassThrough"), TEXT("MainPS"), SF_Pixel);
-
 FRCPassPostProcessBufferInspector::FRCPassPostProcessBufferInspector(FRHICommandList& RHICmdList)
 {
 	// AdjustGBufferRefCount(-1) call is done when the pass gets executed
@@ -63,7 +21,7 @@ FRCPassPostProcessBufferInspector::FRCPassPostProcessBufferInspector(FRHICommand
 FShader* FRCPassPostProcessBufferInspector::SetShaderTempl(const FRenderingCompositePassContext& Context)
 {
 	TShaderMapRef<FPostProcessVS> VertexShader(Context.GetShaderMap());
-	TShaderMapRef<FPostProcessBufferInspectorPS> PixelShader(Context.GetShaderMap());
+	TShaderMapRef<FPostProcessPassThroughPS> PixelShader(Context.GetShaderMap());
 
 	static FGlobalBoundShaderState BoundShaderState;
 	SetGlobalBoundShaderState(Context.RHICmdList, Context.GetFeatureLevel(), BoundShaderState, GFilterVertexDeclaration.VertexDeclarationRHI, *VertexShader, *PixelShader);
@@ -296,27 +254,7 @@ void FRCPassPostProcessBufferInspector::Process(FRenderingCompositePassContext&
 			EDRF_UseTriangleOptimization);
 	}
 
-	// this is a helper class for FCanvas to be able to get screen size
-	class FRenderTargetTemp : public FRenderTarget
-	{
-	public:
-		const FSceneView& View;
-		const FTexture2DRHIRef Texture;
-
-		FRenderTargetTemp(const FSceneView& InView, const FTexture2DRHIRef InTexture)
-			: View(InView), Texture(InTexture)
-		{
-		}
-		virtual FIntPoint GetSizeXY() const
-		{
-			return View.ViewRect.Size();
-		};
-		virtual const FTexture2DRHIRef& GetRenderTargetTexture() const
-		{
-			return Texture;
-		}
-	} TempRenderTarget(View, (const FTexture2DRHIRef&)DestRenderTarget.TargetableTexture);
-
+	FRenderTargetTemp TempRenderTarget(View, (const FTexture2DRHIRef&)DestRenderTarget.TargetableTexture);
 	FCanvas Canvas(&TempRenderTarget, NULL, ViewFamily.CurrentRealTime, ViewFamily.CurrentWorldTime, ViewFamily.DeltaWorldTime, Context.GetFeatureLevel());
 	FLinearColor LabelColor(1, 1, 1);
 	Canvas.DrawShadowedString(100, 50, TEXT("Pixel Inspector On"), GetStatsFont(), LabelColor);
diff --git a/Engine/Source/Runtime/Renderer/Private/PostProcess/PostProcessMobile.cpp b/Engine/Source/Runtime/Renderer/Private/PostProcess/PostProcessMobile.cpp
--- a/Engine/Source/Runtime/Renderer/Private/PostProcess/PostProcessMobile.cpp
+++ b/Engine/Source/Runtime/Renderer/Private/PostProcess/PostProcessMobile.cpp
@@ -1628,7 +1628,7 @@ void FRCPassPostProcessSunMergeES2::Process(FRenderingCompositePassContext& Cont
 	Context.RHICmdList.CopyToResolveTarget(DestRenderTarget.TargetableTexture, DestRenderTarget.ShaderResourceTexture, false, FResolveParams());
 
 	// Double buffer sun+bloom+vignette composite.
-	if(Context.View.FinalPostProcessSettings.AntiAliasingMethod == AAM_TemporalAA)
+	if(Context.View.AntiAliasingMethod == AAM_TemporalAA)
 	{
 		FSceneViewState* ViewState = (FSceneViewState*)Context.View.State;
 		if(ViewState) 
@@ -1822,7 +1822,7 @@ void FRCPassPostProcessSunMergeSmallES2::Process(FRenderingCompositePassContext&
 
 	// Double buffer sun+bloom+vignette composite.
 
-	if (Context.View.FinalPostProcessSettings.AntiAliasingMethod == AAM_TemporalAA)
+	if (Context.View.AntiAliasingMethod == AAM_TemporalAA)
 	{
 		FSceneViewState* ViewState = (FSceneViewState*)Context.View.State;
 		if(ViewState) 
diff --git a/Engine/Source/Runtime/Renderer/Private/PostProcess/PostProcessStreamingAccuracyLegend.cpp b/Engine/Source/Runtime/Renderer/Private/PostProcess/PostProcessStreamingAccuracyLegend.cpp
new file mode 100644
--- /dev/null
+++ b/Engine/Source/Runtime/Renderer/Private/PostProcess/PostProcessStreamingAccuracyLegend.cpp
@@ -0,0 +1,83 @@
+// Copyright 1998-2016 Epic Games, Inc. All Rights Reserved.
+
+/*=============================================================================
+PostProcessVisualizeComplexity.cpp: Contains definitions for complexity viewmode.
+=============================================================================*/
+
+#include "RendererPrivate.h"
+#include "ScenePrivate.h"
+#include "SceneFilterRendering.h"
+#include "PostProcessStreamingAccuracyLegend.h"
+#include "SceneUtils.h"
+#include "PostProcessing.h"
+#include "DebugViewModeRendering.h"
+
+void FRCPassPostProcessStreamingAccuracyLegend::DrawDesc(FCanvas& Canvas, float PosX, float PosY, const FText& Text)
+{
+	Canvas.DrawShadowedText(PosX + 18, PosY, Text, GetStatsFont(), FLinearColor(0.7f, 0.7f, 0.7f), FLinearColor::Black);
+}
+
+void FRCPassPostProcessStreamingAccuracyLegend::DrawBox(FCanvas& Canvas, float PosX, float PosY, const FLinearColor& Color, const FText& Text)
+{
+	Canvas.DrawTile(PosX, PosY, 16, 16, 0, 0, 1, 1, FLinearColor::Black);
+	Canvas.DrawTile(PosX + 1, PosY + 1, 14, 14, 0, 0, 1, 1, Color);
+	Canvas.DrawShadowedText(PosX + 18, PosY, Text, GetStatsFont(), FLinearColor(0.7f, 0.7f, 0.7f), FLinearColor::Black);
+}
+
+void FRCPassPostProcessStreamingAccuracyLegend::DrawCheckerBoard(FCanvas& Canvas, float PosX, float PosY, const FLinearColor& Color0, const FLinearColor& Color1, const FText& Text)
+{
+	Canvas.DrawTile(PosX, PosY, 16, 16, 0, 0, 1, 1, FLinearColor::Black);
+	Canvas.DrawTile(PosX + 1, PosY + 1, 14, 14, 0, 0, 1, 1, Color0);
+	Canvas.DrawTile(PosX + 1, PosY + 1, 7, 7, 0, 0, 1, 1, Color1);
+	Canvas.DrawTile(PosX + 8, PosY + 8, 7, 7, 0, 0, 1, 1, Color1);
+	Canvas.DrawShadowedText(PosX + 18, PosY, Text, GetStatsFont(), FLinearColor(0.7f, 0.7f, 0.7f), FLinearColor::Black);
+}
+
+#define LOCTEXT_NAMESPACE "TextureStreamingBuild"
+
+void FRCPassPostProcessStreamingAccuracyLegend::DrawCustom(FRenderingCompositePassContext& Context)
+{
+	if (Colors.Num() != NumStreamingAccuracyColors) return;
+
+	SCOPED_DRAW_EVENT(Context.RHICmdList, PostProcessStreamingAccuracyLegend);
+
+	const FSceneView& View = Context.View;
+	const FSceneViewFamily& ViewFamily = *(View.Family);
+	FIntRect DestRect = View.UnscaledViewRect;
+
+	FRenderTargetTemp TempRenderTarget(View, (const FTexture2DRHIRef&)PassOutputs[0].RequestSurface(Context).TargetableTexture);
+	FCanvas Canvas(&TempRenderTarget, NULL, ViewFamily.CurrentRealTime, ViewFamily.CurrentWorldTime, ViewFamily.DeltaWorldTime, Context.GetFeatureLevel());
+
+	if (ViewFamily.GetDebugViewShaderMode() == DVSM_MaterialTexCoordScalesAccuracy)
+	{
+		DrawDesc(Canvas, DestRect.Min.X + 115, DestRect.Max.Y - 75, LOCTEXT("DescMaterialTexCoordScalesAccuracy", "Shows under/over texture streaming caused by material texcoord scales."));
+	}
+	else if (ViewFamily.GetDebugViewShaderMode() == DVSM_MeshTexCoordSizeAccuracy)
+	{
+		DrawDesc(Canvas, DestRect.Min.X + 115, DestRect.Max.Y - 75, LOCTEXT("DescMaterialTexCoordScalesAccuracy", "Shows under/over texture streaming caused by mesh texcoord mappings in world units."));
+	}
+	else if (ViewFamily.GetDebugViewShaderMode() == DVSM_PrimitiveDistanceAccuracy)
+	{
+		DrawDesc(Canvas, DestRect.Min.X + 115, DestRect.Max.Y - 100, LOCTEXT("DescPrimitiveDistanceAccuracy", "Shows under/over texture streaming caused by the difference between the streamer calculated"));
+		DrawDesc(Canvas, DestRect.Min.X + 165, DestRect.Max.Y - 75, LOCTEXT("DescPrimitiveDistanceAccuracy2", "distance-to-mesh via bounding box versus the actual per-pixel depth value."));
+	}
+
+	DrawBox(Canvas, DestRect.Min.X + 115, DestRect.Max.Y - 50, Colors[0], LOCTEXT("2XUnder", "2X Under"));
+	DrawBox(Canvas, DestRect.Min.X + 215, DestRect.Max.Y - 50, Colors[1], LOCTEXT("1XUnder", "1X Under"));
+	DrawBox(Canvas, DestRect.Min.X + 315, DestRect.Max.Y - 50, Colors[2], LOCTEXT("Good", "Good"));
+	DrawBox(Canvas, DestRect.Min.X + 415, DestRect.Max.Y - 50, Colors[3], LOCTEXT("1xOver", "1X Over"));
+	DrawBox(Canvas, DestRect.Min.X + 515, DestRect.Max.Y - 50, Colors[4], LOCTEXT("2XOver", "2X Over"));
+
+	const FLinearColor UndefColor(UndefinedStreamingAccuracyIntensity, UndefinedStreamingAccuracyIntensity, UndefinedStreamingAccuracyIntensity, 1.f);
+	DrawBox(Canvas, DestRect.Min.X + 615, DestRect.Max.Y - 50, UndefColor, LOCTEXT("Undefined", "Undefined"));
+	
+	if (ViewFamily.GetDebugViewShaderMode() == DVSM_MaterialTexCoordScalesAccuracy)
+	{
+		DrawCheckerBoard(Canvas, DestRect.Min.X + 115, DestRect.Max.Y - 25, Colors[0], Colors[4], LOCTEXT("WorstUnderAndOver", "Worst Under / Worst Over"));
+		DrawCheckerBoard(Canvas, DestRect.Min.X + 415, DestRect.Max.Y - 25, FLinearColor::Black, FLinearColor::White, LOCTEXT("MissingShader", "Missing Shader"));
+	}
+
+	Canvas.Flush_RenderThread(Context.RHICmdList);
+}
+	
+#undef LOCTEXT_NAMESPACE
diff --git a/Engine/Source/Runtime/Renderer/Private/PostProcess/PostProcessStreamingAccuracyLegend.h b/Engine/Source/Runtime/Renderer/Private/PostProcess/PostProcessStreamingAccuracyLegend.h
new file mode 100644
--- /dev/null
+++ b/Engine/Source/Runtime/Renderer/Private/PostProcess/PostProcessStreamingAccuracyLegend.h
@@ -0,0 +1,33 @@
+// Copyright 1998-2016 Epic Games, Inc. All Rights Reserved.
+
+/*=============================================================================
+	PostProcessStreamingAccuracyLegend.h: PP to apply the streaming accuracy legend
+=============================================================================*/
+
+#pragma once
+
+#include "RenderingCompositionGraph.h"
+#include "PostProcessPassThrough.h"
+
+
+
+class FRCPassPostProcessStreamingAccuracyLegend : public FRCPassPostProcessPassThrough
+{
+public:
+
+	FRCPassPostProcessStreamingAccuracyLegend(const TArray<FLinearColor>& InColors)
+		: FRCPassPostProcessPassThrough(nullptr), Colors(InColors)
+	{}
+
+protected:
+
+	virtual void DrawCustom(FRenderingCompositePassContext& Context) override;
+
+	void DrawDesc(FCanvas& Canvas, float PosX, float PosY, const FText& Text);
+	void DrawBox(FCanvas& Canvas, float PosX, float PosY, const FLinearColor& Color, const FText& Text);
+	void DrawCheckerBoard(FCanvas& Canvas, float PosX, float PosY, const FLinearColor& Color0, const FLinearColor& Color1, const FText& Text);
+
+private: 
+
+	TArray<FLinearColor> Colors;
+};
\ No newline at end of file
diff --git a/Engine/Source/Runtime/Renderer/Private/PostProcess/PostProcessSubsurface.cpp b/Engine/Source/Runtime/Renderer/Private/PostProcess/PostProcessSubsurface.cpp
--- a/Engine/Source/Runtime/Renderer/Private/PostProcess/PostProcessSubsurface.cpp
+++ b/Engine/Source/Runtime/Renderer/Private/PostProcess/PostProcessSubsurface.cpp
@@ -259,27 +259,7 @@ void FRCPassPostProcessSubsurfaceVisualize::Process(FRenderingCompositePassConte
 		EDRF_UseTriangleOptimization);
 
 	{
-		// this is a helper class for FCanvas to be able to get screen size
-		class FRenderTargetTemp : public FRenderTarget
-		{
-		public:
-			const FSceneView& View;
-			const FTexture2DRHIRef Texture;
-
-			FRenderTargetTemp(const FSceneView& InView, const FTexture2DRHIRef InTexture)
-				: View(InView), Texture(InTexture)
-			{
-			}
-			virtual FIntPoint GetSizeXY() const
-			{
-				return View.ViewRect.Size();
-			};
-			virtual const FTexture2DRHIRef& GetRenderTargetTexture() const
-			{
-				return Texture;
-			}
-		} TempRenderTarget(View, (const FTexture2DRHIRef&)DestRenderTarget.TargetableTexture);
-
+		FRenderTargetTemp TempRenderTarget(View, (const FTexture2DRHIRef&)DestRenderTarget.TargetableTexture);
 		FCanvas Canvas(&TempRenderTarget, NULL, ViewFamily.CurrentRealTime, ViewFamily.CurrentWorldTime, ViewFamily.DeltaWorldTime, Context.GetFeatureLevel());
 
 		float X = 30;
@@ -868,7 +848,7 @@ void FRCPassPostProcessSubsurfaceRecombine::Process(FRenderingCompositePassConte
 	{
 		if(QualityCVar == -1)
 		{
-			RecombineQuality = (View.FinalPostProcessSettings.AntiAliasingMethod == AAM_TemporalAA) ? 0 : 1;
+			RecombineQuality = (View.AntiAliasingMethod == AAM_TemporalAA) ? 0 : 1;
 		}
 		else if(QualityCVar == 1)
 		{
diff --git a/Engine/Source/Runtime/Renderer/Private/PostProcess/PostProcessTonemap.cpp b/Engine/Source/Runtime/Renderer/Private/PostProcess/PostProcessTonemap.cpp
--- a/Engine/Source/Runtime/Renderer/Private/PostProcess/PostProcessTonemap.cpp
+++ b/Engine/Source/Runtime/Renderer/Private/PostProcess/PostProcessTonemap.cpp
@@ -1389,7 +1389,7 @@ void FRCPassPostProcessTonemap::Process(FRenderingCompositePassContext& Context)
 	if (ViewFamily.Scene && ViewFamily.Scene->GetShadingPath() == EShadingPath::Mobile)
 	{
 		// Double buffer tonemapper output for temporal AA.
-		if(View.FinalPostProcessSettings.AntiAliasingMethod == AAM_TemporalAA)
+		if(View.AntiAliasingMethod == AAM_TemporalAA)
 		{
 			FSceneViewState* ViewState = (FSceneViewState*)View.State;
 			if(ViewState) 
@@ -1830,7 +1830,7 @@ void FRCPassPostProcessTonemapES2::Process(FRenderingCompositePassContext& Conte
 	Context.RHICmdList.CopyToResolveTarget(DestRenderTarget.TargetableTexture, DestRenderTarget.ShaderResourceTexture, false, FResolveParams());
 
 	// Double buffer tonemapper output for temporal AA.
-	if(Context.View.FinalPostProcessSettings.AntiAliasingMethod == AAM_TemporalAA)
+	if(Context.View.AntiAliasingMethod == AAM_TemporalAA)
 	{
 		FSceneViewState* ViewState = (FSceneViewState*)Context.View.State;
 		if(ViewState) 
diff --git a/Engine/Source/Runtime/Renderer/Private/PostProcess/PostProcessVisualizeHDR.cpp b/Engine/Source/Runtime/Renderer/Private/PostProcess/PostProcessVisualizeHDR.cpp
--- a/Engine/Source/Runtime/Renderer/Private/PostProcess/PostProcessVisualizeHDR.cpp
+++ b/Engine/Source/Runtime/Renderer/Private/PostProcess/PostProcessVisualizeHDR.cpp
@@ -227,28 +227,7 @@ void FRCPassPostProcessVisualizeHDR::Process(FRenderingCompositePassContext& Con
 		*VertexShader,
 		EDRF_UseTriangleOptimization);
 
-
-	// this is a helper class for FCanvas to be able to get screen size
-	class FRenderTargetTemp : public FRenderTarget
-	{
-	public:
-		const FSceneView& View;
-		const FTexture2DRHIRef Texture;
-
-		FRenderTargetTemp(const FSceneView& InView, const FTexture2DRHIRef InTexture)
-			: View(InView), Texture(InTexture)
-		{
-		}
-		virtual FIntPoint GetSizeXY() const
-		{
-			return View.ViewRect.Size();
-		};
-		virtual const FTexture2DRHIRef& GetRenderTargetTexture() const
-		{
-			return Texture;
-		}
-	} TempRenderTarget(View, (const FTexture2DRHIRef&)DestRenderTarget.TargetableTexture);
-
+	FRenderTargetTemp TempRenderTarget(View, (const FTexture2DRHIRef&)DestRenderTarget.TargetableTexture);
 	FCanvas Canvas(&TempRenderTarget, NULL, ViewFamily.CurrentRealTime, ViewFamily.CurrentWorldTime, ViewFamily.DeltaWorldTime, Context.GetFeatureLevel());
 
 	float X = 30;
diff --git a/Engine/Source/Runtime/Renderer/Private/PostProcess/PostProcessing.cpp b/Engine/Source/Runtime/Renderer/Private/PostProcess/PostProcessing.cpp
--- a/Engine/Source/Runtime/Renderer/Private/PostProcess/PostProcessing.cpp
+++ b/Engine/Source/Runtime/Renderer/Private/PostProcess/PostProcessing.cpp
@@ -51,6 +51,7 @@
 #include "IHeadMountedDisplay.h"
 #include "BufferVisualizationData.h"
 #include "PostProcessLpvIndirect.h"
+#include "PostProcessStreamingAccuracyLegend.h"
 
 
 /** The global center for all post processing activities. */
@@ -429,7 +430,7 @@ static void AddPostProcessDepthOfFieldBokeh(FPostprocessContext& Context, FRende
 	FSceneViewState* ViewState = (FSceneViewState*)Context.View.State;
 	
 	FRenderingCompositePass* DOFInputPass = DOFSetup;
-	if( Context.View.FinalPostProcessSettings.AntiAliasingMethod == AAM_TemporalAA && ViewState )
+	if( Context.View.AntiAliasingMethod == AAM_TemporalAA && ViewState )
 	{
 		FRenderingCompositePass* HistoryInput;
 		if( ViewState->DOFHistoryRT && ViewState->bDOFHistory && !Context.View.bCameraCut )
@@ -477,7 +478,7 @@ static bool AddPostProcessDepthOfFieldGaussian(FPostprocessContext& Context, FDe
 
 			FRenderingCompositeOutputRef DOFInputPass = DOFSetup;
 			const bool bMobileQuality = (Context.View.GetFeatureLevel() <= ERHIFeatureLevel::ES3_1);
-			if (Context.View.FinalPostProcessSettings.AntiAliasingMethod == AAM_TemporalAA && ViewState && !bMobileQuality)
+			if (Context.View.AntiAliasingMethod == AAM_TemporalAA && ViewState && !bMobileQuality)
 			{
 				// If no history use current as history
 				FRenderingCompositeOutputRef HistoryInput = DOFSetup;
@@ -607,7 +608,7 @@ static void AddPostProcessDepthOfFieldCircle(FPostprocessContext& Context, FDept
 	FSceneViewState* ViewState = (FSceneViewState*)Context.View.State;
 
 	FRenderingCompositePass* DOFInputPass = DOFSetup;
-	if( Context.View.FinalPostProcessSettings.AntiAliasingMethod == AAM_TemporalAA && ViewState )
+	if( Context.View.AntiAliasingMethod == AAM_TemporalAA && ViewState )
 	{
 		FRenderingCompositePass* HistoryInput;
 		if( ViewState->DOFHistoryRT && !ViewState->bDOFHistory && !Context.View.bCameraCut )
@@ -1337,7 +1338,7 @@ void FPostProcessing::Process(FRHICommandListImmediate& RHICmdList, const FViewI
 
 			AddPostProcessMaterial(Context, BL_BeforeTonemapping, SeparateTranslucency);
 
-			EAntiAliasingMethod AntiAliasingMethod = Context.View.FinalPostProcessSettings.AntiAliasingMethod;
+			EAntiAliasingMethod AntiAliasingMethod = Context.View.AntiAliasingMethod;
 
 			if( AntiAliasingMethod == AAM_TemporalAA && ViewState)
 			{
@@ -1721,6 +1722,14 @@ void FPostProcessing::Process(FRHICommandListImmediate& RHICmdList, const FViewI
 			bResultsUpsampled = true;
 		}
 
+		if (DebugViewShaderMode == DVSM_PrimitiveDistanceAccuracy || DebugViewShaderMode == DVSM_MeshTexCoordSizeAccuracy || DebugViewShaderMode == DVSM_MaterialTexCoordScalesAccuracy)
+		{
+			FRenderingCompositePass* Node = Context.Graph.RegisterPass(new(FMemStack::Get()) FRCPassPostProcessStreamingAccuracyLegend(GEngine->StreamingAccuracyColors));
+			Node->SetInput(ePId_Input0, FRenderingCompositeOutputRef(Context.FinalOutput));
+			Context.FinalOutput = FRenderingCompositeOutputRef(Node);
+			bResultsUpsampled = true;
+		}
+
 		if(View.Family->EngineShowFlags.VisualizeLightCulling) 
 		{
 			float ComplexityScale = 1.f / (float)(GEngine->LightComplexityColors.Num() - 1) / .1f; // .1f comes from the values used in LightAccumulator_GetResult
@@ -1983,7 +1992,7 @@ void FPostProcessing::ProcessES2(FRHICommandListImmediate& RHICmdList, const FVi
 		FRenderingCompositeOutputRef BloomOutput;
 		FRenderingCompositeOutputRef DofOutput;
 
-		bool bUseAa = View.FinalPostProcessSettings.AntiAliasingMethod == AAM_TemporalAA;
+		bool bUseAa = View.AntiAliasingMethod == AAM_TemporalAA;
 
 		// AA with Mobile32bpp mode requires this outside of bUsePost.
 		if(bUseAa)
@@ -2423,7 +2432,7 @@ void FPostProcessing::ProcessPlanarReflection(FRHICommandListImmediate& RHICmdLi
 		}
 
 		FSceneViewState* ViewState = Context.View.ViewState;
-		EAntiAliasingMethod AntiAliasingMethod = Context.View.FinalPostProcessSettings.AntiAliasingMethod;
+		EAntiAliasingMethod AntiAliasingMethod = Context.View.AntiAliasingMethod;
 
 		if (AntiAliasingMethod == AAM_TemporalAA && ViewState)
 		{
diff --git a/Engine/Source/Runtime/Renderer/Private/PostProcess/RenderTargetPool.cpp b/Engine/Source/Runtime/Renderer/Private/PostProcess/RenderTargetPool.cpp
--- a/Engine/Source/Runtime/Renderer/Private/PostProcess/RenderTargetPool.cpp
+++ b/Engine/Source/Runtime/Renderer/Private/PostProcess/RenderTargetPool.cpp
@@ -855,25 +855,7 @@ void FRenderTargetPool::PresentContent(FRHICommandListImmediate& RHICmdList, con
 			RHICmdList.SetRasterizerState(TStaticRasterizerState<>::GetRHI());
 			RHICmdList.SetDepthStencilState(TStaticDepthStencilState<false, CF_Always>::GetRHI());
 
-			// this is a helper class for FCanvas to be able to get screen size
-			class FRenderTargetTemp : public FRenderTarget
-			{
-			public:
-				const FSceneView& View;
-
-				FRenderTargetTemp(const FSceneView& InView) : View(InView)
-				{
-				}
-				virtual FIntPoint GetSizeXY() const
-				{
-					return View.UnscaledViewRect.Size();
-				};
-				virtual const FTexture2DRHIRef& GetRenderTargetTexture() const
-				{
-					return View.Family->RenderTarget->GetRenderTargetTexture();
-				}
-			} TempRenderTarget(View);
-
+			FRenderTargetTemp TempRenderTarget(View, View.UnscaledViewRect.Size());
 			FCanvas Canvas(&TempRenderTarget, NULL, View.Family->CurrentRealTime, View.Family->CurrentWorldTime, View.Family->DeltaWorldTime, View.GetFeatureLevel());
 
 			// TinyFont property
