Fix static analysis warning

diff --git a/Engine/Source/Runtime/Engine/Private/Animation/AnimationRuntime.cpp b/Engine/Source/Runtime/Engine/Private/Animation/AnimationRuntime.cpp
--- a/Engine/Source/Runtime/Engine/Private/Animation/AnimationRuntime.cpp
+++ b/Engine/Source/Runtime/Engine/Private/Animation/AnimationRuntime.cpp
@@ -1403,6 +1403,11 @@ static int32 FindMorphTarget(const TArray<FActiveMorphTarget>& ActiveMorphTarget
 
 void FAnimationRuntime::AppendActiveMorphTargets(const USkeletalMesh* InSkeletalMesh, const TMap<FName, float>& MorphCurveAnims, TArray<FActiveMorphTarget>& InOutActiveMorphTargets, TArray<float>& InOutMorphTargetWeights)
 {
+	if (!InSkeletalMesh)
+	{
+		return;
+	}
+
 	// Then go over the CurveKeys finding morph targets by name
 	for(auto CurveIter=MorphCurveAnims.CreateConstIterator(); CurveIter; ++CurveIter)
 	{
@@ -1427,7 +1432,7 @@ void FAnimationRuntime::AppendActiveMorphTargets(const USkeletalMesh* InSkeletal
 		{
 			// Find morph reference
 			int32 SkeletalMorphIndex = INDEX_NONE;
-			UMorphTarget* Target = InSkeletalMesh ? InSkeletalMesh->FindMorphTargetAndIndex(CurveName, SkeletalMorphIndex) : nullptr;
+			UMorphTarget* Target = InSkeletalMesh->FindMorphTargetAndIndex(CurveName, SkeletalMorphIndex);
 			if (Target != nullptr)
 			{
 				// See if this morph target already has an entry
