Copying //UE4/Dev-Framework to //UE4/Dev-Main (Source: //UE4/Dev-Framework @ 3108373)

==========================
MAJOR FEATURES + CHANGES
==========================

Change 3090652 on 2016/08/16 by [REDACTED_PERSON]

	UE-33627 When slicing procmesh collision shapes, ensure at least 4 verts in resulting geometry
	UE-33301 Fix normals on cap sections when slicing proc mesh

Change 3093719 on 2016/08/18 by [REDACTED_PERSON]

	Remove unused bIsAudible

Change 3094646 on 2016/08/19 by [REDACTED_PERSON]

	DumpActiveSounds console command to help debugging audio

Change 3094667 on 2016/08/19 by [REDACTED_PERSON]

	Remove inadvertently checked in lines

Change 3094877 on 2016/08/19 by [REDACTED_PERSON]

	Improve warning when creating constraint between different scenes to report components being connected

Change 3094881 on 2016/08/19 by [REDACTED_PERSON]

	PR #2609: Fix TimeSecondsToString format so it matches comment

Change 3094883 on 2016/08/19 by [REDACTED_PERSON]

	PR #2651: Add GetBoneMass and GetSkeletalCenterOfMass functions for skeletal meshes.

Change 3094888 on 2016/08/19 by [REDACTED_PERSON]

	PR #2629 : Fix typo in ESuggestProjVelocityTraceOption::OnlyTraceWhileAsceding to OnlyTraceWhileAscending

	PR #2580: Fix typo in comment for DeprojectScreenToWorld

Change 3094891 on 2016/08/19 by [REDACTED_PERSON]

	Allow focusing of Actors with 0 size bounds.

Change 3094892 on 2016/08/19 by [REDACTED_PERSON]

	Put null bounds of an empty text render component at the correct location.

Change 3094893 on 2016/08/19 by [REDACTED_PERSON]

	UE-34951 - Zero-volume vorbis decoded sounds are too expensive

	-Adding an audio settings parameter to disable zero-volume playback globally
	-Adding a new bool on sound waves to allow opt-in to virtualize when at zero-volume

Change 3094895 on 2016/08/19 by [REDACTED_PERSON]

	Added several new features for collision component
	  - 'Attach End' option
	  - Collision support (inc friction)
	  - Stiffness option
	  - External forces
	  - GetCableParticleLocations function
	Removed need for CableUpDir setting
	Added perf stats for cable sim

Change 3094980 on 2016/08/19 by [REDACTED_PERSON]

	Fix compile error

Change 3095035 on 2016/08/19 by [REDACTED_PERSON]

	increased layer limit for navmesh generation

Change 3095036 on 2016/08/19 by [REDACTED_PERSON]

	fixed navmesh generation getting stuck in infite loop during contour trace step

Change 3095148 on 2016/08/19 by [REDACTED_PERSON]

	Adjust UT to correct spelling from CL# 3094888

Change 3095228 on 2016/08/19 by [REDACTED_PERSON]

	PR #2389 - Expanded exposure of cheat manager (Contributed by [REDACTED_PERSON])

Change 3095305 on 2016/08/19 by [REDACTED_PERSON]

	PR #2692: Set mouse position : PlayerController (Contributed by [REDACTED_PERSON])

Change 3096570 on 2016/08/22 by [REDACTED_PERSON]

	Move bodies label in PhAT viewport to be easier to read.

Change 3096742 on 2016/08/22 by [REDACTED_PERSON]

	Fix constraints being improperly edited (To Hinge, To Prismatic, etc...) when an active constraint profile is selected.
	Fix undo redo not working with constraint profile instances.

Change 3096883 on 2016/08/22 by [REDACTED_PERSON]

	fixed landscape's geometry export for navigation: sliced version and mirroring for holes

Change 3096941 on 2016/08/22 by [REDACTED_PERSON]

	Fix phat incorrectly simulating using the wrong physics asset. (was missing virtual)

Change 3098107 on 2016/08/23 by [REDACTED_PERSON]

	Move skeletal mesh component physics updating into EndPhysics so that other components can use it safely in PostPhysics

Change 3098365 on 2016/08/23 by [REDACTED_PERSON]

	Make sure to search guid before searching by name.
	 - Was waiting for change from main to come through before making this change.

Change 3098468 on 2016/08/23 by [REDACTED_PERSON]

	Fix CIS. Was going to make registration functions for cloth private, but I guess some games are relying on this directly.

Change 3098668 on 2016/08/23 by [REDACTED_PERSON]

	Expose the ability to separate reverse and brake given from PR#1418

Change 3099270 on 2016/08/24 by [REDACTED_PERSON]

	Remove WheeledVehicle and related headers from Engine.h

Change 3099358 on 2016/08/24 by [REDACTED_PERSON]

	Fix non-unity build for removing wheeledvehicle out of engine.h

Change 3099434 on 2016/08/24 by [REDACTED_PERSON]

	Fix vehicle suspension force offset to work from the spring location (in the z-axis) instead of the vehicle COM.

Change 3099534 on 2016/08/24 by [REDACTED_PERSON]

	Update vehicle template to use the new suspension mode.

Change 3100097 on 2016/08/24 by [REDACTED_PERSON]

	Morphtarget viewer support

	PR #2327: Morph target viewer (Contributed by [REDACTED_PERSON])

Change 3101285 on 2016/08/25 by [REDACTED_PERSON]

	Adds a simple test case for the FSkeletalMeshMerge tool to QA game.

Change 3101356 on 2016/08/25 by [REDACTED_PERSON]

	Fix invalid bound calculation in CalculateLodCount.

Change 3101981 on 2016/08/25 by [REDACTED_PERSON]

	Fixed BTDecorator_CompareBBEntries not working as expected if first value pulled from BB is '<' then the other #UE4

	It's an alternative fix to the one suggested as PR#2736

Change 3103399 on 2016/08/26 by [REDACTED_PERSON]

	Added SimpleWheeledVehicleMovementComponent which allows for n wheels with no motor simulation.
	Exposed WheeledVehicleMovementComponent4W so that it can be added to any blueprint.

Change 3103414 on 2016/08/26 by [REDACTED_PERSON]

	fixed UBTDecorator_CompareBBEntries::CalculateRawConditionValue condition calue calculation #UE4

Change 3104345 on 2016/08/29 by [REDACTED_PERSON]

	Fix typo (reverseAsBrake instead of reverseAsBreak)

Change 3104412 on 2016/08/29 by [REDACTED_PERSON]

	Fix crash when using Don't Create Shape for vehicle wheel

Change 3104423 on 2016/08/29 by [REDACTED_PERSON]

	Expose SimpleWheeledVehicleMovementComponent torque and brake

Change 3104798 on 2016/08/29 by [REDACTED_PERSON]

	Even better phrasing of the condition in UBTDecorator_CompareBBEntries::CalculateRawConditionValue #UE4

Change 3105781 on 2016/08/30 by [REDACTED_PERSON]

	added log warnings when opening AIGraphs (behavior tree, EQS query) with missing node instance classes
	#2721

Change 3106161 on 2016/08/30 by [REDACTED_PERSON]

	Fixing a comment that's being pulled into documentation inappropriately.

Change 3106394 on 2016/08/30 by [REDACTED_PERSON]

	Don't crash trying to browse to the CAC actor template

Change 3106638 on 2016/08/30 by [REDACTED_PERSON]

	Fixed UK2Node_GetInputAxisValue using default bOverrideParent value on the created binding, causing previous bindings to be squashed.

Change 3106788 on 2016/08/30 by [REDACTED_PERSON]

	Fix CIS warning resulting from 3106638

Change 3106946 on 2016/08/30 by [REDACTED_PERSON]

	Fix is to not include application volume changes (tabbing or application volume) when evaluating sound waves for inclusion in wave instance list, etc. Only use the "actual" volume when setting the volume on playing sound sources.

	#tests tab out of game when running, audio returns to normal like it should (including 1-shot sounds)

Change 3106980 on 2016/08/30 by [REDACTED_PERSON]

	#Fix transform of poseablemeshcomponent not applying child transform and when set transform by bone name in editor

	rushing this for Ryan

Change 3108285 on 2016/08/31 by [REDACTED_PERSON]

	removed deprecated nodes from troll and werewolf behavior trees (both unused)

diff --git a/Engine/Config/BaseEngine.ini b/Engine/Config/BaseEngine.ini
--- a/Engine/Config/BaseEngine.ini
+++ b/Engine/Config/BaseEngine.ini
@@ -912,6 +912,8 @@ bUseStreamingPause=false
 +EnumRedirects=(EnumName="EKeys",OldEnumEntry="EKeys::XboxTypeS_DPad_Right",NewEnumEntry="EKeys::Gamepad_DPad_Right")
 +EnumRedirects=(EnumName="EKeys",OldEnumEntry="EKeys::XboxTypeS_DPad_Left",NewEnumEntry="EKeys::Gamepad_DPad_Left")
 
++EnumRedirects=(EnumName="ESuggestProjVelocityTraceOption",OldEnumEntry="OnlyTraceWhileAsceding",NewEnumEntry="OnlyTraceWhileAscending")
+
 ; ModifyFrequency -> Mobility
 +TaggedPropertyRedirects=(ClassName="SceneComponent",OldPropertyName="ModifyFrequency",NewPropertyName="Mobility")
 +K2FieldRedirects=(OldFieldName="SceneComponent.ModifyFrequency",NewFieldName="Mobility")
diff --git a/Engine/Plugins/Runtime/CableComponent/Source/CableComponent/Classes/CableActor.h b/Engine/Plugins/Runtime/CableComponent/Source/CableComponent/Classes/CableActor.h
--- a/Engine/Plugins/Runtime/CableComponent/Source/CableComponent/Classes/CableActor.h
+++ b/Engine/Plugins/Runtime/CableComponent/Source/CableComponent/Classes/CableActor.h
@@ -5,7 +5,7 @@
 #include "CableActor.generated.h"
 
 /** An actor that renders a simulated cable */
-UCLASS(hidecategories=(Input,Collision,Replication), showcategories=("Input|MouseInput", "Input|TouchInput"))
+UCLASS(hidecategories=(Input,Replication), showcategories=("Input|MouseInput", "Input|TouchInput"))
 class ACableActor : public AActor
 {
 	GENERATED_UCLASS_BODY()
diff --git a/Engine/Plugins/Runtime/CableComponent/Source/CableComponent/Classes/CableComponent.h b/Engine/Plugins/Runtime/CableComponent/Source/CableComponent/Classes/CableComponent.h
--- a/Engine/Plugins/Runtime/CableComponent/Source/CableComponent/Classes/CableComponent.h
+++ b/Engine/Plugins/Runtime/CableComponent/Source/CableComponent/Classes/CableComponent.h
@@ -22,7 +22,7 @@ struct FCableParticle
 };
 
 /** Component that allows you to specify custom triangle mesh geometry */
-UCLASS(hidecategories=(Object, Physics, Collision, Activation, "Components|Activation"), editinlinenew, meta=(BlueprintSpawnableComponent), ClassGroup=Rendering)
+UCLASS(hidecategories=(Object, Physics, Activation, "Components|Activation"), editinlinenew, meta=(BlueprintSpawnableComponent), ClassGroup=Rendering)
 class CABLECOMPONENT_API UCableComponent : public UMeshComponent
 {
 	GENERATED_UCLASS_BODY()
@@ -37,6 +37,10 @@ public:
 
 	//~ Begin USceneComponent Interface.
 	virtual FBoxSphereBounds CalcBounds(const FTransform& LocalToWorld) const override;
+	virtual void QuerySupportedSockets(TArray<FComponentSocketDescription>& OutSockets) const override;
+	virtual bool HasAnySockets() const override;
+	virtual bool DoesSocketExist(FName InSocketName) const override;
+	virtual FTransform GetSocketTransform(FName InSocketName, ERelativeTransformSpace TransformSpace = RTS_World) const override;
 	//~ Begin USceneComponent Interface.
 
 	//~ Begin UPrimitiveComponent Interface.
@@ -48,7 +52,21 @@ public:
 	//~ End UMeshComponent Interface.
 
 
-	/** Actor or Component that the end of the cable should be attached to */
+	/**
+	 *	Should we fix the start to something, or leave it free.
+	 *	If false, component transform is just used for initial location of start of cable
+	 */
+	UPROPERTY(EditAnywhere, BlueprintReadWrite, Category = "Cable")
+	bool bAttachStart;
+
+	/** 
+	 *	Should we fix the end to something (using AttachEndTo and EndLocation), or leave it free. 
+	 *	If false, AttachEndTo and EndLocation are just used for initial location of end of cable
+	 */
+	UPROPERTY(EditAnywhere, BlueprintReadWrite, Category = "Cable")
+	bool bAttachEnd;
+
+	/** Actor or Component that the defines the end position of the cable */
 	UPROPERTY(EditAnywhere, Category="Cable")
 	FComponentReference AttachEndTo;
 
@@ -68,6 +86,10 @@ public:
 	UFUNCTION(BlueprintCallable, Category = "Cable")
 	USceneComponent* GetAttachedComponent() const;
 
+	/** Get array of locations of particles (in world space) making up the cable simulation. */
+	UFUNCTION(BlueprintCallable, Category = "Cable")
+	void GetCableParticleLocations(TArray<FVector>& Locations) const;
+
 	/** Rest length of the cable */
 	UPROPERTY(EditAnywhere, BlueprintReadWrite, Category="Cable", meta=(ClampMin = "0.0", UIMin = "0.0", UIMax = "1000.0"))
 	float CableLength;
@@ -81,9 +103,32 @@ public:
 	float SubstepTime;
 
 	/** The number of solver iterations controls how 'stiff' the cable is */
-	UPROPERTY(EditAnywhere, BlueprintReadWrite, Category="Cable", meta=(ClampMin = "1", ClampMax = "8"))
+	UPROPERTY(EditAnywhere, BlueprintReadWrite, Category="Cable", meta=(ClampMin = "1", ClampMax = "16"))
 	int32 SolverIterations;
 
+	/** Add stiffness constraints to cable. */
+	UPROPERTY(EditAnywhere, BlueprintReadWrite, AdvancedDisplay, Category = "Cable")
+	bool bEnableStiffness;
+
+	/** 
+	 *	EXPERIMENTAL. Perform sweeps for each cable particle, each substep, to avoid collisions with the world. 
+	 *	Uses the Collision Preset on the component to determine what is collided with.
+	 *	This greatly increases the cost of the cable simulation.
+	 */
+	UPROPERTY(EditAnywhere, BlueprintReadWrite, AdvancedDisplay, Category = "Cable")
+	bool bEnableCollision;
+
+	/** If collision is enabled, control how much sliding friction is applied when cable is in contact. */
+	UPROPERTY(EditAnywhere, BlueprintReadWrite, AdvancedDisplay, Category = "Cable", meta = (ClampMin = "0.0", ClampMax = "1.0", EditCondition = "bEnableCollision"))
+	float CollisionFriction;
+
+	/** Force vector (world space) applied to all particles in cable. */
+	UPROPERTY(EditAnywhere, BlueprintReadWrite, Category = "Cable Forces")
+	FVector CableForce;
+
+	/** Scaling applied to world gravity affecting this cable. */
+	UPROPERTY(EditAnywhere, BlueprintReadWrite, Category = "Cable Forces")
+	float CableGravityScale;
 
 	/** How wide the cable geometry is */
 	UPROPERTY(EditAnywhere, BlueprintReadWrite, Category="Cable Rendering", meta=(ClampMin = "0.01", UIMin = "0.01", UIMax = "50.0"))
@@ -97,16 +142,14 @@ public:
 	UPROPERTY(EditAnywhere, BlueprintReadWrite, Category="Cable Rendering", meta=(UIMin = "0.1", UIMax = "8"))
 	float TileMaterial;
 
-	/** Axis (in component space) that is used when building cable mesh geometry. Should be orthogonal to cable direction.  */
-	UPROPERTY(EditAnywhere, AdvancedDisplay, BlueprintReadOnly, Category = "Cable")
-	FVector CableUpDir;
-
 private:
 
 	/** Solve the cable spring constraints */
 	void SolveConstraints();
 	/** Integrate cable point positions */
 	void VerletIntegrate(float InSubstepTime, const FVector& Gravity);
+	/** Perform collision traces for particles */
+	void PerformCableCollision();
 	/** Perform a simulation substep */
 	void PerformSubstep(float InSubstepTime, const FVector& Gravity);
 	/** Get start and end position for the cable */
diff --git a/Engine/Plugins/Runtime/CableComponent/Source/CableComponent/Private/CableComponent.cpp b/Engine/Plugins/Runtime/CableComponent/Source/CableComponent/Private/CableComponent.cpp
--- a/Engine/Plugins/Runtime/CableComponent/Source/CableComponent/Private/CableComponent.cpp
+++ b/Engine/Plugins/Runtime/CableComponent/Source/CableComponent/Private/CableComponent.cpp
@@ -6,6 +6,15 @@
 #include "LocalVertexFactory.h"
 #include "Engine/Engine.h"
 
+DECLARE_CYCLE_STAT(TEXT("Cable Sim"), STAT_Cable_SimTime, STATGROUP_CableComponent);
+DECLARE_CYCLE_STAT(TEXT("Cable Solve"), STAT_Cable_SolveTime, STATGROUP_CableComponent);
+DECLARE_CYCLE_STAT(TEXT("Cable Collision"), STAT_Cable_CollisionTime, STATGROUP_CableComponent);
+DECLARE_CYCLE_STAT(TEXT("Cable Integrate"), STAT_Cable_IntegrateTime, STATGROUP_CableComponent);
+
+static FName CableEndSocketName(TEXT("CableEnd"));
+static FName CableStartSocketName(TEXT("CableStart"));
+
+//////////////////////////////////////////////////////////////////////////
 
 /** Vertex Buffer */
 class FCableVertexBuffer : public FVertexBuffer 
@@ -101,7 +110,6 @@ public:
 		, CableWidth(Component->CableWidth)
 		, NumSides(Component->NumSides)
 		, TileMaterial(Component->TileMaterial)
-		, CableUpDir(Component->CableUpDir)
 	{
 		VertexBuffer.NumVerts = GetRequiredVertexCount();
 		IndexBuffer.NumIndices = GetRequiredIndexCount();
@@ -170,11 +178,12 @@ public:
 			const int32 NextIndex = FMath::Min(PointIdx+1, NumPoints-1);
 			const FVector ForwardDir = (InPoints[NextIndex] - InPoints[PrevIndex]).GetSafeNormal();
 
-			// Find a side vector at this point
-			const FVector RightDir = (CableUpDir ^ ForwardDir).GetSafeNormal();
+			// Find quat from up (Z) vector to forward
+			const FQuat DeltaQuat = FQuat::FindBetween(FVector(0, 0, -1), ForwardDir);
 
-			// Find an up vector
-			const FVector UpDir = (ForwardDir ^ RightDir).GetSafeNormal();
+			// Apply quat orth vectors
+			const FVector RightDir = DeltaQuat.RotateVector(FVector(0, 1, 0));
+			const FVector UpDir = DeltaQuat.RotateVector(FVector(1, 0, 0));
 
 			// Generate a ring of verts
 			for(int32 VertIdx = 0; VertIdx<NumRingVerts; VertIdx++)
@@ -331,8 +340,6 @@ private:
 	int32 NumSides;
 
 	float TileMaterial;
-
-	FVector CableUpDir;
 };
 
 
@@ -346,18 +353,20 @@ UCableComponent::UCableComponent( const FObjectInitializer& ObjectInitializer )
 	bTickInEditor = true;
 	bAutoActivate = true;
 
+	bAttachStart = true;
+	bAttachEnd = true;
 	CableWidth = 10.f;
 	NumSegments = 10;
 	NumSides = 4;
-	EndLocation = FVector(0,0,0);
+	EndLocation = FVector(100.f,0,0);
 	CableLength = 100.f;
 	SubstepTime = 0.02f;
 	SolverIterations = 1;
 	TileMaterial = 1.f;
+	CollisionFriction = 0.2f;
+	CableGravityScale = 1.f;
 
-	CableUpDir = FVector(0,0,1);
-
-	SetCollisionProfileName(UCollisionProfile::NoCollision_ProfileName);
+	SetCollisionProfileName(UCollisionProfile::PhysicsActor_ProfileName);
 }
 
 FPrimitiveSceneProxy* UCableComponent::CreateSceneProxy()
@@ -393,21 +402,14 @@ void UCableComponent::OnRegister()
 
 		Particle.Position = InitialPosition;
 		Particle.OldPosition = InitialPosition;
-
-		// fix particles at ends
-		if(ParticleIdx == 0 || ParticleIdx == NumParticles-1)
-		{
-			Particle.bFree = false;
-		}
-		else
-		{
-			Particle.bFree = true;
-		}
+		Particle.bFree = true; // default to free, will be fixed if desired in TickComponent
 	}
 }
 
 void UCableComponent::VerletIntegrate(float InSubstepTime, const FVector& Gravity)
 {
+	SCOPE_CYCLE_COUNTER(STAT_Cable_IntegrateTime);
+
 	const int32 NumParticles = NumSegments+1;
 	const float SubstepTimeSqr = InSubstepTime * InSubstepTime;
 
@@ -416,8 +418,13 @@ void UCableComponent::VerletIntegrate(float InSubstepTime, const FVector& Gravit
 		FCableParticle& Particle = Particles[ParticleIdx];
 		if(Particle.bFree)
 		{
+			// Calc overall force
+			const FVector ParticleForce = Gravity + CableForce;
+
+			// Find vel
 			const FVector Vel = Particle.Position - Particle.OldPosition;
-			const FVector NewPosition = Particle.Position + Vel + (SubstepTimeSqr * Gravity);
+			// Update position
+			const FVector NewPosition = Particle.Position + Vel + (SubstepTimeSqr * ParticleForce);
 
 			Particle.OldPosition = Particle.Position;
 			Particle.Position = NewPosition;
@@ -452,26 +459,109 @@ static void SolveDistanceConstraint(FCableParticle& ParticleA, FCableParticle& P
 
 void UCableComponent::SolveConstraints()
 {
+	SCOPE_CYCLE_COUNTER(STAT_Cable_SolveTime);
+
 	const float SegmentLength = CableLength/(float)NumSegments;
 
 	// For each iteration..
-	for(int32 IterationIdx=0; IterationIdx<SolverIterations; IterationIdx++)
+	for (int32 IterationIdx = 0; IterationIdx < SolverIterations; IterationIdx++)
 	{
-		// For each segment..
-		for(int32 SegIdx=0; SegIdx<NumSegments; SegIdx++)
+		// Solve distance constraint for each segment
+		for (int32 SegIdx = 0; SegIdx < NumSegments; SegIdx++)
 		{
 			FCableParticle& ParticleA = Particles[SegIdx];
-			FCableParticle& ParticleB = Particles[SegIdx+1];
+			FCableParticle& ParticleB = Particles[SegIdx + 1];
 			// Solve for this pair of particles
 			SolveDistanceConstraint(ParticleA, ParticleB, SegmentLength);
 		}
+
+		// If desired, solve stiffness constraints (distance constraints between every other particle)
+		if (bEnableStiffness)
+		{
+			for (int32 SegIdx = 0; SegIdx < NumSegments-1; SegIdx++)
+			{
+				FCableParticle& ParticleA = Particles[SegIdx];
+				FCableParticle& ParticleB = Particles[SegIdx + 2];
+				SolveDistanceConstraint(ParticleA, ParticleB, 2.f*SegmentLength);
+			}
+		}
+	}
+}
+
+void UCableComponent::PerformCableCollision()
+{
+	SCOPE_CYCLE_COUNTER(STAT_Cable_CollisionTime);
+
+	UWorld* World = GetWorld();
+	// If we have a world, and collision is not disabled
+	if (World && GetCollisionEnabled() != ECollisionEnabled::NoCollision)
+	{
+		// Get collision settings from component
+		static FName CableCollisionName(TEXT("CableCollision"));
+		FCollisionQueryParams Params(CableCollisionName);
+
+		ECollisionChannel TraceChannel = GetCollisionObjectType();
+		FCollisionResponseParams ResponseParams(GetCollisionResponseToChannels());
+
+		// Iterate over each particle
+		for (int32 ParticleIdx = 0; ParticleIdx < Particles.Num(); ParticleIdx++)
+		{
+			FCableParticle& Particle = Particles[ParticleIdx];
+			// If particle is free
+			if (Particle.bFree)
+			{
+				// Do sphere sweep
+				FHitResult Result;
+				bool bHit = World->SweepSingleByChannel(Result, Particle.OldPosition, Particle.Position, FQuat::Identity, TraceChannel, FCollisionShape::MakeSphere(0.5f * CableWidth), Params, ResponseParams);
+				// If we got a hit, resolve it
+				if (bHit)
+				{
+					if (Result.bStartPenetrating)
+					{
+						Particle.Position += (Result.Normal * Result.PenetrationDepth);
+					}
+					else
+					{
+						Particle.Position = Result.Location;
+					}
+
+					// Find new velocity, after fixing collision
+					FVector Delta = Particle.Position - Particle.OldPosition;
+					// Find component in normal
+					float NormalDelta = Delta | Result.Normal;
+					// Find component in plane
+					FVector PlaneDelta = Delta - (NormalDelta * Result.Normal);
+
+					// Zero out any positive separation velocity, basically zero restitution
+					Particle.OldPosition += (NormalDelta * Result.Normal);
+
+					// Apply friction in plane of collision if desired
+					if (CollisionFriction > KINDA_SMALL_NUMBER)
+					{
+						// Scale plane delta  by 'friction'
+						FVector ScaledPlaneDelta = PlaneDelta * CollisionFriction;
+
+						// Apply delta to old position reduce implied velocity in collision plane
+						Particle.OldPosition += ScaledPlaneDelta;
+					}
+				}
+			}
+		}
 	}
 }
 
 void UCableComponent::PerformSubstep(float InSubstepTime, const FVector& Gravity)
 {
+	SCOPE_CYCLE_COUNTER(STAT_Cable_SimTime);
+
 	VerletIntegrate(InSubstepTime, Gravity);
+
 	SolveConstraints();
+
+	if (bEnableCollision)
+	{
+		PerformCableCollision();
+	}
 }
 
 void UCableComponent::SetAttachEndTo(AActor* Actor, FName ComponentProperty)
@@ -490,6 +580,16 @@ USceneComponent* UCableComponent::GetAttachedComponent() const
 	return AttachEndTo.GetComponent(GetOwner());
 }
 
+void UCableComponent::GetCableParticleLocations(TArray<FVector>& Locations) const
+{
+	Locations.Empty();
+	for (const FCableParticle& Particle : Particles)
+	{
+		Locations.Add(Particle.Position);
+	}
+}
+
+
 void UCableComponent::GetEndPositions(FVector& OutStartPosition, FVector& OutEndPosition)
 {
 	// Start position is just component position
@@ -509,17 +609,34 @@ void UCableComponent::TickComponent(float DeltaTime, enum ELevelTick TickType, F
 {
 	Super::TickComponent(DeltaTime, TickType, ThisTickFunction);
 
-	const FVector Gravity = FVector(0, 0, GetWorld()->GetGravityZ());
+	const FVector Gravity = FVector(0, 0, GetWorld()->GetGravityZ()) * CableGravityScale;
 
 	// Update end points
 	FVector CableStart, CableEnd;
 	GetEndPositions(CableStart, CableEnd);
 
 	FCableParticle& StartParticle = Particles[0];
-	StartParticle.Position = StartParticle.OldPosition = CableStart;
+
+	if (bAttachStart)
+	{
+		StartParticle.Position = StartParticle.OldPosition = CableStart;
+		StartParticle.bFree = false;
+	}
+	else
+	{
+		StartParticle.bFree = true;
+	}
 
 	FCableParticle& EndParticle = Particles[NumSegments];
-	EndParticle.Position = EndParticle.OldPosition = CableEnd;
+	if (bAttachEnd)
+	{
+		EndParticle.Position = EndParticle.OldPosition = CableEnd;
+		EndParticle.bFree = false;
+	}
+	else
+	{
+		EndParticle.bFree = true;
+	}
 
 	// Ensure a non-zero substep
 	float UseSubstep = FMath::Max(SubstepTime, 0.005f);
@@ -575,7 +692,72 @@ FBoxSphereBounds UCableComponent::CalcBounds(const FTransform& LocalToWorld) con
 		CableBox += Particle.Position;
 	}
 
-	// Expand by cable width
-	return FBoxSphereBounds(CableBox.ExpandBy(CableWidth));
+	// Expand by cable radius (half cable width)
+	return FBoxSphereBounds(CableBox.ExpandBy(0.5f * CableWidth));
+}
+
+void UCableComponent::QuerySupportedSockets(TArray<FComponentSocketDescription>& OutSockets) const 
+{
+	OutSockets.Add(FComponentSocketDescription(CableEndSocketName, EComponentSocketType::Socket));
+	OutSockets.Add(FComponentSocketDescription(CableStartSocketName, EComponentSocketType::Socket));
+}
+
+FTransform UCableComponent::GetSocketTransform(FName InSocketName, ERelativeTransformSpace TransformSpace) const
+{
+	int32 NumParticles = Particles.Num();
+	if ((InSocketName == CableEndSocketName || InSocketName == CableStartSocketName) && NumParticles >= 2)
+	{
+		FVector ForwardDir, Pos;
+		if (InSocketName == CableEndSocketName)
+		{
+			FVector LastPos = Particles[NumParticles - 1].Position;
+			FVector PreviousPos = Particles[NumParticles - 2].Position;
+
+			ForwardDir = (LastPos - PreviousPos).GetSafeNormal();
+			Pos = LastPos;
+		}
+		else
+		{
+			FVector FirstPos = Particles[0].Position;
+			FVector NextPos = Particles[1].Position;
+
+			ForwardDir = (NextPos - FirstPos).GetSafeNormal();
+			Pos = FirstPos;
+		}
+
+		const FQuat RotQuat = FQuat::FindBetween(FVector(1, 0, 0), ForwardDir);
+		FTransform WorldSocketTM = FTransform(RotQuat, Pos, FVector(1, 1, 1));
+
+		switch (TransformSpace)
+		{
+			case RTS_World:
+			{
+				return WorldSocketTM;
+			}
+			case RTS_Actor:
+			{
+				if (const AActor* Actor = GetOwner())
+				{
+					return WorldSocketTM.GetRelativeTransform(GetOwner()->GetTransform());
+				}
+				break;
+			}
+			case RTS_Component:
+			{
+				return WorldSocketTM.GetRelativeTransform(ComponentToWorld);
+			}
+		}
+	}
+
+	return Super::GetSocketTransform(InSocketName, TransformSpace);
 }
 
+bool UCableComponent::HasAnySockets() const
+{
+	return (Particles.Num() >= 2);
+}
+
+bool UCableComponent::DoesSocketExist(FName InSocketName) const
+{
+	return (InSocketName == CableEndSocketName) || (InSocketName == CableStartSocketName);
+}
\ No newline at end of file
diff --git a/Engine/Plugins/Runtime/CableComponent/Source/CableComponent/Private/CableComponentPluginPrivatePCH.h b/Engine/Plugins/Runtime/CableComponent/Source/CableComponent/Private/CableComponentPluginPrivatePCH.h
--- a/Engine/Plugins/Runtime/CableComponent/Source/CableComponent/Private/CableComponentPluginPrivatePCH.h
+++ b/Engine/Plugins/Runtime/CableComponent/Source/CableComponent/Private/CableComponentPluginPrivatePCH.h
@@ -10,4 +10,8 @@
 #include "ModuleManager.h"
 
 #include "CableComponent.h"
-#include "CableActor.h"
\ No newline at end of file
+#include "CableActor.h"
+
+#pragma once
+
+DECLARE_STATS_GROUP(TEXT("Cable"), STATGROUP_CableComponent, STATCAT_Advanced);
diff --git a/Engine/Plugins/Runtime/ProceduralMeshComponent/Source/ProceduralMeshComponent/Private/KismetProceduralMeshLibrary.cpp b/Engine/Plugins/Runtime/ProceduralMeshComponent/Source/ProceduralMeshComponent/Private/KismetProceduralMeshLibrary.cpp
--- a/Engine/Plugins/Runtime/ProceduralMeshComponent/Source/ProceduralMeshComponent/Private/KismetProceduralMeshLibrary.cpp
+++ b/Engine/Plugins/Runtime/ProceduralMeshComponent/Source/ProceduralMeshComponent/Private/KismetProceduralMeshLibrary.cpp
@@ -526,7 +526,7 @@ FProcMeshVertex InterpolateVert(const FProcMeshVertex& V0, const FProcMeshVertex
 /** Transform triangle from 2D to 3D static-mesh triangle. */
 void Transform2DPolygonTo3D(const FUtilPoly2D& InPoly, const FMatrix& InMatrix, TArray<FProcMeshVertex>& OutVerts, FBox& OutBox)
 {
-	FVector PolyNormal = InMatrix.GetUnitAxis(EAxis::Z);
+	FVector PolyNormal = -InMatrix.GetUnitAxis(EAxis::Z);
 	FProcMeshTangent PolyTangent(InMatrix.GetUnitAxis(EAxis::X), false);
 
 	for (int32 VertexIndex = 0; VertexIndex < InPoly.Verts.Num(); VertexIndex++)
@@ -644,11 +644,11 @@ bool TriangulatePoly(TArray<int32>& OutTris, const TArray<FProcMeshVertex>& Poly
 /** Util to slice a convex hull with a plane */
 void SliceConvexElem(const FKConvexElem& InConvex, const FPlane& SlicePlane, TArray<FVector>& OutConvexVerts)
 {
-	// Get set of planes taht make up hull
+	// Get set of planes that make up hull
 	TArray<FPlane> ConvexPlanes;
 	InConvex.GetPlanes(ConvexPlanes);
 
-	if (ConvexPlanes.Num() > 4)
+	if (ConvexPlanes.Num() >= 4)
 	{
 		// Add on the slicing plane (need to flip as it culls geom in the opposite sense to our geom culling code)
 		ConvexPlanes.Add(SlicePlane.Flip());
@@ -1060,7 +1060,7 @@ void UKismetProceduralMeshLibrary::SliceProceduralMesh(UProceduralMeshComponent*
 				TArray<FVector> SlicedConvexVerts;
 				SliceConvexElem(BaseConvex, SlicePlane, SlicedConvexVerts);
 				// If we got something valid, add it
-				if (SlicedConvexVerts.Num() > 0)
+				if (SlicedConvexVerts.Num() >= 4)
 				{
 					SlicedCollision.Add(SlicedConvexVerts);
 				}
@@ -1070,7 +1070,7 @@ void UKismetProceduralMeshLibrary::SliceProceduralMesh(UProceduralMeshComponent*
 				{
 					TArray<FVector> OtherSlicedConvexVerts;
 					SliceConvexElem(BaseConvex, SlicePlane.Flip(), OtherSlicedConvexVerts);
-					if (OtherSlicedConvexVerts.Num() > 0)
+					if (OtherSlicedConvexVerts.Num() >= 4)
 					{
 						OtherSlicedCollision.Add(OtherSlicedConvexVerts);
 					}
diff --git a/Engine/Source/Editor/AIGraph/Classes/AIGraph.h b/Engine/Source/Editor/AIGraph/Classes/AIGraph.h
--- a/Engine/Source/Editor/AIGraph/Classes/AIGraph.h
+++ b/Engine/Source/Editor/AIGraph/Classes/AIGraph.h
@@ -31,6 +31,10 @@ class AIGRAPH_API UAIGraph : public UEdGraph
 	void LockUpdates();
 	void UnlockUpdates();
 
+	//~ Begin UObject Interface.
+	virtual void Serialize(FArchive& Ar) override;
+	//~ End UObject Interface.
+
 protected:
 
 	/** if set, graph modifications won't cause updates in internal tree structure
diff --git a/Engine/Source/Editor/AIGraph/Private/AIGraph.cpp b/Engine/Source/Editor/AIGraph/Private/AIGraph.cpp
--- a/Engine/Source/Editor/AIGraph/Private/AIGraph.cpp
+++ b/Engine/Source/Editor/AIGraph/Private/AIGraph.cpp
@@ -70,29 +70,67 @@ bool UAIGraph::UpdateUnknownNodeClasses()
 	return bUpdated;
 }
 
+void UpdateAIGraphNodeErrorMessage(UAIGraphNode& Node)
+{
+	// Broke out setting error message in to own function so it can be reused when iterating nodes collection.
+	if (Node.NodeInstance)
+	{
+		Node.ErrorMessage = FGraphNodeClassHelper::GetDeprecationMessage(Node.NodeInstance->GetClass());
+	}
+	else
+	{
+		// Null instance. Do we have any meaningful class data?
+		FString StoredClassName = Node.ClassData.GetClassName();
+		StoredClassName.RemoveFromEnd(TEXT("_C"));
+
+		if (!StoredClassName.IsEmpty())
+		{
+			// There is class data here but the instance was not be created.
+			static const FString IsMissingClassMessage(" class missing. Referenced by ");
+			Node.ErrorMessage = StoredClassName + IsMissingClassMessage + Node.GetFullName();
+		}
+	}
+
+	if (Node.HasErrors())
+	{
+		UE_LOG(LogAIGraph, Error, TEXT("%s"), *Node.ErrorMessage);
+	}
+}
+
 void UAIGraph::UpdateDeprecatedClasses()
 {
-	for (int32 Idx = 0; Idx < Nodes.Num(); Idx++)
+	// This function sets error messages and logs errors about nodes.
+
+	for (int32 Idx = 0, IdxNum = Nodes.Num(); Idx < IdxNum; ++Idx)
 	{
 		UAIGraphNode* Node = Cast<UAIGraphNode>(Nodes[Idx]);
-		if (Node)
+		if (Node != nullptr)
 		{
-			if (Node->NodeInstance)
-			{
-				Node->ErrorMessage = FGraphNodeClassHelper::GetDeprecationMessage(Node->NodeInstance->GetClass());
-			}
-
-			for (int32 SubIdx = 0; SubIdx < Node->SubNodes.Num(); SubIdx++)
+			UpdateAIGraphNodeErrorMessage(*Node);
+			
+			for (int32 SubIdx = 0, SubIdxNum = Node->SubNodes.Num(); SubIdx < SubIdxNum; ++SubIdx)
 			{
-				if (Node->SubNodes[SubIdx] && Node->SubNodes[SubIdx]->NodeInstance)
+				if (Node->SubNodes[SubIdx] != nullptr)
 				{
-					Node->SubNodes[SubIdx]->ErrorMessage = FGraphNodeClassHelper::GetDeprecationMessage(Node->SubNodes[SubIdx]->NodeInstance->GetClass());
+					UpdateAIGraphNodeErrorMessage(*Node->SubNodes[SubIdx]);
 				}
 			}
 		}
 	}
 }
 
+void UAIGraph::Serialize(FArchive& Ar)
+{
+	// Overridden to flags up errors in the behavior tree while cooking.
+	Super::Serialize(Ar);
+
+	if (Ar.IsSaving() || Ar.IsCooking())
+	{
+		// Logging of errors happens in UpdateDeprecatedClasses
+		UpdateDeprecatedClasses();
+	}
+}
+
 void UAIGraph::UpdateClassData()
 {
 	for (int32 Idx = 0; Idx < Nodes.Num(); Idx++)
diff --git a/Engine/Source/Editor/AIGraph/Private/AIGraphModule.cpp b/Engine/Source/Editor/AIGraph/Private/AIGraphModule.cpp
--- a/Engine/Source/Editor/AIGraph/Private/AIGraphModule.cpp
+++ b/Engine/Source/Editor/AIGraph/Private/AIGraphModule.cpp
@@ -6,6 +6,8 @@
 
 IMPLEMENT_MODULE(FAIGraphModule, AIGraph);
 
+DEFINE_LOG_CATEGORY(LogAIGraph);
+
 void FAIGraphModule::StartupModule()
 {
 }
diff --git a/Engine/Source/Editor/AIGraph/Public/AIGraphModule.h b/Engine/Source/Editor/AIGraph/Public/AIGraphModule.h
--- a/Engine/Source/Editor/AIGraph/Public/AIGraphModule.h
+++ b/Engine/Source/Editor/AIGraph/Public/AIGraphModule.h
@@ -4,6 +4,8 @@
 
 #include "ModuleInterface.h"
 
+DECLARE_LOG_CATEGORY_EXTERN(LogAIGraph, Display, All);
+
 class FAIGraphModule : public IModuleInterface
 {
 public:
diff --git a/Engine/Source/Editor/BlueprintGraph/Private/K2Node_GetInputAxisValue.cpp b/Engine/Source/Editor/BlueprintGraph/Private/K2Node_GetInputAxisValue.cpp
--- a/Engine/Source/Editor/BlueprintGraph/Private/K2Node_GetInputAxisValue.cpp
+++ b/Engine/Source/Editor/BlueprintGraph/Private/K2Node_GetInputAxisValue.cpp
@@ -96,6 +96,7 @@ void UK2Node_GetInputAxisValue::RegisterDynamicBinding(UDynamicBlueprintBinding*
 	Binding.InputAxisName = InputAxisName;
 	Binding.bConsumeInput = bConsumeInput;
 	Binding.bExecuteWhenPaused = bExecuteWhenPaused;
+	Binding.bOverrideParentBinding = false;
 
 	InputAxisBindingObject->InputAxisDelegateBindings.Add(Binding);
 }
diff --git a/Engine/Source/Editor/Persona/Private/AnimViewportShowCommands.cpp b/Engine/Source/Editor/Persona/Private/AnimViewportShowCommands.cpp
--- a/Engine/Source/Editor/Persona/Private/AnimViewportShowCommands.cpp
+++ b/Engine/Source/Editor/Persona/Private/AnimViewportShowCommands.cpp
@@ -28,7 +28,11 @@ void FAnimViewportShowCommands::RegisterCommands()
 	UI_COMMAND(ShowDisplayInfoSkelControls, "Skeletal Controls", "Display selected skeletal control info in Viewport", EUserInterfaceActionType::RadioButton, FInputChord());
 	UI_COMMAND(HideDisplayInfo, "None", "Hide All Display Info in Viewport", EUserInterfaceActionType::RadioButton, FInputChord());
 
-	UI_COMMAND( ShowBoneWeight, "View Selected Bone Weight", "Display Selected Bone Weight in Viewport", EUserInterfaceActionType::ToggleButton, FInputChord() );
+	UI_COMMAND( ShowOverlayNone, "None", "Clear Overlay Display", EUserInterfaceActionType::RadioButton, FInputChord());
+	UI_COMMAND( ShowBoneWeight, "Selected Bone Weight", "Display color overlay of the weight from Selected Bone in Viewport", EUserInterfaceActionType::RadioButton, FInputChord() );
+	UI_COMMAND( ShowMorphTargetVerts, "Selected Morphtarget Vertices", "Display color overlay with the chnage of Selected Morphtarget in Viewport", EUserInterfaceActionType::RadioButton, FInputChord());
+
+
 	UI_COMMAND( ShowRawAnimation, "Uncompressed Animation", "Display Skeleton With Uncompressed Animation Data", EUserInterfaceActionType::ToggleButton, FInputChord() );
 	UI_COMMAND( ShowNonRetargetedAnimation, "NonRetargeted Animation", "Display Skeleton With non retargeted Animation Data", EUserInterfaceActionType::ToggleButton, FInputChord() );
 	UI_COMMAND( ShowAdditiveBaseBones, "Additive Base", "Display Skeleton In Additive Base Pose", EUserInterfaceActionType::ToggleButton, FInputChord() );
diff --git a/Engine/Source/Editor/Persona/Private/AnimViewportShowCommands.h b/Engine/Source/Editor/Persona/Private/AnimViewportShowCommands.h
--- a/Engine/Source/Editor/Persona/Private/AnimViewportShowCommands.h
+++ b/Engine/Source/Editor/Persona/Private/AnimViewportShowCommands.h
@@ -81,8 +81,10 @@ public:
 	TSharedPtr< FUICommandInfo > ShowDisplayInfoSkelControls;
 	TSharedPtr< FUICommandInfo > HideDisplayInfo;
 
-	/** Show selected bone weight */
+	/** Show overlay material option */
+	TSharedPtr< FUICommandInfo > ShowOverlayNone;
 	TSharedPtr< FUICommandInfo > ShowBoneWeight;
+	TSharedPtr< FUICommandInfo > ShowMorphTargetVerts;
 
 	/** Show socket hit point diamonds */
 	TSharedPtr< FUICommandInfo > ShowSockets;
diff --git a/Engine/Source/Editor/Persona/Private/Persona.cpp b/Engine/Source/Editor/Persona/Private/Persona.cpp
--- a/Engine/Source/Editor/Persona/Private/Persona.cpp
+++ b/Engine/Source/Editor/Persona/Private/Persona.cpp
@@ -2368,6 +2368,34 @@ void FPersona::SetSelectedBone(USkeleton* InTargetSkeleton, const FName& BoneNam
 	}
 }
 
+void FPersona::SetSelectedMorphTargets(USkeletalMesh *SkeletalMesh, const TArray<FName>& SelectedMorphTargetNames)
+{
+	UDebugSkelMeshComponent* Preview = GetPreviewMeshComponent();
+	if (Preview)
+	{
+		Preview->MorphTargetOfInterests.Reset();
+
+		if (SelectedMorphTargetNames.Num() > 0)
+		{
+			for (const FName& MorphTargetName : SelectedMorphTargetNames)
+			{
+				int32 MorphtargetIdx;
+				UMorphTarget* MorphTarget = SkeletalMesh->FindMorphTargetAndIndex(MorphTargetName, MorphtargetIdx);
+				if (MorphTarget != nullptr)
+				{
+					Preview->MorphTargetOfInterests.AddUnique(MorphTarget);
+				}
+			}
+
+			if (Viewport.IsValid())
+			{
+				Viewport.Pin()->GetLevelViewportClient().Invalidate();
+			}
+			Preview->PostInitMeshObject(PreviewComponent->MeshObject);
+		}
+	}
+}
+
 void FPersona::SetSelectedSocket( const FSelectedSocketInfo& SocketInfo, bool bRebroadcast /* = true */ )
 {
 	if ( UDebugSkelMeshComponent* Preview = GetPreviewMeshComponent() )
diff --git a/Engine/Source/Editor/Persona/Private/Persona.h b/Engine/Source/Editor/Persona/Private/Persona.h
--- a/Engine/Source/Editor/Persona/Private/Persona.h
+++ b/Engine/Source/Editor/Persona/Private/Persona.h
@@ -95,6 +95,9 @@ public:
 	// Sets the selected bone to preview
 	void SetSelectedBone(class USkeleton* Skeleton, const FName& BoneName, bool bRebroadcast = true );
 
+	// Sets the selected morph target
+	void SetSelectedMorphTargets(class USkeletalMesh *SkeletonMesh, const TArray<FName>& SelectedMorphTargetNames);
+
 	/** Helper method to clear out any selected bones */
 	void ClearSelectedBones();
 
diff --git a/Engine/Source/Editor/Persona/Private/SAnimViewportToolBar.cpp b/Engine/Source/Editor/Persona/Private/SAnimViewportToolBar.cpp
--- a/Engine/Source/Editor/Persona/Private/SAnimViewportToolBar.cpp
+++ b/Engine/Source/Editor/Persona/Private/SAnimViewportToolBar.cpp
@@ -506,7 +506,12 @@ TSharedRef<SWidget> SAnimViewportToolBar::GenerateShowMenu() const
 
 				ShowMenuBuilder.AddMenuEntry( Actions.ShowSockets );
 				ShowMenuBuilder.AddMenuEntry( Actions.ShowBoneNames );
-				ShowMenuBuilder.AddMenuEntry( Actions.ShowBoneWeight );
+
+				ShowMenuBuilder.AddSubMenu(
+					LOCTEXT("AnimViewportOverlayDrawSubMenu", "Mesh Overlay"),
+					LOCTEXT("AnimViewportOverlayDrawSubMenuToolTip", "Options"),
+					FNewMenuDelegate::CreateRaw(this, &SAnimViewportToolBar::FillShowOverlayDrawMenu)
+				);
 			}
 			ShowMenuBuilder.EndSection();
 
@@ -604,6 +609,19 @@ void SAnimViewportToolBar::FillShowBoneDrawMenu(FMenuBuilder& MenuBuilder) const
 	MenuBuilder.EndSection();
 }
 
+void SAnimViewportToolBar::FillShowOverlayDrawMenu(FMenuBuilder& MenuBuilder) const
+{
+	const FAnimViewportShowCommands& Actions = FAnimViewportShowCommands::Get();
+
+	MenuBuilder.BeginSection("AnimViewportPreviewOverlayDraw", LOCTEXT("ShowMenu_Actions_Overlay", "Overlay Draw"));
+	{
+		MenuBuilder.AddMenuEntry(Actions.ShowOverlayNone);
+		MenuBuilder.AddMenuEntry(Actions.ShowBoneWeight);
+		MenuBuilder.AddMenuEntry(Actions.ShowMorphTargetVerts);
+	}
+	MenuBuilder.EndSection();
+}
+
 void SAnimViewportToolBar::FillShowAdvancedMenu(FMenuBuilder& MenuBuilder) const
 {
 	const FAnimViewportShowCommands& Actions = FAnimViewportShowCommands::Get();
diff --git a/Engine/Source/Editor/Persona/Private/SAnimViewportToolBar.h b/Engine/Source/Editor/Persona/Private/SAnimViewportToolBar.h
--- a/Engine/Source/Editor/Persona/Private/SAnimViewportToolBar.h
+++ b/Engine/Source/Editor/Persona/Private/SAnimViewportToolBar.h
@@ -52,6 +52,11 @@ private:
 	*/
 	void FillShowBoneDrawMenu(FMenuBuilder& MenuBuilder) const;
 
+	/**
+	* Generates the Show -> Overlay sub menu content
+	*/
+	void FillShowOverlayDrawMenu(FMenuBuilder& MenuBuilder) const;
+
 	/**
 	* Generates the Show -> Clothing sub menu content
 	*/
diff --git a/Engine/Source/Editor/Persona/Private/SAnimationEditorViewport.cpp b/Engine/Source/Editor/Persona/Private/SAnimationEditorViewport.cpp
--- a/Engine/Source/Editor/Persona/Private/SAnimationEditorViewport.cpp
+++ b/Engine/Source/Editor/Persona/Private/SAnimationEditorViewport.cpp
@@ -557,12 +557,24 @@ void SAnimationEditorViewportTabBody::BindCommands()
 		FCanExecuteAction(),
 		FIsActionChecked::CreateSP(this, &SAnimationEditorViewportTabBody::IsShowingMeshInfo, (int32)EDisplayInfoMode::None));
 
-	//Bone weight
+	//Material overlay option
+	CommandList.MapAction(
+		ViewportShowMenuCommands.ShowOverlayNone,
+		FExecuteAction::CreateSP(this, &SAnimationEditorViewportTabBody::OnShowOverlayNone),
+		FCanExecuteAction(),
+		FIsActionChecked::CreateSP(this, &SAnimationEditorViewportTabBody::IsShowingOverlayNone));
+
 	CommandList.MapAction( 
 		ViewportShowMenuCommands.ShowBoneWeight,
-		FExecuteAction::CreateSP(this, &SAnimationEditorViewportTabBody::OnShowBoneWeight),
+		FExecuteAction::CreateSP(this, &SAnimationEditorViewportTabBody::OnShowOverlayBoneWeight),
 		FCanExecuteAction(),
-		FIsActionChecked::CreateSP(this, &SAnimationEditorViewportTabBody::IsShowingBoneWeight));
+		FIsActionChecked::CreateSP(this, &SAnimationEditorViewportTabBody::IsShowingOverlayBoneWeight));
+
+	CommandList.MapAction(
+		ViewportShowMenuCommands.ShowMorphTargetVerts,
+		FExecuteAction::CreateSP(this, &SAnimationEditorViewportTabBody::OnShowOverlayMorphTargetVert),
+		FCanExecuteAction(),
+		FIsActionChecked::CreateSP(this, &SAnimationEditorViewportTabBody::IsShowingOverlayMorphTargetVerts));
 
 	// Show sockets
 	CommandList.MapAction( 
@@ -974,7 +986,26 @@ bool SAnimationEditorViewportTabBody::IsShowingMeshInfo(int32 DisplayInfoMode) c
 	return GetAnimationViewportClient()->GetShowMeshStats() == DisplayInfoMode;
 }
 
-void SAnimationEditorViewportTabBody::OnShowBoneWeight()
+void SAnimationEditorViewportTabBody::OnShowOverlayNone()
+{
+	UDebugSkelMeshComponent* PreviewComponent = PersonaPtr.Pin()->PreviewComponent;
+	if (PreviewComponent)
+	{
+		PreviewComponent->SetShowBoneWeight(false);
+		PreviewComponent->SetShowMorphTargetVerts(false);
+		UpdateShowFlagForMeshEdges();
+		PreviewComponent->MarkRenderStateDirty();
+		RefreshViewport();
+	}
+}
+
+bool SAnimationEditorViewportTabBody::IsShowingOverlayNone() const
+{
+	UDebugSkelMeshComponent* PreviewComponent = PersonaPtr.Pin()->PreviewComponent;
+	return PreviewComponent != NULL && !PreviewComponent->bDrawBoneInfluences && !PreviewComponent->bDrawMorphTargetVerts;
+}
+
+void SAnimationEditorViewportTabBody::OnShowOverlayBoneWeight()
 {
 	UDebugSkelMeshComponent* PreviewComponent = PersonaPtr.Pin()->PreviewComponent;
 	if( PreviewComponent )
@@ -986,12 +1017,30 @@ void SAnimationEditorViewportTabBody::OnShowBoneWeight()
 	}
 }
 
-bool SAnimationEditorViewportTabBody::IsShowingBoneWeight() const
+bool SAnimationEditorViewportTabBody::IsShowingOverlayBoneWeight() const
 {
 	UDebugSkelMeshComponent* PreviewComponent = PersonaPtr.Pin()->PreviewComponent;
 	return PreviewComponent != NULL && PreviewComponent->bDrawBoneInfluences;
 }
 
+void SAnimationEditorViewportTabBody::OnShowOverlayMorphTargetVert()
+{
+	UDebugSkelMeshComponent* PreviewComponent = PersonaPtr.Pin()->PreviewComponent;
+	if (PreviewComponent)
+	{
+		PreviewComponent->SetShowMorphTargetVerts(!PreviewComponent->bDrawMorphTargetVerts);
+		UpdateShowFlagForMeshEdges();
+		PreviewComponent->MarkRenderStateDirty();
+		RefreshViewport();
+	}
+}
+
+bool SAnimationEditorViewportTabBody::IsShowingOverlayMorphTargetVerts() const
+{
+	UDebugSkelMeshComponent* PreviewComponent = PersonaPtr.Pin()->PreviewComponent;
+	return PreviewComponent != NULL && PreviewComponent->bDrawMorphTargetVerts;
+}
+
 void SAnimationEditorViewportTabBody::OnSetBoneDrawMode(int32 BoneDrawMode)
 {
 	TSharedRef<FAnimationViewportClient> AnimViewportClient = StaticCastSharedRef<FAnimationViewportClient>(LevelViewportClient.ToSharedRef());
@@ -1346,10 +1395,10 @@ float SAnimationEditorViewportTabBody::GetViewMaxInput() const
 
 void SAnimationEditorViewportTabBody::UpdateShowFlagForMeshEdges()
 {
-	bool bDrawBonesInfluence = false;
+	bool bUseOverlayMaterial = false;
 	if (UDebugSkelMeshComponent* PreviewComponent = PersonaPtr.Pin()->PreviewComponent)
 	{
-		bDrawBonesInfluence = PreviewComponent->bDrawBoneInfluences;
+		bUseOverlayMaterial = PreviewComponent->bDrawBoneInfluences || PreviewComponent->bDrawMorphTargetVerts;
 	}
 
 	//@TODO: SNOWPOCALYPSE: broke UnlitWithMeshEdges
@@ -1358,7 +1407,7 @@ void SAnimationEditorViewportTabBody::UpdateShowFlagForMeshEdges()
 	bShowMeshEdgesViewMode = (CurrentViewMode == EAnimationEditorViewportMode::UnlitWithMeshEdges);
 #endif
 
-	LevelViewportClient->EngineShowFlags.SetMeshEdges(bDrawBonesInfluence || bShowMeshEdgesViewMode);
+	LevelViewportClient->EngineShowFlags.SetMeshEdges(bUseOverlayMaterial || bShowMeshEdgesViewMode);
 }
 
 bool SAnimationEditorViewportTabBody::IsLODModelSelected(int32 LODSelectionType) const
diff --git a/Engine/Source/Editor/Persona/Private/SAnimationEditorViewport.h b/Engine/Source/Editor/Persona/Private/SAnimationEditorViewport.h
--- a/Engine/Source/Editor/Persona/Private/SAnimationEditorViewport.h
+++ b/Engine/Source/Editor/Persona/Private/SAnimationEditorViewport.h
@@ -228,10 +228,22 @@ private:
 	bool IsShowingBoneNames() const;
 	
 	/** Function to show/hide selected bone weight */
-	void OnShowBoneWeight();
+	void OnShowOverlayNone();
 
 	/** Function to check whether bone weights are displayed or not*/
-	bool IsShowingBoneWeight() const;
+	bool IsShowingOverlayNone() const;
+
+	/** Function to show/hide selected bone weight */
+	void OnShowOverlayBoneWeight();
+
+	/** Function to check whether bone weights are displayed or not*/
+	bool IsShowingOverlayBoneWeight() const;
+
+	/** Function to show/hide selected morphtarget overlay*/
+	void OnShowOverlayMorphTargetVert();
+
+	/** Function to check whether morphtarget overlay is displayed or not*/
+	bool IsShowingOverlayMorphTargetVerts() const;
 
 	/** Function to set Local axes mode of the specificed type */
 	void OnSetBoneDrawMode(int32 BoneDrawMode);
diff --git a/Engine/Source/Editor/Persona/Private/SMorphTargetViewer.cpp b/Engine/Source/Editor/Persona/Private/SMorphTargetViewer.cpp
--- a/Engine/Source/Editor/Persona/Private/SMorphTargetViewer.cpp
+++ b/Engine/Source/Editor/Persona/Private/SMorphTargetViewer.cpp
@@ -338,6 +338,7 @@ void SMorphTargetViewer::Construct(const FArguments& InArgs)
 			.ListItemsSource( &MorphTargetList )
 			.OnGenerateRow( this, &SMorphTargetViewer::GenerateMorphTargetRow )
 			.OnContextMenuOpening( this, &SMorphTargetViewer::OnGetContextMenuContent )
+			.OnSelectionChanged( this, &SMorphTargetViewer::OnRowsSelectedChanged )
 			.ItemHeight( 22.0f )
 			.HeaderRow
 			(
@@ -443,6 +444,7 @@ void SMorphTargetViewer::CreateMorphTargetList( const FString& SearchText )
 		}
 	}
 
+	NotifySelectionChange();
 	MorphTargetListView->RequestListRefresh();
 }
 
@@ -510,6 +512,27 @@ SMorphTargetViewer::~SMorphTargetViewer()
 void SMorphTargetViewer::OnPostUndo()
 {
 	CreateMorphTargetList();
+	NotifySelectionChange();
+}
+
+void SMorphTargetViewer::NotifySelectionChange() const
+{
+	TArray< TSharedPtr< FDisplayedMorphTargetInfo > > SelectedRows = MorphTargetListView->GetSelectedItems();
+
+	TArray<FName> SelectedMorphtargetNames;
+	for (auto ItemIt = SelectedRows.CreateIterator(); ItemIt; ++ItemIt)
+	{
+		TSharedPtr< FDisplayedMorphTargetInfo > RowItem = (*ItemIt);
+		SelectedMorphtargetNames.AddUnique(RowItem->Name);
+	}
+
+	// stil have to call this even if empty, otherwise it won't clear it
+	PersonaPtr.Pin()->SetSelectedMorphTargets(SkeletalMesh, SelectedMorphtargetNames);
+}
+
+void SMorphTargetViewer::OnRowsSelectedChanged(TSharedPtr<FDisplayedMorphTargetInfo> Item, ESelectInfo::Type SelectInfo)
+{
+	NotifySelectionChange();
 }
 
 #undef LOCTEXT_NAMESPACE
diff --git a/Engine/Source/Editor/Persona/Private/SMorphTargetViewer.h b/Engine/Source/Editor/Persona/Private/SMorphTargetViewer.h
--- a/Engine/Source/Editor/Persona/Private/SMorphTargetViewer.h
+++ b/Engine/Source/Editor/Persona/Private/SMorphTargetViewer.h
@@ -143,6 +143,9 @@ private:
 	/** Handler for context menus */
 	TSharedPtr<SWidget> OnGetContextMenuContent() const;
 
+	/** Handler for row selection change */
+	void OnRowsSelectedChanged(TSharedPtr<FDisplayedMorphTargetInfo> Item, ESelectInfo::Type SelectInfo);
+
 	/**
 	* Clears and rebuilds the table, according to an optional search string
 	*
@@ -168,4 +171,7 @@ private:
 
 	/** Current text typed into NameFilterBox */
 	FText FilterText;
-};
\ No newline at end of file
+
+	/** notify selection change to Persona */
+	void NotifySelectionChange() const;
+};
diff --git a/Engine/Source/Editor/PhAT/Classes/PhATEdSkeletalMeshComponent.h b/Engine/Source/Editor/PhAT/Classes/PhATEdSkeletalMeshComponent.h
--- a/Engine/Source/Editor/PhAT/Classes/PhATEdSkeletalMeshComponent.h
+++ b/Engine/Source/Editor/PhAT/Classes/PhATEdSkeletalMeshComponent.h
@@ -71,7 +71,4 @@ class UPhATEdSkeletalMeshComponent : public USkeletalMeshComponent
 	FTransform GetPrimitiveTransform(FTransform& BoneTM, int32 BodyIndex, EKCollisionPrimitiveType PrimType, int32 PrimIndex, float Scale);
 	FColor GetPrimitiveColor(int32 BodyIndex, EKCollisionPrimitiveType PrimitiveType, int32 PrimitiveIndex);
 	UMaterialInterface* GetPrimitiveMaterial(int32 BodyIndex, EKCollisionPrimitiveType PrimitiveType, int32 PrimitiveIndex, bool bHitTest);
-
-	/** Returns the physics asset for this PhATEd component - note: This hides the implementation in the USkinnedMeshComponent base class */
-	class UPhysicsAsset* GetPhysicsAsset() const;
 };
diff --git a/Engine/Source/Editor/PhAT/Private/PhAT.cpp b/Engine/Source/Editor/PhAT/Private/PhAT.cpp
--- a/Engine/Source/Editor/PhAT/Private/PhAT.cpp
+++ b/Engine/Source/Editor/PhAT/Private/PhAT.cpp
@@ -181,6 +181,7 @@ TSharedRef<SDockTab> FPhAT::SpawnTab( const FSpawnTabArgs& TabSpawnArgs, FName T
 					FName ProfileName = CS->GetCurrentConstraintProfileName();
 					if (!CS->ContainsConstraintProfile(ProfileName))
 					{
+						CS->Modify();
 						CS->AddConstraintProfile(ProfileName);
 					}
 				}
@@ -199,6 +200,7 @@ TSharedRef<SDockTab> FPhAT::SpawnTab( const FSpawnTabArgs& TabSpawnArgs, FName T
 			{
 				if (UPhysicsConstraintTemplate* CS = Cast<UPhysicsConstraintTemplate>(WeakObj.Get()))
 				{
+					CS->Modify();
 					FName ProfileName = CS->GetCurrentConstraintProfileName();
 					CS->RemoveConstraintProfile(ProfileName);
 				}
@@ -1578,22 +1580,22 @@ void FPhAT::BindCommands()
 	ToolkitCommands->MapAction(
 		Commands.ConvertToBallAndSocket,
 		FExecuteAction::CreateSP(this, &FPhAT::OnConvertToBallAndSocket),
-		FCanExecuteAction::CreateSP(this, &FPhAT::IsSelectedEditConstraintMode));
+		FCanExecuteAction::CreateSP(this, &FPhAT::CanEditConstraintProperties));
 
 	ToolkitCommands->MapAction(
 		Commands.ConvertToHinge,
 		FExecuteAction::CreateSP(this, &FPhAT::OnConvertToHinge),
-		FCanExecuteAction::CreateSP(this, &FPhAT::IsSelectedEditConstraintMode));
+		FCanExecuteAction::CreateSP(this, &FPhAT::CanEditConstraintProperties));
 
 	ToolkitCommands->MapAction(
 		Commands.ConvertToPrismatic,
 		FExecuteAction::CreateSP(this, &FPhAT::OnConvertToPrismatic),
-		FCanExecuteAction::CreateSP(this, &FPhAT::IsSelectedEditConstraintMode));
+		FCanExecuteAction::CreateSP(this, &FPhAT::CanEditConstraintProperties));
 
 	ToolkitCommands->MapAction(
 		Commands.ConvertToSkeletal,
 		FExecuteAction::CreateSP(this, &FPhAT::OnConvertToSkeletal),
-		FCanExecuteAction::CreateSP(this, &FPhAT::IsSelectedEditConstraintMode));
+		FCanExecuteAction::CreateSP(this, &FPhAT::CanEditConstraintProperties));
 
 	ToolkitCommands->MapAction(
 		Commands.DeleteConstraint,
@@ -2219,12 +2221,15 @@ void FPhAT::SnapConstraintToBone(int32 ConstraintIndex, const FTransform& Parent
 
 void FPhAT::CreateOrConvertConstraint(EPhATConstraintType ConstraintType)
 {
+	//we have to manually call PostEditChange to ensure profiles are updated correctly
+	UProperty* DefaultInstanceProperty = FindField<UProperty>(UPhysicsConstraintTemplate::StaticClass(), GET_MEMBER_NAME_CHECKED(UPhysicsConstraintTemplate, DefaultInstance));
+
 	const FScopedTransaction Transaction( LOCTEXT( "CreateConvertConstraint", "Create Or Convert Constraint" ) );
 
 	for(int32 i=0; i<SharedData->SelectedConstraints.Num(); ++i)
 	{
 		UPhysicsConstraintTemplate* ConstraintSetup = SharedData->PhysicsAsset->ConstraintSetup[SharedData->SelectedConstraints[i].Index];
-		ConstraintSetup->Modify();
+		ConstraintSetup->PreEditChange(DefaultInstanceProperty);
 
 		if(ConstraintType == EPCT_BSJoint)
 		{
@@ -2242,6 +2247,9 @@ void FPhAT::CreateOrConvertConstraint(EPhATConstraintType ConstraintType)
 		{
 			ConstraintUtils::ConfigureAsSkelJoint(ConstraintSetup->DefaultInstance);
 		}
+
+		FPropertyChangedEvent PropertyChangedEvent(DefaultInstanceProperty);
+		ConstraintSetup->PostEditChangeProperty(PropertyChangedEvent);
 	}
 
 	RefreshHierachyTree();
@@ -2447,6 +2455,31 @@ bool FPhAT::IsEditConstraintMode() const
 	return IsNotSimulation() && (SharedData->EditingMode == FPhATSharedData::PEM_ConstraintEdit);
 }
 
+bool FPhAT::CanEditConstraintProperties() const
+{
+	if(IsEditConstraintMode() && SharedData->PhysicsAsset)
+	{
+		//If we are currently editing a constraint profile, make sure all selected constraints belong to the profile
+		if(SharedData->PhysicsAsset->CurrentConstraintProfileName != NAME_None)
+		{
+			for (const FPhATSharedData::FSelection& Selection : SharedData->SelectedConstraints)
+			{
+				UPhysicsConstraintTemplate* CS = SharedData->PhysicsAsset->ConstraintSetup[Selection.Index];
+				if(!CS || !CS->ContainsConstraintProfile(SharedData->PhysicsAsset->CurrentConstraintProfileName))
+				{
+					//missing at least one constraint from profile so don't allow editing
+					return false;
+				}
+			}
+		}
+		
+		//no constraint profile so editing is fine
+		return true;
+	}
+
+	return false;
+}
+
 bool FPhAT::IsSelectedEditConstraintMode() const
 {
 	return IsEditConstraintMode() && (SharedData->GetSelectedConstraint());
diff --git a/Engine/Source/Editor/PhAT/Private/PhAT.h b/Engine/Source/Editor/PhAT/Private/PhAT.h
--- a/Engine/Source/Editor/PhAT/Private/PhAT.h
+++ b/Engine/Source/Editor/PhAT/Private/PhAT.h
@@ -158,6 +158,7 @@ private:
 	bool IsEditBodyMode() const;
 	bool IsSelectedEditBodyMode() const;
 	bool IsEditConstraintMode() const;
+	bool CanEditConstraintProperties() const;
 	bool IsSelectedEditConstraintMode() const;
 	bool CanStartSimulation() const;
 	void OnChangeDefaultMesh();
diff --git a/Engine/Source/Editor/PhAT/Private/PhATPreviewViewportClient.cpp b/Engine/Source/Editor/PhAT/Private/PhATPreviewViewportClient.cpp
--- a/Engine/Source/Editor/PhAT/Private/PhATPreviewViewportClient.cpp
+++ b/Engine/Source/Editor/PhAT/Private/PhATPreviewViewportClient.cpp
@@ -116,7 +116,8 @@ void FPhATEdPreviewViewportClient::DrawCanvas( FViewport& InViewport, FSceneView
 	float W, H;
 	PhATFont->GetCharSize(TEXT('L'), W, H);
 
-	const float XOffset = 350.0f;
+	const float XOffset = 5.0f;
+	const float YOffset = 32.0f;
 
 	FCanvasTextItem TextItem( FVector2D::ZeroVector, FText::GetEmpty(), PhATFont, FLinearColor::White );
 
@@ -129,7 +130,7 @@ void FPhATEdPreviewViewportClient::DrawCanvas( FViewport& InViewport, FSceneView
 		FText::AsNumber(SharedData->PhysicsAsset->ConstraintSetup.Num()) ).ToString();
 
 	TextItem.Text = FText::FromString( StatusString );
-	Canvas.DrawItem( TextItem, XOffset, 3 );
+	Canvas.DrawItem( TextItem, XOffset, YOffset);
 	
 	TextItem.Text = FText::GetEmpty();
 	if (SharedData->bRunningSimulation)
diff --git a/Engine/Source/Editor/PhAT/Private/PhATRender.cpp b/Engine/Source/Editor/PhAT/Private/PhATRender.cpp
--- a/Engine/Source/Editor/PhAT/Private/PhATRender.cpp
+++ b/Engine/Source/Editor/PhAT/Private/PhATRender.cpp
@@ -50,11 +50,6 @@ UPhATEdSkeletalMeshComponent::UPhATEdSkeletalMeshComponent(const FObjectInitiali
 	SetCollisionProfileName(CollisionProfileName);
 }
 
-UPhysicsAsset* UPhATEdSkeletalMeshComponent::GetPhysicsAsset() const
-{
-	return SharedData->PhysicsAsset;
-}
-
 void UPhATEdSkeletalMeshComponent::RenderAssetTools(const FSceneView* View, class FPrimitiveDrawInterface* PDI, bool bHitTest)
 {
 	check(SharedData);
diff --git a/Engine/Source/Editor/PhAT/Private/PhATSharedData.cpp b/Engine/Source/Editor/PhAT/Private/PhATSharedData.cpp
--- a/Engine/Source/Editor/PhAT/Private/PhATSharedData.cpp
+++ b/Engine/Source/Editor/PhAT/Private/PhATSharedData.cpp
@@ -107,6 +107,7 @@ void FPhATSharedData::Initialize()
 	// Create SkeletalMeshComponent for rendering skeletal mesh
 	EditorSkelComp = NewObject<UPhATEdSkeletalMeshComponent>();
 	EditorSkelComp->SharedData = this;
+	EditorSkelComp->SetPhysicsAsset(PhysicsAsset);
 
 	PhysicalAnimationComponent = NewObject<UPhysicalAnimationComponent>();
 	PhysicalAnimationComponent->SetSkeletalMeshComponent(EditorSkelComp);
diff --git a/Engine/Source/Editor/PropertyEditor/Private/Presentation/PropertyEditor/PropertyEditor.cpp b/Engine/Source/Editor/PropertyEditor/Private/Presentation/PropertyEditor/PropertyEditor.cpp
--- a/Engine/Source/Editor/PropertyEditor/Private/Presentation/PropertyEditor/PropertyEditor.cpp
+++ b/Engine/Source/Editor/PropertyEditor/Private/Presentation/PropertyEditor/PropertyEditor.cpp
@@ -767,16 +767,21 @@ void FPropertyEditor::SyncToObjectsInNode( const TWeakPtr< FPropertyNode >& Weak
 		}
 
 		// If a single actor is selected, sync to its location in the level editor viewport instead of the content browser.
-		if( Objects.Num() == 1 && Objects[0]->IsA(AActor::StaticClass()) )
+		if( Objects.Num() == 1 && Objects[0]->IsA<AActor>() )
 		{
-			TArray<AActor*> Actors;
-			Actors.Add(Cast<AActor>(Objects[0]));
+			AActor* Actor = CastChecked<AActor>(Objects[0]);
 
-			GEditor->SelectNone(/*bNoteSelectionChange=*/false, /*bDeselectBSPSurfs=*/true);
-			GEditor->SelectActor(Actors[0], /*InSelected=*/true, /*bNotify=*/true, /*bSelectEvenIfHidden=*/true);
+			if (Actor->GetLevel())
+			{
+				TArray<AActor*> Actors;
+				Actors.Add(Actor);
+
+				GEditor->SelectNone(/*bNoteSelectionChange=*/false, /*bDeselectBSPSurfs=*/true);
+				GEditor->SelectActor(Actor, /*InSelected=*/true, /*bNotify=*/true, /*bSelectEvenIfHidden=*/true);
 
-			// Jump to the location of the actor
-			GEditor->MoveViewportCamerasToActor( Actors, /*bActiveViewportOnly=*/false );
+				// Jump to the location of the actor
+				GEditor->MoveViewportCamerasToActor(Actors, /*bActiveViewportOnly=*/false);
+			}
 		}
 		else if ( Objects.Num() > 0 )
 		{
diff --git a/Engine/Source/Editor/UnrealEd/Classes/Animation/DebugSkelMeshComponent.h b/Engine/Source/Editor/UnrealEd/Classes/Animation/DebugSkelMeshComponent.h
--- a/Engine/Source/Editor/UnrealEd/Classes/Animation/DebugSkelMeshComponent.h
+++ b/Engine/Source/Editor/UnrealEd/Classes/Animation/DebugSkelMeshComponent.h
@@ -70,6 +70,10 @@ class UDebugSkelMeshComponent : public USkeletalMeshComponent
 	UPROPERTY(transient)
 	uint32 bDrawBoneInfluences:1;
 
+	/** Morphtarget viewing */
+	UPROPERTY(transient)
+	uint32 bDrawMorphTargetVerts : 1;
+
 	/** Vertex normal viewing */
 	UPROPERTY(transient)
 	uint32 bDrawNormals:1;
@@ -146,6 +150,10 @@ class UDebugSkelMeshComponent : public USkeletalMeshComponent
 	/** Array of bones to render bone weights for */
 	UPROPERTY(transient)
 	TArray<int32> BonesOfInterest;
+
+	/** Array of morphtargets to render verts for */
+	UPROPERTY(transient)
+	TArray<class UMorphTarget*> MorphTargetOfInterests;
 	
 	/** Array of sockets to render manipulation widgets for
 	/	Storing a pointer to the actual socket rather than a name, as we don't care here
@@ -214,6 +222,12 @@ class UDebugSkelMeshComponent : public USkeletalMeshComponent
 	 */
 	UNREALED_API void SetShowBoneWeight(bool bNewShowBoneWeight);
 
+	/**
+	* Update material information depending on color render mode
+	* Refresh/replace materials
+	*/
+	UNREALED_API void SetShowMorphTargetVerts(bool bNewShowMorphTargetVerts);
+
 	/**
 	 * Does it use in-game bounds or bounds calculated from bones
 	 */
@@ -283,6 +297,8 @@ private:
 	// Helper function to generate space bases for current frame
 	void GenSpaceBases(TArray<FTransform>& OutSpaceBases);
 
+	// Helper function to enable overlay material
+	void EnableOverlayMaterial(bool bEnable);
 public:
 	/** Current turn table mode */
 	EPersonaTurnTableMode::Type TurnTableMode;
diff --git a/Engine/Source/Editor/UnrealEd/Private/Animation/DebugSkelMeshComponent.cpp b/Engine/Source/Editor/UnrealEd/Private/Animation/DebugSkelMeshComponent.cpp
--- a/Engine/Source/Editor/UnrealEd/Private/Animation/DebugSkelMeshComponent.cpp
+++ b/Engine/Source/Editor/UnrealEd/Private/Animation/DebugSkelMeshComponent.cpp
@@ -346,7 +346,7 @@ void UDebugSkelMeshComponent::EnablePreview(bool bEnable, UAnimationAsset* Previ
 
 bool UDebugSkelMeshComponent::ShouldCPUSkin()
 {
-	return 	bCPUSkinning || bDrawBoneInfluences || bDrawNormals || bDrawTangents || bDrawBinormals;
+	return 	bCPUSkinning || bDrawBoneInfluences || bDrawNormals || bDrawTangents || bDrawBinormals || bDrawMorphTargetVerts;
 }
 
 
@@ -358,7 +358,11 @@ void UDebugSkelMeshComponent::PostInitMeshObject(FSkeletalMeshObject* InMeshObje
 	{
 		if(bDrawBoneInfluences)
 		{
-			InMeshObject->EnableBlendWeightRendering(true, BonesOfInterest);
+			InMeshObject->EnableOverlayRendering(true, &BonesOfInterest, nullptr);
+		}
+		else if (bDrawMorphTargetVerts)
+		{
+			InMeshObject->EnableOverlayRendering(true, nullptr, &MorphTargetOfInterests);
 		}
 	}
 }
@@ -371,12 +375,24 @@ void UDebugSkelMeshComponent::SetShowBoneWeight(bool bNewShowBoneWeight)
 		return;
 	}
 
+	if (bDrawMorphTargetVerts)
+	{
+		SetShowMorphTargetVerts(false);
+	}
+
 	// if turning on this mode
-	if(bNewShowBoneWeight)
+	EnableOverlayMaterial(bNewShowBoneWeight);
+
+	bDrawBoneInfluences = bNewShowBoneWeight;
+}
+
+void UDebugSkelMeshComponent::EnableOverlayMaterial(bool bEnable)
+{
+	if (bEnable)
 	{
 		SkelMaterials.Empty();
 		int32 NumMaterials = GetNumMaterials();
-		for (int32 i=0; i<NumMaterials; i++)
+		for (int32 i = 0; i < NumMaterials; i++)
 		{
 			// Back up old material
 			SkelMaterials.Add(GetMaterial(i));
@@ -389,14 +405,30 @@ void UDebugSkelMeshComponent::SetShowBoneWeight(bool bNewShowBoneWeight)
 	{
 		int32 NumMaterials = GetNumMaterials();
 		check(NumMaterials == SkelMaterials.Num());
-		for (int32 i=0; i<NumMaterials; i++)
+		for (int32 i = 0; i < NumMaterials; i++)
 		{
 			// restore original material
 			SetMaterial(i, SkelMaterials[i]);
 		}
 	}
+}
+void UDebugSkelMeshComponent::SetShowMorphTargetVerts(bool bNewShowMorphTargetVerts)
+{
+	// Check we are actually changing it!
+	if (bNewShowMorphTargetVerts == bDrawMorphTargetVerts)
+	{
+		return;
+	}
 
-	bDrawBoneInfluences = bNewShowBoneWeight;
+	if (bDrawBoneInfluences)
+	{
+		SetShowBoneWeight(false);
+	}
+
+	// if turning on this mode
+	EnableOverlayMaterial(bNewShowMorphTargetVerts);
+
+	bDrawMorphTargetVerts = bNewShowMorphTargetVerts;
 }
 
 void UDebugSkelMeshComponent::GenSpaceBases(TArray<FTransform>& OutSpaceBases)
diff --git a/Engine/Source/Editor/UnrealEd/Private/EditorServer.cpp b/Engine/Source/Editor/UnrealEd/Private/EditorServer.cpp
--- a/Engine/Source/Editor/UnrealEd/Private/EditorServer.cpp
+++ b/Engine/Source/Editor/UnrealEd/Private/EditorServer.cpp
@@ -6833,7 +6833,7 @@ void UEditorEngine::AutoMergeStaticMeshes()
 void UEditorEngine::MoveViewportCamerasToBox(const FBox& BoundingBox, bool bActiveViewportOnly) const
 {
 	// Make sure we had at least one non-null actor in the array passed in.
-	if (BoundingBox.GetSize() != FVector::ZeroVector)
+	if (BoundingBox.GetSize() != FVector::ZeroVector || BoundingBox.GetCenter() != FVector::ZeroVector)
 	{
 		if (bActiveViewportOnly)
 		{
diff --git a/Engine/Source/Editor/UnrealEd/Private/EditorViewportClient.cpp b/Engine/Source/Editor/UnrealEd/Private/EditorViewportClient.cpp
--- a/Engine/Source/Editor/UnrealEd/Private/EditorViewportClient.cpp
+++ b/Engine/Source/Editor/UnrealEd/Private/EditorViewportClient.cpp
@@ -587,7 +587,7 @@ void FEditorViewportClient::ToggleOrbitCamera( bool bEnableOrbitCamera )
 void FEditorViewportClient::FocusViewportOnBox( const FBox& BoundingBox, bool bInstant /* = false */ )
 {
 	const FVector Position = BoundingBox.GetCenter();
-	float Radius = BoundingBox.GetExtent().Size();
+	float Radius = FMath::Max(BoundingBox.GetExtent().Size(), 10.f);
 
 	float AspectToUse = AspectRatio;
 	FIntPoint ViewportSize = Viewport->GetSizeXY();
diff --git a/Engine/Source/Runtime/AIModule/Private/AIController.cpp b/Engine/Source/Runtime/AIModule/Private/AIController.cpp
--- a/Engine/Source/Runtime/AIModule/Private/AIController.cpp
+++ b/Engine/Source/Runtime/AIModule/Private/AIController.cpp
@@ -1028,7 +1028,7 @@ bool AAIController::SuggestTossVelocity(FVector& OutTossVelocity, FVector Start,
 	// pawn's physics volume gets 2nd priority
 	APhysicsVolume const* const PhysicsVolume = GetPawn() ? GetPawn()->GetPawnPhysicsVolume() : NULL;
 	float const GravityOverride = PhysicsVolume ? PhysicsVolume->GetGravityZ() : 0.f;
-	ESuggestProjVelocityTraceOption::Type const TraceOption = bOnlyTraceUp ? ESuggestProjVelocityTraceOption::OnlyTraceWhileAsceding : ESuggestProjVelocityTraceOption::TraceFullPath;
+	ESuggestProjVelocityTraceOption::Type const TraceOption = bOnlyTraceUp ? ESuggestProjVelocityTraceOption::OnlyTraceWhileAscending : ESuggestProjVelocityTraceOption::TraceFullPath;
 
 	return UGameplayStatics::SuggestProjectileVelocity(this, OutTossVelocity, Start, End, TossSpeed, bPreferHighArc, CollisionRadius, GravityOverride, TraceOption);
 }
@@ -1089,4 +1089,4 @@ void AAIController::SetGenericTeamId(const FGenericTeamId& NewTeamID)
 		TeamID = NewTeamID;
 		// @todo notify perception system that a controller changed team ID
 	}
-}
\ No newline at end of file
+}
diff --git a/Engine/Source/Runtime/AIModule/Private/BehaviorTree/Decorators/BTDecorator_CompareBBEntries.cpp b/Engine/Source/Runtime/AIModule/Private/BehaviorTree/Decorators/BTDecorator_CompareBBEntries.cpp
--- a/Engine/Source/Runtime/AIModule/Private/BehaviorTree/Decorators/BTDecorator_CompareBBEntries.cpp
+++ b/Engine/Source/Runtime/AIModule/Private/BehaviorTree/Decorators/BTDecorator_CompareBBEntries.cpp
@@ -39,16 +39,9 @@ bool UBTDecorator_CompareBBEntries::CalculateRawConditionValue(UBehaviorTreeComp
 	const UBlackboardComponent* BlackboardComp = OwnerComp.GetBlackboardComponent();
 	if (BlackboardComp)
 	{
-		//BlackboardComp->GetKeyType
-		const EBlackboardCompare::Type Result = BlackboardComp->CompareKeyValues(BlackboardKeyA.SelectedKeyType, 
-			BlackboardKeyA.GetSelectedKeyID(), BlackboardKeyB.GetSelectedKeyID());
+		const EBlackboardCompare::Type Result = BlackboardComp->CompareKeyValues(BlackboardKeyA.SelectedKeyType, BlackboardKeyA.GetSelectedKeyID(), BlackboardKeyB.GetSelectedKeyID());
 
-		static_assert(int32(EBlackboardCompare::Equal) == int32(EBlackBoardEntryComparison::Equal)
-			&& int32(EBlackboardCompare::NotEqual) == int32(EBlackBoardEntryComparison::NotEqual),
-			"These values need to be equal");
-
-		// note that this comparison relies on assumption asserted above
-		return int32(Operator) == int32(Result);
+		return ((Result == EBlackboardCompare::Equal) == (Operator == EBlackBoardEntryComparison::Equal));
 	}
 
 	return false;
diff --git a/Engine/Source/Runtime/AnimGraphRuntime/Private/BoneControllers/AnimNode_WheelHandler.cpp b/Engine/Source/Runtime/AnimGraphRuntime/Private/BoneControllers/AnimNode_WheelHandler.cpp
--- a/Engine/Source/Runtime/AnimGraphRuntime/Private/BoneControllers/AnimNode_WheelHandler.cpp
+++ b/Engine/Source/Runtime/AnimGraphRuntime/Private/BoneControllers/AnimNode_WheelHandler.cpp
@@ -14,6 +14,7 @@
 
 FAnimNode_WheelHandler::FAnimNode_WheelHandler()
 {
+	AnimInstanceProxy = nullptr;
 }
 
 void FAnimNode_WheelHandler::GatherDebugData(FNodeDebugData& DebugData)
@@ -25,17 +26,19 @@ void FAnimNode_WheelHandler::GatherDebugData(FNodeDebugData& DebugData)
 	DebugLine += ")";
 
 	DebugData.AddDebugItem(DebugLine);
-	for(const auto & WheelSim : WheelSimulators)
+
+	const TArray<FWheelAnimData>& WheelAnimData = AnimInstanceProxy->GetWheelAnimData();
+	for (const FWheelLookupData& Wheel : Wheels)
 	{
-		if (WheelSim.BoneReference.BoneIndex != INDEX_NONE)
+		if (Wheel.BoneReference.BoneIndex != INDEX_NONE)
 		{
-			DebugLine = FString::Printf(TEXT(" [Wheel Index : %d] Bone: %s , Rotation Offset : %s, Location Offset : %s"), 
-				WheelSim.WheelIndex, *WheelSim.BoneReference.BoneName.ToString(), *WheelSim.RotOffset.ToString(), *WheelSim.LocOffset.ToString());
+			DebugLine = FString::Printf(TEXT(" [Wheel Index : %d] Bone: %s , Rotation Offset : %s, Location Offset : %s"),
+				Wheel.WheelIndex, *Wheel.BoneReference.BoneName.ToString(), *WheelAnimData[Wheel.WheelIndex].RotOffset.ToString(), *WheelAnimData[Wheel.WheelIndex].LocOffset.ToString());
 		}
 		else
 		{
 			DebugLine = FString::Printf(TEXT(" [Wheel Index : %d] Bone: %s (invalid bone)"),
-					WheelSim.WheelIndex, *WheelSim.BoneReference.BoneName.ToString());
+				Wheel.WheelIndex, *Wheel.BoneReference.BoneName.ToString());
 		}
 
 		DebugData.AddDebugItem(DebugLine);
@@ -48,12 +51,14 @@ void FAnimNode_WheelHandler::EvaluateBoneTransforms(USkeletalMeshComponent* Skel
 {
 	check(OutBoneTransforms.Num() == 0);
 
+	const TArray<FWheelAnimData>& WheelAnimData = AnimInstanceProxy->GetWheelAnimData();
+
 	const FBoneContainer& BoneContainer = MeshBases.GetPose().GetBoneContainer();
-	for(const auto & WheelSim : WheelSimulators)
+	for(const FWheelLookupData& Wheel : Wheels)
 	{
-		if (WheelSim.BoneReference.IsValid(BoneContainer))
+		if (Wheel.BoneReference.IsValid(BoneContainer))
 		{
-			FCompactPoseBoneIndex WheelSimBoneIndex = WheelSim.BoneReference.GetCompactPoseIndex(BoneContainer);
+			FCompactPoseBoneIndex WheelSimBoneIndex = Wheel.BoneReference.GetCompactPoseIndex(BoneContainer);
 
 			// the way we apply transform is same as FMatrix or FTransform
 			// we apply scale first, and rotation, and translation
@@ -61,13 +66,13 @@ void FAnimNode_WheelHandler::EvaluateBoneTransforms(USkeletalMeshComponent* Skel
 			FTransform NewBoneTM = MeshBases.GetComponentSpaceTransform(WheelSimBoneIndex);
 
 			FAnimationRuntime::ConvertCSTransformToBoneSpace(SkelComp, MeshBases, NewBoneTM, WheelSimBoneIndex, BCS_ComponentSpace);
-			
+
 			// Apply rotation offset
-			const FQuat BoneQuat(WheelSim.RotOffset);
+			const FQuat BoneQuat(WheelAnimData[Wheel.WheelIndex].RotOffset);
 			NewBoneTM.SetRotation(BoneQuat * NewBoneTM.GetRotation());
 
 			// Apply loc offset
-			NewBoneTM.AddToTranslation(WheelSim.LocOffset);
+			NewBoneTM.AddToTranslation(WheelAnimData[Wheel.WheelIndex].LocOffset);
 
 			// Convert back to Component Space.
 			FAnimationRuntime::ConvertBoneSpaceTransformToCS(SkelComp, MeshBases, NewBoneTM, WheelSimBoneIndex, BCS_ComponentSpace);
@@ -81,10 +86,10 @@ void FAnimNode_WheelHandler::EvaluateBoneTransforms(USkeletalMeshComponent* Skel
 bool FAnimNode_WheelHandler::IsValidToEvaluate(const USkeleton* Skeleton, const FBoneContainer& RequiredBones) 
 {
 	// if both bones are valid
-	for(const auto & WheelSim : WheelSimulators)
+	for (const FWheelLookupData& Wheel : Wheels)
 	{
 		// if one of them is valid
-		if (WheelSim.BoneReference.IsValid(RequiredBones) == true)
+		if (Wheel.BoneReference.IsValid(RequiredBones) == true)
 		{
 			return true;
 		}
@@ -95,71 +100,23 @@ bool FAnimNode_WheelHandler::IsValidToEvaluate(const USkeleton* Skeleton, const
 
 void FAnimNode_WheelHandler::InitializeBoneReferences(const FBoneContainer& RequiredBones) 
 {
-	for (auto & WheelSim : WheelSimulators)
-	{
-		WheelSim.BoneReference.Initialize(RequiredBones);
-	}
+	const TArray<FWheelAnimData>& WheelAnimData = AnimInstanceProxy->GetWheelAnimData();
+	const int32 NumWheels = WheelAnimData.Num();
+	Wheels.Empty(NumWheels);
 
-	// sort by bone indices
-	WheelSimulators.Sort([](FWheelSimulator L, FWheelSimulator R) { return L.BoneReference.BoneIndex < R.BoneReference.BoneIndex; });
-}
-
-void FAnimNode_WheelHandler::UpdateInternal(const FAnimationUpdateContext& Context)
-{
-	if(VehicleSimComponent)
+	for (int32 WheelIndex = 0; WheelIndex < NumWheels; ++WheelIndex)
 	{
-		for(auto & WheelSim : WheelSimulators)
-		{
-			if(VehicleSimComponent->Wheels.IsValidIndex(WheelSim.WheelIndex))
-			{
-				UVehicleWheel* Wheel = VehicleSimComponent->Wheels[WheelSim.WheelIndex];
-				if(Wheel != nullptr)
-				{
-					WheelSim.RotOffset.Pitch	= Wheel->GetRotationAngle();
-					WheelSim.RotOffset.Yaw		= Wheel->GetSteerAngle();
-					WheelSim.RotOffset.Roll		= 0.f;
-
-					WheelSim.LocOffset.X		= 0.f;
-					WheelSim.LocOffset.Y		= 0.f;
-					WheelSim.LocOffset.Z		= Wheel->GetSuspensionOffset();
-				}
-			}
-		}
+		FWheelLookupData* Wheel = new(Wheels)FWheelLookupData();
+		Wheel->WheelIndex = WheelIndex;
+		Wheel->BoneReference.BoneName = WheelAnimData[WheelIndex].BoneName;
+		Wheel->BoneReference.Initialize(RequiredBones);
 	}
 
-	FAnimNode_SkeletalControlBase::UpdateInternal(Context);
+	// sort by bone indices
+	Wheels.Sort([](const FWheelLookupData& L, const FWheelLookupData& R) { return L.BoneReference.BoneIndex < R.BoneReference.BoneIndex; });
 }
 
 void FAnimNode_WheelHandler::Initialize(const FAnimationInitializeContext& Context)
 {
-	// TODO: only check vehicle anim instance
-	// UVehicleAnimInstance
-	AWheeledVehicle * Vehicle = Cast<AWheeledVehicle> (Context.AnimInstanceProxy->GetSkelMeshComponent()->GetOwner());
-	;
-	// we only support vehicle for this class
-	if(Vehicle != nullptr)
-	{
-		VehicleSimComponent = Vehicle->GetVehicleMovementComponent();
-
-		int32 NumOfwheels = VehicleSimComponent->WheelSetups.Num();
-		if(NumOfwheels > 0)
-		{
-			WheelSimulators.Empty(NumOfwheels);
-			WheelSimulators.AddZeroed(NumOfwheels);
-			// now add wheel data
-			for(int32 WheelIndex = 0; WheelIndex<WheelSimulators.Num(); ++WheelIndex)
-			{
-				FWheelSimulator & WheelSim = WheelSimulators[WheelIndex];
-				const FWheelSetup& WheelSetup = VehicleSimComponent->WheelSetups[WheelIndex];
-
-				// set data
-				WheelSim.WheelIndex = WheelIndex;
-				WheelSim.BoneReference.BoneName = WheelSetup.BoneName;
-				WheelSim.LocOffset = FVector::ZeroVector;
-				WheelSim.RotOffset = FRotator::ZeroRotator;
-			}
-		}
-	}
-
-	FAnimNode_SkeletalControlBase::Initialize(Context);
+	AnimInstanceProxy = (FVehicleAnimInstanceProxy*)Context.AnimInstanceProxy;	//TODO: This is cached for now because we need it in eval bone transforms.
 }
\ No newline at end of file
diff --git a/Engine/Source/Runtime/AnimGraphRuntime/Public/BoneControllers/AnimNode_WheelHandler.h b/Engine/Source/Runtime/AnimGraphRuntime/Public/BoneControllers/AnimNode_WheelHandler.h
--- a/Engine/Source/Runtime/AnimGraphRuntime/Public/BoneControllers/AnimNode_WheelHandler.h
+++ b/Engine/Source/Runtime/AnimGraphRuntime/Public/BoneControllers/AnimNode_WheelHandler.h
@@ -2,17 +2,10 @@
 
 #pragma once
 
+#include "Vehicles/VehicleAnimInstance.h"
 #include "AnimNode_SkeletalControlBase.h"
 #include "AnimNode_WheelHandler.generated.h"
 
-struct FWheelSimulator
-{
-	int32					WheelIndex;
-	FBoneReference			BoneReference;
-	FRotator				RotOffset;
-	FVector					LocOffset;
-};
-
 /**
  *	Simple controller that replaces or adds to the translation/rotation of a single bone.
  */
@@ -21,12 +14,6 @@ struct ANIMGRAPHRUNTIME_API FAnimNode_WheelHandler : public FAnimNode_SkeletalCo
 {
 	GENERATED_USTRUCT_BODY()
 
-	/** Current Asset being played **/
-	UPROPERTY(transient)
-	class UWheeledVehicleMovementComponent* VehicleSimComponent;
-
-	TArray<FWheelSimulator>			WheelSimulators;
-
 	FAnimNode_WheelHandler();
 
 	// FAnimNode_Base interface
@@ -36,13 +23,20 @@ struct ANIMGRAPHRUNTIME_API FAnimNode_WheelHandler : public FAnimNode_SkeletalCo
 	// FAnimNode_SkeletalControlBase interface
 	virtual void EvaluateBoneTransforms(USkeletalMeshComponent* SkelComp, FCSPose<FCompactPose>& MeshBases, TArray<FBoneTransform>& OutBoneTransforms) override;
 	virtual bool IsValidToEvaluate(const USkeleton* Skeleton, const FBoneContainer& RequiredBones) override;
-	virtual void UpdateInternal(const FAnimationUpdateContext& Context) override;
 	virtual void Initialize(const FAnimationInitializeContext& Context) override;
-	virtual bool CanUpdateInWorkerThread() const override { return false; }
 	// End of FAnimNode_SkeletalControlBase interface
 
 private:
 	// FAnimNode_SkeletalControlBase interface
 	virtual void InitializeBoneReferences(const FBoneContainer& RequiredBones) override;
 	// End of FAnimNode_SkeletalControlBase interface
+
+	struct FWheelLookupData
+	{
+		int32 WheelIndex;
+		FBoneReference BoneReference;
+	};
+
+	TArray<FWheelLookupData> Wheels;
+	const FVehicleAnimInstanceProxy* AnimInstanceProxy;	//TODO: we only cache this to use in eval where it's safe. Should change API to pass proxy into eval
 };
diff --git a/Engine/Source/Runtime/Core/Public/UObject/FrameworkObjectVersion.h b/Engine/Source/Runtime/Core/Public/UObject/FrameworkObjectVersion.h
--- a/Engine/Source/Runtime/Core/Public/UObject/FrameworkObjectVersion.h
+++ b/Engine/Source/Runtime/Core/Public/UObject/FrameworkObjectVersion.h
@@ -50,6 +50,9 @@ struct CORE_API FFrameworkObjectVersion
 		// Switched render data for clothing over to unreal data, reskinned to the simulation mesh
 		AddInternalClothingGraphicalSkinning,
 
+		// Wheel force offset is now applied at the wheel instead of vehicle COM
+		WheelOffsetIsFromWheel,
+
 		// -----<new versions can be added above this line>-------------------------------------------------
 		VersionPlusOne,
 		LatestVersion = VersionPlusOne - 1
diff --git a/Engine/Source/Runtime/Engine/Classes/Animation/AnimMontage.h b/Engine/Source/Runtime/Engine/Classes/Animation/AnimMontage.h
--- a/Engine/Source/Runtime/Engine/Classes/Animation/AnimMontage.h
+++ b/Engine/Source/Runtime/Engine/Classes/Animation/AnimMontage.h
@@ -345,8 +345,8 @@ public:
 	 * to accumulate and update curve data/notifies/branching points
 	 */
 	void UpdateWeight(float DeltaTime);
-	/** Simulate is same as Advance, but without calling any events or touching any of the instance data. So it performs a simulation of advancing the timeline. 
-	 * @fixme laurent can we make Advance use that, so we don't have 2 code paths which risk getting out of sync? */
+	//~ @fixme laurent can we make Advance use that, so we don't have 2 code paths which risk getting out of sync?
+	/** Simulate is same as Advance, but without calling any events or touching any of the instance data. So it performs a simulation of advancing the timeline. */
 	bool SimulateAdvance(float DeltaTime, float& InOutPosition, struct FRootMotionMovementParams & OutRootMotionParams) const;
 	void Advance(float DeltaTime, struct FRootMotionMovementParams * OutRootMotionParams, bool bBlendRootMotion);
 
diff --git a/Engine/Source/Runtime/Engine/Classes/Components/InputComponent.h b/Engine/Source/Runtime/Engine/Classes/Components/InputComponent.h
--- a/Engine/Source/Runtime/Engine/Classes/Components/InputComponent.h
+++ b/Engine/Source/Runtime/Engine/Classes/Components/InputComponent.h
@@ -64,6 +64,18 @@ struct TInputUnifiedDelegate
 		FuncDynDelegate.Unbind();
 	}
 
+	/** Returns a const reference to the Function Delegate. */
+	inline const DelegateType& GetDelegate() const
+	{
+		return FuncDelegate;
+	}
+
+	/** Returns a const reference to the Dynamic Function Delegate. */
+	inline const DynamicDelegateType& GetDynamicDelegate() const
+	{
+		return FuncDynDelegate;
+	}
+
 protected:
 	/** Holds the delegate to call. */
 	DelegateType FuncDelegate;
diff --git a/Engine/Source/Runtime/Engine/Classes/Components/SkeletalMeshComponent.h b/Engine/Source/Runtime/Engine/Classes/Components/SkeletalMeshComponent.h
--- a/Engine/Source/Runtime/Engine/Classes/Components/SkeletalMeshComponent.h
+++ b/Engine/Source/Runtime/Engine/Classes/Components/SkeletalMeshComponent.h
@@ -266,10 +266,10 @@ enum class EAllowKinematicDeferral
 class USkeletalMeshComponent;
 
 /**
-* Tick function that does post physics work on skeletal mesh component
+* Tick function that does post physics work on skeletal mesh component. This executes in EndPhysics (after physics is done)
 **/
 USTRUCT()
-struct FSkeletalMeshComponentPostPhysicsTickFunction : public FTickFunction
+struct FSkeletalMeshComponentEndPhysicsTickFunction : public FTickFunction
 {
 	GENERATED_USTRUCT_BODY()
 
@@ -955,8 +955,10 @@ public:
 	virtual void InitializeComponent() override;
 	virtual void TickComponent(float DeltaTime, enum ELevelTick TickType, FActorComponentTickFunction *ThisTickFunction) override;
 
+	//Handle registering our end physics tick function
+	void RegisterEndPhysicsTick(bool bRegister);
+
 	//Handle registering our pre cloth tick function
-	void RegisterPostPhysicsTick(bool bRegister);
 	void RegisterClothTick(bool bRegister);
 
 	//~ End UActorComponent Interface.
@@ -1065,6 +1067,23 @@ public:
 	virtual void SetAllPhysicsLinearVelocity(FVector NewVel,bool bAddToCurrent = false) override;
 	virtual void SetAllMassScale(float InMassScale = 1.f) override;
 	virtual float GetMass() const override;
+
+	/**
+	*	Returns the mass (in kg) of the given bone
+	*
+	*	@param BoneName		Name of the body to return. 'None' indicates root body.
+	*	@param bScaleMass	If true, the mass is scaled by the bone's MassScale.
+	*/
+	UFUNCTION(BlueprintCallable, Category = "Physics")
+	float GetBoneMass(FName BoneName = NAME_None, bool bScaleMass = true) const;
+
+	/**
+	*	Returns the center of mass of the skeletal mesh, instead of the root body's location
+	*/
+	UFUNCTION(BlueprintCallable, Category = "Physics")
+	FVector GetSkeletalCenterOfMass() const;
+
+
 	virtual float CalculateMass(FName BoneName = NAME_None) override;
 	virtual bool DoCustomNavigableGeometryExport(FNavigableGeometryExport& GeomExport) const override;
 
@@ -1436,11 +1455,11 @@ protected:
 	
 private:
 
-	FSkeletalMeshComponentPostPhysicsTickFunction PostPhysicsTickFunction;
-	friend struct FSkeletalMeshComponentPostPhysicsTickFunction;
+	FSkeletalMeshComponentEndPhysicsTickFunction EndPhysicsTickFunction;
+	friend struct FSkeletalMeshComponentEndPhysicsTickFunction;
 
 	/** Update systems after physics sim is done */
-	void PostPhysicsTickComponent(FSkeletalMeshComponentPostPhysicsTickFunction& ThisTickFunction);
+	void EndPhysicsTickComponent(FSkeletalMeshComponentEndPhysicsTickFunction& ThisTickFunction);
 
 	/** Evaluate Anim System **/
 	void EvaluateAnimation(const USkeletalMesh* InSkeletalMesh, UAnimInstance* InAnimInstance, TArray<FTransform>& OutBoneSpaceTransforms, FVector& OutRootBoneTranslation, FBlendedHeapCurve& OutCurve) const;
@@ -1499,13 +1518,13 @@ public:
 	
 private:
 	// Returns whether we need to run the Pre Cloth Tick or not
-	bool ShouldRunPostPhysicsTick() const;
+	bool ShouldRunEndPhysicsTick() const;
 
 	// Returns whether we need to run the Cloth Tick or not
 	bool ShouldRunClothTick() const;
 
 	// Handles registering/unregistering the pre cloth tick as it is needed
-	void UpdatePostPhysicsTickRegisteredState();
+	void UpdateEndPhysicsTickRegisteredState();
 
 	// Handles registering/unregistering the cloth tick as it is needed
 	void UpdateClothTickRegisteredState();
diff --git a/Engine/Source/Runtime/Engine/Classes/Engine/DebugCameraController.h b/Engine/Source/Runtime/Engine/Classes/Engine/DebugCameraController.h
--- a/Engine/Source/Runtime/Engine/Classes/Engine/DebugCameraController.h
+++ b/Engine/Source/Runtime/Engine/Classes/Engine/DebugCameraController.h
@@ -79,19 +79,23 @@ public:
 	/** @todo document */
 	class UPlayer* OriginalPlayer;
 
-	/* Allows control over the speed of the spectator pawn. This scales the speed based on the InitialMaxSpeed. */
-	UPROPERTY(EditDefaultsOnly, BlueprintReadWrite, Category="Debug Camera")
+	/** Allows control over the speed of the spectator pawn. This scales the speed based on the InitialMaxSpeed. Use Set Pawn Movement Speed Scale during runtime */
+	UPROPERTY(EditDefaultsOnly, BlueprintReadOnly, Category="Debug Camera")
 	float SpeedScale;
 	
-	/* Initial max speed of the spectator pawn when we start possession. */
+	/** Sets the pawn movement speed scale. */
+	UFUNCTION(BlueprintCallable, Category="Debug Camera")
+	void SetPawnMovementSpeedScale(float NewSpeedScale);
+	
+	/** Initial max speed of the spectator pawn when we start possession. */
 	UPROPERTY(EditDefaultsOnly, BlueprintReadOnly, Category="Debug Camera")
 	float InitialMaxSpeed;
 
-	/* Initial acceleration of the spectator pawn when we start possession. */
+	/** Initial acceleration of the spectator pawn when we start possession. */
 	UPROPERTY(EditDefaultsOnly, BlueprintReadOnly, Category="Debug Camera")
 	float InitialAccel;
 
-	/* Initial deceleration of the spectator pawn when we start possession. */
+	/** Initial deceleration of the spectator pawn when we start possession. */
 	UPROPERTY(EditDefaultsOnly, BlueprintReadOnly, Category="Debug Camera")
 	float InitialDecel;
 
@@ -109,17 +113,17 @@ public:
 	*/
 	UFUNCTION(BlueprintImplementableEvent, meta=(DisplayName = "OnActivate"))
 	void ReceiveOnActivate(class APlayerController* OriginalPC);
-	
+
 	/** Function called on activation debug camera controller */
 	virtual void OnActivate(class APlayerController* OriginalPC);
-
+	
 	/** 
 	* Function called on deactivation of debug camera controller.
 	* @param RestoredPC The Player Controller that the player input is being returned to.
 	*/
 	UFUNCTION(BlueprintImplementableEvent, meta=(DisplayName = "OnDeactivate"))
 	void ReceiveOnDeactivate(class APlayerController* RestoredPC);
-	
+
 	/** Function called on deactivation debug camera controller */
 	virtual void OnDeactivate(class APlayerController* RestoredPC);
 
@@ -146,6 +150,9 @@ protected:
 
 	/**
 	 * Called when an actor has been selected with the primary key (e.g. left mouse button).
+	 *
+	 * The selection trace starts from the center of the debug camera's view.
+	 *
 	 * @param SelectHitLocation The exact world-space location where the selection trace hit the New Selected Actor.
 	 * @param SelectHitNormal The world-space surface normal of the New Selected Actor at the hit location.
 	 */
diff --git a/Engine/Source/Runtime/Engine/Classes/GameFramework/CheatManager.h b/Engine/Source/Runtime/Engine/Classes/GameFramework/CheatManager.h
--- a/Engine/Source/Runtime/Engine/Classes/GameFramework/CheatManager.h
+++ b/Engine/Source/Runtime/Engine/Classes/GameFramework/CheatManager.h
@@ -54,7 +54,11 @@ struct FDebugTraceInfo
 
 };
 
-UCLASS(Within=PlayerController, Blueprintable)
+/** 
+	Cheat Manager is a central blueprint to implement test and debug code and actions that are not to ship with the game.
+	As the Cheat Manager is not instanced in shipping builds, it is for debugging purposes only
+*/
+UCLASS(Blueprintable, Within=PlayerController)
 class ENGINE_API UCheatManager : public UObject
 {
 	GENERATED_UCLASS_BODY()
@@ -63,11 +67,8 @@ class ENGINE_API UCheatManager : public UObject
 	UPROPERTY()
 	class ADebugCameraController* DebugCameraControllerRef;
 
-	/** 
-	* Debug camera - used to have independent camera without stopping gameplay. 
-	* If the Outer PlayerController::DebugCameraClass is valid then it is used instead of default DebugCameraController class. 
-	*/
-	UPROPERTY(Category="Debug Camera", EditDefaultsOnly, BlueprintReadOnly)
+	/** Debug camera - used to have independent camera without stopping gameplay */
+	UPROPERTY(EditDefaultsOnly, BlueprintReadOnly, Category="Debug Camera")
 	TSubclassOf<class ADebugCameraController>  DebugCameraControllerClass;
 
 	// Trace/Sweep debug start
@@ -114,43 +115,43 @@ class ENGINE_API UCheatManager : public UObject
 	int32 CurrentTracePawnIndex;
 
 	/** Pause the game for Delay seconds. */
-	UFUNCTION(exec)
+	UFUNCTION(exec,BlueprintCallable,Category="Cheat Manager")
 	virtual void FreezeFrame(float Delay);
 
 	/* Teleport to surface player is looking at. */
-	UFUNCTION(exec)
+	UFUNCTION(exec,BlueprintCallable,Category="Cheat Manager")
 	virtual void Teleport();
 
 	/* Scale the player's size to be F * default size. */
-	UFUNCTION(exec)
+	UFUNCTION(exec,BlueprintCallable,Category="Cheat Manager")
 	virtual void ChangeSize(float F);
 
 	/** Pawn can fly. */
-	UFUNCTION(exec)
+	UFUNCTION(exec,BlueprintCallable,Category="Cheat Manager")
 	virtual void Fly();
 
 	/** Return to walking movement mode from Fly or Ghost cheat. */
-	UFUNCTION(exec)
+	UFUNCTION(exec,BlueprintCallable,Category="Cheat Manager")
 	virtual void Walk();
 
 	/** Pawn no longer collides with the world, and can fly */
-	UFUNCTION(exec)
+	UFUNCTION(exec,BlueprintCallable,Category="Cheat Manager")
 	virtual void Ghost();
 
 	/** Invulnerability cheat. */
-	UFUNCTION(exec)
+	UFUNCTION(exec,BlueprintCallable,Category="Cheat Manager")
 	virtual void God();
 
 	/** Modify time dilation to change apparent speed of passage of time. e.g. "Slomo 0.1" makes everything move very slowly, while "Slomo 10" makes everything move very fast. */
-	UFUNCTION(exec)
+	UFUNCTION(exec,BlueprintCallable,Category="Cheat Manager")
 	virtual void Slomo(float NewTimeDilation);
 
 	/** Damage the actor you're looking at (sourced from the player). */
-	UFUNCTION(exec)
+	UFUNCTION(exec,BlueprintCallable,Category="Cheat Manager")
 	void DamageTarget(float DamageAmount);
 
 	/** Destroy the actor you're looking at. */
-	UFUNCTION(exec)
+	UFUNCTION(exec,BlueprintCallable,Category="Cheat Manager")
 	virtual void DestroyTarget();
 
 	/** Destroy all actors of class aClass */
@@ -170,7 +171,7 @@ class ENGINE_API UCheatManager : public UObject
 	virtual void Summon(const FString& ClassName);
 
 	/** Freeze everything in the level except for players. */
-	UFUNCTION(exec)
+	UFUNCTION(exec,BlueprintCallable,Category="Cheat Manager")
 	virtual void PlayersOnly();
 
 	/** Make controlled pawn the viewtarget again. */
@@ -350,16 +351,23 @@ class ENGINE_API UCheatManager : public UObject
 	/** streaming level debugging */
 	virtual void SetLevelStreamingStatus(FName PackageName, bool bShouldBeLoaded, bool bShouldBeVisible);
 
-	/** Called when the Cheat Manager is first initialized. */
+	/** 
+	* BP implementable event for when CheatManager is created to allow any needed initialization.  
+	*/
 	UFUNCTION(BlueprintImplementableEvent, meta=(DisplayName = "Init Cheat Manager", Keywords="Begin Play"))
 	void ReceiveInitCheatManager();
 
 	/** 
 	 * Called when CheatManager is created to allow any needed initialization.  
-	 * This is not an actor, so we need a stand in for PostInitializeComponents 
 	 */
 	virtual void InitCheatManager();
 
+	/**
+	* This is the End Play event for the CheatManager
+	*/
+	UFUNCTION(BlueprintImplementableEvent, meta=(DisplayName = "Shutdown", keywords="endplay"))
+	void ReceiveEndPlay();
+	
 	/**
 	* Called before destroying the object.  This is called immediately upon deciding to destroy the object, to allow the object to begin an
 	* asynchronous cleanup process.
@@ -374,9 +382,11 @@ protected:
 	virtual bool DoGameSpecificBugItLog(FOutputDevice& OutputFile) { return true; }
 
 	/** Switch controller to debug camera without locking gameplay and with locking local player controller input */
+	UFUNCTION(BlueprintCallable,Category="Cheat Manager")
 	virtual void EnableDebugCamera();
 	
 	/** Switch controller from debug camera back to normal controller */
+	UFUNCTION(BlueprintCallable,Category="Cheat Manager")
 	virtual void DisableDebugCamera();
 };
 
diff --git a/Engine/Source/Runtime/Engine/Classes/GameFramework/PlayerController.h b/Engine/Source/Runtime/Engine/Classes/GameFramework/PlayerController.h
--- a/Engine/Source/Runtime/Engine/Classes/GameFramework/PlayerController.h
+++ b/Engine/Source/Runtime/Engine/Classes/GameFramework/PlayerController.h
@@ -238,11 +238,11 @@ class ENGINE_API APlayerController : public AController
 	int32 ClientCap;
 	
 	/** Object that manages "cheat" commands.  Not instantiated in shipping builds. */
-	UPROPERTY(transient, BlueprintReadOnly, Category=PlayerController)
+	UPROPERTY(transient, BlueprintReadOnly, Category="Cheat Manager")
 	class UCheatManager* CheatManager;
 	
-	/** class of my CheatManager. */
-	UPROPERTY(EditDefaultsOnly, BlueprintReadOnly, Category=PlayerController)
+	/** Class of my CheatManager.  The Cheat Manager is not created in shipping builds */
+	UPROPERTY(EditDefaultsOnly, BlueprintReadOnly, Category="Cheat Manager")
 	TSubclassOf<class UCheatManager> CheatClass;
 
 	/** Object that manages player input. */
@@ -486,6 +486,9 @@ public:
 	 */
 	bool ProjectWorldLocationToScreenWithDistance(FVector WorldLocation, FVector& ScreenLocation) const;
 	
+	/** Positions the mouse cursor in screen space, in pixels. */
+	UFUNCTION( BlueprintCallable, Category="Game|Player", meta = (DisplayName = "SetMousePosition", Keywords = "mouse" ))
+	void SetMouseLocation( const int X, const int Y );
 	/**
 	  * Updates the rotation of player, based on ControlRotation after RotationInput has been applied.
 	  * This may then be modified by the PlayerCamera, and is passed to Pawn->FaceRotation().
diff --git a/Engine/Source/Runtime/Engine/Classes/Kismet/GameplayStatics.h b/Engine/Source/Runtime/Engine/Classes/Kismet/GameplayStatics.h
--- a/Engine/Source/Runtime/Engine/Classes/Kismet/GameplayStatics.h
+++ b/Engine/Source/Runtime/Engine/Classes/Kismet/GameplayStatics.h
@@ -22,7 +22,7 @@ namespace ESuggestProjVelocityTraceOption
 	{
 		DoNotTrace,
 		TraceFullPath,
-		OnlyTraceWhileAsceding,
+		OnlyTraceWhileAscending,
 	};
 }
 
@@ -875,7 +875,7 @@ class ENGINE_API UGameplayStatics : public UBlueprintFunctionLibrary
 	 * @param Player			Deproject using this player's view.
 	 * @param ScreenPosition	2D screen space to deproject.
 	 * @param WorldPosition		(out) Corresponding 3D position in world space.
-	 * @param WorldDirection	(out) World space direction vector away from the camera at the given 2d poiunt.
+	 * @param WorldDirection	(out) World space direction vector away from the camera at the given 2d point.
 	 */
 	UFUNCTION(BlueprintPure, Category = "Utilities", meta = (Keywords = "unproject"))
 	static bool DeprojectScreenToWorld(APlayerController const* Player, const FVector2D& ScreenPosition, FVector& WorldPosition, FVector& WorldDirection);
diff --git a/Engine/Source/Runtime/Engine/Classes/Sound/AudioSettings.h b/Engine/Source/Runtime/Engine/Classes/Sound/AudioSettings.h
--- a/Engine/Source/Runtime/Engine/Classes/Sound/AudioSettings.h
+++ b/Engine/Source/Runtime/Engine/Classes/Sound/AudioSettings.h
@@ -63,6 +63,10 @@ class ENGINE_API UAudioSettings : public UDeveloperSettings
 	UPROPERTY(config, EditAnywhere, Category="Quality")
 	TArray<FAudioQualitySettings> QualityLevels;
 
+	/** Allows sounds to play at 0 volume. */
+	UPROPERTY(config, EditAnywhere, Category = "Quality")
+	uint32 bAllowVirtualizedSounds:1;
+
 	/**
 	 * The format string to use when generating the filename for contexts within dialogue waves. This must generate unique names for your project.
 	 * Available format markers:
diff --git a/Engine/Source/Runtime/Engine/Classes/Sound/SoundWave.h b/Engine/Source/Runtime/Engine/Classes/Sound/SoundWave.h
--- a/Engine/Source/Runtime/Engine/Classes/Sound/SoundWave.h
+++ b/Engine/Source/Runtime/Engine/Classes/Sound/SoundWave.h
@@ -153,6 +153,10 @@ class ENGINE_API USoundWave : public USoundBase
 	UPROPERTY(EditAnywhere, Category=Subtitles )
 	uint32 bSingleLine:1;
 
+	/** Allows sound to play at 0 volume, otherwise will stop the sound when the sound is silent. */
+	UPROPERTY(EditAnywhere, Category=Sound)
+	uint32 bVirtualizeWhenSilent:1;
+
 	/** Whether this SoundWave was decompressed from OGG. */
 	uint32 bDecompressedFromOgg:1;
 
diff --git a/Engine/Source/Runtime/Engine/Classes/Vehicles/SimpleWheeledVehicleMovementComponent.h b/Engine/Source/Runtime/Engine/Classes/Vehicles/SimpleWheeledVehicleMovementComponent.h
new file mode 100644
--- /dev/null
+++ b/Engine/Source/Runtime/Engine/Classes/Vehicles/SimpleWheeledVehicleMovementComponent.h
@@ -0,0 +1,35 @@
+// Copyright 1998-2016 Epic Games, Inc. All Rights Reserved.
+
+/*
+* Simple n-wheeled vehicle with suspension and tire friction. If you need a motor sim see UWheeledVehicleMovementComponent4W
+*/
+#pragma once
+#include "Vehicles/WheeledVehicleMovementComponent.h"
+#include "SimpleWheeledVehicleMovementComponent.generated.h"
+
+UCLASS(meta = (BlueprintSpawnableComponent), hidecategories = (PlanarMovement, "Components|Movement|Planar", Activation, "Components|Activation"))
+class ENGINE_API USimpleWheeledVehicleMovementComponent : public UWheeledVehicleMovementComponent
+{
+	GENERATED_BODY()
+public:
+	/**Set the brake torque to be applied to a specific wheel*/
+	UFUNCTION(BlueprintCallable, Category = Vehicle)
+	void SetBrakeTorque(float BrakeTorque, int32 WheelIndex);
+
+	/**Set the drive torque to be applied to a specific wheel*/
+	UFUNCTION(BlueprintCallable, Category = Vehicle)
+	void SetDriveTorque(float DriveTorque, int32 WheelIndex);
+
+	/**Set the steer angle (in degrees) to be applied to a specific wheel*/
+	UFUNCTION(BlueprintCallable, Category = Vehicle)
+	void SetSteerAngle(float SteerAngle, int32 WheelIndex);
+
+protected:
+
+#if WITH_VEHICLE
+
+	/** Allocate and setup the PhysX vehicle */
+	virtual void SetupVehicleDrive(physx::PxVehicleWheelsSimData* PWheelsSimData) override;
+
+#endif // WITH_VEHICLE
+};
\ No newline at end of file
diff --git a/Engine/Source/Runtime/Engine/Classes/Vehicles/VehicleAnimInstance.h b/Engine/Source/Runtime/Engine/Classes/Vehicles/VehicleAnimInstance.h
--- a/Engine/Source/Runtime/Engine/Classes/Vehicles/VehicleAnimInstance.h
+++ b/Engine/Source/Runtime/Engine/Classes/Vehicles/VehicleAnimInstance.h
@@ -5,8 +5,51 @@
  *
  */
 #include "Animation/AnimInstance.h"
+#include "Animation/AnimInstanceProxy.h"
 #include "VehicleAnimInstance.generated.h"
 
+struct FWheelAnimData
+{
+	FName BoneName;
+	FRotator RotOffset;
+	FVector LocOffset;
+};
+
+class UWheeledVehicleMovementComponent;
+
+ /** Proxy override for this UAnimInstance-derived class */
+USTRUCT()
+struct ENGINE_API FVehicleAnimInstanceProxy : public FAnimInstanceProxy
+{
+	GENERATED_BODY()
+
+	FVehicleAnimInstanceProxy()
+		: FAnimInstanceProxy()
+	{
+	}
+
+	FVehicleAnimInstanceProxy(UAnimInstance* Instance)
+		: FAnimInstanceProxy(Instance)
+	{
+	}
+
+public:
+
+	void SetWheeledVehicleMovementComponent(const UWheeledVehicleMovementComponent* InWheeledVehicleMovementComponent);
+
+	/** FAnimInstanceProxy interface begin*/
+	virtual void PreUpdate(UAnimInstance* InAnimInstance, float DeltaSeconds) override;
+	/** FAnimInstanceProxy interface end*/
+
+	const TArray<FWheelAnimData>& GetWheelAnimData() const
+	{
+		return WheelInstances;
+	}
+
+private:
+	TArray<FWheelAnimData> WheelInstances;
+};
+
 UCLASS(transient)
 class ENGINE_API UVehicleAnimInstance : public UAnimInstance
 {
@@ -15,6 +58,32 @@ class ENGINE_API UVehicleAnimInstance : public UAnimInstance
 	/** Makes a montage jump to the end of a named section. */
 	UFUNCTION(BlueprintCallable, Category="Animation")
 	class AWheeledVehicle * GetVehicle();
+
+public:
+	TArray<FWheelAnimData> WheelData;
+
+public:
+	void SetWheeledVehicleMovementComponent(const UWheeledVehicleMovementComponent* InWheeledVehicleMovementComponent)
+	{
+		WheeledVehicleMovementComponent = InWheeledVehicleMovementComponent;
+		AnimInstanceProxy.SetWheeledVehicleMovementComponent(InWheeledVehicleMovementComponent);
+	}
+
+	const UWheeledVehicleMovementComponent* GetWheeledVehicleMovementComponent() const
+	{
+		return WheeledVehicleMovementComponent;
+	}
+
+private:
+	/** UAnimInstance interface begin*/
+	virtual FAnimInstanceProxy* CreateAnimInstanceProxy() override;
+	virtual void DestroyAnimInstanceProxy(FAnimInstanceProxy* InProxy) override;
+	/** UAnimInstance interface end*/
+
+	FVehicleAnimInstanceProxy AnimInstanceProxy;
+	
+	UPROPERTY(transient)
+	const UWheeledVehicleMovementComponent* WheeledVehicleMovementComponent;
 };
 
 
diff --git a/Engine/Source/Runtime/Engine/Classes/Vehicles/VehicleWheel.h b/Engine/Source/Runtime/Engine/Classes/Vehicles/VehicleWheel.h
--- a/Engine/Source/Runtime/Engine/Classes/Vehicles/VehicleWheel.h
+++ b/Engine/Source/Runtime/Engine/Classes/Vehicles/VehicleWheel.h
@@ -29,7 +29,7 @@ class ENGINE_API UVehicleWheel : public UObject
 	class UStaticMesh*								CollisionMesh;
 
 	/** If set, shape won't be created, but mapped from chassis mesh */
-	UPROPERTY(EditDefaultsOnly, Category=Shape)
+	UPROPERTY(EditDefaultsOnly, Category=Shape, meta=(DisplayName="UsePhysAssetShape"))
 	bool											bDontCreateShape;
 
 	/** 
@@ -86,7 +86,7 @@ class ENGINE_API UVehicleWheel : public UObject
 	UPROPERTY(EditAnywhere, Category=Tire)
 	float											LongStiffValue;
 
-	/** Vertical offset from vehicle center of mass where suspension forces are applied */
+	/** Vertical offset from where suspension forces are applied (along Z-axis) */
 	UPROPERTY(EditAnywhere, Category=Suspension)
 	float											SuspensionForceOffset;
 
@@ -166,13 +166,13 @@ class ENGINE_API UVehicleWheel : public UObject
 	FVector											Velocity;
 
 	UFUNCTION(BlueprintCallable, Category="Game|Components|WheeledVehicleMovement")
-	float GetSteerAngle();
+	float GetSteerAngle() const;
 
 	UFUNCTION(BlueprintCallable, Category="Game|Components|WheeledVehicleMovement")
-	float GetRotationAngle();
+	float GetRotationAngle() const;
 
 	UFUNCTION(BlueprintCallable, Category="Game|Components|WheeledVehicleMovement")
-	float GetSuspensionOffset();
+	float GetSuspensionOffset() const;
 
 #if WITH_PHYSX
 
diff --git a/Engine/Source/Runtime/Engine/Classes/Vehicles/WheeledVehicleMovementComponent.h b/Engine/Source/Runtime/Engine/Classes/Vehicles/WheeledVehicleMovementComponent.h
--- a/Engine/Source/Runtime/Engine/Classes/Vehicles/WheeledVehicleMovementComponent.h
+++ b/Engine/Source/Runtime/Engine/Classes/Vehicles/WheeledVehicleMovementComponent.h
@@ -159,10 +159,14 @@ struct ENGINE_API FVehicleInputRate
  * Component to handle the vehicle simulation for an actor.
  */
 UCLASS(Abstract, hidecategories=(PlanarMovement, "Components|Movement|Planar", Activation, "Components|Activation"))
-class ENGINE_API UWheeledVehicleMovementComponent : public UPawnMovementComponent, public IRVOAvoidanceInterface
+class ENGINE_API UWheeledVehicleMovementComponent : public UMovementComponent, public IRVOAvoidanceInterface
 {
 	GENERATED_UCLASS_BODY()
 
+	// Supports the old (before 4.14) way of applying spring forces. We used to offset from the vehicle center of mass instead of the spring location center of mass. You should only use this for existing content that hasn't been re-tuned
+	UPROPERTY(EditAnywhere, Category = VehicleSetup)
+	uint8 bDeprecatedSpringOffsetMode : 1;
+
 	/** Wheels to create */
 	UPROPERTY(EditAnywhere, Category=VehicleSetup)
 	TArray<FWheelSetup> WheelSetups;
@@ -184,6 +188,10 @@ class ENGINE_API UWheeledVehicleMovementComponent : public UPawnMovementComponen
 	UPROPERTY(EditAnywhere, Category = VehicleSetup, meta = (ClampMin = "0.01", UIMin = "0.01"))
 	float ChassisHeight;
 
+	/** If true, the brake and reverse controls will behave in a more arcade fashion where holding reverse also functions as brake. For a more realistic approach turn this off*/
+	UPROPERTY(EditAnywhere, Category = VehicleSetup)
+	bool bReverseAsBrake;
+
 	// Drag area in cm^2
 	UPROPERTY(transient)
 	float DragArea;
@@ -251,6 +259,10 @@ class ENGINE_API UWheeledVehicleMovementComponent : public UPawnMovementComponen
 	bool CheckSlipThreshold(float AbsLongSlipThreshold, float AbsLatSlipThreshold) const;
 	float GetMaxSpringForce() const;
 
+	/** UObject interface */
+	virtual void Serialize(FArchive& Ar) override;
+	/** End UObject interface*/
+
 #if WITH_VEHICLE
 
 	// The instanced PhysX vehicle
@@ -304,6 +316,10 @@ class ENGINE_API UWheeledVehicleMovementComponent : public UPawnMovementComponen
 	/** Set the user input for the vehicle throttle */
 	UFUNCTION(BlueprintCallable, Category="Game|Components|WheeledVehicleMovement")
 	void SetThrottleInput(float Throttle);
+
+	/** Set the user input for the vehicle Brake */
+	UFUNCTION(BlueprintCallable, Category = "Game|Components|WheeledVehicleMovement")
+	void SetBrakeInput(float Brake);
 	
 	/** Set the user input for the vehicle steering */
 	UFUNCTION(BlueprintCallable, Category="Game|Components|WheeledVehicleMovement")
@@ -445,6 +461,10 @@ protected:
 	UPROPERTY(Transient)
 	float RawThrottleInput;
 
+	// What the player has the brake set to. Range -1...1
+	UPROPERTY(Transient)
+	float RawBrakeInput;
+
 	// True if the player is holding the handbrake
 	UPROPERTY(Transient)
 	uint32 bRawHandbrakeInput : 1;
@@ -575,7 +595,10 @@ protected:
 	virtual void UpdateSimulation( float DeltaTime );
 
 	/** Allocate and setup the PhysX vehicle */
-	virtual void SetupVehicle() {};
+	virtual void SetupVehicle();
+
+	/** Create the specific vehicle drive (4w drive vs tank etc...)*/
+	virtual void SetupVehicleDrive(physx::PxVehicleWheelsSimData* PWheelsSimData);
 
 	/** Do some final setup after the physx vehicle gets created */
 	virtual void PostSetupVehicle();
diff --git a/Engine/Source/Runtime/Engine/Classes/Vehicles/WheeledVehicleMovementComponent4W.h b/Engine/Source/Runtime/Engine/Classes/Vehicles/WheeledVehicleMovementComponent4W.h
--- a/Engine/Source/Runtime/Engine/Classes/Vehicles/WheeledVehicleMovementComponent4W.h
+++ b/Engine/Source/Runtime/Engine/Classes/Vehicles/WheeledVehicleMovementComponent4W.h
@@ -145,7 +145,7 @@ struct FVehicleTransmissionData
 	float ClutchStrength;
 };
 
-UCLASS()
+UCLASS(meta = (BlueprintSpawnableComponent), hidecategories = (PlanarMovement, "Components|Movement|Planar", Activation, "Components|Activation"))
 class ENGINE_API UWheeledVehicleMovementComponent4W : public UWheeledVehicleMovementComponent
 {
 	GENERATED_UCLASS_BODY()
@@ -182,7 +182,7 @@ protected:
 #if WITH_VEHICLE
 
 	/** Allocate and setup the PhysX vehicle */
-	virtual void SetupVehicle() override;
+	virtual void SetupVehicleDrive(physx::PxVehicleWheelsSimData* PWheelsSimData) override;
 
 	virtual void UpdateSimulation(float DeltaTime) override;
 
diff --git a/Engine/Source/Runtime/Engine/Private/AI/Navigation/RecastNavMeshGenerator.cpp b/Engine/Source/Runtime/Engine/Private/AI/Navigation/RecastNavMeshGenerator.cpp
--- a/Engine/Source/Runtime/Engine/Private/AI/Navigation/RecastNavMeshGenerator.cpp
+++ b/Engine/Source/Runtime/Engine/Private/AI/Navigation/RecastNavMeshGenerator.cpp
@@ -492,6 +492,14 @@ void ExportPxHeightField(PxHeightField const * const HeightField, const FTransfo
 	{
 		for (int32 X = 0; X < NumCols - 1; X++)
 		{
+			const int32 SampleIdx = (bMirrored ? X : (NumCols - X - 1))*NumCols + Y;
+			const PxHeightFieldSample& Sample = HFSamples[SampleIdx];
+			const bool bIsHole = (Sample.materialIndex0 == PxHeightFieldMaterial::eHOLE);
+			if (bIsHole)
+			{
+				continue;
+			}
+
 			const int32 I00 = X + 0 + (Y + 0)*NumCols;
 			int32 I01 = X + 0 + (Y + 1)*NumCols;
 			int32 I10 = X + 1 + (Y + 0)*NumCols;
@@ -502,20 +510,13 @@ void ExportPxHeightField(PxHeightField const * const HeightField, const FTransfo
 				Swap(I01, I10);
 			}
 
-			const int32 SampleIdx = (NumCols - X - 1)*NumCols + Y;
-			const PxHeightFieldSample& Sample = HFSamples[SampleIdx];
-			const bool HoleQuad = (Sample.materialIndex0 == PxHeightFieldMaterial::eHOLE);
+			IndexBuffer.Add(VertOffset + I00);
+			IndexBuffer.Add(VertOffset + I11);
+			IndexBuffer.Add(VertOffset + I10);
 
-			if (!HoleQuad)
-			{
-				IndexBuffer.Add(VertOffset + I00);
-				IndexBuffer.Add(VertOffset + I11);
-				IndexBuffer.Add(VertOffset + I10);
-
-				IndexBuffer.Add(VertOffset + I00);
-				IndexBuffer.Add(VertOffset + I01);
-				IndexBuffer.Add(VertOffset + I11);
-			}
+			IndexBuffer.Add(VertOffset + I00);
+			IndexBuffer.Add(VertOffset + I01);
+			IndexBuffer.Add(VertOffset + I11);
 		}
 	}
 }
@@ -532,76 +533,75 @@ void ExportHeightFieldSlice(const FNavHeightfieldSamples& PrefetchedHeightfieldS
 
 	// calculate the actual start and number of columns we want
 	const FBox LocalBox = SliceBox.TransformBy(LocalToWorld.Inverse());
+	const bool bMirrored = (LocalToWorld.GetDeterminant() < 0.f);
 
-	const int32 StartingRow = FMath::Max(FMath::FloorToInt(LocalBox.Min.Y) - 1, 0);
-	const int32 StartingColumn = FMath::Max(FMath::FloorToInt(LocalBox.Min.X) - 1, 0);
-	const int32 RowLimit = FMath::Min(FMath::CeilToInt(LocalBox.Max.Y) + 1, NumRows);
-	const int32 ColumnLimit = FMath::Min(FMath::CeilToInt(LocalBox.Max.X) + 1, NumCols);
-	
-	const int32 VertexCount = (RowLimit - StartingRow)*(ColumnLimit - StartingColumn);
+	const int32 MinX = FMath::Clamp(FMath::FloorToInt(LocalBox.Min.X) - 1, 0, NumCols);
+	const int32 MinY = FMath::Clamp(FMath::FloorToInt(LocalBox.Min.Y) - 1, 0, NumRows);
+	const int32 MaxX = FMath::Clamp(FMath::CeilToInt(LocalBox.Max.X) + 1, 0, NumCols);
+	const int32 MaxY = FMath::Clamp(FMath::CeilToInt(LocalBox.Max.Y) + 1, 0, NumRows);
+	const int32 SizeX = MaxX - MinX;
+	const int32 SizeY = MaxY - MinY;
 
-	const int32 VertOffset = VertexBuffer.Num() / 3;
-	const int32 RowCount = RowLimit - StartingRow;
-	const int32 CollumnCount = ColumnLimit - StartingColumn;
-	const int32 NumQuads = (RowCount - 1)*(CollumnCount - 1);
+	if (SizeX <= 0 || SizeY <= 0)
+	{
+		// slice is outside bounds, skip
+		return;
+	}
 
+	const int32 VertOffset = VertexBuffer.Num() / 3;
+	const int32 NumVerts = SizeX * SizeY;
+	const int32 NumQuads = (SizeX - 1) * (SizeY - 1);
 	{
 		QUICK_SCOPE_CYCLE_COUNTER(STAT_RecastGeometryExport_AllocatingMemory);
-		VertexBuffer.Reserve(VertexBuffer.Num() + VertexCount * 3);
-		IndexBuffer.Reserve(IndexBuffer.Num() + NumQuads * 6);
+		VertexBuffer.Reserve(VertexBuffer.Num() + NumVerts * 3);
+		IndexBuffer.Reserve(IndexBuffer.Num() + NumQuads * 3 * 2);
 	}
 
-	const bool bMirrored = (LocalToWorld.GetDeterminant() < 0.f);
-	
-	TNavStatArray<float> TmpVertexBuffer;
-	TmpVertexBuffer.AddUninitialized(VertexCount * 3);
-	FVector* Vertex = ((FVector*)TmpVertexBuffer.GetData());
-
-	for (int32 Y = StartingRow; Y < RowLimit; Y++)
+	for (int32 IdxY = 0; IdxY < SizeY; IdxY++)
 	{
-		for (int32 X = StartingColumn; X < ColumnLimit; X++)
+		for (int32 IdxX = 0; IdxX < SizeX; IdxX++)
 		{
-			const int32 SampleIdx = (bMirrored ? X : (NumCols - X - 1))*NumCols + Y;
+			const int32 CoordX = IdxX + MinX;
+			const int32 CoordY = IdxY + MinY;
+			const int32 SampleIdx = ((bMirrored ? CoordX : (NumCols - CoordX - 1)) * NumCols) + CoordY;
 
-			const FVector UnrealCoords = LocalToWorld.TransformPosition(FVector(X, Y, PrefetchedHeightfieldSamples.Heights[SampleIdx]));
-			*Vertex++ = UnrealCoords;
+			const FVector UnrealCoords = LocalToWorld.TransformPosition(FVector(CoordX, CoordY, PrefetchedHeightfieldSamples.Heights[SampleIdx]));
+			VertexBuffer.Add(UnrealCoords.X);
+			VertexBuffer.Add(UnrealCoords.Y);
+			VertexBuffer.Add(UnrealCoords.Z);
 		}
 	}
-	VertexBuffer.Append(TmpVertexBuffer);
 
-	for (int32 Y = 0; Y < RowCount - 1; Y++)
+	for (int32 IdxY = 0; IdxY < SizeY - 1; IdxY++)
 	{
-		for (int32 X = 0; X < CollumnCount - 1; X++)
+		for (int32 IdxX = 0; IdxX < SizeX - 1; IdxX++)
 		{
-			const int32 SampleIdx = (NumCols - X - 1)*NumCols + Y;
-			if (PrefetchedHeightfieldSamples.Holes[SampleIdx] == false)
+			const int32 CoordX = IdxX + MinX;
+			const int32 CoordY = IdxY + MinY;
+			const int32 SampleIdx = ((bMirrored ? CoordX : (NumCols - CoordX - 1)) * NumCols) + CoordY;
+
+			const bool bIsHole = PrefetchedHeightfieldSamples.Holes[SampleIdx];
+			if (bIsHole)
 			{
-				const int32 I00 = X + 0 + (Y + 0)*CollumnCount;
-				const int32 I01 = X + 0 + (Y + 1)*CollumnCount;
-				const int32 I10 = X + 1 + (Y + 0)*CollumnCount;
-				const int32 I11 = X + 1 + (Y + 1)*CollumnCount;
+				continue;
+			}
 
-				if (bMirrored == false)
-				{
-					IndexBuffer.Add(VertOffset + I00);
-					IndexBuffer.Add(VertOffset + I11);
-					IndexBuffer.Add(VertOffset + I10);
+			const int32 I00 = (IdxX + 0) + (IdxY + 0)*SizeX;
+			int32 I01 = (IdxX + 0) + (IdxY + 1)*SizeX;
+			int32 I10 = (IdxX + 1) + (IdxY + 0)*SizeX;
+			const int32 I11 = (IdxX + 1) + (IdxY + 1)*SizeX;
+			if (bMirrored)
+			{
+				Swap(I01, I10);
+			}
 
-					IndexBuffer.Add(VertOffset + I00);
-					IndexBuffer.Add(VertOffset + I01);
-					IndexBuffer.Add(VertOffset + I11);
-				}
-				else
-				{
-					IndexBuffer.Add(VertOffset + I00);
-					IndexBuffer.Add(VertOffset + I01);
-					IndexBuffer.Add(VertOffset + I11);
+			IndexBuffer.Add(VertOffset + I00);
+			IndexBuffer.Add(VertOffset + I11);
+			IndexBuffer.Add(VertOffset + I10);
 
-					IndexBuffer.Add(VertOffset + I00);
-					IndexBuffer.Add(VertOffset + I11);
-					IndexBuffer.Add(VertOffset + I10);
-				}
-			}
+			IndexBuffer.Add(VertOffset + I00);
+			IndexBuffer.Add(VertOffset + I01);
+			IndexBuffer.Add(VertOffset + I11);
 		}
 	}
 }
diff --git a/Engine/Source/Runtime/Engine/Private/ActiveSound.cpp b/Engine/Source/Runtime/Engine/Private/ActiveSound.cpp
--- a/Engine/Source/Runtime/Engine/Private/ActiveSound.cpp
+++ b/Engine/Source/Runtime/Engine/Private/ActiveSound.cpp
@@ -48,7 +48,6 @@ FActiveSound::FActiveSound()
 #endif
 	, bEnableLowPassFilter(false)
 	, bOcclusionAsyncTrace(true)
-	, bIsAudible(true)
 	, UserIndex(0)
 	, bIsOccluded(false)
 	, bAsyncOcclusionPending(false)
@@ -262,9 +261,6 @@ void FActiveSound::UpdateWaveInstances( TArray<FWaveInstance*> &InWaveInstances,
 	// The apparent max distance factors the actual max distance of the sound scaled with the distance scale due to focus effects
 	float ApparentMaxDistance = MaxDistance * FocusDistanceScale;
 
-	// Update whether or not his sound is out of range
-	bIsAudible = AudioDevice->LocationIsAudible(Transform.GetTranslation(), ClosestListenerPtr->Transform, ApparentMaxDistance);
-
 	FSoundParseParameters ParseParams;
 	ParseParams.Transform = Transform;
 	ParseParams.StartTime = RequestedStartTime;
@@ -276,15 +272,12 @@ void FActiveSound::UpdateWaveInstances( TArray<FWaveInstance*> &InWaveInstances,
 	UpdateAdjustVolumeMultiplier(DeltaTime);
 
 	// If the sound is a preview sound, then ignore the transient master volume and application volume
-	float MasterVolume = AudioDevice->GetTransientMasterVolume(); 
-	float ApplicationVolume = FApp::GetVolumeMultiplier();
-	if (bIsPreviewSound)
+	if (!bIsPreviewSound)
 	{
-		MasterVolume = 1.0f;
-		ApplicationVolume = 1.0f;
+		ParseParams.VolumeApp = AudioDevice->GetTransientMasterVolume() * FApp::GetVolumeMultiplier();
 	}
 
-	ParseParams.VolumeMultiplier = VolumeMultiplier * Sound->GetVolumeMultiplier() * CurrentAdjustVolumeMultiplier * MasterVolume * ApplicationVolume * ConcurrencyVolumeScale;
+	ParseParams.VolumeMultiplier = VolumeMultiplier * Sound->GetVolumeMultiplier() * CurrentAdjustVolumeMultiplier * ConcurrencyVolumeScale;
 
 	ParseParams.Priority = Priority;
 	ParseParams.Pitch *= PitchMultiplier * Sound->GetPitchMultiplier();
@@ -342,7 +335,7 @@ void FActiveSound::UpdateWaveInstances( TArray<FWaveInstance*> &InWaveInstances,
 			VolumeConcurrency = 0.0f;
 			for (const FWaveInstance* WaveInstance : ThisSoundsWaveInstances)
 			{
-				const float WaveInstanceVolume = WaveInstance->GetActualVolume();
+				const float WaveInstanceVolume = WaveInstance->GetVolume();
 				if (WaveInstanceVolume > VolumeConcurrency)
 				{
 					VolumeConcurrency = WaveInstanceVolume;
diff --git a/Engine/Source/Runtime/Engine/Private/Animation/Skeleton.cpp b/Engine/Source/Runtime/Engine/Private/Animation/Skeleton.cpp
--- a/Engine/Source/Runtime/Engine/Private/Animation/Skeleton.cpp
+++ b/Engine/Source/Runtime/Engine/Private/Animation/Skeleton.cpp
@@ -1236,7 +1236,12 @@ bool USkeleton::VerifySmartNameInternal(const FName&  ContainerName, FSmartName&
 
 #if WITH_EDITOR
 		// if I find the name, fill up the data
-		if (FillSmartNameByDisplayName(Mapping, DisplayName, InOutSmartName) == false)
+		// look for same guid, the name might have been changed
+		if (Mapping->GetNameByGuid(InOutSmartName.Guid, DisplayName))
+		{
+			ensureAlways(FillSmartNameByDisplayName(Mapping, DisplayName, InOutSmartName));
+		}
+		else if (FillSmartNameByDisplayName(Mapping, DisplayName, InOutSmartName) == false)
 		{
 			// look for same guid, the name might have been changed
 			if (Mapping->GetNameByGuid(InOutSmartName.Guid, DisplayName))
diff --git a/Engine/Source/Runtime/Engine/Private/Audio.cpp b/Engine/Source/Runtime/Engine/Private/Audio.cpp
--- a/Engine/Source/Runtime/Engine/Private/Audio.cpp
+++ b/Engine/Source/Runtime/Engine/Private/Audio.cpp
@@ -174,7 +174,7 @@ FString FSoundSource::Describe(bool bUseLongName)
 {
 	return FString::Printf(TEXT("Wave: %s, Volume: %6.2f, Owner: %s"), 
 		bUseLongName ? *WaveInstance->WaveData->GetPathName() : *WaveInstance->WaveData->GetName(),
-		WaveInstance->GetActualVolume(), 
+		WaveInstance->GetVolume(), 
 		WaveInstance->ActiveSound ? *WaveInstance->ActiveSound->GetOwnerName() : TEXT("None"));
 }
 
@@ -561,6 +561,7 @@ FWaveInstance::FWaveInstance( FActiveSound* InActiveSound )
 ,	ActiveSound( InActiveSound )
 ,	Volume( 0.0f )
 ,	VolumeMultiplier( 1.0f )
+,	VolumeApp( 1.0f )
 ,	Priority( 1.0f )
 ,	VoiceCenterChannelVolume( 0.0f )
 ,	RadioFilterVolume( 0.0f )
@@ -659,9 +660,15 @@ void FWaveInstance::AddReferencedObjects( FReferenceCollector& Collector )
 
 float FWaveInstance::GetActualVolume() const
 {
-	return Volume * VolumeMultiplier;
+	// Include all volumes 
+	return GetVolume() * VolumeApp;
 }
 
+float FWaveInstance::GetVolume() const
+{
+	// Include all volumes 
+	return Volume * VolumeMultiplier;
+}
 
 bool FWaveInstance::ShouldStopDueToMaxConcurrency() const
 {
@@ -671,7 +678,7 @@ bool FWaveInstance::ShouldStopDueToMaxConcurrency() const
 float FWaveInstance::GetVolumeWeightedPriority() const
 {
 	// This will result in zero-volume sounds still able to be sorted due to priority but give non-zero volumes higher priority than 0 volumes
-	float ActualVolume = GetActualVolume();
+	float ActualVolume = GetVolume();
 	if (ActualVolume > 0.0f)
 	{
 		return ActualVolume * Priority;
diff --git a/Engine/Source/Runtime/Engine/Private/AudioDevice.cpp b/Engine/Source/Runtime/Engine/Private/AudioDevice.cpp
--- a/Engine/Source/Runtime/Engine/Private/AudioDevice.cpp
+++ b/Engine/Source/Runtime/Engine/Private/AudioDevice.cpp
@@ -106,6 +106,7 @@ FAudioDevice::FAudioDevice()
 	, bIsDeviceMuted(false)
 	, bIsInitialized(false)
 	, bHasActivatedReverb(false)
+	, bAllowVirtualizedSounds(true)
 #if !UE_BUILD_SHIPPING
 	, RequestedAudioStats(0)
 #endif
@@ -141,6 +142,8 @@ bool FAudioDevice::Init(int32 InMaxChannels)
 		PlatformAudioHeadroom = FMath::Pow(10.0f, Headroom / 20.0f);
 	}
 
+	bAllowVirtualizedSounds = GetDefault<UAudioSettings>()->bAllowVirtualizedSounds;
+
 	const FStringAssetReference DefaultBaseSoundMixName = GetDefault<UAudioSettings>()->DefaultBaseSoundMix;
 	if (DefaultBaseSoundMixName.IsValid())
 	{
@@ -383,7 +386,7 @@ bool FAudioDevice::HandleListWavesCommand(const TCHAR* Cmd, FOutputDevice& Ar)
 		FSoundSource* Source = WaveInstanceSourceMap.FindRef(WaveInstance);
 		UAudioComponent* AudioComponent = UAudioComponent::GetAudioComponentFromID(WaveInstance->ActiveSound->GetAudioComponentID());
 		AActor* SoundOwner = AudioComponent ? AudioComponent->GetOwner() : nullptr;
-		Ar.Logf(TEXT("%4i.    %s %6.2f %6.2f  %s   %s"), InstanceIndex, Source ? TEXT("Yes") : TEXT(" No"), WaveInstance->ActiveSound->PlaybackTime, WaveInstance->GetActualVolume(), *WaveInstance->WaveData->GetPathName(), SoundOwner ? *SoundOwner->GetName() : TEXT("None"));
+		Ar.Logf(TEXT("%4i.    %s %6.2f %6.2f  %s   %s"), InstanceIndex, Source ? TEXT("Yes") : TEXT(" No"), WaveInstance->ActiveSound->PlaybackTime, WaveInstance->GetVolume(), *WaveInstance->WaveData->GetPathName(), SoundOwner ? *SoundOwner->GetName() : TEXT("None"));
 	}
 
 	Ar.Logf(TEXT("Total: %i"), WaveInstances.Num()-FirstActiveIndex);
@@ -1581,7 +1584,7 @@ void FAudioDevice::UpdatePassiveSoundMixModifiers(TArray<FWaveInstance*>& WaveIn
 			USoundClass* SoundClass = WaveInstance->SoundClass;
 			if (SoundClass) 
 			{
-				const float WaveInstanceActualVolume = WaveInstance->GetActualVolume();
+				const float WaveInstanceActualVolume = WaveInstance->GetVolume();
 				// Check each SoundMix individually for volume levels
 				for (const FPassiveSoundMixModifier& PassiveSoundMixModifier : SoundClass->PassiveSoundMixModifiers)
 				{
@@ -2657,7 +2660,7 @@ void FAudioDevice::StopSources(TArray<FWaveInstance*>& WaveInstances, int32 Firs
 			Source->LastUpdate = CurrentTick;
 
 			// If they are still audible, mark them as such
-			float VolumeWeightedPriority = WaveInstance->GetActualVolume();
+			float VolumeWeightedPriority = WaveInstance->GetVolume();
 			if (VolumeWeightedPriority > 0.0f)
 			{
 				Source->LastHeardUpdate = CurrentTick;
@@ -2713,7 +2716,7 @@ void FAudioDevice::StopSources(TArray<FWaveInstance*>& WaveInstances, int32 Firs
 	for (int32 InstanceIndex = 0; InstanceIndex < FirstActiveIndex; InstanceIndex++)
 	{
 		FWaveInstance* WaveInstance = WaveInstances[ InstanceIndex ];
-		if (WaveInstance->GetActualVolume() > 0.1f)
+		if (WaveInstance->GetVolume() > 0.1f) 
 		{
 			AudibleInactiveSounds++;
 		}
@@ -3016,7 +3019,7 @@ void FAudioDevice::SendUpdateResultsToGameThread(const int32 FirstActiveIndex)
 				FAudioStats::FStatWaveInstanceInfo WaveInstanceInfo;
 				FSoundSource* Source = WaveInstanceSourceMap.FindRef(WaveInstance);
 				WaveInstanceInfo.Description = Source ? Source->Describe((RequestedAudioStats & ERequestedAudioStats::LongSoundNames) != 0) : FString(TEXT("No source"));
-				WaveInstanceInfo.ActualVolume = WaveInstance->GetActualVolume();
+				WaveInstanceInfo.ActualVolume = WaveInstance->GetVolume();
 				WaveInstanceInfo.InstanceIndex = InstanceIndex;
 				StatSoundInfos[*SoundInfoIndex].WaveInstanceInfos.Add(MoveTemp(WaveInstanceInfo));
 			}
@@ -4154,3 +4157,41 @@ void FAudioDevice::SetPlatformAudioHeadroom(const float InPlatformHeadRoom)
 
 	PlatformAudioHeadroom = InPlatformHeadRoom;
 }
+
+#if !UE_BUILD_SHIPPING
+void HandleDumpActiveSounds(UWorld* World)
+{
+	if (FAudioDevice* AudioDevice = (World ? World->GetAudioDevice() : GEngine->GetMainAudioDevice()))
+	{
+		AudioDevice->DumpActiveSounds();
+	}
+}
+
+static FAutoConsoleCommandWithWorld DumpActiveSounds(TEXT("Audio.DumpActiveSounds"), TEXT("Outputs data about all the currently active sounds."), FConsoleCommandWithWorldDelegate::CreateStatic(&HandleDumpActiveSounds), ECVF_Cheat);
+
+void FAudioDevice::DumpActiveSounds() const
+{
+	check(IsInGameThread());
+
+	FAudioThreadSuspendContext SuspendAudio;
+
+	UE_LOG(LogAudio, Display, TEXT("Active Sound Count: %d"), ActiveSounds.Num());
+	UE_LOG(LogAudio, Display, TEXT("------------------------"), ActiveSounds.Num());
+
+	for (const FActiveSound* ActiveSound : ActiveSounds)
+	{
+		if (ActiveSound)
+		{
+			UE_LOG(LogAudio, Display, TEXT("%s (%.3g) - %s"), *ActiveSound->GetSound()->GetName(), ActiveSound->GetSound()->GetDuration(), *ActiveSound->GetAudioComponentName());
+
+			for (const TPair<UPTRINT, FWaveInstance*>& WaveInstancePair : ActiveSound->WaveInstances)
+			{
+				const FWaveInstance* WaveInstance = WaveInstancePair.Value;
+				UE_LOG(LogAudio, Display, TEXT("   %s (%.3g) (%d) - %.3g"), *WaveInstance->GetName(), WaveInstance->WaveData->GetDuration(), WaveInstance->WaveData->GetResourceSize(EResourceSizeMode::Inclusive), WaveInstance->GetActualVolume());
+			}
+		}
+	}
+
+
+}
+#endif
diff --git a/Engine/Source/Runtime/Engine/Private/Components/MovementComponent.cpp b/Engine/Source/Runtime/Engine/Private/Components/MovementComponent.cpp
--- a/Engine/Source/Runtime/Engine/Private/Components/MovementComponent.cpp
+++ b/Engine/Source/Runtime/Engine/Private/Components/MovementComponent.cpp
@@ -89,22 +89,10 @@ void UMovementComponent::InitializeComponent()
 	TGuardValue<bool> InInitializeComponentGuard(bInInitializeComponent, true);
 	Super::InitializeComponent();
 
-	USceneComponent* NewUpdatedComponent = NULL;
-	if (UpdatedComponent != NULL)
-	{
-		NewUpdatedComponent = UpdatedComponent;
-	}
-	else if (bAutoRegisterUpdatedComponent)
+	if (bSnapToPlaneAtStart)
 	{
-		// Auto-register owner's root component if found.
-		AActor* MyActor = GetOwner();
-		if (MyActor)
-		{
-			NewUpdatedComponent = MyActor->GetRootComponent();
-		}
+		SnapUpdatedComponentToPlane();
 	}
-
-	SetUpdatedComponent(NewUpdatedComponent);
 }
 
 
@@ -126,11 +114,19 @@ void UMovementComponent::OnRegister()
 		PlaneConstraintNormal = PlaneConstraintNormal.GetSafeNormal();
 	}
 
-	if (bSnapToPlaneAtStart)
+	USceneComponent* NewUpdatedComponent = UpdatedComponent;
+	if (!UpdatedComponent && bAutoRegisterUpdatedComponent)
 	{
-		SnapUpdatedComponentToPlane();
+		// Auto-register owner's root component if found.
+		AActor* MyActor = GetOwner();
+		if (MyActor)
+		{
+			NewUpdatedComponent = MyActor->GetRootComponent();
+		}
 	}
 
+	SetUpdatedComponent(NewUpdatedComponent);
+
 #if WITH_EDITOR
 	// Reset so next PIE session warns.
 	bEditorWarnedStaticMobilityMove = false;
diff --git a/Engine/Source/Runtime/Engine/Private/Components/PoseableMeshComponent.cpp b/Engine/Source/Runtime/Engine/Private/Components/PoseableMeshComponent.cpp
--- a/Engine/Source/Runtime/Engine/Private/Components/PoseableMeshComponent.cpp
+++ b/Engine/Source/Runtime/Engine/Private/Components/PoseableMeshComponent.cpp
@@ -65,6 +65,7 @@ void UPoseableMeshComponent::RefreshBoneTransforms(FActorComponentTickFunction*
 	FillComponentSpaceTransforms();
 	FinalizeBoneTransform();
 
+	UpdateChildTransforms();
 	MarkRenderDynamicDataDirty();
 }
 
@@ -154,7 +155,7 @@ void UPoseableMeshComponent::SetBoneTransformByName(FName BoneName, const FTrans
 			}
 
 			// Need to send new state to render thread
-			MarkRenderDynamicDataDirty();
+			RefreshBoneTransforms();
 		}
 	}
 }
diff --git a/Engine/Source/Runtime/Engine/Private/Components/SkeletalMeshComponent.cpp b/Engine/Source/Runtime/Engine/Private/Components/SkeletalMeshComponent.cpp
--- a/Engine/Source/Runtime/Engine/Private/Components/SkeletalMeshComponent.cpp
+++ b/Engine/Source/Runtime/Engine/Private/Components/SkeletalMeshComponent.cpp
@@ -158,9 +158,9 @@ USkeletalMeshComponent::USkeletalMeshComponent(const FObjectInitializer& ObjectI
 	bGenerateOverlapEvents = false;
 	LineCheckBoundsScale = FVector(1.0f, 1.0f, 1.0f);
 
-	PostPhysicsTickFunction.TickGroup = TG_PostPhysics;
-	PostPhysicsTickFunction.bCanEverTick = true;
-	PostPhysicsTickFunction.bStartWithTickEnabled = true;
+	EndPhysicsTickFunction.TickGroup = TG_EndPhysics;
+	EndPhysicsTickFunction.bCanEverTick = true;
+	EndPhysicsTickFunction.bStartWithTickEnabled = true;
 
 	ClothTickFunction.TickGroup = TG_PrePhysics;
 	ClothTickFunction.EndTickGroup = TG_PostPhysics;
@@ -209,34 +209,30 @@ void USkeletalMeshComponent::RegisterComponentTickFunctions(bool bRegister)
 {
 	Super::RegisterComponentTickFunctions(bRegister);
 
-	UpdatePostPhysicsTickRegisteredState();
+	UpdateEndPhysicsTickRegisteredState();
 	UpdateClothTickRegisteredState();
 }
 
-void USkeletalMeshComponent::RegisterPostPhysicsTick(bool bRegister)
+void USkeletalMeshComponent::RegisterEndPhysicsTick(bool bRegister)
 {
-	if (bRegister != PostPhysicsTickFunction.IsTickFunctionRegistered())
+	if (bRegister != EndPhysicsTickFunction.IsTickFunctionRegistered())
 	{
 		if (bRegister)
 		{
-			if (SetupActorComponentTickFunction(&PostPhysicsTickFunction))
+			if (SetupActorComponentTickFunction(&EndPhysicsTickFunction))
 			{
-				PostPhysicsTickFunction.Target = this;
-				// Set a prereq for the pre cloth tick to happen after physics is finished
+				EndPhysicsTickFunction.Target = this;
+				// Make sure our EndPhysicsTick gets called after physics simulation is finished
 				UWorld* World = GetWorld();
 				if (World != nullptr)
 				{
-					PostPhysicsTickFunction.AddPrerequisite(World, World->EndPhysicsTickFunction);
+					EndPhysicsTickFunction.AddPrerequisite(World, World->EndPhysicsTickFunction);
 				}
-
-				// 4.11.2 Hack (UE-24725): set up this tick prereq in case this other function is used.
-				// It's usually not, but our tick function is private and if someone else needs to set up a prereq they can use the public one in the base class.
-				PostPhysicsComponentTick.AddPrerequisite(this, PostPhysicsTickFunction);
 			}
 		}
 		else
 		{
-			PostPhysicsTickFunction.UnRegisterTickFunction();
+			EndPhysicsTickFunction.UnRegisterTickFunction();
 		}
 	}
 }
@@ -251,7 +247,7 @@ void USkeletalMeshComponent::RegisterClothTick(bool bRegister)
 			{
 				ClothTickFunction.Target = this;
 				ClothTickFunction.AddPrerequisite(this, PrimaryComponentTick);
-				ClothTickFunction.AddPrerequisite(this, PostPhysicsTickFunction);	//If this tick function is running it means that we are doing physics blending so we should wait for its results
+				ClothTickFunction.AddPrerequisite(this, EndPhysicsTickFunction);	//If this tick function is running it means that we are doing physics blending so we should wait for its results
 			}
 		}
 		else
@@ -261,15 +257,15 @@ void USkeletalMeshComponent::RegisterClothTick(bool bRegister)
 	}
 }
 
-bool USkeletalMeshComponent::ShouldRunPostPhysicsTick() const
+bool USkeletalMeshComponent::ShouldRunEndPhysicsTick() const
 {
 	return	(bEnablePhysicsOnDedicatedServer || !IsNetMode(NM_DedicatedServer)) && // Early out if we are on a dedicated server and not running physics.
 			(IsSimulatingPhysics() || ShouldBlendPhysicsBones());
 }
 
-void USkeletalMeshComponent::UpdatePostPhysicsTickRegisteredState()
+void USkeletalMeshComponent::UpdateEndPhysicsTickRegisteredState()
 {
-	RegisterPostPhysicsTick(PrimaryComponentTick.IsTickFunctionRegistered() && ShouldRunPostPhysicsTick());
+	RegisterEndPhysicsTick(PrimaryComponentTick.IsTickFunctionRegistered() && ShouldRunEndPhysicsTick());
 }
 
 bool USkeletalMeshComponent::ShouldRunClothTick() const
@@ -753,7 +749,7 @@ static TAutoConsoleVariable<int32> CVarHiPriSkinnedMeshesTicks(
 
 void USkeletalMeshComponent::TickComponent(float DeltaTime, enum ELevelTick TickType, FActorComponentTickFunction *ThisTickFunction)
 {
-	UpdatePostPhysicsTickRegisteredState();
+	UpdateEndPhysicsTickRegisteredState();
 	UpdateClothTickRegisteredState();
 
 	// clear and add morphtarget curves that are added via SetMorphTarget
@@ -766,7 +762,7 @@ void USkeletalMeshComponent::TickComponent(float DeltaTime, enum ELevelTick Tick
 
 	/** Update the end group and tick priority */
 	const bool bDoLateEnd = CVarAnimationDelaysEndGroup.GetValueOnGameThread() > 0;
-	const bool bRequiresPhysics = PostPhysicsTickFunction.IsTickFunctionRegistered();
+	const bool bRequiresPhysics = EndPhysicsTickFunction.IsTickFunctionRegistered();
 	const ETickingGroup EndTickGroup = bDoLateEnd && !bRequiresPhysics ? TG_PostPhysics : TG_PrePhysics;
 	ThisTickFunction->EndTickGroup = EndTickGroup;
 
diff --git a/Engine/Source/Runtime/Engine/Private/Components/TextRenderComponent.cpp b/Engine/Source/Runtime/Engine/Private/Components/TextRenderComponent.cpp
--- a/Engine/Source/Runtime/Engine/Private/Components/TextRenderComponent.cpp
+++ b/Engine/Source/Runtime/Engine/Private/Components/TextRenderComponent.cpp
@@ -858,7 +858,7 @@ FBoxSphereBounds UTextRenderComponent::CalcBounds(const FTransform& LocalToWorld
 	}
 	else
 	{
-		return FBoxSphereBounds(ForceInit);
+		return FBoxSphereBounds(ForceInit).TransformBy(LocalToWorld);
 	}
 }
 
diff --git a/Engine/Source/Runtime/Engine/Private/DebugCameraController.cpp b/Engine/Source/Runtime/Engine/Private/DebugCameraController.cpp
--- a/Engine/Source/Runtime/Engine/Private/DebugCameraController.cpp
+++ b/Engine/Source/Runtime/Engine/Private/DebugCameraController.cpp
@@ -480,6 +480,11 @@ void ADebugCameraController::ApplySpeedScale()
 		}
 	}
 }
+void ADebugCameraController::SetPawnMovementSpeedScale(const float NewSpeedScale)
+{ 
+	SpeedScale = NewSpeedScale;
+	ApplySpeedScale();
+}
 
 void ADebugCameraController::IncreaseFOV()
 {
diff --git a/Engine/Source/Runtime/Engine/Private/GameFramework/SpringArmComponent.cpp b/Engine/Source/Runtime/Engine/Private/GameFramework/SpringArmComponent.cpp
--- a/Engine/Source/Runtime/Engine/Private/GameFramework/SpringArmComponent.cpp
+++ b/Engine/Source/Runtime/Engine/Private/GameFramework/SpringArmComponent.cpp
@@ -196,27 +196,6 @@ void USpringArmComponent::OnRegister()
 
 	// Init deprecated var, for old code that may refer to it.
 	SetDeprecatedControllerViewRotation(*this, bUsePawnControlRotation);
-
-	// 4.11.2 Hack (UE-24725): Set up tick dependency on simulated objects in the parent hierarchy, so we use post-physics transforms.
-	// Probably only relevant for simulated SkeletalMeshComponents.
-	if (!IsTemplate())
-	{
-		USceneComponent* Parent = GetAttachParent();
-		while (Parent != nullptr)
-		{
-			if (UPrimitiveComponent* PrimitiveParent = Cast<UPrimitiveComponent>(Parent))
-			{
-				if (PrimitiveParent->IsSimulatingPhysics())
-				{
-					PrimitiveParent->PostPhysicsComponentTick.bCanEverTick = true;
-					PrimitiveParent->PostPhysicsComponentTick.SetTickFunctionEnable(true);
-					PrimaryComponentTick.AddPrerequisite(PrimitiveParent, PrimitiveParent->PostPhysicsComponentTick);
-					break;
-				}
-			}
-			Parent = Parent->GetAttachParent();
-		}
-	}
 }
 
 void USpringArmComponent::PostLoad()
diff --git a/Engine/Source/Runtime/Engine/Private/GameplayStatics.cpp b/Engine/Source/Runtime/Engine/Private/GameplayStatics.cpp
--- a/Engine/Source/Runtime/Engine/Private/GameplayStatics.cpp
+++ b/Engine/Source/Runtime/Engine/Private/GameplayStatics.cpp
@@ -1766,7 +1766,7 @@ bool UGameplayStatics::SuggestProjectileVelocity(UObject* WorldContextObject, FV
 				// d = vt + .5 a t^2
 				const FVector TraceEnd = Start + PrioritizedProjVelocities[CurrentSolutionIdx]*TimeInFlight + FVector(0.f, 0.f, 0.5f * -GravityZ * FMath::Square(TimeInFlight) - CollisionRadius);
 
-				if ( (TraceOption == ESuggestProjVelocityTraceOption::OnlyTraceWhileAsceding) && (TraceEnd.Z < TraceStart.Z) )
+				if ( (TraceOption == ESuggestProjVelocityTraceOption::OnlyTraceWhileAscending) && (TraceEnd.Z < TraceStart.Z) )
 				{
 					// falling, we are done tracing
 					if (!bDrawDebug)
@@ -2117,4 +2117,4 @@ int32 UGameplayStatics::GetIntOption( const FString& Options, const FString& Key
 	return DefaultValue;
 }
 
-#undef LOCTEXT_NAMESPACE
\ No newline at end of file
+#undef LOCTEXT_NAMESPACE
diff --git a/Engine/Source/Runtime/Engine/Private/InputAxisDelegateBinding.cpp b/Engine/Source/Runtime/Engine/Private/InputAxisDelegateBinding.cpp
--- a/Engine/Source/Runtime/Engine/Private/InputAxisDelegateBinding.cpp
+++ b/Engine/Source/Runtime/Engine/Private/InputAxisDelegateBinding.cpp
@@ -4,6 +4,20 @@
 #include "BlueprintUtilities.h"
 #include "Engine/InputAxisDelegateBinding.h"
 
+// Determines whether or not the Dynamic Delegate for the given binding
+// is a NAME_None function delegate. That implies it is a "dummy" delegate
+// only used for capturing axis values (for things like GetAxisValue).
+static inline bool IsFunctionNameNone(const FInputAxisBinding& Binding)
+{
+	const FInputAxisHandlerDynamicSignature& DynamicDelegate = Binding.AxisDelegate.GetDynamicDelegate();
+	const FInputAxisHandlerSignature& NonDynamicDelegate = Binding.AxisDelegate.GetDelegate();
+
+	// Dynamic delegates will return unbound if the function is NAME_None.
+	// Therefore, just ensure that we don't have a normal delegate bound.
+	return !NonDynamicDelegate.IsBound() && DynamicDelegate.GetFunctionName().IsNone();
+}
+
+
 UInputAxisDelegateBinding::UInputAxisDelegateBinding(const FObjectInitializer& ObjectInitializer)
 	: Super(ObjectInitializer)
 {
@@ -15,27 +29,44 @@ void UInputAxisDelegateBinding::BindToInputComponent(UInputComponent* InputCompo
 
 	for (int32 BindIndex=0; BindIndex<InputAxisDelegateBindings.Num(); ++BindIndex)
 	{
+		bool bFoundNameNoneDelegate = false;
 		const FBlueprintInputAxisDelegateBinding& Binding = InputAxisDelegateBindings[BindIndex];
 
-		FInputAxisBinding AB( Binding.InputAxisName );
-		AB.bConsumeInput = Binding.bConsumeInput;
-		AB.bExecuteWhenPaused = Binding.bExecuteWhenPaused;
-		AB.AxisDelegate.BindDelegate(InputComponent->GetOwner(), Binding.FunctionNameToBind);
-
-		if (Binding.bOverrideParentBinding)
+		// Only allow a single NAME_None delegate for any given axis.
+		if (Binding.FunctionNameToBind == NAME_None)
 		{
-			for (int32 ExistingIndex = InputComponent->AxisBindings.Num() - 1; ExistingIndex >= 0; --ExistingIndex)
+			for (const FInputAxisBinding& ExistingBinding : InputComponent->AxisBindings)
 			{
-				const FInputAxisBinding& ExistingBind = InputComponent->AxisBindings[ExistingIndex];
-				if (ExistingBind.AxisName == AB.AxisName)
+				if (Binding.InputAxisName == ExistingBinding.AxisName && IsFunctionNameNone(ExistingBinding))
 				{
-					InputComponent->AxisBindings.RemoveAt(ExistingIndex);
+					bFoundNameNoneDelegate = true;
+					break;
 				}
 			}
 		}
 
-		// To avoid binds in the same layer being removed by the parent override temporarily put them in this array and add later
-		BindsToAdd.Add(AB);
+		if (!bFoundNameNoneDelegate)
+		{
+			FInputAxisBinding AB( Binding.InputAxisName );
+			AB.bConsumeInput = Binding.bConsumeInput;
+			AB.bExecuteWhenPaused = Binding.bExecuteWhenPaused;
+			AB.AxisDelegate.BindDelegate(InputComponent->GetOwner(), Binding.FunctionNameToBind);
+
+			if (Binding.bOverrideParentBinding)
+			{
+				for (int32 ExistingIndex = InputComponent->AxisBindings.Num() - 1; ExistingIndex >= 0; --ExistingIndex)
+				{
+					const FInputAxisBinding& ExistingBind = InputComponent->AxisBindings[ExistingIndex];
+					if (ExistingBind.AxisName == AB.AxisName && !IsFunctionNameNone(ExistingBind))
+					{
+							InputComponent->AxisBindings.RemoveAt(ExistingIndex);
+					}
+				}
+			}
+
+			// To avoid binds in the same layer being removed by the parent override temporarily put them in this array and add later
+			BindsToAdd.Add(AB);
+		}
 	}
 
 	for (int32 Index=0; Index < BindsToAdd.Num(); ++Index)
diff --git a/Engine/Source/Runtime/Engine/Private/KismetStringLibrary.cpp b/Engine/Source/Runtime/Engine/Private/KismetStringLibrary.cpp
--- a/Engine/Source/Runtime/Engine/Private/KismetStringLibrary.cpp
+++ b/Engine/Source/Runtime/Engine/Private/KismetStringLibrary.cpp
@@ -449,5 +449,5 @@ FString UKismetStringLibrary::TimeSecondsToString(float InSeconds)
 	const int32 NumCentiseconds = FMath::FloorToInt((InSeconds - FMath::FloorToFloat(InSeconds)) * 100.f);
 
 	// Create string, including leading zeroes
-	return FString::Printf(TEXT("%02d:%02d:%02d"), NumMinutes, NumSeconds, NumCentiseconds);
+	return FString::Printf(TEXT("%02d:%02d.%02d"), NumMinutes, NumSeconds, NumCentiseconds);
 }
diff --git a/Engine/Source/Runtime/Engine/Private/PhysicsEngine/ConstraintInstance.cpp b/Engine/Source/Runtime/Engine/Private/PhysicsEngine/ConstraintInstance.cpp
--- a/Engine/Source/Runtime/Engine/Private/PhysicsEngine/ConstraintInstance.cpp
+++ b/Engine/Source/Runtime/Engine/Private/PhysicsEngine/ConstraintInstance.cpp
@@ -288,10 +288,16 @@ PxScene* GetPScene_LockFree(const FBodyInstance* Body1, const FBodyInstance* Bod
 	else if(SceneIndex1 >= 0 && SceneIndex2 >= 0 && SceneIndex1 != SceneIndex2)
 	{
 #if !(UE_BUILD_SHIPPING || UE_BUILD_TEST)
+		UPrimitiveComponent* PrimComp1 = Body1 ? Body1->OwnerComponent.Get() : nullptr;
+		UPrimitiveComponent* PrimComp2 = Body2 ? Body2->OwnerComponent.Get() : nullptr;
+
 		FMessageLog("PIE").Warning()
 			->AddToken(FTextToken::Create(LOCTEXT("JointBetweenScenesStart", "Constraint")))
 			->AddToken(FUObjectToken::Create(DebugOwner))
-			->AddToken(FTextToken::Create(LOCTEXT("JointBetweenScenesEnd", "attempting to create a joint between two actors in different scenes.  No joint created.")));
+			->AddToken(FTextToken::Create(LOCTEXT("JointBetweenScenesMid", "attempting to create a joint between two actors in different scenes (")))
+			->AddToken(FUObjectToken::Create(PrimComp1))
+			->AddToken(FUObjectToken::Create(PrimComp2))
+			->AddToken(FTextToken::Create(LOCTEXT("JointBetweenScenesEnd", ").  No joint created.")));
 #endif // !(UE_BUILD_SHIPPING || UE_BUILD_TEST)
 	}
 	else
diff --git a/Engine/Source/Runtime/Engine/Private/PlayerController.cpp b/Engine/Source/Runtime/Engine/Private/PlayerController.cpp
--- a/Engine/Source/Runtime/Engine/Private/PlayerController.cpp
+++ b/Engine/Source/Runtime/Engine/Private/PlayerController.cpp
@@ -1384,6 +1384,11 @@ void APlayerController::EndPlay(const EEndPlayReason::Type EndPlayReason)
 		}
 	}
 
+	if (CheatManager)
+	{
+		CheatManager->ReceiveEndPlay();
+	}
+	
 	Super::EndPlay(EndPlayReason);
 }
 
@@ -1975,6 +1980,23 @@ bool APlayerController::GetHitResultAtScreenPosition(const FVector2D ScreenPosit
 	return false;
 }
 
+void APlayerController::SetMouseLocation(const int X, const int Y)
+{
+	ULocalPlayer* LocalPlayer = Cast<ULocalPlayer>( Player );
+	if (LocalPlayer)
+	{
+		UGameViewportClient* ViewportClient = LocalPlayer->ViewportClient;
+		if (ViewportClient)
+		{
+			FViewport* Viewport = ViewportClient->Viewport;
+			if (Viewport)
+			{
+				Viewport->SetMouse( X, Y );
+			}
+		}
+	}
+}
+
 /* PlayerTick is only called if the PlayerController has a PlayerInput object.  Therefore, it will not be called on servers for non-locally controlled playercontrollers. */
 void APlayerController::PlayerTick( float DeltaTime )
 {
diff --git a/Engine/Source/Runtime/Engine/Private/SkeletalMeshComponentPhysics.cpp b/Engine/Source/Runtime/Engine/Private/SkeletalMeshComponentPhysics.cpp
--- a/Engine/Source/Runtime/Engine/Private/SkeletalMeshComponentPhysics.cpp
+++ b/Engine/Source/Runtime/Engine/Private/SkeletalMeshComponentPhysics.cpp
@@ -59,18 +59,18 @@ FString FSkeletalMeshComponentClothTickFunction::DiagnosticMessage()
 	return TEXT("FSkeletalMeshComponentClothTickFunction");
 }
 
-void FSkeletalMeshComponentPostPhysicsTickFunction::ExecuteTick(float DeltaTime, enum ELevelTick TickType, ENamedThreads::Type CurrentThread, const FGraphEventRef& MyCompletionGraphEvent)
+void FSkeletalMeshComponentEndPhysicsTickFunction::ExecuteTick(float DeltaTime, enum ELevelTick TickType, ENamedThreads::Type CurrentThread, const FGraphEventRef& MyCompletionGraphEvent)
 {
-	QUICK_SCOPE_CYCLE_COUNTER(FSkeletalMeshComponentPostPhysicsTickFunction_ExecuteTick);
+	QUICK_SCOPE_CYCLE_COUNTER(FSkeletalMeshComponentEndPhysicsTickFunction_ExecuteTick);
 	FActorComponentTickFunction::ExecuteTickHelper(Target, /*bTickInEditor=*/ false, DeltaTime, TickType, [this](float DilatedTime)
 	{
-		Target->PostPhysicsTickComponent(*this);
+		Target->EndPhysicsTickComponent(*this);
 	});
 }
 
-FString FSkeletalMeshComponentPostPhysicsTickFunction::DiagnosticMessage()
+FString FSkeletalMeshComponentEndPhysicsTickFunction::DiagnosticMessage()
 {
-	return TEXT("FSkeletalMeshComponentPostPhysicsTickFunction");
+	return TEXT("FSkeletalMeshComponentEndPhysicsTickFunction");
 }
 
 #if WITH_APEX_CLOTHING
@@ -482,7 +482,7 @@ void USkeletalMeshComponent::SetSimulatePhysics(bool bSimulate)
 		SetRootBodyIndex(RootBodyData.BodyIndex);	//Update the root body data cache in case animation has moved root body relative to root joint
 	}
 
-	UpdatePostPhysicsTickRegisteredState();
+	UpdateEndPhysicsTickRegisteredState();
 	UpdateClothTickRegisteredState();
 }
 
@@ -1102,7 +1102,7 @@ void USkeletalMeshComponent::SetAllBodiesSimulatePhysics(bool bNewSimulate)
 
 	SetRootBodyIndex(RootBodyData.BodyIndex);	//Update the root body data cache in case animation has moved root body relative to root joint
 
-	UpdatePostPhysicsTickRegisteredState();
+	UpdateEndPhysicsTickRegisteredState();
 	UpdateClothTickRegisteredState();
 }
 
@@ -1138,7 +1138,7 @@ void USkeletalMeshComponent::SetAllBodiesBelowSimulatePhysics( const FName& InBo
 			SetRootBodyIndex(RootBodyData.BodyIndex);	//Update the root body data cache in case animation has moved root body relative to root joint
 		}
 
-		UpdatePostPhysicsTickRegisteredState();
+		UpdateEndPhysicsTickRegisteredState();
 		UpdateClothTickRegisteredState();
 	}
 }
@@ -1366,7 +1366,7 @@ void USkeletalMeshComponent::SetAllBodiesPhysicsBlendWeight(float PhysicsBlendWe
 
 	bBlendPhysics = false;
 
-	UpdatePostPhysicsTickRegisteredState();
+	UpdateEndPhysicsTickRegisteredState();
 	UpdateClothTickRegisteredState();
 }
 
@@ -1382,7 +1382,7 @@ void USkeletalMeshComponent::SetAllBodiesBelowPhysicsBlendWeight( const FName& I
 	{
 		bBlendPhysics = false;
 
-		UpdatePostPhysicsTickRegisteredState();
+		UpdateEndPhysicsTickRegisteredState();
 		UpdateClothTickRegisteredState();
 	}
 }
@@ -1399,7 +1399,7 @@ void USkeletalMeshComponent::AccumulateAllBodiesBelowPhysicsBlendWeight( const F
 	{
 		bBlendPhysics = false;
 
-		UpdatePostPhysicsTickRegisteredState();
+		UpdateEndPhysicsTickRegisteredState();
 		UpdateClothTickRegisteredState();
 	}
 }
@@ -3411,13 +3411,13 @@ void USkeletalMeshComponent::ProcessClothCollisionWithEnvironment()
 
 #endif// #if WITH_CLOTH_COLLISION_DETECTION
 
-void USkeletalMeshComponent::PostPhysicsTickComponent(FSkeletalMeshComponentPostPhysicsTickFunction& ThisTickFunction)
+void USkeletalMeshComponent::EndPhysicsTickComponent(FSkeletalMeshComponentEndPhysicsTickFunction& ThisTickFunction)
 {
 	//IMPORTANT!
 	//
-	// The decision on whether to use PostPhysicsTickComponent or not is made by ShouldRunPostPhysicsTick()
-	// Any changes that are made to PostPhysicsTickComponent that affect whether it should be run or not
-	// have to be reflected in ShouldRunPostPhysicsTick() as well
+	// The decision on whether to use EndPhysicsTickComponent or not is made by ShouldRunEndPhysicsTick()
+	// Any changes that are made to EndPhysicsTickComponent that affect whether it should be run or not
+	// have to be reflected in ShouldRunEndPhysicsTick() as well
 	
 	// if physics is disabled on dedicated server, no reason to be here. 
 	if (!bEnablePhysicsOnDedicatedServer && IsRunningDedicatedServer())
@@ -5255,6 +5255,40 @@ float USkeletalMeshComponent::GetMass() const
 	return Mass;
 }
 
+float USkeletalMeshComponent::GetBoneMass(FName BoneName, bool bScaleMass) const
+{
+	if (FBodyInstance* BI = GetBodyInstance(BoneName))
+	{
+		float Scale = 1.0f;
+		if (bScaleMass)
+		{
+			Scale = BI->MassScale;
+		}
+		return Scale*BI->GetBodyMass();
+	}
+
+	return 0.0f;
+}
+
+FVector USkeletalMeshComponent::GetSkeletalCenterOfMass() const
+{
+	FVector Location = FVector::ZeroVector;
+	float TotalMass = 0.0f;
+	for (int32 i = 0; i < Bodies.Num(); i++)
+	{
+		FBodyInstance* BI = Bodies[i];
+		if (BI->IsValidBodyInstance())
+		{
+			float Mass = BI->MassScale*BI->GetBodyMass();
+			Location += Mass*BI->GetCOMPosition();
+			TotalMass += Mass;
+		}
+	}
+
+	return Location / TotalMass;
+}
+
+
 // blueprint callable methods 
 float USkeletalMeshComponent::GetClothMaxDistanceScale()
 {
diff --git a/Engine/Source/Runtime/Engine/Private/SkeletalMeshMerge.cpp b/Engine/Source/Runtime/Engine/Private/SkeletalMeshMerge.cpp
--- a/Engine/Source/Runtime/Engine/Private/SkeletalMeshMerge.cpp
+++ b/Engine/Source/Runtime/Engine/Private/SkeletalMeshMerge.cpp
@@ -704,7 +704,7 @@ int32 FSkeletalMeshMerge::CalculateLodCount(const TArray<USkeletalMesh*>& Source
 	// But, make sure there is at least one.
 
 	LodCount -= StripTopLODs;
-	LodCount = FMath::Min(LodCount, 1);
+	LodCount = FMath::Max(LodCount, 1);
 	
 	return LodCount;
 }
diff --git a/Engine/Source/Runtime/Engine/Private/SkeletalRenderCPUSkin.cpp b/Engine/Source/Runtime/Engine/Private/SkeletalRenderCPUSkin.cpp
--- a/Engine/Source/Runtime/Engine/Private/SkeletalRenderCPUSkin.cpp
+++ b/Engine/Source/Runtime/Engine/Private/SkeletalRenderCPUSkin.cpp
@@ -45,7 +45,15 @@ static void SkinVertices(FFinalSkinVertex* DestVertex, FMatrix* ReferenceToLocal
  * @param LOD - LOD model corresponding to DestVertex 
  * @param BonesOfInterest - array of bones we want to display
  */
-static void CalculateBoneWeights(FFinalSkinVertex* DestVertex, FStaticLODModel& LOD, const TArray<int32>& BonesOfInterest);
+static void CalculateBoneWeights(FFinalSkinVertex* DestVertex, FStaticLODModel& LOD, TArray<int32> InBonesOfInterest);
+
+/**
+* Modify the vertex buffer to store morph target weights in the UV coordinates for rendering
+* @param DestVertex - already filled out vertex buffer from SkinVertices
+* @param LOD - LOD model corresponding to DestVertex
+* @param MorphTargetOfInterest - array of morphtargets to draw
+*/
+static void CalculateMorphTargetWeights(FFinalSkinVertex* DestVertex, FStaticLODModel& LOD, int LODIndex, TArray<UMorphTarget*> InMorphTargetOfInterest);
 
 /*-----------------------------------------------------------------------------
 	FFinalSkinVertexBuffer
@@ -123,7 +131,7 @@ FSkeletalMeshObjectCPUSkin::FSkeletalMeshObjectCPUSkin(USkinnedMeshComponent* In
 	: FSkeletalMeshObject(InMeshComponent, InSkeletalMeshResource, InFeatureLevel)
 ,	DynamicData(NULL)
 ,	CachedVertexLOD(INDEX_NONE)
-,	bRenderBoneWeight(false)
+,	bRenderOverlayMaterial(false)
 {
 	// create LODs to match the base mesh
 	for( int32 LODIndex=0;LODIndex < SkeletalMeshResource->LODModels.Num();LODIndex++ )
@@ -160,12 +168,22 @@ void FSkeletalMeshObjectCPUSkin::ReleaseResources()
 	}
 }
 
-void FSkeletalMeshObjectCPUSkin::EnableBlendWeightRendering(bool bEnabled, const TArray<int32>& InBonesOfInterest)
+
+void FSkeletalMeshObjectCPUSkin::EnableOverlayRendering(bool bEnabled, const TArray<int32>* InBonesOfInterest, const TArray<UMorphTarget*>* InMorphTargetOfInterest)
 {
-	bRenderBoneWeight = bEnabled;
+	bRenderOverlayMaterial = bEnabled;
+	
+	BonesOfInterest.Reset();
+	MorphTargetOfInterest.Reset();
 
-	BonesOfInterest.Empty(InBonesOfInterest.Num());
-	BonesOfInterest.Append(InBonesOfInterest);
+	if (InBonesOfInterest)
+	{
+		BonesOfInterest.Append(*InBonesOfInterest);
+	}
+	else if (InMorphTargetOfInterest)
+	{
+		MorphTargetOfInterest.Append(*InMorphTargetOfInterest);
+	}
 }
 
 void FSkeletalMeshObjectCPUSkin::UpdateRecomputeTangent(int32 MaterialIndex, bool bRecomputeTangent)
@@ -292,10 +310,20 @@ void FSkeletalMeshObjectCPUSkin::CacheVertices(int32 LODIndex, bool bForce) cons
 				}
 			}
 
-			if (bRenderBoneWeight)
+			if (bRenderOverlayMaterial)
 			{
-				//Transfer bone weights we're interested in to the UV channels
-				CalculateBoneWeights(DestVertex, LOD, BonesOfInterest);
+				if (MorphTargetOfInterest.Num() > 0)
+				{
+					//Transfer morph target weights we're interested in to the UV channels
+					CalculateMorphTargetWeights(DestVertex, LOD, LODIndex, MorphTargetOfInterest);
+				}
+				else // default is bones of interest
+					// this can go if no morphtarget is selected but enabled to render
+					// but that doesn't matter since it will only draw empty overlay
+				{
+					//Transfer bone weights we're interested in to the UV channels
+					CalculateBoneWeights(DestVertex, LOD, BonesOfInterest);
+				}
 			}
 		}
 
@@ -902,7 +930,7 @@ static FORCEINLINE void CalculateSectionBoneWeights(FFinalSkinVertex*& DestVerte
  * @param LOD - LOD model corresponding to DestVertex 
  * @param BonesOfInterest - array of bones we want to display
  */
-static void CalculateBoneWeights(FFinalSkinVertex* DestVertex, FStaticLODModel& LOD, const TArray<int32>& BonesOfInterest)
+static void CalculateBoneWeights(FFinalSkinVertex* DestVertex, FStaticLODModel& LOD, TArray<int32> InBonesOfInterest)
 {
 	const float INV255 = 1.f/255.f;
 
@@ -916,11 +944,33 @@ static void CalculateBoneWeights(FFinalSkinVertex* DestVertex, FStaticLODModel&
 
 		if (LOD.VertexBufferGPUSkin.HasExtraBoneInfluences())
 		{
-			CalculateSectionBoneWeights<true>(DestVertex, LOD.VertexBufferGPUSkin, Section, BonesOfInterest);
+			CalculateSectionBoneWeights<true>(DestVertex, LOD.VertexBufferGPUSkin, Section, InBonesOfInterest);
 		}
 		else
 		{
-			CalculateSectionBoneWeights<false>(DestVertex, LOD.VertexBufferGPUSkin, Section, BonesOfInterest);
+			CalculateSectionBoneWeights<false>(DestVertex, LOD.VertexBufferGPUSkin, Section, InBonesOfInterest);
+		}
+	}
+}
+
+static void CalculateMorphTargetWeights(FFinalSkinVertex* DestVertex, FStaticLODModel& LOD, int LODIndex, TArray<UMorphTarget*> InMorphTargetOfInterest)
+{
+	const FFinalSkinVertex* EndVert = DestVertex + LOD.NumVertices;
+
+	for (FFinalSkinVertex* ClearVert = DestVertex; ClearVert != EndVert; ++ClearVert)
+	{
+		ClearVert->U = 0.f;
+		ClearVert->V = 0.f;
+	}
+
+	for (const UMorphTarget* Morphtarget : InMorphTargetOfInterest)
+	{
+		const FMorphTargetLODModel& MTLOD = Morphtarget->MorphLODModels[LODIndex];
+		for (int32 MorphVertexIndex = 0; MorphVertexIndex < MTLOD.Vertices.Num(); ++MorphVertexIndex)
+		{
+			FFinalSkinVertex* SetVert = DestVertex + MTLOD.Vertices[MorphVertexIndex].SourceIdx;
+			SetVert->U = 1.0f;
+			SetVert->V = 1.0f;
 		}
 	}
 }
diff --git a/Engine/Source/Runtime/Engine/Private/SkeletalRenderCPUSkin.h b/Engine/Source/Runtime/Engine/Private/SkeletalRenderCPUSkin.h
--- a/Engine/Source/Runtime/Engine/Private/SkeletalRenderCPUSkin.h
+++ b/Engine/Source/Runtime/Engine/Private/SkeletalRenderCPUSkin.h
@@ -152,7 +152,7 @@ public:
 	virtual void Update(int32 LODIndex,USkinnedMeshComponent* InMeshComponent,const TArray<FActiveMorphTarget>& ActiveMorphTargets, const TArray<float>& MorphTargetsWeights) override;
 	void UpdateDynamicData_RenderThread(FRHICommandListImmediate& RHICmdList, FDynamicSkelMeshObjectDataCPUSkin* InDynamicData, uint32 FrameNumberToPrepare);
 	virtual void UpdateRecomputeTangent(int32 MaterialIndex, bool bRecomputeTangent) override;
-	virtual void EnableBlendWeightRendering(bool bEnabled, const TArray<int32>& InBonesOfInterest) override;
+	virtual void EnableOverlayRendering(bool bEnabled, const TArray<int32>* InBonesOfInterest, const TArray<UMorphTarget*>* InMorphTargetOfInterest) override;
 	virtual void CacheVertices(int32 LODIndex, bool bForce) const override;
 	virtual bool IsCPUSkinned() const override { return true; }
 	virtual const FVertexFactory* GetSkinVertexFactory(const FSceneView* View, int32 LODIndex, int32 ChunkIdx) const override;
@@ -255,8 +255,9 @@ private:
 
 	/** Array of bone's to render bone weights for */
 	TArray<int32> BonesOfInterest;
+	TArray<UMorphTarget*> MorphTargetOfInterest;
 
 	/** Bone weight viewing in editor */
-	bool bRenderBoneWeight;
+	bool bRenderOverlayMaterial;
 };
 
diff --git a/Engine/Source/Runtime/Engine/Private/SoundWave.cpp b/Engine/Source/Runtime/Engine/Private/SoundWave.cpp
--- a/Engine/Source/Runtime/Engine/Private/SoundWave.cpp
+++ b/Engine/Source/Runtime/Engine/Private/SoundWave.cpp
@@ -643,6 +643,7 @@ void USoundWave::Parse( FAudioDevice* AudioDevice, const UPTRINT NodeWaveInstanc
 		// Propagate properties and add WaveInstance to outgoing array of FWaveInstances.
 		WaveInstance->Volume = ParseParams.Volume * Volume;
 		WaveInstance->VolumeMultiplier = ParseParams.VolumeMultiplier;
+		WaveInstance->VolumeApp = ParseParams.VolumeApp;
 		WaveInstance->Pitch = ParseParams.Pitch * Pitch;
 		WaveInstance->bEnableLowPassFilter = ParseParams.bEnableLowPassFilter;
 		WaveInstance->bIsOccluded = ParseParams.bIsOccluded;
@@ -733,7 +734,11 @@ void USoundWave::Parse( FAudioDevice* AudioDevice, const UPTRINT NodeWaveInstanc
 			WaveInstance->SpatializationAlgorithm = ParseParams.SpatializationAlgorithm;
 		}
 
-		WaveInstances.Add(WaveInstance);
+		// Only append to the wave instances list if we're virtual (always append) or we're audible (non-zero volume)
+		if (WaveInstance->GetVolume() > KINDA_SMALL_NUMBER || (bVirtualizeWhenSilent && AudioDevice->VirtualSoundsEnabled()))
+		{
+			WaveInstances.Add(WaveInstance);
+		}
 
 		// We're still alive.
 		ActiveSound.bFinished = false;
diff --git a/Engine/Source/Runtime/Engine/Private/Vehicles/SimpleWheeledVehicleMovementComponent.cpp b/Engine/Source/Runtime/Engine/Private/Vehicles/SimpleWheeledVehicleMovementComponent.cpp
new file mode 100644
--- /dev/null
+++ b/Engine/Source/Runtime/Engine/Private/Vehicles/SimpleWheeledVehicleMovementComponent.cpp
@@ -0,0 +1,86 @@
+// Copyright 1998-2016 Epic Games, Inc. All Rights Reserved.
+
+#include "EnginePrivate.h"
+#include "PhysicsPublic.h"
+#include "Vehicles/SimpleWheeledVehicleMovementComponent.h"
+
+#if WITH_PHYSX
+#include "PhysXVehicleManager.h"
+
+void USimpleWheeledVehicleMovementComponent::SetBrakeTorque(float BrakeTorque, int32 WheelIndex)
+{
+#if WITH_VEHICLE
+	if (PVehicle && UpdatedPrimitive)
+	{
+		PxVehicleWheels* PVehicleLocal = PVehicle;
+		if (WheelSetups.IsValidIndex(WheelIndex))
+		{
+			ExecuteOnPxRigidDynamicReadWrite(UpdatedPrimitive->GetBodyInstance(), [PVehicleLocal, WheelIndex, BrakeTorque](PxRigidDynamic* PRigidDynamic)
+			{
+				((PxVehicleNoDrive*)PVehicleLocal)->setBrakeTorque(WheelIndex, M2ToCm2(BrakeTorque));
+			});
+		}
+	}
+#endif
+}
+
+void USimpleWheeledVehicleMovementComponent::SetDriveTorque(float DriveTorque, int32 WheelIndex)
+{
+#if WITH_VEHICLE
+	if (PVehicle && UpdatedPrimitive)
+	{
+		PxVehicleWheels* PVehicleLocal = PVehicle;
+		if (WheelSetups.IsValidIndex(WheelIndex))
+		{
+			ExecuteOnPxRigidDynamicReadWrite(UpdatedPrimitive->GetBodyInstance(), [PVehicleLocal, WheelIndex, DriveTorque](PxRigidDynamic* PRigidDynamic)
+			{
+				((PxVehicleNoDrive*)PVehicleLocal)->setDriveTorque(WheelIndex, M2ToCm2(DriveTorque));
+			});
+		}
+	}
+#endif
+}
+
+void USimpleWheeledVehicleMovementComponent::SetSteerAngle(float SteerAngle, int32 WheelIndex)
+{
+#if WITH_VEHICLE
+	if (PVehicle && UpdatedPrimitive)
+	{
+		PxVehicleWheels* PVehicleLocal = PVehicle;
+		if (WheelSetups.IsValidIndex(WheelIndex))
+		{
+			const float SteerRad = FMath::DegreesToRadians(SteerAngle);
+			ExecuteOnPxRigidDynamicReadWrite(UpdatedPrimitive->GetBodyInstance(), [PVehicleLocal, WheelIndex, SteerRad](PxRigidDynamic* PRigidDynamic)
+			{
+				((PxVehicleNoDrive*)PVehicleLocal)->setSteerAngle(WheelIndex, SteerRad);
+			});
+		}
+	}
+#endif
+}
+
+#if WITH_VEHICLE
+
+void USimpleWheeledVehicleMovementComponent::SetupVehicleDrive(PxVehicleWheelsSimData* PWheelsSimData)
+{
+	//Use a simple PxNoDrive which will give us suspension but no engine forces which we leave to the user
+
+	// Create the vehicle
+	PxVehicleNoDrive* PVehicleNoDrive = PxVehicleNoDrive::allocate(WheelSetups.Num());
+	check(PVehicleNoDrive);
+
+	ExecuteOnPxRigidDynamicReadWrite(UpdatedPrimitive->GetBodyInstance(), [&](PxRigidDynamic* PRigidDynamic)
+	{
+		PVehicleNoDrive->setup(GPhysXSDK, PRigidDynamic, *PWheelsSimData);
+		PVehicleNoDrive->setToRestState();
+
+		// cleanup
+		PWheelsSimData->free();
+	});
+
+	// cache values
+	PVehicle = PVehicleNoDrive;
+}
+
+#endif
+#endif
\ No newline at end of file
diff --git a/Engine/Source/Runtime/Engine/Private/Vehicles/VehicleAnimInstance.cpp b/Engine/Source/Runtime/Engine/Private/Vehicles/VehicleAnimInstance.cpp
--- a/Engine/Source/Runtime/Engine/Private/Vehicles/VehicleAnimInstance.cpp
+++ b/Engine/Source/Runtime/Engine/Private/Vehicles/VehicleAnimInstance.cpp
@@ -7,6 +7,8 @@
 
 #include "EnginePrivate.h"
 #include "Vehicles/VehicleAnimInstance.h"
+#include "Vehicles/WheeledVehicleMovementComponent.h"
+#include "Vehicles/VehicleWheel.h"
 #include "GameFramework/WheeledVehicle.h"
 #include "AnimationRuntime.h"
 
@@ -24,7 +26,65 @@ class AWheeledVehicle * UVehicleAnimInstance::GetVehicle()
 	return Cast<AWheeledVehicle> (GetOwningActor());
 }
 
+FAnimInstanceProxy* UVehicleAnimInstance::CreateAnimInstanceProxy()
+{
+	return &AnimInstanceProxy;
+}
+
+void UVehicleAnimInstance::DestroyAnimInstanceProxy(FAnimInstanceProxy* InProxy)
+{
+}
+
 /////////////////////////////////////////////////////
 /////////////////////////////////////////////////////
 /////////////////////////////////////////////////////
+//// PROXY ///
+
+void FVehicleAnimInstanceProxy::SetWheeledVehicleMovementComponent(const UWheeledVehicleMovementComponent* InWheeledVehicleMovementComponent)
+{
+	const UWheeledVehicleMovementComponent* WheeledVehicleMovementComponent = InWheeledVehicleMovementComponent;
+
+	//initialize wheel data
+	const int32 NumOfwheels = WheeledVehicleMovementComponent->WheelSetups.Num();
+	WheelInstances.Empty(NumOfwheels);
+	if (NumOfwheels > 0)
+	{
+		WheelInstances.AddZeroed(NumOfwheels);
+		// now add wheel data
+		for (int32 WheelIndex = 0; WheelIndex<WheelInstances.Num(); ++WheelIndex)
+		{
+			FWheelAnimData& WheelInstance = WheelInstances[WheelIndex];
+			const FWheelSetup& WheelSetup = WheeledVehicleMovementComponent->WheelSetups[WheelIndex];
+
+			// set data
+			WheelInstance.BoneName = WheelSetup.BoneName;
+			WheelInstance.LocOffset = FVector::ZeroVector;
+			WheelInstance.RotOffset = FRotator::ZeroRotator;
+		}
+	}
+}
+
+void FVehicleAnimInstanceProxy::PreUpdate(UAnimInstance* InAnimInstance, float DeltaSeconds)
+{
+	const UVehicleAnimInstance* VehicleAnimInstance = CastChecked<UVehicleAnimInstance>(InAnimInstance);
+	if(const UWheeledVehicleMovementComponent* WheeledVehicleMovementComponent = VehicleAnimInstance->GetWheeledVehicleMovementComponent())
+	{
+		for (int32 WheelIndex = 0; WheelIndex < WheelInstances.Num(); ++ WheelIndex)
+		{
+			FWheelAnimData& WheelInstance = WheelInstances[WheelIndex];
+			if (WheeledVehicleMovementComponent->Wheels.IsValidIndex(WheelIndex))
+			{
+				if (const UVehicleWheel* VehicleWheel = WheeledVehicleMovementComponent->Wheels[WheelIndex])
+				{
+					WheelInstance.RotOffset.Pitch = VehicleWheel->GetRotationAngle();
+					WheelInstance.RotOffset.Yaw = VehicleWheel->GetSteerAngle();
+					WheelInstance.RotOffset.Roll = 0.f;
 
+					WheelInstance.LocOffset.X = 0.f;
+					WheelInstance.LocOffset.Y = 0.f;
+					WheelInstance.LocOffset.Z = VehicleWheel->GetSuspensionOffset();
+				}
+			}
+		}
+	}
+}
\ No newline at end of file
diff --git a/Engine/Source/Runtime/Engine/Private/Vehicles/VehicleWheel.cpp b/Engine/Source/Runtime/Engine/Private/Vehicles/VehicleWheel.cpp
--- a/Engine/Source/Runtime/Engine/Private/Vehicles/VehicleWheel.cpp
+++ b/Engine/Source/Runtime/Engine/Private/Vehicles/VehicleWheel.cpp
@@ -41,7 +41,7 @@ UVehicleWheel::UVehicleWheel(const FObjectInitializer& ObjectInitializer)
 	SuspensionDampingRatio = 1.0f;
 }
 
-float UVehicleWheel::GetSteerAngle()
+float UVehicleWheel::GetSteerAngle() const
 {
 #if WITH_VEHICLE
 	FPhysXVehicleManager* VehicleManager = VehicleSim->GetWorld()->GetPhysicsScene()->GetVehicleManager();
@@ -52,7 +52,7 @@ float UVehicleWheel::GetSteerAngle()
 #endif // WITH_PHYSX
 }
 
-float UVehicleWheel::GetRotationAngle()
+float UVehicleWheel::GetRotationAngle() const
 {
 #if WITH_VEHICLE
 	FPhysXVehicleManager* VehicleManager = VehicleSim->GetWorld()->GetPhysicsScene()->GetVehicleManager();
@@ -66,7 +66,7 @@ float UVehicleWheel::GetRotationAngle()
 #endif // WITH_PHYSX
 }
 
-float UVehicleWheel::GetSuspensionOffset()
+float UVehicleWheel::GetSuspensionOffset() const
 {
 #if WITH_VEHICLE
 	FPhysXVehicleManager* VehicleManager = VehicleSim->GetWorld()->GetPhysicsScene()->GetVehicleManager();
diff --git a/Engine/Source/Runtime/Engine/Private/Vehicles/WheeledVehicleMovementComponent.cpp b/Engine/Source/Runtime/Engine/Private/Vehicles/WheeledVehicleMovementComponent.cpp
--- a/Engine/Source/Runtime/Engine/Private/Vehicles/WheeledVehicleMovementComponent.cpp
+++ b/Engine/Source/Runtime/Engine/Private/Vehicles/WheeledVehicleMovementComponent.cpp
@@ -8,9 +8,11 @@
 #include "Vehicles/VehicleWheel.h"
 #include "Vehicles/WheeledVehicleMovementComponent.h"
 #include "Vehicles/TireType.h"
+#include "Vehicles/VehicleAnimInstance.h"
 #include "DisplayDebugHelpers.h"
 #include "PhysicsEngine/ConstraintInstance.h"
 #include "PhysicsEngine/PhysicsAsset.h"
+#include "FrameworkObjectVersion.h"
 
 #if WITH_PHYSX
 #include "../PhysicsEngine/PhysXSupport.h"
@@ -112,6 +114,8 @@ UWheeledVehicleMovementComponent::UWheeledVehicleMovementComponent(const FObject
 	SteeringInputRate.RiseRate = 2.5f;
 	SteeringInputRate.FallRate = 5.0f;
 
+	bDeprecatedSpringOffsetMode = false;	//This is just for backwards compat. Newly tuned content should never need to use this
+
 	bUseRVOAvoidance = false;
 	AvoidanceVelocity = FVector::ZeroVector;
 	AvoidanceLockVelocity = FVector::ZeroVector;
@@ -128,6 +132,8 @@ UWheeledVehicleMovementComponent::UWheeledVehicleMovementComponent(const FObject
     ThresholdLongitudinalSpeed = 5.f;
     LowForwardSpeedSubStepCount = 3;
     HighForwardSpeedSubStepCount = 1;
+	
+	bReverseAsBrake = true;	//Treats reverse button as break for a more arcade feel (also automatically goes into reverse)
 
 #if WITH_VEHICLE
 	// tire load filtering
@@ -222,41 +228,57 @@ void UWheeledVehicleMovementComponent::SetupVehicleShapes()
 			const FVector WheelOffset = GetWheelRestingPosition(WheelSetup);
 			const PxTransform PLocalPose = PxTransform(U2PVector(WheelOffset));
 			PxShape* PWheelShape = NULL;
-
+			
 			// Prepare shape
+			const UBodySetup* WheelBodySetup = nullptr;
+			FVector MeshScaleV(1.f, 1.f, 1.f);
 			if (Wheel->bDontCreateShape)
 			{
-				continue;
+				//don't create shape so grab it directly from the bodies associated with the vehicle
+				if(USkinnedMeshComponent* SkinnedMesh = GetMesh())
+				{
+					if(const FBodyInstance* WheelBI = SkinnedMesh->GetBodyInstance(WheelSetup.BoneName))
+					{
+						WheelBodySetup = WheelBI->BodySetup.Get();
+					}
+				}
+				
 			}
 			else if (Wheel->CollisionMesh && Wheel->CollisionMesh->BodySetup)
 			{
+				WheelBodySetup = Wheel->CollisionMesh->BodySetup;
+
 				FBoxSphereBounds MeshBounds = Wheel->CollisionMesh->GetBounds();
-				FVector MeshScaleV = FVector(1, 1, 1);
 				if (Wheel->bAutoAdjustCollisionSize)
 				{
 					MeshScaleV.X = Wheel->ShapeRadius / MeshBounds.BoxExtent.X;
 					MeshScaleV.Y = Wheel->ShapeWidth / MeshBounds.BoxExtent.Y;
 					MeshScaleV.Z = Wheel->ShapeRadius / MeshBounds.BoxExtent.Z;
 				}
+			}
+
+			if(WheelBodySetup)
+			{
 				PxMeshScale MeshScale(U2PVector(UpdatedComponent->RelativeScale3D * MeshScaleV), PxQuat::createIdentity());
-				if (Wheel->CollisionMesh->BodySetup->TriMeshes.Num())
+
+				if (WheelBodySetup->AggGeom.ConvexElems.Num() == 1)
 				{
-					PxTriangleMesh* TriMesh = Wheel->CollisionMesh->BodySetup->TriMeshes[0];
+					PxConvexMesh* ConvexMesh = WheelBodySetup->AggGeom.ConvexElems[0].ConvexMesh;
+					PWheelShape = PVehicleActor->createShape(PxConvexMeshGeometry(ConvexMesh, MeshScale), *WheelMaterial, PLocalPose);
+				}
+				else if (WheelBodySetup->TriMeshes.Num())
+				{
+					PxTriangleMesh* TriMesh = WheelBodySetup->TriMeshes[0];
 
 					// No eSIMULATION_SHAPE flag for wheels
 					PWheelShape = PVehicleActor->createShape(PxTriangleMeshGeometry(TriMesh, MeshScale), *WheelMaterial, PxShapeFlag::eSCENE_QUERY_SHAPE | PxShapeFlag::eVISUALIZATION);
 					PWheelShape->setLocalPose(PLocalPose);
 				}
-				else if (Wheel->CollisionMesh->BodySetup->AggGeom.ConvexElems.Num() == 1)
-				{
-					PxConvexMesh* ConvexMesh = Wheel->CollisionMesh->BodySetup->AggGeom.ConvexElems[0].ConvexMesh;
-					PWheelShape = PVehicleActor->createShape(PxConvexMeshGeometry(ConvexMesh, MeshScale), *WheelMaterial, PLocalPose);
-				}
-
-				check(PWheelShape);
 			}
-			else
+			
+			if(!PWheelShape)
 			{
+				//fallback onto simple spheres
 				PWheelShape = PVehicleActor->createShape(PxSphereGeometry(Wheel->ShapeRadius), *WheelMaterial, PLocalPose);
 			}
 
@@ -317,19 +339,20 @@ void UWheeledVehicleMovementComponent::SetupWheels( PxVehicleWheelsSimData* PWhe
 		// Prealloc data for the sprung masses
 		PxVec3 WheelOffsets[32];
 		float SprungMasses[32];
-		check(WheelSetups.Num() <= 32);
+		
+		const int32 NumWheels = FMath::Min(32, WheelSetups.Num());
 
 		// Calculate wheel offsets first, necessary for sprung masses
-		for (int32 WheelIdx = 0; WheelIdx < WheelSetups.Num(); ++WheelIdx)
+		for (int32 WheelIdx = 0; WheelIdx < NumWheels; ++WheelIdx)
 		{
 			WheelOffsets[WheelIdx] = U2PVector(GetWheelRestingPosition(WheelSetups[WheelIdx]));
 		}
 
 		// Now that we have all the wheel offsets, calculate the sprung masses
 		PxVec3 PLocalCOM = U2PVector(GetLocalCOM());
-		PxVehicleComputeSprungMasses(WheelSetups.Num(), WheelOffsets, PLocalCOM, PVehicleActor->getMass(), 2, SprungMasses);
+		PxVehicleComputeSprungMasses(NumWheels, WheelOffsets, PLocalCOM, PVehicleActor->getMass(), /*gravityDirection=*/2, SprungMasses);
 
-		for (int32 WheelIdx = 0; WheelIdx < WheelSetups.Num(); ++WheelIdx)
+		for (int32 WheelIdx = 0; WheelIdx < NumWheels; ++WheelIdx)
 		{
 			UVehicleWheel* Wheel = WheelSetups[WheelIdx].WheelClass.GetDefaultObject();
 
@@ -366,7 +389,8 @@ void UWheeledVehicleMovementComponent::SetupWheels( PxVehicleWheelsSimData* PWhe
 
 			PxVec3 PSuspTravelDirection = PxVec3(0.0f, 0.0f, -1.0f);
 			PxVec3 PWheelCentreCMOffset = PWheelOffset - PLocalCOM;
-			PxVec3 PSuspForceAppCMOffset = PxVec3(PWheelCentreCMOffset.x, PWheelCentreCMOffset.y, Wheel->SuspensionForceOffset);
+			PxVec3 PSuspForceAppCMOffset = !bDeprecatedSpringOffsetMode ? PxVec3(PWheelCentreCMOffset.x, PWheelCentreCMOffset.y, Wheel->SuspensionForceOffset + PWheelCentreCMOffset.z)
+																							 : PxVec3(PWheelCentreCMOffset.x, PWheelCentreCMOffset.y, Wheel->SuspensionForceOffset);
 			PxVec3 PTireForceAppCMOffset = PSuspForceAppCMOffset;
 
 			// finalize sim data
@@ -380,22 +404,31 @@ void UWheeledVehicleMovementComponent::SetupWheels( PxVehicleWheelsSimData* PWhe
 		}
 
 		const int32 NumShapes = PVehicleActor->getNbShapes();
-		const int32 NumChassisShapes = NumShapes - WheelSetups.Num();
-		check(NumChassisShapes >= 1);
+		const int32 NumChassisShapes = NumShapes - NumWheels;
+		if(NumChassisShapes >= 1)
+		{
+			TArray<PxShape*> Shapes;
+			Shapes.AddZeroed(NumShapes);
 
-		TArray<PxShape*> Shapes;
-		Shapes.AddZeroed(NumShapes);
+			PVehicleActor->getShapes(Shapes.GetData(), NumShapes);
 
-		PVehicleActor->getShapes(Shapes.GetData(), NumShapes);
+			for (int32 WheelIdx = 0; WheelIdx < NumWheels; ++WheelIdx)
+			{
+				const int32 WheelShapeIndex = NumChassisShapes + WheelIdx;
 
-		for (int32 WheelIdx = 0; WheelIdx < WheelSetups.Num(); ++WheelIdx)
+				PWheelsSimData->setWheelShapeMapping(WheelIdx, WheelShapeIndex);
+				PWheelsSimData->setSceneQueryFilterData(WheelIdx, Shapes[WheelShapeIndex]->getQueryFilterData());
+			}
+		}
+		else
 		{
-			const int32 WheelShapeIndex = NumChassisShapes + WheelIdx;
-
-			PWheelsSimData->setWheelShapeMapping(WheelIdx, WheelShapeIndex);
-			PWheelsSimData->setSceneQueryFilterData(WheelIdx, Shapes[WheelShapeIndex]->getQueryFilterData());
+#if !(UE_BUILD_SHIPPING || UE_BUILD_TEST)
+			UE_LOG(LogPhysics, Warning, TEXT("Missing wheel shapes. Please ensure there's a body associated with each wheel, or deselect Don'tCreateShape in your wheel class for vehicle''%s''"), *GetPathNameSafe(this));
+#endif
 		}
 
+		
+
 		// tire load filtering
 		PxVehicleTireLoadFilterData PTireLoadFilter;
 		PTireLoadFilter.mMinNormalisedLoad = MinNormalizedTireLoad;
@@ -642,6 +675,10 @@ void UWheeledVehicleMovementComponent::OnCreatePhysicsState()
 				if(USkeletalMeshComponent* MeshComp = Cast<USkeletalMeshComponent>(GetMesh()))
 				{
 					MeshOnPhysicsStateChangeHandle = MeshComp->RegisterOnPhysicsCreatedDelegate(FOnSkelMeshPhysicsCreated::CreateUObject(this, &UWheeledVehicleMovementComponent::RecreatePhysicsState));
+					if(UVehicleAnimInstance* VehicleAnimInstance = Cast<UVehicleAnimInstance>(MeshComp->GetAnimInstance()))
+					{
+						VehicleAnimInstance->SetWheeledVehicleMovementComponent(this);
+					}
 				}
 			}
 		}
@@ -805,6 +842,46 @@ void UWheeledVehicleMovementComponent::PreTick(float DeltaTime)
 	}
 }
 
+void UWheeledVehicleMovementComponent::SetupVehicle()
+{
+	if (!UpdatedPrimitive)
+	{
+		return;
+	}
+
+	if (WheelSetups.Num() == 0)
+	{
+		PVehicle = nullptr;
+		PVehicleDrive = nullptr;
+		return;
+	}
+
+	for (int32 WheelIdx = 0; WheelIdx < WheelSetups.Num(); ++WheelIdx)
+	{
+		const FWheelSetup& WheelSetup = WheelSetups[WheelIdx];
+		if (WheelSetup.BoneName == NAME_None)
+		{
+			return;
+		}
+	}
+
+	// Setup the chassis and wheel shapes
+	SetupVehicleShapes();
+
+	// Setup mass properties
+	SetupVehicleMass();
+
+	// Setup the wheels
+	PxVehicleWheelsSimData* PWheelsSimData = PxVehicleWheelsSimData::allocate(WheelSetups.Num());
+	SetupWheels(PWheelsSimData);
+
+	SetupVehicleDrive(PWheelsSimData);
+}
+
+void UWheeledVehicleMovementComponent::SetupVehicleDrive(PxVehicleWheelsSimData* PWheelsSimData)
+{	
+}
+
 void UWheeledVehicleMovementComponent::UpdateSimulation( float DeltaTime )
 {
 }
@@ -850,19 +927,23 @@ void UWheeledVehicleMovementComponent::UpdateState( float DeltaTime )
 	APawn* MyOwner = UpdatedComponent ? Cast<APawn>(UpdatedComponent->GetOwner()) : NULL;
 	if (MyOwner && MyOwner->IsLocallyControlled())
 	{
-		//Manual shifting between reverse and first gear
-		if (FMath::Abs(GetForwardSpeed()) < WrongDirectionThreshold)	//we only shift between reverse and first if the car is slow enough. This isn't 100% correct since we really only care about engine speed, but good enough
+		if(bReverseAsBrake)
 		{
-			if (RawThrottleInput < 0.f && GetCurrentGear() >= 0 && GetTargetGear() >= 0)
-			{
-				SetTargetGear(-1, true);
-			}
-			else if (RawThrottleInput > 0.f && GetCurrentGear() <= 0 && GetTargetGear() <= 0)
+			//for reverse as state we want to automatically shift between reverse and first gear
+			if (FMath::Abs(GetForwardSpeed()) < WrongDirectionThreshold)	//we only shift between reverse and first if the car is slow enough. This isn't 100% correct since we really only care about engine speed, but good enough
 			{
-				SetTargetGear(1, true);
+				if (RawThrottleInput < 0.f && GetCurrentGear() >= 0 && GetTargetGear() >= 0)
+				{
+					SetTargetGear(-1, true);
+				}
+				else if (RawThrottleInput > 0.f && GetCurrentGear() <= 0 && GetTargetGear() <= 0)
+				{
+					SetTargetGear(1, true);
+				}
 			}
 		}
 		
+		
 		if (bUseRVOAvoidance)
 		{
 			CalculateAvoidanceVelocity(DeltaTime);
@@ -933,45 +1014,53 @@ float UWheeledVehicleMovementComponent::CalcSteeringInput()
 
 float UWheeledVehicleMovementComponent::CalcBrakeInput()
 {	
-	const float ForwardSpeed = GetForwardSpeed();
-
-	float NewBrakeInput = 0.0f;
-
-	// if player wants to move forwards...
-	if ( RawThrottleInput > 0.f )
+	if(bReverseAsBrake)
 	{
-		// if vehicle is moving backwards, then press brake
-		if ( ForwardSpeed < -WrongDirectionThreshold)
-		{
-			NewBrakeInput = 1.0f;			
-		}
+		const float ForwardSpeed = GetForwardSpeed();
 
-	}
+		float NewBrakeInput = 0.0f;
 
-	// if player wants to move backwards...
-	else if ( RawThrottleInput < 0.f )
-	{
-		// if vehicle is moving forwards, then press brake
-		if (ForwardSpeed > WrongDirectionThreshold)
+		// if player wants to move forwards...
+		if (RawThrottleInput > 0.f)
 		{
-			NewBrakeInput = 1.0f;			// Seems a bit severe to have 0 or 1 braking. Better control can be had by allowing continuous brake input values
-		}	
-	}
+			// if vehicle is moving backwards, then press brake
+			if (ForwardSpeed < -WrongDirectionThreshold)
+			{
+				NewBrakeInput = 1.0f;
+			}
 
-	// if player isn't pressing forward or backwards...
-	else
-	{
-		if (ForwardSpeed < StopThreshold && ForwardSpeed > -StopThreshold)	//auto break 
+		}
+
+		// if player wants to move backwards...
+		else if (RawThrottleInput < 0.f)
 		{
-			NewBrakeInput = 1.f;
+			// if vehicle is moving forwards, then press brake
+			if (ForwardSpeed > WrongDirectionThreshold)
+			{
+				NewBrakeInput = 1.0f;			// Seems a bit severe to have 0 or 1 braking. Better control can be had by allowing continuous brake input values
+			}
 		}
+
+		// if player isn't pressing forward or backwards...
 		else
 		{
-			NewBrakeInput = IdleBrakeInput;
+			if (ForwardSpeed < StopThreshold && ForwardSpeed > -StopThreshold)	//auto break 
+			{
+				NewBrakeInput = 1.f;
+			}
+			else
+			{
+				NewBrakeInput = IdleBrakeInput;
+			}
 		}
-	}
 
-	return FMath::Clamp<float>(NewBrakeInput, 0.0, 1.0);
+		return FMath::Clamp<float>(NewBrakeInput, 0.0, 1.0);
+	}
+	else
+	{
+		return FMath::Abs(RawBrakeInput);
+	}
+	
 }
 
 float UWheeledVehicleMovementComponent::CalcHandbrakeInput()
@@ -996,10 +1085,13 @@ float UWheeledVehicleMovementComponent::CalcThrottleInput()
 		}		
 	}
 
-	//If the user is changing direction we should really be braking first and not applying any gas, so wait until they've changed gears
-	if ( (RawThrottleInput > 0.f && GetTargetGear() < 0) || (RawThrottleInput < 0.f && GetTargetGear() > 0)) 
+	if(bReverseAsBrake)
 	{
-		return 0.f;
+		//If the user is changing direction we should really be braking first and not applying any gas, so wait until they've changed gears
+		if ((RawThrottleInput > 0.f && GetTargetGear() < 0) || (RawThrottleInput < 0.f && GetTargetGear() > 0))
+		{
+			return 0.f;
+		}
 	}
 
 	return FMath::Abs(RawThrottleInput);
@@ -1018,6 +1110,12 @@ void UWheeledVehicleMovementComponent::SetThrottleInput( float Throttle )
 	RawThrottleInput = FMath::Clamp( Throttle, -1.0f, 1.0f );
 }
 
+void UWheeledVehicleMovementComponent::SetBrakeInput(float Brake)
+{
+	RawBrakeInput = FMath::Clamp(Brake, -1.0f, 1.0f);
+}
+
+
 void UWheeledVehicleMovementComponent::SetSteeringInput( float Steering )
 {
 	RawSteeringInput = FMath::Clamp( Steering, -1.0f, 1.0f );
@@ -1291,6 +1389,18 @@ float UWheeledVehicleMovementComponent::GetMaxSpringForce() const
 
 }
 
+void UWheeledVehicleMovementComponent::Serialize(FArchive& Ar)
+{
+	Super::Serialize(Ar);
+
+	Ar.UsingCustomVersion(FFrameworkObjectVersion::GUID);
+
+	if (Ar.CustomVer(FFrameworkObjectVersion::GUID) < FFrameworkObjectVersion::WheelOffsetIsFromWheel)
+	{
+		bDeprecatedSpringOffsetMode = true;	//Existing content is tuned with the old way of applying spring force offset. There's no easy way to re-compute this at the wheel level since it's a shared asset
+	}
+}
+
 void UWheeledVehicleMovementComponent::DrawDebug(UCanvas* Canvas, float& YL, float& YPos)
 {
 	if (PVehicle == NULL)
diff --git a/Engine/Source/Runtime/Engine/Private/Vehicles/WheeledVehicleMovementComponent4W.cpp b/Engine/Source/Runtime/Engine/Private/Vehicles/WheeledVehicleMovementComponent4W.cpp
--- a/Engine/Source/Runtime/Engine/Private/Vehicles/WheeledVehicleMovementComponent4W.cpp
+++ b/Engine/Source/Runtime/Engine/Private/Vehicles/WheeledVehicleMovementComponent4W.cpp
@@ -242,39 +242,15 @@ void SetupDriveHelper(const UWheeledVehicleMovementComponent4W* VehicleData, con
 	DriveData.setAutoBoxData(AutoBoxSetup);
 }
 
-void UWheeledVehicleMovementComponent4W::SetupVehicle()
+void UWheeledVehicleMovementComponent4W::SetupVehicleDrive(PxVehicleWheelsSimData* PWheelsSimData)
 {
-	if (!UpdatedPrimitive)
+	if(WheelSetups.Num() != 4)
 	{
+		PVehicle = nullptr;
+		PVehicleDrive = nullptr;
 		return;
 	}
 
-	if (WheelSetups.Num() != 4)
-	{
-		PVehicle = NULL;
-		PVehicleDrive = NULL;
-		return;
-	}
-
-	for (int32 WheelIdx = 0; WheelIdx < WheelSetups.Num(); ++WheelIdx)
-	{
-		const FWheelSetup& WheelSetup = WheelSetups[WheelIdx];
-		if (WheelSetup.BoneName == NAME_None)
-		{
-			return;
-		}
-	}
-
-	// Setup the chassis and wheel shapes
-	SetupVehicleShapes();
-
-	// Setup mass properties
-	SetupVehicleMass();
-
-	// Setup the wheels
-	PxVehicleWheelsSimData* PWheelsSimData = PxVehicleWheelsSimData::allocate(4);
-	SetupWheels(PWheelsSimData);
-
 	// Setup drive data
 	PxVehicleDriveSimData4W DriveData;
 	SetupDriveHelper(this, PWheelsSimData, DriveData);
@@ -291,8 +267,6 @@ void UWheeledVehicleMovementComponent4W::SetupVehicle()
 		// cleanup
 		PWheelsSimData->free();
 	});
-	
-	PWheelsSimData = NULL;
 
 	// cache values
 	PVehicle = PVehicleDrive4W;
diff --git a/Engine/Source/Runtime/Engine/Public/ActiveSound.h b/Engine/Source/Runtime/Engine/Public/ActiveSound.h
--- a/Engine/Source/Runtime/Engine/Public/ActiveSound.h
+++ b/Engine/Source/Runtime/Engine/Public/ActiveSound.h
@@ -28,6 +28,9 @@ struct FSoundParseParameters
 	
 	float Volume;
 	float VolumeMultiplier;
+	
+	// Volume due to application-level volume scaling (tabbing, master volume)
+	float VolumeApp;
 
 	// The multiplier to apply if the sound class desires
 	float InteriorVolumeMultiplier;
@@ -78,6 +81,7 @@ struct FSoundParseParameters
 		, Velocity(ForceInit)
 		, Volume(1.f)
 		, VolumeMultiplier(1.f)
+		, VolumeApp(1.f)
 		, InteriorVolumeMultiplier(1.f)
 		, Pitch(1.f)
 		, StartTime(-1.f)
@@ -238,10 +242,6 @@ public:
 	/** Whether or not to use an async trace for occlusion */
 	uint8 bOcclusionAsyncTrace : 1;
 
-private:
-	/** Whether or not this sound is audible */
-	uint8 bIsAudible : 1;
-
 public:
 	uint8 UserIndex;
 
diff --git a/Engine/Source/Runtime/Engine/Public/Audio.h b/Engine/Source/Runtime/Engine/Public/Audio.h
--- a/Engine/Source/Runtime/Engine/Public/Audio.h
+++ b/Engine/Source/Runtime/Engine/Public/Audio.h
@@ -198,6 +198,8 @@ struct ENGINE_API FWaveInstance
 	float				Volume;
 	/** Current volume multiplier - used to zero the volume without stopping the source */
 	float				VolumeMultiplier;
+	/** The volume of the wave instance due to application volume or tab-state */
+	float				VolumeApp;
 	/** An audio component priority value that scales with volume (post all gain stages) and is used to determine voice playback priority. */
 	float				Priority;
 	/** Voice center channel volume */
@@ -302,9 +304,12 @@ struct ENGINE_API FWaveInstance
 	/** Returns the actual volume the wave instance will play at */
 	bool ShouldStopDueToMaxConcurrency() const;
 
-	/** Returns the actual volume the wave instance will play at */
+	/** Returns the actual volume the wave instance will play at, including all gain stages. */
 	float GetActualVolume() const;
 
+	/** Returns the volume of the wave instance (ignoring application muting) */
+	float GetVolume() const;
+
 	/** Returns the weighted priority of the wave instance. */
 	float GetVolumeWeightedPriority() const;
 
diff --git a/Engine/Source/Runtime/Engine/Public/AudioDevice.h b/Engine/Source/Runtime/Engine/Public/AudioDevice.h
--- a/Engine/Source/Runtime/Engine/Public/AudioDevice.h
+++ b/Engine/Source/Runtime/Engine/Public/AudioDevice.h
@@ -998,6 +998,8 @@ public:
 	static bool CanUseVRAudioDevice();
 
 #if !UE_BUILD_SHIPPING
+	void DumpActiveSounds() const;
+
 	void RenderStatReverb(UWorld* World, FViewport* Viewport, FCanvas* Canvas, int32 X, int32& Y, const FVector* ViewLocation, const FRotator* ViewRotation) const;
 
 	void UpdateSoundShowFlags(const uint8 OldSoundShowFlags, const uint8 NewSoundShowFlags);
@@ -1034,6 +1036,9 @@ public:
 		return SoundClasses;
 	}
 
+	/** Whether or not virtual sounds are enabled, */
+	bool VirtualSoundsEnabled() const { return bAllowVirtualizedSounds; }
+
 public:
 
 	/** The maximum number of concurrent audible sounds */
@@ -1158,6 +1163,9 @@ private:
 	/** Whether the value in HighestPriorityActivatedReverb should be used - Audio Thread owned */
 	uint8 bHasActivatedReverb:1;
 
+	/** Whether or not we're supporting zero volume wave instances */
+	uint8 bAllowVirtualizedSounds:1;
+
 #if !UE_BUILD_SHIPPING
 	uint8 RequestedAudioStats;
 
diff --git a/Engine/Source/Runtime/Engine/Public/Engine.h b/Engine/Source/Runtime/Engine/Public/Engine.h
--- a/Engine/Source/Runtime/Engine/Public/Engine.h
+++ b/Engine/Source/Runtime/Engine/Public/Engine.h
@@ -80,7 +80,6 @@
 #include "GameFramework/FloatingPawnMovement.h"
 #include "GameFramework/DefaultPawn.h"
 #include "GameFramework/SpectatorPawn.h"
-#include "GameFramework/WheeledVehicle.h"
 #include "Engine/ReflectionCapture.h"
 #include "Engine/BoxReflectionCapture.h"
 #include "Engine/PlaneReflectionCapture.h"
@@ -99,9 +98,6 @@
 #include "Engine/DataAsset.h"
 #include "Engine/TextureStreamingTypes.h"
 #include "GameFramework/SpectatorPawnMovement.h"
-#include "Vehicles/VehicleWheel.h"
-#include "Vehicles/WheeledVehicleMovementComponent.h"
-#include "Vehicles/WheeledVehicleMovementComponent4W.h"
 #include "GameFramework/ProjectileMovementComponent.h"
 #include "GameFramework/RotatingMovementComponent.h"
 #include "Components/PawnNoiseEmitterComponent.h"
diff --git a/Engine/Source/Runtime/Engine/Public/SkeletalRenderPublic.h b/Engine/Source/Runtime/Engine/Public/SkeletalRenderPublic.h
--- a/Engine/Source/Runtime/Engine/Public/SkeletalRenderPublic.h
+++ b/Engine/Source/Runtime/Engine/Public/SkeletalRenderPublic.h
@@ -94,10 +94,12 @@ public:
 	/** 
 	 * Enable blend weight rendering in the editor
 	 * @param bEnabled - turn on or off the rendering mode
-	 * @param BonesOfInterest - array of bone indices to capture weights for
+	 * optional parameters will decide which one to draw
+	 * @param (optional) BonesOfInterest - array of bone indices to capture weights for
+	 * @param (optional) MorphTargetsOfInterest - array of morphtargets to render for
 	 */
-	virtual void EnableBlendWeightRendering(bool bEnabled, const TArray<int32>& InBonesOfInterest) {}
-	
+	virtual void EnableOverlayRendering(bool bEnabled, const TArray<int32>* InBonesOfInterest, const TArray<UMorphTarget*>* MorphTargetOfInterest) {}
+
 	/** 
 	* Draw Normals/Tangents based on skinned vertex data 
 	* @param PDI			- Draw Interface
diff --git a/Engine/Source/Runtime/Navmesh/Private/DetourTileCache/DetourTileCacheBuilder.cpp b/Engine/Source/Runtime/Navmesh/Private/DetourTileCache/DetourTileCacheBuilder.cpp
--- a/Engine/Source/Runtime/Navmesh/Private/DetourTileCache/DetourTileCacheBuilder.cpp
+++ b/Engine/Source/Runtime/Navmesh/Private/DetourTileCache/DetourTileCacheBuilder.cpp
@@ -722,8 +722,11 @@ dtStatus dtBuildTileCacheContours(dtTileCacheAlloc* alloc, dtTileCacheLayer& lay
 				{
 					const int ax = x + getDirOffsetX(dir);
 					const int ay = y + getDirOffsetY(dir);
-					const int aidx = ax+ay*w;
-					r = layer.regs[aidx];
+					if (ax >= 0 && ay >= 0 && ax < w && ay < h)
+					{
+						const int aidx = ax + ay*w;
+						r = layer.regs[aidx];
+					}
 				}
 
 				if (r == ri)
diff --git a/Engine/Source/Runtime/Navmesh/Private/Recast/RecastLayers.cpp b/Engine/Source/Runtime/Navmesh/Private/Recast/RecastLayers.cpp
--- a/Engine/Source/Runtime/Navmesh/Private/Recast/RecastLayers.cpp
+++ b/Engine/Source/Runtime/Navmesh/Private/Recast/RecastLayers.cpp
@@ -56,7 +56,36 @@ inline bool overlapRange(const unsigned short amin, const unsigned short amax,
 	return (amin > bmax || amax < bmin) ? false : true;
 }
 
+static void fixLayerConnections(rcHeightfieldLayer* layer)
+{
+	// [UE4: break one directional connections, contour tracing gets stuck in infinite loop]
+	const int lw = layer->width;
+	const int lh = layer->height;
 
+	for (int y = 0; y < lh; ++y)
+	{
+		for (int x = 0; x < lw; ++x)
+		{
+			const int idx = x + y*lw;
+			const int con = layer->cons[idx];
+
+			for (int dir = 0; dir < 4; ++dir)
+			{
+				if ((con & (1 << dir)) == 0)
+				{
+					const int nx = x + rcGetDirOffsetX(dir);
+					const int ny = y + rcGetDirOffsetY(dir);
+					if (nx >= 0 && ny >= 0 && nx < lw && ny < lh)
+					{
+						const int nidx = nx + ny*lw;
+						const int oppDir = (dir + 2) % 4;
+						layer->cons[nidx] &= ~(1 << oppDir);
+					}
+				}
+			}
+		}
+	}
+}
 
 struct rcLayerSweepSpan
 {
@@ -787,10 +816,6 @@ static bool SplitAndStoreLayerRegions(rcContext* ctx, rcCompactHeightfield& chf,
 								if (nx >= 0 && ny >= 0 && nx < lw && ny < lh)
 								{
 									con |= (unsigned char)(1 << dir);
-
-									// [UE4: make sure that connections are bidirectional, otherwise contour tracing will stuck in infinite loop]
-									const int nidx = nx + (ny * lw);
-									layer->cons[nidx] |= (unsigned char)(1 << ((dir + 2) % 4));
 								}
 							}
 						}
@@ -801,6 +826,8 @@ static bool SplitAndStoreLayerRegions(rcContext* ctx, rcCompactHeightfield& chf,
 			}
 		}
 
+		fixLayerConnections(layer);
+
 		if (layer->minx > layer->maxx)
 			layer->minx = layer->maxx = 0;
 		if (layer->miny > layer->maxy)
@@ -1416,12 +1443,12 @@ bool rcBuildHeightfieldLayers(rcContext* ctx, rcCompactHeightfield& chf,
 						{
 							const int ax = cx + rcGetDirOffsetX(dir);
 							const int ay = cy + rcGetDirOffsetY(dir);
-							const int ai = (int)chf.cells[ax+ay*w].index + rcGetCon(s, dir);
+							const int ai = (int)chf.cells[ax + ay*w].index + rcGetCon(s, dir);
 							unsigned short alid = (srcReg[ai] < nreg) ? regions[srcReg[ai]].layerId : 0xffff;
 							// Portal mask
 							if (chf.areas[ai] != RC_NULL_AREA && lid != alid)
 							{
-								portal |= (unsigned char)(1<<dir);
+								portal |= (unsigned char)(1 << dir);
 								// Update height so that it matches on both sides of the portal.
 								const rcCompactSpan& as = chf.spans[ai];
 								if (as.y > hmin)
@@ -1433,7 +1460,9 @@ bool rcBuildHeightfieldLayers(rcContext* ctx, rcCompactHeightfield& chf,
 								const int nx = ax - borderSize;
 								const int ny = ay - borderSize;
 								if (nx >= 0 && ny >= 0 && nx < lw && ny < lh)
-									con |= (unsigned char)(1<<dir);
+								{
+									con |= (unsigned char)(1 << dir);
+								}
 							}
 						}
 					}
@@ -1443,6 +1472,8 @@ bool rcBuildHeightfieldLayers(rcContext* ctx, rcCompactHeightfield& chf,
 			}
 		}
 
+		fixLayerConnections(layer);
+
 		if (layer->minx > layer->maxx)
 			layer->minx = layer->maxx = 0;
 		if (layer->miny > layer->maxy)
diff --git a/Engine/Source/Runtime/Navmesh/Public/Recast/Recast.h b/Engine/Source/Runtime/Navmesh/Public/Recast/Recast.h
--- a/Engine/Source/Runtime/Navmesh/Public/Recast/Recast.h
+++ b/Engine/Source/Runtime/Navmesh/Public/Recast/Recast.h
@@ -355,8 +355,8 @@ struct rcCompactSpan
 {
 	unsigned short y;			///< The lower extent of the span. (Measured from the heightfield's base.)
 	unsigned short reg;			///< The id of the region the span belongs to. (Or zero if not in a region.)
-	unsigned int con : 24;		///< Packed neighbor connection data.
-	unsigned int h : 8;			///< The height of the span.  (Measured from #y.)
+	unsigned int con;			///< Packed neighbor connection data.
+	unsigned char h;			///< The height of the span.  (Measured from #y.)
 };
 
 /// A compact, static heightfield representing unobstructed space.
@@ -625,7 +625,7 @@ static const unsigned char RC_WALKABLE_AREA = 63;
 
 /// The value returned by #rcGetCon if the specified direction is not connected
 /// to another span. (Has no neighbor.)
-static const int RC_NOT_CONNECTED = 0x3f;
+static const int RC_NOT_CONNECTED = 0xff;
 
 /// @name General helper functions
 /// @{
@@ -1180,9 +1180,9 @@ NAVMESH_API bool rcBuildRegionsChunky(rcContext* ctx, rcCompactHeightfield& chf,
 ///  @param[in]		i		The index of the neighbor span.
 inline void rcSetCon(rcCompactSpan& s, int dir, int i)
 {
-	const unsigned int shift = (unsigned int)dir*6;
+	const unsigned int shift = (unsigned int)dir * 8;
 	unsigned int con = s.con;
-	s.con = (con & ~(0x3f << shift)) | (((unsigned int)i & 0x3f) << shift);
+	s.con = (con & ~(0xff << shift)) | (((unsigned int)i & 0xff) << shift);
 }
 
 /// Gets neighbor connection data for the specified direction.
@@ -1192,8 +1192,8 @@ inline void rcSetCon(rcCompactSpan& s, int dir, int i)
 ///  	or #RC_NOT_CONNECTED if there is no connection.
 inline int rcGetCon(const rcCompactSpan& s, int dir)
 {
-	const unsigned int shift = (unsigned int)dir*6;
-	return (s.con >> shift) & 0x3f;
+	const unsigned int shift = (unsigned int)dir * 8;
+	return (s.con >> shift) & 0xff;
 }
 
 /// Gets the standard width (x-axis) offset for the specified direction.
