Only set the updated component if in a game world

diff --git a/Engine/Source/Runtime/Engine/Private/Components/MovementComponent.cpp b/Engine/Source/Runtime/Engine/Private/Components/MovementComponent.cpp
--- a/Engine/Source/Runtime/Engine/Private/Components/MovementComponent.cpp
+++ b/Engine/Source/Runtime/Engine/Private/Components/MovementComponent.cpp
@@ -112,20 +112,20 @@ void UMovementComponent::OnRegister()
 	if (MyWorld && MyWorld->IsGameWorld())
 	{
 		PlaneConstraintNormal = PlaneConstraintNormal.GetSafeNormal();
-	}
 
-	USceneComponent* NewUpdatedComponent = UpdatedComponent;
-	if (!UpdatedComponent && bAutoRegisterUpdatedComponent)
-	{
-		// Auto-register owner's root component if found.
-		AActor* MyActor = GetOwner();
-		if (MyActor)
+		USceneComponent* NewUpdatedComponent = UpdatedComponent;
+		if (!UpdatedComponent && bAutoRegisterUpdatedComponent)
 		{
-			NewUpdatedComponent = MyActor->GetRootComponent();
+			// Auto-register owner's root component if found.
+			AActor* MyActor = GetOwner();
+			if (MyActor)
+			{
+				NewUpdatedComponent = MyActor->GetRootComponent();
+			}
 		}
-	}
 
-	SetUpdatedComponent(NewUpdatedComponent);
+		SetUpdatedComponent(NewUpdatedComponent);
+	}
 
 #if WITH_EDITOR
 	// Reset so next PIE session warns.
