Copy - Fixed DistanceFromPlaneFadeEnd_DEPRECATED lost default in existing content, causing planar reflections to disappear

diff --git a/Engine/Source/Runtime/Engine/Private/Components/SceneCaptureComponent.cpp b/Engine/Source/Runtime/Engine/Private/Components/SceneCaptureComponent.cpp
--- a/Engine/Source/Runtime/Engine/Private/Components/SceneCaptureComponent.cpp
+++ b/Engine/Source/Runtime/Engine/Private/Components/SceneCaptureComponent.cpp
@@ -579,6 +579,8 @@ UPlanarReflectionComponent::UPlanarReflectionComponent(const FObjectInitializer&
 	PrefilterRoughnessDistance = 10000;
 	ScreenPercentage = 50;
 	NormalDistortionStrength = 500;
+	DistanceFromPlaneFadeStart_DEPRECATED = 400;
+	DistanceFromPlaneFadeEnd_DEPRECATED = 600;
 	DistanceFromPlaneFadeoutStart = 60;
 	DistanceFromPlaneFadeoutEnd = 100;
 	AngleFromPlaneFadeStart = 20;
