Fix audio occlusion to be (mostly) thread safe and (mostly) work correctly

diff --git a/Engine/Source/Runtime/Engine/Private/ActiveSound.cpp b/Engine/Source/Runtime/Engine/Private/ActiveSound.cpp
--- a/Engine/Source/Runtime/Engine/Private/ActiveSound.cpp
+++ b/Engine/Source/Runtime/Engine/Private/ActiveSound.cpp
@@ -22,7 +22,6 @@ FActiveSound::FActiveSound()
 	, ConcurrencySettings(nullptr)
 	, SoundClassOverride(nullptr)
 	, bHasCheckedOcclusion(false)
-	, bIsTraceDelegateBound(false)
 	, bAllowSpatialization(true)
 	, bHasAttenuationSettings(false)
 	, bShouldRemainActiveIfDropped(false)
@@ -47,7 +46,6 @@ FActiveSound::FActiveSound()
 	, bWarnedAboutOrphanedLooping(false)
 #endif
 	, bEnableLowPassFilter(false)
-	, bOcclusionAsyncTrace(true)
 	, UserIndex(0)
 	, bIsOccluded(false)
 	, bAsyncOcclusionPending(false)
@@ -69,7 +67,7 @@ FActiveSound::FActiveSound()
 	, FocusDistanceScale(1.0f)
 	, VolumeConcurrency(0.0f)
 	, OcclusionCheckInterval(0.f)
-	, LastOcclusionCheckTime(0.f)
+	, LastOcclusionCheckTime(TNumericLimits<float>::Lowest())
 	, MaxDistance(WORLD_MAX)
 	, LastLocation(FVector::ZeroVector)
 	, AudioVolumeID(0)
@@ -427,25 +425,31 @@ void FActiveSound::UpdateOcclusion(const FAttenuationSettings* AttenuationSettin
 		CurrentOcclusionVolumeAttenuation.Set(1.0f, InterpolationTime);
 	}
 
-	UWorld* WorldPtr = World.Get();
-	CurrentOcclusionFilterFrequency.Update(WorldPtr->DeltaTimeSeconds);
-	CurrentOcclusionVolumeAttenuation.Update(WorldPtr->DeltaTimeSeconds);
+	const float DeltaTime = FApp::GetDeltaTime();
+	CurrentOcclusionFilterFrequency.Update(DeltaTime);
+	CurrentOcclusionVolumeAttenuation.Update(DeltaTime);
 }
 
 void FActiveSound::OcclusionTraceDone(const FTraceHandle& TraceHandle, FTraceDatum& TraceDatum)
 {
 	// Look for any results that resulted in a blocking hit
-	bIsOccluded = false;
+	bool bFoundBlockingHit = false;
 	for (const FHitResult& HitResult : TraceDatum.OutHits)
 	{
 		if (HitResult.bBlockingHit)
 		{
-			bIsOccluded = true;
+			bFoundBlockingHit = true;
 			break;
 		}
 	}
 
-	bAsyncOcclusionPending = false;
+	FActiveSound* ActiveSound = this;
+
+	FAudioThread::RunCommandOnAudioThread([ActiveSound, bFoundBlockingHit]()
+	{
+		ActiveSound->bIsOccluded = bFoundBlockingHit;
+		ActiveSound->bAsyncOcclusionPending = false;
+	});
 }
 
 void FActiveSound::CheckOcclusion(const FVector ListenerLocation, const FVector SoundLocation, const FAttenuationSettings* AttenuationSettingsPtr)
@@ -453,49 +457,52 @@ void FActiveSound::CheckOcclusion(const FVector ListenerLocation, const FVector
 	check(AttenuationSettingsPtr);
 	check(AttenuationSettingsPtr->bEnableOcclusion);
 
-	UWorld* WorldPtr = World.Get();
-	float WorldTime = WorldPtr->GetTimeSeconds();
-
-	if ((WorldTime - LastOcclusionCheckTime) > OcclusionCheckInterval)
+	if (!bAsyncOcclusionPending && (PlaybackTime - LastOcclusionCheckTime) > OcclusionCheckInterval)
 	{
-		LastOcclusionCheckTime = WorldTime;
+		LastOcclusionCheckTime = PlaybackTime;
 		static FName NAME_SoundOcclusion = FName(TEXT("SoundOcclusion"));
 
-		FCollisionQueryParams Params(NAME_SoundOcclusion, AttenuationSettingsPtr->bUseComplexCollisionForOcclusion);
-		if (OwnerID > 0)
-		{
-			Params.AddIgnoredActor(OwnerID);
-		}
+		const bool bUseComplexCollisionForOcclusion = AttenuationSettingsPtr->bUseComplexCollisionForOcclusion;
+		const ECollisionChannel OcclusionTraceChannel = AttenuationSettingsPtr->OcclusionTraceChannel;
 
-		if (bOcclusionAsyncTrace)
+		if (!OcclusionTraceDelegate.IsBound())
 		{
-			ECollisionChannel OcclusionTraceChannel = AttenuationSettingsPtr->OcclusionTraceChannel;
+			OcclusionTraceDelegate.BindRaw(this, &FActiveSound::OcclusionTraceDone);
 
-			// If the audio thread is running, then only do a sync trace
-			if (FAudioThread::IsAudioThreadRunning())
+			FCollisionQueryParams Params(NAME_SoundOcclusion, bUseComplexCollisionForOcclusion);
+			if (OwnerID > 0)
 			{
-				bIsOccluded = WorldPtr->LineTraceTestByChannel(SoundLocation, ListenerLocation, OcclusionTraceChannel, Params);
+				Params.AddIgnoredActor(OwnerID);
 			}
-			else
+
+			if (UWorld* WorldPtr = World.Get())
 			{
-				// Check if we've not already bound our trace delegate
-				if (!bIsTraceDelegateBound)
-				{
-					bIsTraceDelegateBound = true;
+				// LineTraceTestByChannel is generally threadsafe, but there is a very narrow race condition here 
+				// if World goes invalid before the scene lock and queries begin.
+				bIsOccluded = WorldPtr->LineTraceTestByChannel(SoundLocation, ListenerLocation, OcclusionTraceChannel, Params);
+			}
+		}
+		else
+		{
+			bAsyncOcclusionPending = true;
 
-					// Bind our async occlusion trace delegate (so next update we'll have it bound)
-					OcclusionTraceDelegate.BindRaw(this, &FActiveSound::OcclusionTraceDone);
+			const uint32 SoundOwnerID = OwnerID;
+			FTraceDelegate* TraceDelegate = &OcclusionTraceDelegate;
+			TWeakObjectPtr<UWorld> SoundWorld = World;
 
-					// Only do async occlusion trace if we've already made one. The first trace must be synchronous to avoid issues with sounds starting playing as occluded
-					bIsOccluded = WorldPtr->LineTraceTestByChannel(SoundLocation, ListenerLocation, OcclusionTraceChannel, Params);
-				}
-				// don't need to do another async trace if we've already got one pending
-				else if (!bAsyncOcclusionPending)
+			FAudioThread::RunCommandOnGameThread([SoundWorld, SoundLocation, ListenerLocation, OcclusionTraceChannel, SoundOwnerID, bUseComplexCollisionForOcclusion, TraceDelegate]
+			{
+				if (UWorld* WorldPtr = SoundWorld.Get())
 				{
-					bAsyncOcclusionPending = true;
-					WorldPtr->AsyncLineTraceByChannel(EAsyncTraceType::Test, SoundLocation, ListenerLocation, OcclusionTraceChannel, Params, FCollisionResponseParams::DefaultResponseParam, &OcclusionTraceDelegate);
+					FCollisionQueryParams Params(NAME_SoundOcclusion, bUseComplexCollisionForOcclusion);
+					if (SoundOwnerID > 0)
+					{
+						Params.AddIgnoredActor(SoundOwnerID);
+					}
+
+					WorldPtr->AsyncLineTraceByChannel(EAsyncTraceType::Test, SoundLocation, ListenerLocation, OcclusionTraceChannel, Params, FCollisionResponseParams::DefaultResponseParam, TraceDelegate);
 				}
-			}
+			});
 		}
 	}
 
diff --git a/Engine/Source/Runtime/Engine/Public/ActiveSound.h b/Engine/Source/Runtime/Engine/Public/ActiveSound.h
--- a/Engine/Source/Runtime/Engine/Public/ActiveSound.h
+++ b/Engine/Source/Runtime/Engine/Public/ActiveSound.h
@@ -168,9 +168,6 @@ public:
 	/** Whether or not the sound has checked if it was occluded already. Used to initialize a sound as occluded and bypassing occlusion interpolation. */
 	uint8 bHasCheckedOcclusion:1;
 
-	/** Flag to trigger binding our trace delegate for async trace calls */
-	uint8 bIsTraceDelegateBound:1;
-
 	/** Is this sound allowed to be spatialized? */
 	uint8 bAllowSpatialization:1;
 
@@ -239,9 +236,6 @@ public:
 	/** Whether or not we have a low-pass filter enabled on this active sound. */
 	uint8 bEnableLowPassFilter : 1;
 
-	/** Whether or not to use an async trace for occlusion */
-	uint8 bOcclusionAsyncTrace : 1;
-
 public:
 	uint8 UserIndex;
 
