Copying //UE4/Dev-Sequencer to //UE4/Dev-Main (Source: //UE4/Dev-Sequencer @ 3116759)

diff --git a/Engine/Source/Editor/UMGEditor/Private/Animation/Sequencer2DTransformTrackEditor.cpp b/Engine/Source/Editor/UMGEditor/Private/Animation/Sequencer2DTransformTrackEditor.cpp
--- a/Engine/Source/Editor/UMGEditor/Private/Animation/Sequencer2DTransformTrackEditor.cpp
+++ b/Engine/Source/Editor/UMGEditor/Private/Animation/Sequencer2DTransformTrackEditor.cpp
@@ -20,8 +20,10 @@ class F2DTransformSection
 {
 public:
 
-	F2DTransformSection( UMovieSceneSection& InSectionObject, const FText& SectionName )
-		: FPropertySection(InSectionObject, SectionName) {}
+	F2DTransformSection(ISequencer* InSequencer, FGuid InObjectBinding, FName InPropertyName, const FString& InPropertyPath, UMovieSceneSection& InSectionObject, const FText& SectionName)
+		: FPropertySection(InSequencer, InObjectBinding, InPropertyName, InPropertyPath, InSectionObject, SectionName)
+	{
+	}
 
 	virtual void GenerateSectionLayout( class ISectionLayoutBuilder& LayoutBuilder ) const override
 	{
@@ -30,79 +32,103 @@ public:
 
 		UMovieScene2DTransformSection* TransformSection = Cast<UMovieScene2DTransformSection>(&SectionObject);
 
-		TranslationXKeyArea = MakeShareable(new FFloatCurveKeyArea(&TransformSection->GetTranslationCurve(EAxis::X), TransformSection, RedKeyAreaColor));
-		TranslationYKeyArea = MakeShareable(new FFloatCurveKeyArea(&TransformSection->GetTranslationCurve(EAxis::Y), TransformSection, GreenKeyAreaColor));
+		// Translation
+		TAttribute<TOptional<float>> TranslationXExternalValue = TAttribute<TOptional<float>>::Create(
+			TAttribute<TOptional<float>>::FGetter::CreateRaw(this, &F2DTransformSection::GetTranslationXValue));
+		TSharedRef<FFloatCurveKeyArea> TranslationXKeyArea = MakeShareable(new FFloatCurveKeyArea(&TransformSection->GetTranslationCurve(EAxis::X), TranslationXExternalValue, TransformSection, RedKeyAreaColor));
+
+		TAttribute<TOptional<float>> TranslationYExternalValue = TAttribute<TOptional<float>>::Create(
+			TAttribute<TOptional<float>>::FGetter::CreateRaw(this, &F2DTransformSection::GetTranslationYValue));
+		TSharedRef<FFloatCurveKeyArea> TranslationYKeyArea = MakeShareable(new FFloatCurveKeyArea(&TransformSection->GetTranslationCurve(EAxis::Y), TranslationYExternalValue, TransformSection, GreenKeyAreaColor));
 
-		RotationKeyArea = MakeShareable(new FFloatCurveKeyArea(&TransformSection->GetRotationCurve(), TransformSection));
+		// Rotation
+		TAttribute<TOptional<float>> RotationExternalValue = TAttribute<TOptional<float>>::Create(
+			TAttribute<TOptional<float>>::FGetter::CreateRaw(this, &F2DTransformSection::GetRotationValue));
+		TSharedRef<FFloatCurveKeyArea> RotationKeyArea = MakeShareable(new FFloatCurveKeyArea(&TransformSection->GetRotationCurve(), RotationExternalValue, TransformSection));
 
-		ScaleXKeyArea = MakeShareable(new FFloatCurveKeyArea(&TransformSection->GetScaleCurve(EAxis::X), TransformSection, RedKeyAreaColor));
-		ScaleYKeyArea = MakeShareable(new FFloatCurveKeyArea(&TransformSection->GetScaleCurve(EAxis::Y), TransformSection, GreenKeyAreaColor));
+		// Scale
+		TAttribute<TOptional<float>> ScaleXExternalValue = TAttribute<TOptional<float>>::Create(
+			TAttribute<TOptional<float>>::FGetter::CreateRaw(this, &F2DTransformSection::GetScaleXValue));
+		TSharedRef<FFloatCurveKeyArea> ScaleXKeyArea = MakeShareable(new FFloatCurveKeyArea(&TransformSection->GetScaleCurve(EAxis::X), ScaleXExternalValue, TransformSection, RedKeyAreaColor));
 
-		ShearXKeyArea = MakeShareable(new FFloatCurveKeyArea(&TransformSection->GetSheerCurve(EAxis::X), TransformSection, RedKeyAreaColor));
-		ShearYKeyArea = MakeShareable(new FFloatCurveKeyArea(&TransformSection->GetSheerCurve(EAxis::Y), TransformSection, GreenKeyAreaColor));
+		TAttribute<TOptional<float>> ScaleYExternalValue = TAttribute<TOptional<float>>::Create(
+			TAttribute<TOptional<float>>::FGetter::CreateRaw(this, &F2DTransformSection::GetScaleYValue));
+		TSharedRef<FFloatCurveKeyArea> ScaleYKeyArea = MakeShareable(new FFloatCurveKeyArea(&TransformSection->GetScaleCurve(EAxis::Y), ScaleYExternalValue, TransformSection, GreenKeyAreaColor));
+
+		// Shear
+		TAttribute<TOptional<float>> ShearXExternalValue = TAttribute<TOptional<float>>::Create(
+			TAttribute<TOptional<float>>::FGetter::CreateRaw(this, &F2DTransformSection::GetShearXValue));
+		TSharedRef<FFloatCurveKeyArea> ShearXKeyArea = MakeShareable(new FFloatCurveKeyArea(&TransformSection->GetSheerCurve(EAxis::X), ShearXExternalValue, TransformSection, RedKeyAreaColor));
+
+		TAttribute<TOptional<float>> ShearYExternalValue = TAttribute<TOptional<float>>::Create(
+			TAttribute<TOptional<float>>::FGetter::CreateRaw(this, &F2DTransformSection::GetShearYValue));
+		TSharedRef<FFloatCurveKeyArea> ShearYKeyArea = MakeShareable(new FFloatCurveKeyArea(&TransformSection->GetSheerCurve(EAxis::Y), ShearYExternalValue, TransformSection, GreenKeyAreaColor));
 
 		// This generates the tree structure for the transform section
 		LayoutBuilder.PushCategory("Translation", NSLOCTEXT("F2DTransformSection", "TranslationArea", "Translation"));
-			LayoutBuilder.AddKeyArea("Translation.X", NSLOCTEXT("F2DTransformSection", "TransXArea", "X"), TranslationXKeyArea.ToSharedRef());
-			LayoutBuilder.AddKeyArea("Translation.Y", NSLOCTEXT("F2DTransformSection", "TransYArea", "Y"), TranslationYKeyArea.ToSharedRef());
+			LayoutBuilder.AddKeyArea("Translation.X", NSLOCTEXT("F2DTransformSection", "TransXArea", "X"), TranslationXKeyArea);
+			LayoutBuilder.AddKeyArea("Translation.Y", NSLOCTEXT("F2DTransformSection", "TransYArea", "Y"), TranslationYKeyArea);
 		LayoutBuilder.PopCategory();
 
 		LayoutBuilder.PushCategory("Rotation", NSLOCTEXT("F2DTransformSection", "RotationArea", "Rotation"));
-			LayoutBuilder.AddKeyArea("Rotation.Angle", NSLOCTEXT("F2DTransformSection", "AngleArea", "Angle"), RotationKeyArea.ToSharedRef());
+			LayoutBuilder.AddKeyArea("Rotation.Angle", NSLOCTEXT("F2DTransformSection", "AngleArea", "Angle"), RotationKeyArea);
 		LayoutBuilder.PopCategory();
 
 		LayoutBuilder.PushCategory("Scale", NSLOCTEXT("F2DTransformSection", "ScaleArea", "Scale"));
-			LayoutBuilder.AddKeyArea("Scale.X", NSLOCTEXT("F2DTransformSection", "ScaleXArea", "X"), ScaleXKeyArea.ToSharedRef());
-			LayoutBuilder.AddKeyArea("Scale.Y", NSLOCTEXT("F2DTransformSection", "ScaleYArea", "Y"), ScaleYKeyArea.ToSharedRef());
+			LayoutBuilder.AddKeyArea("Scale.X", NSLOCTEXT("F2DTransformSection", "ScaleXArea", "X"), ScaleXKeyArea);
+			LayoutBuilder.AddKeyArea("Scale.Y", NSLOCTEXT("F2DTransformSection", "ScaleYArea", "Y"), ScaleYKeyArea);
 		LayoutBuilder.PopCategory();
 
 		LayoutBuilder.PushCategory("Shear", NSLOCTEXT("F2DTransformSection", "ShearArea", "Shear"));
-			LayoutBuilder.AddKeyArea("Shear.X", NSLOCTEXT("F2DTransformSection", "SheerXArea", "X"), ShearXKeyArea.ToSharedRef());
-			LayoutBuilder.AddKeyArea("Shear.Y", NSLOCTEXT("F2DTransformSection", "SheerYArea", "Y"), ShearYKeyArea.ToSharedRef());
+			LayoutBuilder.AddKeyArea("Shear.X", NSLOCTEXT("F2DTransformSection", "SheerXArea", "X"), ShearXKeyArea);
+			LayoutBuilder.AddKeyArea("Shear.Y", NSLOCTEXT("F2DTransformSection", "SheerYArea", "Y"), ShearYKeyArea);
 		LayoutBuilder.PopCategory();
 	}
 
-	virtual void SetIntermediateValue( FPropertyChangedParams PropertyChangedParams ) override
-	{
-		FWidgetTransform Transform = PropertyChangedParams.GetPropertyValue<FWidgetTransform>();
-
-		TranslationXKeyArea->SetIntermediateValue(Transform.Translation.X);
-		TranslationYKeyArea->SetIntermediateValue(Transform.Translation.Y);
-
-		RotationKeyArea->SetIntermediateValue(Transform.Angle);
-
-		ScaleXKeyArea->SetIntermediateValue(Transform.Scale.X);
-		ScaleYKeyArea->SetIntermediateValue(Transform.Scale.Y);
+private:
 
-		ShearXKeyArea->SetIntermediateValue(Transform.Shear.X);
-		ShearYKeyArea->SetIntermediateValue(Transform.Shear.Y);
+	TOptional<float> GetTranslationXValue() const
+	{
+		TOptional<FWidgetTransform> Transform = GetPropertyValue<FWidgetTransform>();
+		return Transform.IsSet() ? TOptional<float>(Transform.GetValue().Translation.X) : TOptional<float>();
 	}
 
-	virtual void ClearIntermediateValue() override
+	TOptional<float> GetTranslationYValue() const
 	{
-		TranslationXKeyArea->ClearIntermediateValue();
-		TranslationYKeyArea->ClearIntermediateValue();
-
-		RotationKeyArea->ClearIntermediateValue();
+		TOptional<FWidgetTransform> Transform = GetPropertyValue<FWidgetTransform>();
+		return Transform.IsSet() ? TOptional<float>(Transform.GetValue().Translation.Y) : TOptional<float>();
+	}
 
-		ScaleXKeyArea->ClearIntermediateValue();
-		ScaleYKeyArea->ClearIntermediateValue();
+	TOptional<float> GetRotationValue() const
+	{
+		TOptional<FWidgetTransform> Transform = GetPropertyValue<FWidgetTransform>();
+		return Transform.IsSet() ? TOptional<float>(Transform.GetValue().Angle) : TOptional<float>();
+	}
 
-		ShearXKeyArea->ClearIntermediateValue();
-		ShearYKeyArea->ClearIntermediateValue();
+	TOptional<float> GetScaleXValue() const
+	{
+		TOptional<FWidgetTransform> Transform = GetPropertyValue<FWidgetTransform>();
+		return Transform.IsSet() ? TOptional<float>(Transform.GetValue().Scale.X) : TOptional<float>();
 	}
 
-private:
-	mutable TSharedPtr<FFloatCurveKeyArea> TranslationXKeyArea;
-	mutable TSharedPtr<FFloatCurveKeyArea> TranslationYKeyArea;
+	TOptional<float> GetScaleYValue() const
+	{
+		TOptional<FWidgetTransform> Transform = GetPropertyValue<FWidgetTransform>();
+		return Transform.IsSet() ? TOptional<float>(Transform.GetValue().Scale.Y) : TOptional<float>();
+	}
 
-	mutable TSharedPtr<FFloatCurveKeyArea> ScaleXKeyArea;
-	mutable TSharedPtr<FFloatCurveKeyArea> ScaleYKeyArea;
+	TOptional<float> GetShearXValue() const
+	{
+		TOptional<FWidgetTransform> Transform = GetPropertyValue<FWidgetTransform>();
+		return Transform.IsSet() ? TOptional<float>(Transform.GetValue().Shear.X) : TOptional<float>();
+	}
 
-	mutable TSharedPtr<FFloatCurveKeyArea> ShearXKeyArea;
-	mutable TSharedPtr<FFloatCurveKeyArea> ShearYKeyArea;
+	TOptional<float> GetShearYValue() const
+	{
+		TOptional<FWidgetTransform> Transform = GetPropertyValue<FWidgetTransform>();
+		return Transform.IsSet() ? TOptional<float>(Transform.GetValue().Shear.Y) : TOptional<float>();
+	}
 
-	mutable TSharedPtr<FFloatCurveKeyArea> RotationKeyArea;
 };
 
 
@@ -112,12 +138,11 @@ TSharedRef<ISequencerTrackEditor> F2DTransformTrackEditor::CreateTrackEditor( TS
 }
 
 
-TSharedRef<FPropertySection> F2DTransformTrackEditor::MakePropertySectionInterface( UMovieSceneSection& SectionObject, UMovieSceneTrack& Track )
+TSharedRef<ISequencerSection> F2DTransformTrackEditor::MakeSectionInterface(UMovieSceneSection& SectionObject, UMovieSceneTrack& Track, FGuid ObjectBinding)
 {
-	check( SupportsType( SectionObject.GetOuter()->GetClass() ) );
-
-	UClass* SectionClass = SectionObject.GetOuter()->GetClass();
-	return MakeShareable(new F2DTransformSection(SectionObject, Track.GetDisplayName()));
+	UMovieScenePropertyTrack* PropertyTrack = Cast<UMovieScenePropertyTrack>(&Track);
+	checkf(PropertyTrack != nullptr, TEXT("Incompatible track in F2DTransformTrackEditor"));
+	return MakeShareable(new F2DTransformSection(GetSequencer().Get(), ObjectBinding, PropertyTrack->GetPropertyName(), PropertyTrack->GetPropertyPath(), SectionObject, Track.GetDisplayName()));
 }
 
 
diff --git a/Engine/Source/Editor/UMGEditor/Private/Animation/Sequencer2DTransformTrackEditor.h b/Engine/Source/Editor/UMGEditor/Private/Animation/Sequencer2DTransformTrackEditor.h
--- a/Engine/Source/Editor/UMGEditor/Private/Animation/Sequencer2DTransformTrackEditor.h
+++ b/Engine/Source/Editor/UMGEditor/Private/Animation/Sequencer2DTransformTrackEditor.h
@@ -29,12 +29,14 @@ public:
 	 */
 	static TSharedRef<ISequencerTrackEditor> CreateTrackEditor( TSharedRef<ISequencer> OwningSequencer );
 
+	//~ ISequencerTrackEditor Interface
+
+	virtual TSharedRef<ISequencerSection> MakeSectionInterface(UMovieSceneSection& SectionObject, UMovieSceneTrack& Track, FGuid ObjectBinding) override;
+
 protected:
 
 	// FPropertyTrackEditor interface
 
-	virtual TSharedRef<FPropertySection> MakePropertySectionInterface( UMovieSceneSection& SectionObject, UMovieSceneTrack& Track ) override;
-
 	virtual void GenerateKeysFromPropertyChanged( const FPropertyChangedParams& PropertyChangedParams, TArray<F2DTransformKey>& NewGeneratedKeys, TArray<F2DTransformKey>& DefaultGeneratedKeys ) override;
 
 private:
diff --git a/Engine/Source/Editor/UMGEditor/Private/Animation/UMGDetailKeyframeHandler.cpp b/Engine/Source/Editor/UMGEditor/Private/Animation/UMGDetailKeyframeHandler.cpp
--- a/Engine/Source/Editor/UMGEditor/Private/Animation/UMGDetailKeyframeHandler.cpp
+++ b/Engine/Source/Editor/UMGEditor/Private/Animation/UMGDetailKeyframeHandler.cpp
@@ -26,12 +26,7 @@ void FUMGDetailKeyframeHandler::OnKeyPropertyClicked(const IPropertyHandle& Keye
 	TArray<UObject*> Objects;
 	KeyedPropertyHandle.GetOuterObjects( Objects );
 
-	FKeyPropertyParams KeyPropertyParams(Objects, KeyedPropertyHandle);
-	KeyPropertyParams.KeyParams.bCreateHandleIfMissing = true;
-	KeyPropertyParams.KeyParams.bCreateTrackIfMissing = true;
-	KeyPropertyParams.KeyParams.bCreateKeyIfUnchanged = true;
-	KeyPropertyParams.KeyParams.bCreateKeyIfEmpty = true;
-	KeyPropertyParams.KeyParams.bCreateKeyOnlyWhenAutoKeying = false;
+	FKeyPropertyParams KeyPropertyParams(Objects, KeyedPropertyHandle, ESequencerKeyMode::ManualKeyForced);
 
 	BlueprintEditor.Pin()->GetSequencer()->KeyProperty(KeyPropertyParams);
 }
diff --git a/Engine/Source/Editor/UnrealEd/Private/Editor.cpp b/Engine/Source/Editor/UnrealEd/Private/Editor.cpp
--- a/Engine/Source/Editor/UnrealEd/Private/Editor.cpp
+++ b/Engine/Source/Editor/UnrealEd/Private/Editor.cpp
@@ -590,7 +590,7 @@ namespace EditorUtilities
 	AActor* GetEditorWorldCounterpartActor( AActor* Actor )
 	{
 		const bool bIsSimActor = Actor->GetOutermost()->HasAnyPackageFlags(PKG_PlayInEditor);
-		if( bIsSimActor && GEditor->PlayWorld != NULL )
+		if( bIsSimActor && GEditor && GEditor->PlayWorld != NULL )
 		{
 			// Do we have a counterpart in the editor world?
 			auto* SimWorldActor = Actor;
@@ -621,7 +621,7 @@ namespace EditorUtilities
 	AActor* GetSimWorldCounterpartActor( AActor* Actor )
 	{
 		const bool bIsSimActor = Actor->GetOutermost()->HasAnyPackageFlags(PKG_PlayInEditor);
-		if( !bIsSimActor && GEditor->EditorWorld != NULL )
+		if( !bIsSimActor && GEditor && GEditor->EditorWorld != NULL )
 		{
 			// Do we have a counterpart in the sim world?
 			auto* EditorWorldActor = Actor;
diff --git a/Engine/Source/Runtime/CinematicCamera/Private/CameraRig_Crane.cpp b/Engine/Source/Runtime/CinematicCamera/Private/CameraRig_Crane.cpp
--- a/Engine/Source/Runtime/CinematicCamera/Private/CameraRig_Crane.cpp
+++ b/Engine/Source/Runtime/CinematicCamera/Private/CameraRig_Crane.cpp
@@ -14,6 +14,8 @@ ACameraRig_Crane::ACameraRig_Crane(const FObjectInitializer& ObjectInitializer)
 	CraneYaw = 0.f;
 	CranePitch = 0.f;
 	CraneArmLength = 500.f;
+	bLockMountPitch = false;
+	bLockMountYaw = false;
 
 	// create the root component
 	TransformComponent = CreateDefaultSubobject<USceneComponent>(TEXT("TransformComponent"));
@@ -148,8 +150,16 @@ void ACameraRig_Crane::UpdateCraneComponents()
 
 	// zero the pitch from the camera mount component
 	// this effectively gives us bAbsoluteRotation for only pitch component of an attached camera
-	FRotator NewCameraMountWorldRot = CraneCameraMount->GetComponentRotation();
-	NewCameraMountWorldRot.Pitch = 0.f;
+	FRotator NewCameraMountWorldRot = CraneCameraMount->GetAttachParent() ? CraneCameraMount->GetAttachParent()->GetComponentRotation() : FRotator(0.f, 0.f, 0.f);
+	if (!bLockMountPitch)
+	{
+		NewCameraMountWorldRot.Pitch = 0.f;
+	}
+	if (!bLockMountYaw)
+	{
+		NewCameraMountWorldRot.Yaw = RootComponent->RelativeRotation.Yaw;
+	}
+	NewCameraMountWorldRot.Roll = 0.f;
 	CraneCameraMount->SetWorldRotation(NewCameraMountWorldRot);
 
 	UpdatePreviewMeshes();
@@ -172,6 +182,13 @@ void ACameraRig_Crane::PostEditChangeProperty(FPropertyChangedEvent& PropertyCha
 	UpdateCraneComponents();
 }
 
+void ACameraRig_Crane::PostEditUndo()
+{
+	Super::PostEditUndo();
+
+	UpdateCraneComponents();
+}
+
 USceneComponent* ACameraRig_Crane::GetDefaultAttachComponent() const
 {
 	return CraneCameraMount;
diff --git a/Engine/Source/Runtime/CinematicCamera/Private/CameraRig_Rail.cpp b/Engine/Source/Runtime/CinematicCamera/Private/CameraRig_Rail.cpp
--- a/Engine/Source/Runtime/CinematicCamera/Private/CameraRig_Rail.cpp
+++ b/Engine/Source/Runtime/CinematicCamera/Private/CameraRig_Rail.cpp
@@ -118,6 +118,17 @@ void ACameraRig_Rail::UpdatePreviewMeshes()
 					SplineMeshComp->SetStartAndEnd(StartLoc, StartTangent, EndLoc, EndTangent, true);
 				}
 			}
+
+			// Unregister any owned components that aren't in PreviewRailMeshSegments
+			TArray<USplineMeshComponent*> OwnedSplineMeshComponents;
+			GetComponents(OwnedSplineMeshComponents);
+			for (auto OwnedComponent : OwnedSplineMeshComponents)
+			{
+				if (!PreviewRailMeshSegments.Contains(OwnedComponent))
+				{
+					OwnedComponent->UnregisterComponent();
+				}
+			}
 		}
 
 		// make visualization of the mount follow the contour of the rail
@@ -159,6 +170,12 @@ void ACameraRig_Rail::PostEditChangeProperty(FPropertyChangedEvent& PropertyChan
 	Super::PostEditChangeProperty(PropertyChangedEvent);
 	UpdateRailComponents();
 }
+
+void ACameraRig_Rail::PostEditUndo()
+{
+	Super::PostEditUndo();
+	UpdateRailComponents();
+}
 #endif
 
 void ACameraRig_Rail::Tick(float DeltaTime)
diff --git a/Engine/Source/Runtime/CinematicCamera/Public/CameraRig_Crane.h b/Engine/Source/Runtime/CinematicCamera/Public/CameraRig_Crane.h
--- a/Engine/Source/Runtime/CinematicCamera/Public/CameraRig_Crane.h
+++ b/Engine/Source/Runtime/CinematicCamera/Public/CameraRig_Crane.h
@@ -9,7 +9,7 @@
 class USceneComponent;
 
 /** 
- * A simple rig for to simulate crane-like camera movements.
+ * A simple rig for simulating crane-like camera movements.
  */
 UCLASS(Blueprintable)
 class CINEMATICCAMERA_API ACameraRig_Crane : public AActor
@@ -24,7 +24,7 @@ public:
 	virtual void Tick(float DeltaTime) override;
 	virtual bool ShouldTickIfViewportsOnly() const override;
 
-	/** Controls the pitch of the crane arm (but does not affect the rotation of the mount). */
+	/** Controls the pitch of the crane arm. */
 	UPROPERTY(Interp, EditAnywhere, BlueprintReadWrite, Category = "Crane Controls", meta = (UIMin = "-360", UIMax = "360", Units = deg))
 	float CranePitch;
 	
@@ -36,8 +36,17 @@ public:
 	UPROPERTY(Interp, EditAnywhere, BlueprintReadWrite, Category = "Crane Controls", meta = (ClampMin=0, Units = cm))
 	float CraneArmLength;
 
+	/** Lock the mount pitch so that an attached camera is locked and pitched in the direction of the crane arm */
+	UPROPERTY(Interp, EditAnywhere, BlueprintReadWrite, Category = "Crane Controls")
+	bool bLockMountPitch;
+
+	/** Lock the mount yaw so that an attached camera is locked and oriented in the direction of the crane arm */
+	UPROPERTY(Interp, EditAnywhere, BlueprintReadWrite, Category = "Crane Controls")
+	bool bLockMountYaw;
+
 #if WITH_EDITOR
 	virtual void PostEditChangeProperty(FPropertyChangedEvent& PropertyChangedEvent) override;
+	virtual void PostEditUndo() override;
 	virtual class USceneComponent* GetDefaultAttachComponent() const override;
 #endif
 	
diff --git a/Engine/Source/Runtime/CinematicCamera/Public/CameraRig_Rail.h b/Engine/Source/Runtime/CinematicCamera/Public/CameraRig_Rail.h
--- a/Engine/Source/Runtime/CinematicCamera/Public/CameraRig_Rail.h
+++ b/Engine/Source/Runtime/CinematicCamera/Public/CameraRig_Rail.h
@@ -30,6 +30,7 @@ public:
 #if WITH_EDITOR
 	virtual class USceneComponent* GetDefaultAttachComponent() const override;
 	virtual void PostEditChangeProperty(FPropertyChangedEvent& PropertyChangedEvent) override;
+	virtual void PostEditUndo() override;
 #endif
 
 	/** Returns the spline component that defines the rail path */
diff --git a/Engine/Source/Runtime/CinematicCamera/Public/CineCameraComponent.h b/Engine/Source/Runtime/CinematicCamera/Public/CineCameraComponent.h
--- a/Engine/Source/Runtime/CinematicCamera/Public/CineCameraComponent.h
+++ b/Engine/Source/Runtime/CinematicCamera/Public/CineCameraComponent.h
@@ -105,7 +105,7 @@ struct FCameraTrackingFocusSettings
 	GENERATED_USTRUCT_BODY()
 
 	/** Focus distance will be tied to this actor's location. */
-	UPROPERTY(EditAnywhere, BlueprintReadWrite, Category = "Tracking Focus")
+	UPROPERTY(Interp, EditAnywhere, BlueprintReadWrite, Category = "Tracking Focus")
 	AActor* ActorToTrack;
 
 	/** Offset from actor position to track. Relative to actor if tracking an actor, relative to world otherwise. */
diff --git a/Engine/Source/Runtime/Core/Public/Containers/Queue.h b/Engine/Source/Runtime/Core/Public/Containers/Queue.h
--- a/Engine/Source/Runtime/Core/Public/Containers/Queue.h
+++ b/Engine/Source/Runtime/Core/Public/Containers/Queue.h
@@ -141,7 +141,7 @@ public:
 	 * @return true if an item was returned, false if the queue was empty.
 	 * @see Dequeue, Enqueue, IsEmpty
 	 */
-	bool Peek( ItemType& OutItem ) const
+	bool Peek(ItemType& OutItem) const
 	{
 		if (Tail->NextNode == nullptr)
 		{
@@ -170,7 +170,7 @@ private:
 		{ }
 
 		/** Creates and initializes a new node. */
-		TNode( const ItemType& InItem )
+		TNode(const ItemType& InItem)
 			: NextNode(nullptr)
 			, Item(InItem)
 		{ }
@@ -184,7 +184,10 @@ private:
 	TNode* Tail;
 
 private:
-	// Non-copyable
+
+	/** Hidden copy constructor. */
 	TQueue(const TQueue&);
+
+	/** Hidden assignment operator. */
 	TQueue& operator=(const TQueue&);
 };
diff --git a/Engine/Source/Runtime/Engine/Classes/Curves/IntegralCurve.h b/Engine/Source/Runtime/Engine/Classes/Curves/IntegralCurve.h
--- a/Engine/Source/Runtime/Engine/Classes/Curves/IntegralCurve.h
+++ b/Engine/Source/Runtime/Engine/Classes/Curves/IntegralCurve.h
@@ -77,12 +77,21 @@ public:
 
 	/** Get the time for the Key with the specified index. */
 	float GetKeyTime(FKeyHandle KeyHandle) const;
+
+	/** Set the value of the key with the specified index. */
+	void SetKeyValue(FKeyHandle KeyHandle, int32 NewValue);
+
+	/** Get the value for the Key with the specified index. */
+	int32 GetKeyValue(FKeyHandle KeyHandle) const;
 	
 	/** Set the default value for the curve */
 	void SetDefaultValue(int32 InDefaultValue) { DefaultValue = InDefaultValue; }
 
 	/** Get the default value for the curve */
 	int32 GetDefaultValue() const { return DefaultValue; }
+	
+	/** Removes the default value for this curve. */
+	void ClearDefaultValue() { DefaultValue = MAX_int32; }
 
 	/** Sets whether or not the default value should be used for evaluation for time values before the first key. */
 	void SetUseDefaultValueBeforeFirstKey(bool InbUseDefaultValueBeforeFirstKey) { bUseDefaultValueBeforeFirstKey = InbUseDefaultValueBeforeFirstKey; }
diff --git a/Engine/Source/Runtime/Engine/Classes/Curves/RichCurve.h b/Engine/Source/Runtime/Engine/Classes/Curves/RichCurve.h
--- a/Engine/Source/Runtime/Engine/Classes/Curves/RichCurve.h
+++ b/Engine/Source/Runtime/Engine/Classes/Curves/RichCurve.h
@@ -260,6 +260,9 @@ public:
 	/** Get the default value for the curve */
 	float GetDefaultValue() const { return DefaultValue; }
 
+	/** Removes the default value for this curve. */
+	void ClearDefaultValue() { DefaultValue = MAX_flt; }
+
 	/** Shifts all keys forwards or backwards in time by an even amount, preserving order */
 	void ShiftCurve(float DeltaTime);
 	void ShiftCurve(float DeltaTime, TSet<FKeyHandle>& KeyHandles);
diff --git a/Engine/Source/Runtime/Engine/Classes/Curves/StringCurve.h b/Engine/Source/Runtime/Engine/Classes/Curves/StringCurve.h
--- a/Engine/Source/Runtime/Engine/Classes/Curves/StringCurve.h
+++ b/Engine/Source/Runtime/Engine/Classes/Curves/StringCurve.h
@@ -150,6 +150,14 @@ public:
 	 */
 	float GetKeyTime(FKeyHandle KeyHandle) const;
 
+	/**
+	* Get the value for the Key with the specified index.
+	*
+	* @param KeyHandle Handle to the key whose value to get.
+	* @return The key's value.
+	*/
+	FString GetKeyValue(FKeyHandle KeyHandle) const;
+
 	/**
 	 * Set the default value of the curve.
 	 *
@@ -160,6 +168,14 @@ public:
 		DefaultValue = InDefaultValue;
 	}
 
+	/**
+	 * Clears the default value for this curve if it has been set. 
+	 */
+	void ClearDefaultValue()
+	{
+		DefaultValue.Empty();
+	}
+
 	/**
 	 * Move a key to a new time.
 	 *
@@ -170,6 +186,14 @@ public:
 	 */
 	FKeyHandle SetKeyTime(FKeyHandle KeyHandle, float NewTime);
 
+	/**
+	* Assign a new value to a key.
+	*
+	* @param KeyHandle The handle of the key to change.
+	* @param NewTime The new value to set on the key.
+	*/
+	void SetKeyValue(FKeyHandle KeyHandle, FString NewValue);
+
 	/**
 	 * Finds the key at InTime, and updates its value. If it can't find the key within the KeyTimeTolerance, it adds one at that time.
 	 *
diff --git a/Engine/Source/Runtime/Engine/Private/AudioDevice.cpp b/Engine/Source/Runtime/Engine/Private/AudioDevice.cpp
--- a/Engine/Source/Runtime/Engine/Private/AudioDevice.cpp
+++ b/Engine/Source/Runtime/Engine/Private/AudioDevice.cpp
@@ -106,6 +106,7 @@ FAudioDevice::FAudioDevice()
 	, bHRTFEnabledForAll(false)
 	, bIsDeviceMuted(false)
 	, bIsInitialized(false)
+	, AudioClock(0.0)
 	, bHasActivatedReverb(false)
 	, bAllowVirtualizedSounds(true)
 #if !UE_BUILD_SHIPPING
diff --git a/Engine/Source/Runtime/Engine/Private/Curves/IntegralCurve.cpp b/Engine/Source/Runtime/Engine/Private/Curves/IntegralCurve.cpp
--- a/Engine/Source/Runtime/Engine/Private/Curves/IntegralCurve.cpp
+++ b/Engine/Source/Runtime/Engine/Private/Curves/IntegralCurve.cpp
@@ -159,6 +159,26 @@ float FIntegralCurve::GetKeyTime(FKeyHandle KeyHandle) const
 }
 
 
+void FIntegralCurve::SetKeyValue(FKeyHandle KeyHandle, int32 NewValue)
+{
+	if (IsKeyHandleValid(KeyHandle))
+	{
+		GetKey(KeyHandle).Value = NewValue;
+	}
+}
+
+
+int32 FIntegralCurve::GetKeyValue(FKeyHandle KeyHandle) const
+{
+	if (!IsKeyHandleValid(KeyHandle))
+	{
+		return 0.f;
+	}
+
+	return GetKey(KeyHandle).Value;
+}
+
+
 void FIntegralCurve::ShiftCurve(float DeltaTime)
 {
 	TSet<FKeyHandle> KeyHandles;
diff --git a/Engine/Source/Runtime/Engine/Private/Curves/StringCurve.cpp b/Engine/Source/Runtime/Engine/Private/Curves/StringCurve.cpp
--- a/Engine/Source/Runtime/Engine/Private/Curves/StringCurve.cpp
+++ b/Engine/Source/Runtime/Engine/Private/Curves/StringCurve.cpp
@@ -161,6 +161,17 @@ float FStringCurve::GetKeyTime(FKeyHandle KeyHandle) const
 }
 
 
+FString FStringCurve::GetKeyValue(FKeyHandle KeyHandle) const
+{
+	if (!IsKeyHandleValid(KeyHandle))
+	{
+		return FString();
+	}
+
+	return GetKey(KeyHandle).Value;
+}
+
+
 FKeyHandle FStringCurve::SetKeyTime(FKeyHandle KeyHandle, float NewTime)
 {
 	if (!IsKeyHandleValid(KeyHandle))
@@ -180,6 +191,14 @@ FKeyHandle FStringCurve::SetKeyTime(FKeyHandle KeyHandle, float NewTime)
 	return KeyHandle;
 }
 
+void FStringCurve::SetKeyValue(FKeyHandle KeyHandle, FString NewValue)
+{
+	if (IsKeyHandleValid(KeyHandle))
+	{
+		GetKey(KeyHandle).Value = NewValue;
+	}
+}
+
 
 void FStringCurve::ShiftCurve(float DeltaTime)
 {
diff --git a/Engine/Source/Runtime/Engine/Private/SoundWave.cpp b/Engine/Source/Runtime/Engine/Private/SoundWave.cpp
--- a/Engine/Source/Runtime/Engine/Private/SoundWave.cpp
+++ b/Engine/Source/Runtime/Engine/Private/SoundWave.cpp
@@ -559,13 +559,16 @@ FWaveInstance* USoundWave::HandleStart( FActiveSound& ActiveSound, const UPTRINT
 	if (ActiveSound.bHandleSubtitles && Subtitles.Num() > 0)
 	{
 		FQueueSubtitleParams QueueSubtitleParams(Subtitles);
-		QueueSubtitleParams.AudioComponentID = ActiveSound.GetAudioComponentID();
-		QueueSubtitleParams.WorldPtr = ActiveSound.GetWeakWorld();
-		QueueSubtitleParams.WaveInstance = (PTRINT)WaveInstance;
-		QueueSubtitleParams.SubtitlePriority = ActiveSound.SubtitlePriority;
-		QueueSubtitleParams.Duration = Duration;
-		QueueSubtitleParams.bManualWordWrap = bManualWordWrap;
-		QueueSubtitleParams.bSingleLine = bSingleLine;
+		{
+			QueueSubtitleParams.AudioComponentID = ActiveSound.GetAudioComponentID();
+			QueueSubtitleParams.WorldPtr = ActiveSound.GetWeakWorld();
+			QueueSubtitleParams.WaveInstance = (PTRINT)WaveInstance;
+			QueueSubtitleParams.SubtitlePriority = ActiveSound.SubtitlePriority;
+			QueueSubtitleParams.Duration = Duration;
+			QueueSubtitleParams.bManualWordWrap = bManualWordWrap;
+			QueueSubtitleParams.bSingleLine = bSingleLine;
+			QueueSubtitleParams.RequestedStartTime = ActiveSound.RequestedStartTime;
+		}
 
 		FSubtitleManager::QueueSubtitles(QueueSubtitleParams);
 	}
diff --git a/Engine/Source/Runtime/Engine/Private/SubtitleManager.cpp b/Engine/Source/Runtime/Engine/Private/SubtitleManager.cpp
--- a/Engine/Source/Runtime/Engine/Private/SubtitleManager.cpp
+++ b/Engine/Source/Runtime/Engine/Private/SubtitleManager.cpp
@@ -1,43 +1,42 @@
 // Copyright 1998-2016 Epic Games, Inc. All Rights Reserved.
 
-
 #include "EnginePrivate.h"
 #include "SubtitleManager.h"
 #include "SoundDefinitions.h"
 #include "Engine/Font.h"
 #include "AudioThread.h"
 
+
 DEFINE_LOG_CATEGORY(LogSubtitle);
 
-// Modifier to the spacing between lines in subtitles
-#define MULTILINE_SPACING_SCALING 1.1f
+#define MULTILINE_SPACING_SCALING 1.1f		// Modifier to the spacing between lines in subtitles
+#define SUBTITLE_CHAR_HEIGHT 24
+
+/** The default offset of the outline box. */
+FIntRect DrawStringOutlineBoxOffset(2, 2, 4, 4);
+
 
-/**
- * Kills all the subtitles
- */
-void FSubtitleManager::KillAllSubtitles( void )
+void FSubtitleManager::KillAllSubtitles()
 {
 	ActiveSubtitles.Empty();
 }
 
-/**
- * Kills all the subtitles associated with the subtitle ID.
- * This is called from AudioComponent::Stop
- */
-void FSubtitleManager::KillSubtitles( PTRINT SubtitleID )
+
+void FSubtitleManager::KillSubtitles(PTRINT SubtitleID)
 {
-	ActiveSubtitles.Remove( SubtitleID );
+	ActiveSubtitles.Remove(SubtitleID);
 }
 
+
 void FSubtitleManager::QueueSubtitles(const FQueueSubtitleParams& QueueSubtitleParams)
 {
 	check(IsInAudioThread());
 
 	DECLARE_CYCLE_STAT(TEXT("FGameThreadAudioTask.QueueSubtitles"), STAT_AudioQueueSubtitles, STATGROUP_TaskGraphTasks);
 
-	FAudioThread::RunCommandOnGameThread([QueueSubtitleParams]()
-	{
+	FAudioThread::RunCommandOnGameThread([QueueSubtitleParams]() {
 		UAudioComponent* AudioComponent = UAudioComponent::GetAudioComponentFromID(QueueSubtitleParams.AudioComponentID);
+
 		if (AudioComponent && AudioComponent->OnQueueSubtitles.IsBound())
 		{
 			// intercept the subtitles if the delegate is set
@@ -47,131 +46,141 @@ void FSubtitleManager::QueueSubtitles(const FQueueSubtitleParams& QueueSubtitleP
 		{
 			// otherwise, pass them on to the subtitle manager for display
 			// Subtitles are hashed based on the associated sound (wave instance).
-			FSubtitleManager::GetSubtitleManager()->QueueSubtitles( QueueSubtitleParams.WaveInstance, QueueSubtitleParams.SubtitlePriority, QueueSubtitleParams.bManualWordWrap, QueueSubtitleParams.bSingleLine, QueueSubtitleParams.Duration, QueueSubtitleParams.Subtitles, World->GetAudioTimeSeconds() );
+			FSubtitleManager::GetSubtitleManager()->QueueSubtitles(
+				QueueSubtitleParams.WaveInstance,
+				QueueSubtitleParams.SubtitlePriority,
+				QueueSubtitleParams.bManualWordWrap,
+				QueueSubtitleParams.bSingleLine,
+				QueueSubtitleParams.Duration,
+				QueueSubtitleParams.Subtitles,
+				QueueSubtitleParams.RequestedStartTime,
+				World->GetAudioTimeSeconds()
+			);
 		}
 	}, GET_STATID(STAT_AudioQueueSubtitles));
 }
 
-/**
- * Add an array of subtitles to the active list
- * This is called from AudioComponent::Play
- * As there is only a global subtitle manager, this system will not have different subtitles for different viewports
- *
- * @param  SubTitleID - the controlling id that groups sets of subtitles together
- * @param  Priority - used to prioritize subtitles; higher values have higher priority.  Subtitles with a priority 0.0f are not displayed.
- * @param  StartTime - float of seconds that the subtitles start at
- * @param  SoundDuration - time in seconds after which the subtitles do not display
- * @param  Subtitles - TArray of lines of subtitle and time offset to play them
- */
-void FSubtitleManager::QueueSubtitles( PTRINT SubtitleID, float Priority, bool bManualWordWrap, bool bSingleLine, float SoundDuration, const TArray<FSubtitleCue>& Subtitles, float InStartTime )
+
+void FSubtitleManager::QueueSubtitles(PTRINT SubtitleID, float Priority, bool bManualWordWrap, bool bSingleLine, float SoundDuration, const TArray<FSubtitleCue>& Subtitles, float InStartTime, float InCurrentTime)
 {
-	check( GEngine );
+	check(GEngine);
 	check(IsInGameThread());
 
-	// No subtitles to display
-	if( !Subtitles.Num() )
+	if ((Subtitles.Num() == 0) || (Priority == 0.0f))
 	{
-		return;
+		return; // nothing to show, or subtitle suppressed
 	}
 
-	// Return when a subtitle is played with its subtitle suppressed.
-	if (Priority == 0.0f)
+	if (SoundDuration == 0.0f)
 	{
+		UE_LOG(LogSubtitle, Warning, TEXT("Received subtitle with no sound duration."));
 		return;
 	}
 
-	// No sound associated with the subtitle (used for automatic timing)
-	if( SoundDuration == 0.0f )
+	// skip subtitles that have already been displayed
+	TArray<FSubtitleCue> SubtitlesToAdd;
 	{
-		UE_LOG(LogSubtitle, Warning, TEXT( "Received subtitle with no sound duration." ) );
-		return;
-	}
+		for (int32 SubtitleIndex = 1; SubtitleIndex < Subtitles.Num(); ++SubtitleIndex)
+		{
+			if (Subtitles[SubtitleIndex].Time >= InStartTime)
+			{
+				SubtitlesToAdd.Add(Subtitles[SubtitleIndex - 1]);
+			}
+		}
 
-	// Add in the new subtitles
-	FActiveSubtitle& NewSubtitle = ActiveSubtitles.Add( SubtitleID, FActiveSubtitle( 0, Priority, bManualWordWrap, bSingleLine, Subtitles ) );
+		SubtitlesToAdd.Add(Subtitles.Last());
+	}
 
-	// Resolve time offsets to absolute time
-	for( FSubtitleCue& Subtitle : NewSubtitle.Subtitles )
+	// resolve time offsets to absolute time
+	for (FSubtitleCue& Subtitle : SubtitlesToAdd)
 	{
 		if (Subtitle.Time > SoundDuration)
 		{
 			Subtitle.Time = SoundDuration;
-			UE_LOG(LogSubtitle, Warning, TEXT( "Subtitle has time offset greater than length of sound - clamping" ) );
+			UE_LOG(LogSubtitle, Warning, TEXT("Subtitle has time offset greater than length of sound - clamping"));
 		}
 
-		Subtitle.Time += InStartTime;
+		Subtitle.Time -= InStartTime;
+		Subtitle.Time += InCurrentTime;
 	}
 
-	// Add on a blank at the end to clear
+	// append new subtitles
+	FActiveSubtitle& NewSubtitle = ActiveSubtitles.Add(SubtitleID, FActiveSubtitle(0, Priority, bManualWordWrap, bSingleLine, SubtitlesToAdd));
+
+	// add on a blank at the end to clear
 	FSubtitleCue& Temp = NewSubtitle.Subtitles[NewSubtitle.Subtitles.AddDefaulted()];
-	Temp.Text = FText::GetEmpty();
-	Temp.Time = InStartTime + SoundDuration;
+	{
+		Temp.Text = FText::GetEmpty();
+		Temp.Time = InCurrentTime + (SoundDuration - InStartTime);
+	}
 }
 
-/**
- * Draws a subtitle at the specified pixel location.
- */
-#define SUBTITLE_CHAR_HEIGHT		24
-/** The default offset of the outline box */
-FIntRect DrawStringOutlineBoxOffset(2, 2, 4, 4);
 
-void FSubtitleManager::DisplaySubtitle( FCanvas* Canvas, FActiveSubtitle* Subtitle, FIntRect& Parms, const FLinearColor& Color )
+void FSubtitleManager::DisplaySubtitle(FCanvas* Canvas, FActiveSubtitle* Subtitle, FIntRect& Parms, const FLinearColor& Color)
 {
 	// These should be valid in here
-	check( GEngine );
-	check( Canvas );
+	check(GEngine);
+	check(Canvas);
 
 	CurrentSubtitleHeight = 0.0f;
 
-	// This can be NULL when there's an asset mixup (especially with localisation)
-	if( !GEngine->GetSubtitleFont() )
+	// This can be NULL when there's an asset mixup (especially with localization)
+	if (!GEngine->GetSubtitleFont())
 	{
-		UE_LOG(LogSubtitle, Warning, TEXT( "NULL GEngine->GetSubtitleFont() - subtitles not rendering!" ) );
+		UE_LOG(LogSubtitle, Warning, TEXT("NULL GEngine->GetSubtitleFont() - subtitles not rendering!"));
 		return;
 	}
 
 	float FontHeight = GEngine->GetSubtitleFont()->GetMaxCharHeight();
 	float HeightTest = Canvas->GetRenderTarget()->GetSizeXY().Y;
-	int32 SubtitleHeight = FMath::TruncToInt( ( FontHeight * MULTILINE_SPACING_SCALING ) );
+	int32 SubtitleHeight = FMath::TruncToInt((FontHeight * MULTILINE_SPACING_SCALING));
 	FIntRect BackgroundBoxOffset = DrawStringOutlineBoxOffset;
 
 	// Needed to add a drop shadow and doing all 4 sides was the only way to make them look correct.  If this shows up as a framerate hit we'll have to think of a different way of dealing with this.
-	if( Subtitle->bSingleLine )
+	if (Subtitle->bSingleLine)
 	{
 		const FText& SubtitleText = Subtitle->Subtitles[Subtitle->Index].Text;
-		if ( !SubtitleText.IsEmpty() )
+
+		if (!SubtitleText.IsEmpty())
 		{
 			
 			// Display lines up from the bottom of the region
 			Parms.Max.Y -= SUBTITLE_CHAR_HEIGHT;
 			
-			FCanvasTextItem TextItem( FVector2D( Parms.Min.X + ( Parms.Width() / 2 ), Parms.Max.Y ), SubtitleText , GEngine->GetSubtitleFont(), Color );
-			TextItem.Depth = SUBTITLE_SCREEN_DEPTH_FOR_3D;
-			TextItem.bOutlined = true;
-			TextItem.bCentreX = true;
-			TextItem.OutlineColor = FLinearColor::Black;
-			Canvas->DrawItem( TextItem );
-			
+			FCanvasTextItem TextItem(FVector2D(Parms.Min.X + (Parms.Width() / 2), Parms.Max.Y), SubtitleText , GEngine->GetSubtitleFont(), Color);
+			{
+				TextItem.Depth = SUBTITLE_SCREEN_DEPTH_FOR_3D;
+				TextItem.bOutlined = true;
+				TextItem.bCentreX = true;
+				TextItem.OutlineColor = FLinearColor::Black;
+			}
+
+			Canvas->DrawItem(TextItem);		
 			CurrentSubtitleHeight += SubtitleHeight;
 		}
 	}
 	else
 	{
-		for( int32 Idx = Subtitle->Subtitles.Num() - 1; Idx >= 0; Idx-- )
+		for (int32 Idx = Subtitle->Subtitles.Num() - 1; Idx >= 0; Idx--)
 		{
 			const FText& SubtitleText = Subtitle->Subtitles[Idx].Text;
+
 			// Display lines up from the bottom of the region
-			if ( !SubtitleText.IsEmpty() )
+			if (!SubtitleText.IsEmpty())
 			{
-				FCanvasTextItem TextItem( FVector2D( Parms.Min.X + ( Parms.Width() / 2 ), Parms.Max.Y ), SubtitleText, GEngine->GetSubtitleFont(), Color );
-				TextItem.Depth = SUBTITLE_SCREEN_DEPTH_FOR_3D;
-				TextItem.bOutlined = true;
-				TextItem.bCentreX = true;
-				TextItem.OutlineColor = FLinearColor::Black;
-				Canvas->DrawItem( TextItem );
+				FCanvasTextItem TextItem(FVector2D(Parms.Min.X + (Parms.Width() / 2), Parms.Max.Y), SubtitleText, GEngine->GetSubtitleFont(), Color);
+				{
+					TextItem.Depth = SUBTITLE_SCREEN_DEPTH_FOR_3D;
+					TextItem.bOutlined = true;
+					TextItem.bCentreX = true;
+					TextItem.OutlineColor = FLinearColor::Black;
+				}
+
+				Canvas->DrawItem(TextItem);
 
 				Parms.Max.Y -= SubtitleHeight;
 				CurrentSubtitleHeight += SubtitleHeight;
+
 				// Don't overlap subsequent boxes...
 				BackgroundBoxOffset.Max.Y = BackgroundBoxOffset.Min.Y;
 			}
@@ -179,54 +188,51 @@ void FSubtitleManager::DisplaySubtitle( FCanvas* Canvas, FActiveSubtitle* Subtit
 	}
 }
 
-/**
- * If any of the active subtitles need to be split into multiple lines, do so now
- * - caveat - this assumes the width of the subtitle region does not change while the subtitle is active
- */
-void FSubtitleManager::SplitLinesToSafeZone( FCanvas* Canvas, FIntRect& SubtitleRegion )
+
+void FSubtitleManager::SplitLinesToSafeZone(FCanvas* Canvas, FIntRect& SubtitleRegion)
 {
-	int32				i;
-	FString			Concatenated;
-	float			SoundDuration;
-	float			StartTime, SecondsPerChar, Cumulative;
+	FString Concatenated;
+	float SoundDuration;
+	float StartTime, SecondsPerChar, Cumulative;
 
-	for( TMap<PTRINT, FActiveSubtitle>::TIterator It( ActiveSubtitles ); It; ++It )
+	for (TMap<PTRINT, FActiveSubtitle>::TIterator It(ActiveSubtitles); It; ++It)
 	{
 		FActiveSubtitle& Subtitle = It.Value();
-		if( Subtitle.bSplit )
+		if (Subtitle.bSplit)
 		{
 			continue;
 		}
 
 		// Concatenate the lines into one (in case the lines were partially manually split)
-		Concatenated.Empty( 256 );
+		Concatenated.Empty(256);
 
 		// Set up the base data
-		FSubtitleCue& Initial = Subtitle.Subtitles[ 0 ];
+		FSubtitleCue& Initial = Subtitle.Subtitles[0];
 		Concatenated = Initial.Text.ToString();
 		StartTime = Initial.Time;
 		SoundDuration = 0.0f;
-		for( i = 1; i < Subtitle.Subtitles.Num(); i++ )
+
+		for (int32 SubtitleIndex = 1; SubtitleIndex < Subtitle.Subtitles.Num(); ++SubtitleIndex)
 		{
-			FSubtitleCue& Subsequent = Subtitle.Subtitles[ i ];
+			FSubtitleCue& Subsequent = Subtitle.Subtitles[SubtitleIndex];
 			Concatenated += Subsequent.Text.ToString();
 			// Last blank entry sets the cutoff time to the duration of the sound
 			SoundDuration = Subsequent.Time - StartTime;
 		}
 
 		// Adjust concat string to use \n character instead of "/n" or "\n"
-		int32 SubIdx = Concatenated.Find( TEXT( "/n" ) );
-		while( SubIdx >= 0 )
+		int32 SubIdx = Concatenated.Find(TEXT("/n"));
+		while (SubIdx >= 0)
 		{
-			Concatenated = Concatenated.Left( SubIdx ) + ( ( TCHAR )L'\n' ) + Concatenated.Right( Concatenated.Len() - ( SubIdx + 2 ) );
-			SubIdx = Concatenated.Find( TEXT( "/n" ) );
+			Concatenated = Concatenated.Left(SubIdx) + ((TCHAR)L'\n') + Concatenated.Right(Concatenated.Len() - (SubIdx + 2));
+			SubIdx = Concatenated.Find(TEXT("/n"));
 		}
 
-		SubIdx = Concatenated.Find( TEXT( "\\n" ) );
-		while( SubIdx >= 0 )
+		SubIdx = Concatenated.Find(TEXT("\\n"));
+		while (SubIdx >= 0)
 		{
-			Concatenated = Concatenated.Left( SubIdx ) + ( ( TCHAR )L'\n' ) + Concatenated.Right( Concatenated.Len() - ( SubIdx + 2 ) );
-			SubIdx = Concatenated.Find( TEXT( "\\n" ));
+			Concatenated = Concatenated.Left(SubIdx) + ((TCHAR)L'\n') + Concatenated.Right(Concatenated.Len() - (SubIdx + 2));
+			SubIdx = Concatenated.Find(TEXT("\\n"));
 		}
 
 		// Work out a metric for the length of time a line should be displayed
@@ -234,17 +240,17 @@ void FSubtitleManager::SplitLinesToSafeZone( FCanvas* Canvas, FIntRect& Subtitle
 
 		// Word wrap into lines
 		TArray<FWrappedStringElement> Lines;
-		FTextSizingParameters RenderParms( 0.0f, 0.0f, SubtitleRegion.Width(), 0.0f, GEngine->GetSubtitleFont() );
-		Canvas->WrapString( RenderParms, 0, *Concatenated, Lines );
+		FTextSizingParameters RenderParms(0.0f, 0.0f, SubtitleRegion.Width(), 0.0f, GEngine->GetSubtitleFont());
+		Canvas->WrapString(RenderParms, 0, *Concatenated, Lines);
 
 		// Set up the times
 		Subtitle.Subtitles.Empty();
 		Cumulative = 0.0f;
-		for( i = 0; i < Lines.Num(); i++ )
+
+		for (int32 SubtitleIndex = 0; SubtitleIndex < Lines.Num(); ++SubtitleIndex)
 		{
-			FSubtitleCue& Line = Subtitle.Subtitles[Subtitle.Subtitles.AddDefaulted()];
-			
-			Line.Text = FText::FromString(Lines[ i ].Value);
+			FSubtitleCue& Line = Subtitle.Subtitles[Subtitle.Subtitles.AddDefaulted()];		
+			Line.Text = FText::FromString(Lines[SubtitleIndex].Value);
 			Line.Time = StartTime + Cumulative;
 
 			Cumulative += SecondsPerChar * Line.Text.ToString().Len();
@@ -252,15 +258,19 @@ void FSubtitleManager::SplitLinesToSafeZone( FCanvas* Canvas, FIntRect& Subtitle
 
 		// Add in the blank terminating line
 		FSubtitleCue& Temp = Subtitle.Subtitles[Subtitle.Subtitles.AddDefaulted()];
-		Temp.Text = FText::GetEmpty();
-		Temp.Time = StartTime + SoundDuration;
+		{
+			Temp.Text = FText::GetEmpty();
+			Temp.Time = StartTime + SoundDuration;
+		}
 
-		UE_LOG(LogAudio, Log, TEXT( "Splitting subtitle:" ) );
-		for( i = 0; i < Subtitle.Subtitles.Num() - 1; i++ )
+		UE_LOG(LogAudio, Log, TEXT("Splitting subtitle:"));
+
+		for (int32 SubtitleIndex = 0; SubtitleIndex < Subtitle.Subtitles.Num() - 1; ++SubtitleIndex)
 		{
-			FSubtitleCue& Cue = Subtitle.Subtitles[ i ];
-			FSubtitleCue& NextCue = Subtitle.Subtitles[ i + 1 ];
-			UE_LOG(LogAudio, Log, TEXT( " ... '%s' at %g to %g" ), *(Cue.Text.ToString()), Cue.Time, NextCue.Time );
+			FSubtitleCue& Cue = Subtitle.Subtitles[SubtitleIndex];
+			FSubtitleCue& NextCue = Subtitle.Subtitles[SubtitleIndex + 1];
+
+			UE_LOG(LogAudio, Log, TEXT(" ... '%s' at %g to %g"), *(Cue.Text.ToString()), Cue.Time, NextCue.Time);
 		}
 
 		// Mark it as split so it doesn't happen again
@@ -268,107 +278,100 @@ void FSubtitleManager::SplitLinesToSafeZone( FCanvas* Canvas, FIntRect& Subtitle
 	}
 }
 
-/**
- * Trim the SubtitleRegion to the safe 80% of the canvas 
- */
-void FSubtitleManager::TrimRegionToSafeZone( FCanvas* Canvas, FIntRect& SubtitleRegion )
+
+void FSubtitleManager::TrimRegionToSafeZone(FCanvas* Canvas, FIntRect& InOutSubtitleRegion)
 {
 	FIntRect SafeZone;
-
-	// Must display all text within text safe area - currently defined as 80% of the screen width and height
-	SafeZone.Min.X = ( 10 * Canvas->GetRenderTarget()->GetSizeXY().X ) / 100;
-	SafeZone.Min.Y = ( 10 * Canvas->GetRenderTarget()->GetSizeXY().Y ) / 100;
-	SafeZone.Max.X = Canvas->GetRenderTarget()->GetSizeXY().X - SafeZone.Min.X;
-	SafeZone.Max.Y = Canvas->GetRenderTarget()->GetSizeXY().Y - SafeZone.Min.Y;
+	{
+		// display all text within text safe area (80% of the screen width and height)
+		SafeZone.Min.X = (10 * Canvas->GetRenderTarget()->GetSizeXY().X) / 100;
+		SafeZone.Min.Y = (10 * Canvas->GetRenderTarget()->GetSizeXY().Y) / 100;
+		SafeZone.Max.X = Canvas->GetRenderTarget()->GetSizeXY().X - SafeZone.Min.X;
+		SafeZone.Max.Y = Canvas->GetRenderTarget()->GetSizeXY().Y - SafeZone.Min.Y;
+	}
 
 	// Trim to safe area, but keep everything central
-	if( SubtitleRegion.Min.X < SafeZone.Min.X || SubtitleRegion.Max.X > SafeZone.Max.X )
+	if ((InOutSubtitleRegion.Min.X < SafeZone.Min.X) || (InOutSubtitleRegion.Max.X > SafeZone.Max.X))
 	{
-		int32 Delta = FMath::Max( SafeZone.Min.X - SubtitleRegion.Min.X, SubtitleRegion.Max.X - SafeZone.Max.X );
-		SubtitleRegion.Min.X += Delta;
-		SubtitleRegion.Max.X -= Delta;
+		int32 Delta = FMath::Max(SafeZone.Min.X - InOutSubtitleRegion.Min.X, InOutSubtitleRegion.Max.X - SafeZone.Max.X);
+		InOutSubtitleRegion.Min.X += Delta;
+		InOutSubtitleRegion.Max.X -= Delta;
 	}
 
-	if( SubtitleRegion.Max.Y > SafeZone.Max.Y )
+	if (InOutSubtitleRegion.Max.Y > SafeZone.Max.Y)
 	{
-		SubtitleRegion.Max.Y = SafeZone.Max.Y;
+		InOutSubtitleRegion.Max.Y = SafeZone.Max.Y;
 	}
 }
 
-/**
- * Find the highest priority subtitle from the list of currently active ones
- */
-PTRINT FSubtitleManager::FindHighestPrioritySubtitle( float CurrentTime )
+
+PTRINT FSubtitleManager::FindHighestPrioritySubtitle(float CurrentTime)
 {
 	// Tick the available subtitles and find the highest priority one
 	float HighestPriority = -1.0f;
 	PTRINT HighestPriorityID = 0;
 
-	for( TMap<PTRINT, FActiveSubtitle>::TIterator It( ActiveSubtitles ); It; ++It )
+	for (TMap<PTRINT, FActiveSubtitle>::TIterator It(ActiveSubtitles); It; ++It)
 	{
 		FActiveSubtitle& Subtitle = It.Value();
 		
 		// Remove when last entry is reached
-		if( Subtitle.Index == Subtitle.Subtitles.Num() - 1 )
+		if (Subtitle.Index == Subtitle.Subtitles.Num() - 1)
 		{
 			It.RemoveCurrent();
 			continue;
 		}
 
-		if( CurrentTime >= Subtitle.Subtitles[ Subtitle.Index + 1 ].Time )
+		FSubtitleCue& SubtitleCue = Subtitle.Subtitles[Subtitle.Index];
+		
+		if (CurrentTime >= Subtitle.Subtitles[Subtitle.Index + 1].Time)
 		{
 			Subtitle.Index++;
 		}
 
-		if( Subtitle.Priority > HighestPriority )
+		if (Subtitle.Priority > HighestPriority)
 		{
 			HighestPriority = Subtitle.Priority;
 			HighestPriorityID = It.Key();
 		}
 	}
 
-	return( HighestPriorityID );
+	return(HighestPriorityID);
 }
 
-/**
- * Display the currently queued subtitles and cleanup after they have finished rendering
- *
- * @param  Canvas - where to render the subtitles
- * @param  CurrentTime - current world time
- */
-void FSubtitleManager::DisplaySubtitles( FCanvas* InCanvas, FIntRect& InSubtitleRegion, float InAudioTimeSeconds )
+
+void FSubtitleManager::DisplaySubtitles(FCanvas* InCanvas, FIntRect& InSubtitleRegion, float InAudioTimeSeconds)
 {
-	check( GEngine );
-	check( InCanvas );
+	check(GEngine);
+	check(InCanvas);
 
-	// Do nothing if subtitles are disabled.
-	if( !GEngine->bSubtitlesEnabled )
+	if (!GEngine->bSubtitlesEnabled)
 	{
-		return;
+		return; // do nothing if subtitles are disabled.
 	}
 
-	if( !GEngine->GetSubtitleFont() )
+	if (!GEngine->GetSubtitleFont())
 	{
-		UE_LOG(LogSubtitle, Warning, TEXT( "NULL GEngine->GetSubtitleFont() - subtitles not rendering!" ) );
+		UE_LOG(LogSubtitle, Warning, TEXT("NULL GEngine->GetSubtitleFont() - subtitles not rendering!"));
 		return;
 	}
 	
 	if (InSubtitleRegion.Area() > 0)
 	{
 		// Work out the safe zones
-		TrimRegionToSafeZone( InCanvas, InSubtitleRegion );
+		TrimRegionToSafeZone(InCanvas, InSubtitleRegion);
 
 		// If the lines have not already been split, split them to the safe zone now
-		SplitLinesToSafeZone( InCanvas, InSubtitleRegion );
+		SplitLinesToSafeZone(InCanvas, InSubtitleRegion);
 
 		// Find the subtitle to display
-		PTRINT HighestPriorityID = FindHighestPrioritySubtitle( InAudioTimeSeconds );
+		PTRINT HighestPriorityID = FindHighestPrioritySubtitle(InAudioTimeSeconds);
 
 		// Display the highest priority subtitle
-		if( HighestPriorityID )
+		if (HighestPriorityID)
 		{
-			FActiveSubtitle* Subtitle = ActiveSubtitles.Find( HighestPriorityID );
-			DisplaySubtitle( InCanvas, Subtitle, InSubtitleRegion, FLinearColor::White );
+			FActiveSubtitle* Subtitle = ActiveSubtitles.Find(HighestPriorityID);
+			DisplaySubtitle(InCanvas, Subtitle, InSubtitleRegion, FLinearColor::White);
 		}
 		else
 		{
@@ -377,20 +380,16 @@ void FSubtitleManager::DisplaySubtitles( FCanvas* InCanvas, FIntRect& InSubtitle
 	}
 }
 
-/**
- * Returns the singleton subtitle manager instance, creating it if necessary.
- */
-FSubtitleManager* FSubtitleManager::GetSubtitleManager( void )
+
+FSubtitleManager* FSubtitleManager::GetSubtitleManager()
 {
 	static FSubtitleManager* SSubtitleManager = nullptr;
 
-	if( !SSubtitleManager )
+	if (!SSubtitleManager)
 	{
 		SSubtitleManager = new FSubtitleManager();
 		SSubtitleManager->CurrentSubtitleHeight = 0.0f;
 	}
 
-	return( SSubtitleManager );
+	return(SSubtitleManager);
 }
-
-// end
diff --git a/Engine/Source/Runtime/Engine/Public/AudioDevice.h b/Engine/Source/Runtime/Engine/Public/AudioDevice.h
--- a/Engine/Source/Runtime/Engine/Public/AudioDevice.h
+++ b/Engine/Source/Runtime/Engine/Public/AudioDevice.h
@@ -1021,6 +1021,9 @@ public:
 	/** Query if the editor is in VR Preview for the current play world. Returns false for non-editor builds */
 	static bool CanUseVRAudioDevice();
 
+	/** Returns the audio clock of the audio device. Not supported on all platforms. */
+	double GetAudioClock() const { return AudioClock; }
+
 #if !UE_BUILD_SHIPPING
 	void DumpActiveSounds() const;
 
@@ -1187,6 +1190,13 @@ private:
 	/** Whether the audio device has been initialized */
 	uint8 bIsInitialized:1;
 
+protected:
+
+	/** The audio clock from the audio hardware. Not supported on all platforms. */
+	double AudioClock;
+
+private:
+
 	/** Whether the value in HighestPriorityActivatedReverb should be used - Audio Thread owned */
 	uint8 bHasActivatedReverb:1;
 
diff --git a/Engine/Source/Runtime/Engine/Public/SubtitleManager.h b/Engine/Source/Runtime/Engine/Public/SubtitleManager.h
--- a/Engine/Source/Runtime/Engine/Public/SubtitleManager.h
+++ b/Engine/Source/Runtime/Engine/Public/SubtitleManager.h
@@ -1,37 +1,43 @@
 // Copyright 1998-2016 Epic Games, Inc. All Rights Reserved.
 
-
 #pragma once
 
 
 #define SUBTITLE_SCREEN_DEPTH_FOR_3D 0.1f
 
+
 /**
  * A collection of subtitles, rendered at a certain priority.
  */
 class FActiveSubtitle
 {
 public:
-	FActiveSubtitle( int32 InIndex, float InPriority, bool InbSplit, bool InbSingleLine, TArray<FSubtitleCue> InSubtitles )
-		:	Index( InIndex )
-		,	Priority( InPriority )
-		,	bSplit( InbSplit )
-		,	bSingleLine( InbSingleLine )
-		,	Subtitles( MoveTemp(InSubtitles) )
-	{}
+
+	FActiveSubtitle(int32 InIndex, float InPriority, bool InbSplit, bool InbSingleLine, TArray<FSubtitleCue> InSubtitles)
+		: Index(InIndex)
+		, Priority(InPriority)
+		, bSplit(InbSplit)
+		, bSingleLine(InbSingleLine)
+		, Subtitles(MoveTemp(InSubtitles))
+	{ }
 
 	/** Index into the Subtitles array for the currently active subtitle in this set. */
-	int32						Index;
+	int32 Index;
+
 	/** Priority for this set of subtitles, used by FSubtitleManager to determine which subtitle to play. */
-	float					Priority;
-	/** true if the subtitles have come in pre split, or after they have been split programmatically. */
-	bool					bSplit;
-	/** true if the subtitles should be displayed as a sequence of single lines */
-	bool					bSingleLine;
+	float Priority;
+	
+	/** Whether subtitles have come in pre-split, or after they have been split programmatically. */
+	bool bSplit;
+	
+	/** Whether subtitles should be displayed as a sequence of single lines. */
+	bool bSingleLine;
+	
 	/** A set of subtitles. */
 	TArray<FSubtitleCue>	Subtitles;
 };
 
+
 struct FQueueSubtitleParams
 {
 	FQueueSubtitleParams(const TArray<FSubtitleCue>& InSubtitles)
@@ -46,94 +52,108 @@ struct FQueueSubtitleParams
 	uint8 bManualWordWrap : 1;
 	uint8 bSingleLine : 1;
 	const TArray<FSubtitleCue>& Subtitles;
+	float RequestedStartTime;
 };
 
 
 /**
- * Subtitle manager.  Handles prioritization and rendering of subtitles.
+ * Subtitle manager. Handles prioritization and rendering of subtitles.
  */
 class FSubtitleManager
 {
 public:
-	/**
-	 * Kills all the subtitles
-	 */
-	void		KillAllSubtitles( void );
+
+	/** Kill all the subtitles. */
+	void KillAllSubtitles(void);
 
 	/**
-	 * Kills all the subtitles associated with the subtitle ID
+	 * Kill specified subtitles.
+	 *
+	 * @param SubtitleId Identifier of the subtitles to kill.
 	 */
-	void		KillSubtitles( PTRINT SubtitleID );
+	void KillSubtitles(PTRINT SubtitleID);
 
 	/**
-	 * Trim the SubtitleRegion to the safe 80% of the canvas 
+	 * Trim the SubtitleRegion to the safe 80% of the canvas.
+	 *
+	 * @param Canvas The canvas to trim to.
+	 * @param InOutSubtitleRegion The trim region.
 	 */
-	void		TrimRegionToSafeZone( FCanvas * Canvas, FIntRect & SubtitleRegion );
+	void TrimRegionToSafeZone(FCanvas* Canvas, FIntRect& InOutSubtitleRegion);
 
 	/**
-	 * If any of the active subtitles need to be split into multiple lines, do so now
-	 * - caveat - this assumes the width of the subtitle region does not change while the subtitle is active
+	 * If any of the active subtitles need to be split into multiple lines, do so now.
+	 *
+	 * Note: This assumes the width of the subtitle region does not change while the subtitle is active.
 	 */
-	void		SplitLinesToSafeZone( FCanvas* Canvas, FIntRect & SubtitleRegion );
+	void SplitLinesToSafeZone(FCanvas* Canvas, FIntRect & SubtitleRegion);
 
 	/**
-	 * Find the highest priority subtitle from the list of currently active ones
+	 * Find the highest priority subtitle from the list of currently active ones.
+	 *
+	 * @param CurrentTime The time at which to find the subtitle.
+	 * @return Found subtitle identifier.
 	 */
-	PTRINT		FindHighestPrioritySubtitle( float CurrentTime );
+	PTRINT FindHighestPrioritySubtitle(float CurrentTime);
 
 	/**
 	 * Add an array of subtitles to the active list
 	 *
-	 * @param  SubTitleID - the controlling id that groups sets of subtitles together
-	 * @param  Priority - used to prioritize subtitles; higher values have higher priority.  Subtitles with a priority 0.0f are not displayed.
-	 * @param  StartTime - float of seconds that the subtitles start at
-	 * @param  SoundDuration - time in seconds after which the subtitles do not display
-	 * @param  Subtitles - TArray of lines of subtitle and time offset to play them
+	 * @param SubTitleID The controlling id that groups sets of subtitles together.
+	 * @param Priority Used to prioritize subtitles; higher values have higher priority.  Subtitles with a priority 0.0f are not displayed.
+	 * @param StartTime Time at which the subtitles start (in seconds).
+	 * @param SoundDuration Time  after which the subtitles do not display (in seconds).
+	 * @param Subtitles Collection of lines of subtitle and time offset to play them.
 	 */
-	void		QueueSubtitles( PTRINT SubtitleID, float Priority, bool bManualWordWrap, bool bSingleLine, float SoundDuration, const TArray<FSubtitleCue>& Subtitles, float InStartTime );
+	void QueueSubtitles(PTRINT SubtitleID, float Priority, bool bManualWordWrap, bool bSingleLine, float SoundDuration, const TArray<FSubtitleCue>& Subtitles, float InStartTime, float InCurrentTime);
 
 	static void QueueSubtitles(const FQueueSubtitleParams& QueueSubtitlesParams);
 
 	/**
 	 * Draws a subtitle at the specified pixel location.
 	 */
-	void		DisplaySubtitle( FCanvas *Canvas, FActiveSubtitle* Subtitle, FIntRect & Parms, const FLinearColor& Color );
+	void DisplaySubtitle(FCanvas *Canvas, FActiveSubtitle* Subtitle, FIntRect & Parms, const FLinearColor& Color);
 
 	/**
-	 * Display the currently queued subtitles and cleanup after they have finished rendering
+	 * Display the currently queued subtitles and cleanup after they have finished rendering.
 	 *
-	 * @param  Canvas - where to render the subtitles
-	 * @param  CurrentTime - current world time
+	 * @param Canvas Where to render the subtitles.
+	 * @param CurrentTime Current world time.
 	 */
-	ENGINE_API void		DisplaySubtitles( FCanvas *InCanvas, FIntRect & SubtitleRegion, float InAudioTimeSeconds );
+	ENGINE_API void DisplaySubtitles(FCanvas *InCanvas, FIntRect & SubtitleRegion, float InAudioTimeSeconds);
 
-	/** @return true if there are any active subtitles */
-	bool		HasSubtitles( void ) 
+	/**
+	 * Whether there are any active subtitles.
+	 *
+	 * @return true if there are any active subtitles, false otherwise.
+	 */
+	bool HasSubtitles() 
 	{ 
-		return ActiveSubtitles.Num() > 0; 
+		return (ActiveSubtitles.Num() > 0); 
 	}
 
 	/**
-	 * @return The height of the currently rendered subtitles, 0 if none are rendering
+	 * Get the height of currently rendered subtitles.
+	 *
+	 * @return Subtitle height (in pixels), or 0 if none are rendering.
 	 */
-	float GetCurrentSubtitlesHeight( void )
+	float GetCurrentSubtitlesHeight()
 	{
 		return CurrentSubtitleHeight;
 	}
 
 	/**
-	 * Returns the singleton subtitle manager instance, creating it if necessary.
+	 * Get the subtitle manager singleton instance.
+	 *
+	 * @return Subtitle manager.
 	 */
-	ENGINE_API static FSubtitleManager* GetSubtitleManager( void );
+	ENGINE_API static FSubtitleManager* GetSubtitleManager();
 
 private:
-	/**
-	 * The set of active, prioritized subtitles.
-	 */
-	TMap<PTRINT, FActiveSubtitle>	ActiveSubtitles;
 
-	/**
-	 * The current height of the subtitles
-	 */
+	/** The set of active, prioritized subtitles. */
+	TMap<PTRINT, FActiveSubtitle> ActiveSubtitles;
+
+	/** The current height of the subtitles. */
 	float CurrentSubtitleHeight;
 };
diff --git a/Engine/Source/Runtime/LevelSequence/Private/LevelSequenceActor.cpp b/Engine/Source/Runtime/LevelSequence/Private/LevelSequenceActor.cpp
--- a/Engine/Source/Runtime/LevelSequence/Private/LevelSequenceActor.cpp
+++ b/Engine/Source/Runtime/LevelSequence/Private/LevelSequenceActor.cpp
@@ -8,6 +8,9 @@
 ALevelSequenceActor::ALevelSequenceActor(const FObjectInitializer& Init)
 	: Super(Init)
 {
+	USceneComponent* SceneComponent = CreateDefaultSubobject<USceneComponent>(TEXT("SceneComp"));
+	RootComponent = SceneComponent;
+
 #if WITH_EDITORONLY_DATA
 	UBillboardComponent* SpriteComponent = CreateEditorOnlyDefaultSubobject<UBillboardComponent>(TEXT("Sprite"));
 
@@ -29,7 +32,6 @@ ALevelSequenceActor::ALevelSequenceActor(const FObjectInitializer& Init)
 			SpriteComponent->bAbsoluteScale = true;
 			SpriteComponent->bReceivesDecals = false;
 			SpriteComponent->bHiddenInGame = true;
-			SetRootComponent(SpriteComponent);
 		}
 	}
 #endif //WITH_EDITORONLY_DATA
@@ -74,6 +76,20 @@ void ALevelSequenceActor::Tick(float DeltaSeconds)
 	}
 }
 
+void ALevelSequenceActor::PostLoad()
+{
+	Super::PostLoad();
+
+#if WITH_EDITORONLY_DATA
+	// Fix sprite component so that it's attached to the root component. In the past, the sprite component was the root component.
+	UBillboardComponent* SpriteComponent = FindComponentByClass<UBillboardComponent>();
+	if (SpriteComponent && SpriteComponent->GetAttachParent() != RootComponent)
+	{
+		SpriteComponent->SetupAttachment(RootComponent);
+	}
+#endif
+}
+
 
 ULevelSequence* ALevelSequenceActor::GetSequence(bool Load) const
 {
diff --git a/Engine/Source/Runtime/LevelSequence/Private/LevelSequencePlayer.cpp b/Engine/Source/Runtime/LevelSequence/Private/LevelSequencePlayer.cpp
--- a/Engine/Source/Runtime/LevelSequence/Private/LevelSequencePlayer.cpp
+++ b/Engine/Source/Runtime/LevelSequence/Private/LevelSequencePlayer.cpp
@@ -136,6 +136,11 @@ void ULevelSequencePlayer::Stop()
 		TimeCursorPosition = PlaybackSettings.PlayRate < 0.f ? GetLength() : 0.f;
 		CurrentNumLoops = 0;
 
+		if (RootMovieSceneInstance.IsValid())
+		{
+			RootMovieSceneInstance->RestoreState(*this);
+		}
+
 		SetTickPrerequisites(false);
 
 		if (OnStop.IsBound())
diff --git a/Engine/Source/Runtime/LevelSequence/Private/LevelSequenceSpawnRegister.cpp b/Engine/Source/Runtime/LevelSequence/Private/LevelSequenceSpawnRegister.cpp
--- a/Engine/Source/Runtime/LevelSequence/Private/LevelSequenceSpawnRegister.cpp
+++ b/Engine/Source/Runtime/LevelSequence/Private/LevelSequenceSpawnRegister.cpp
@@ -55,6 +55,8 @@ UObject* FLevelSequenceSpawnRegister::SpawnObject(const FGuid& BindingId, FMovie
 		SpawnInfo.SpawnCollisionHandlingOverride = ESpawnActorCollisionHandlingMethod::AlwaysSpawn;
 		// @todo: Spawning with a non-CDO template is fraught with issues
 		//SpawnInfo.Template = ObjectTemplate;
+		// allow pre-construction variables to be set.
+		SpawnInfo.bDeferConstruction = true;
 	}
 
 	FTransform SpawnTransform;
@@ -108,6 +110,8 @@ UObject* FLevelSequenceSpawnRegister::SpawnObject(const FGuid& BindingId, FMovie
 	// tag this actor so we know it was spawned by sequencer
 	SpawnedActor->Tags.Add(SequencerActorTag);
 
+	SpawnedActor->FinishSpawning(SpawnTransform);
+
 	Register.Add(Key, FSpawnedObject(BindingId, *SpawnedActor, Spawnable->GetSpawnOwnership()));
 
 	SequenceInstance.OnObjectSpawned(BindingId, *SpawnedActor, Player);
diff --git a/Engine/Source/Runtime/MediaAssets/Private/Assets/FileMediaSource.cpp b/Engine/Source/Runtime/MediaAssets/Private/Assets/FileMediaSource.cpp
--- a/Engine/Source/Runtime/MediaAssets/Private/Assets/FileMediaSource.cpp
+++ b/Engine/Source/Runtime/MediaAssets/Private/Assets/FileMediaSource.cpp
@@ -7,6 +7,34 @@
 static FName PrecacheFileName("PrecacheFile");
 
 
+/* UFileMediaSource interface
+ *****************************************************************************/
+
+void UFileMediaSource::SetFilePath(const FString& Path)
+{
+	if (Path.IsEmpty() || Path.StartsWith(TEXT("./")))
+	{
+		FilePath = Path;
+	}
+	else
+	{
+		FString FullPath = FPaths::ConvertRelativePathToFull(Path);
+		const FString FullGameContentDir = FPaths::ConvertRelativePathToFull(FPaths::GameContentDir());
+
+		if (FullPath.StartsWith(FullGameContentDir))
+		{
+			FPaths::MakePathRelativeTo(FullPath, *FullGameContentDir);
+			FullPath = FString(TEXT("./")) + FullPath;
+		}
+
+		FilePath = FullPath;
+	}
+}
+
+
+/* UMediaSource overrides
+ *****************************************************************************/
+
 bool UFileMediaSource::GetMediaOption(const FName& Key, bool DefaultValue) const
 {
 	if (Key == PrecacheFileName)
@@ -30,6 +58,9 @@ bool UFileMediaSource::Validate() const
 }
 
 
+/* UFileMediaSource implementation
+ *****************************************************************************/
+
 FString UFileMediaSource::GetFullPath() const
 {
 	FString FullPath = FPaths::ConvertRelativePathToFull(
diff --git a/Engine/Source/Runtime/MediaAssets/Private/Assets/MediaPlayer.cpp b/Engine/Source/Runtime/MediaAssets/Private/Assets/MediaPlayer.cpp
--- a/Engine/Source/Runtime/MediaAssets/Private/Assets/MediaPlayer.cpp
+++ b/Engine/Source/Runtime/MediaAssets/Private/Assets/MediaPlayer.cpp
@@ -34,6 +34,32 @@ bool UMediaPlayer::CanPause() const
 }
 
 
+bool UMediaPlayer::CanPlaySource(UMediaSource* MediaSource)
+{
+	if ((MediaSource == nullptr) || !MediaSource->Validate())
+	{
+		return false;
+	}
+
+	FName OutPlayerName;
+
+	return FindPlayerForUrl(MediaSource->GetUrl(), *MediaSource, OutPlayerName).IsValid();
+}
+
+
+bool UMediaPlayer::CanPlayUrl(const FString& Url)
+{
+	if (Url.IsEmpty())
+	{
+		return false;
+	}
+
+	FName OutPlayerName;
+
+	return FindPlayerForUrl(Url, *GetDefault<UMediaSource>(), OutPlayerName).IsValid();
+}
+
+
 void UMediaPlayer::Close()
 {
 	CurrentUrl.Empty();
@@ -455,6 +481,7 @@ void UMediaPlayer::BeginDestroy()
 	if (Player.IsValid())
 	{
 		Player->Close();
+		Player->GetOutput().SetCaptionSink(nullptr);
 		Player->OnMediaEvent().RemoveAll(this);
 		Player.Reset();
 	}
@@ -556,106 +583,107 @@ void UMediaPlayer::PreEditChange(UProperty* PropertyAboutToChange)
 /* UMediaPlayer implementation
  *****************************************************************************/
 
-bool UMediaPlayer::Open(const FString& Url, const IMediaOptions& Options)
+TSharedPtr<IMediaPlayer> UMediaPlayer::FindPlayerForUrl(const FString& Url, const IMediaOptions& Options, FName& OutPlayerName)
 {
-	if (IsRunningDedicatedServer())
+	// reuse existing player if desired
+	if ((DesiredPlayerName == PlayerName) && (PlayerName != NAME_None))
 	{
-		return false;
-	}
+		OutPlayerName = PlayerName;
 
-	if (Url.IsEmpty())
-	{
-		return false;
+		return Player;
 	}
 
+	OutPlayerName = NAME_None;
+
 	// load media module
 	IMediaModule* MediaModule = FModuleManager::LoadModulePtr<IMediaModule>("Media");
 
 	if (MediaModule == nullptr)
 	{
 		UE_LOG(LogMediaAssets, Error, TEXT("Failed to load Media module"));
-
-		return false;
+		return nullptr;
 	}
 
-	TSharedPtr<IMediaPlayer> NewPlayer;
-	FName NewPlayerName = NAME_None;
-
-	// find a native player
+	// try to create desired player
 	if (DesiredPlayerName != NAME_None)
 	{
-		if (DesiredPlayerName != PlayerName)
+		IMediaPlayerFactory* Factory = MediaModule->GetPlayerFactory(DesiredPlayerName);
+
+		if (Factory == nullptr)
 		{
-			// create desired player
-			IMediaPlayerFactory* Factory = MediaModule->GetPlayerFactory(DesiredPlayerName);
-
-			if (Factory == nullptr)
-			{
-				UE_LOG(LogMediaAssets, Error, TEXT("Could not find desired player %s for %s"), *DesiredPlayerName.ToString(), *Url);
-			}
-			else
-			{
-				NewPlayer = Factory->CreatePlayer();
-
-				if (!NewPlayer.IsValid())
-				{
-					UE_LOG(LogMediaAssets, Error, TEXT("Failed to create desired player %s for %s"), *DesiredPlayerName.ToString(), *Url);
-				}
-				else
-				{
-					NewPlayerName = DesiredPlayerName;
-				}
-			}
+			UE_LOG(LogMediaAssets, Error, TEXT("Could not find desired player %s for %s"), *DesiredPlayerName.ToString(), *Url);
+			return nullptr;
 		}
-		else
+
+		TSharedPtr<IMediaPlayer> NewPlayer = Factory->CreatePlayer();
+
+		if (!NewPlayer.IsValid())
 		{
-			NewPlayer = Player; // reuse existing player
+			UE_LOG(LogMediaAssets, Error, TEXT("Failed to create desired player %s for %s"), *DesiredPlayerName.ToString(), *Url);
+			return nullptr;
 		}
+
+		OutPlayerName = DesiredPlayerName;
+
+		return NewPlayer;
 	}
-	else
+
+	// try to reuse existing player
+	if (PlayerName != NAME_None)
 	{
-		if (PlayerName != NAME_None)
+		IMediaPlayerFactory* Factory = MediaModule->GetPlayerFactory(PlayerName);
+
+		if ((Factory != nullptr) && Factory->CanPlayUrl(Url, Options))
 		{
-			// try to reuse existing player
-			IMediaPlayerFactory* Factory = MediaModule->GetPlayerFactory(PlayerName);
+			OutPlayerName = PlayerName;
 
-			if ((Factory != nullptr) && Factory->CanPlayUrl(Url, Options))
-			{
-				NewPlayer = Player;
-			}
+			return Player;
 		}
+	}
 
-		if (!NewPlayer.IsValid())
+	// try to auto-select new player
+	const TArray<IMediaPlayerFactory*>& PlayerFactories = MediaModule->GetPlayerFactories();
+	const FString RunningPlatformName(FPlatformProperties::IniPlatformName());
+
+	for (IMediaPlayerFactory* Factory : PlayerFactories)
+	{
+		if (!Factory->SupportsPlatform(RunningPlatformName) || !Factory->CanPlayUrl(Url, Options))
 		{
-			// automatically select & create new player
-			const TArray<IMediaPlayerFactory*>& PlayerFactories = MediaModule->GetPlayerFactories();
-			const FString RunningPlatformName(FPlatformProperties::IniPlatformName());
-
-			for (IMediaPlayerFactory* Factory : PlayerFactories)
-			{
-				if (!Factory->SupportsPlatform(RunningPlatformName) || !Factory->CanPlayUrl(Url, Options))
-				{
-					continue;
-				}
-
-				NewPlayer = Factory->CreatePlayer();
-
-				if (NewPlayer.IsValid())
-				{
-					NewPlayerName = Factory->GetName();
-
-					break;
-				}
-			}
-
-			if (!NewPlayer.IsValid())
-			{
-				UE_LOG(LogMediaAssets, Error, TEXT("Could not find a native player for %s"), *Url);
-			}
+			continue;
+		}
+
+		TSharedPtr<IMediaPlayer> NewPlayer = Factory->CreatePlayer();
+
+		if (NewPlayer.IsValid())
+		{
+			OutPlayerName = Factory->GetName();
+
+			return NewPlayer;
 		}
 	}
 
-	// initialize new player
+	UE_LOG(LogMediaAssets, Error, TEXT("Could not find a native player for %s"), *Url);
+
+	return nullptr;
+}
+
+
+bool UMediaPlayer::Open(const FString& Url, const IMediaOptions& Options)
+{
+	if (IsRunningDedicatedServer())
+	{
+		return false;
+	}
+
+	if (Url.IsEmpty())
+	{
+		return false;
+	}
+
+	// find & initialize new player
+	FName OutPlayerName = NAME_None;
+	TSharedPtr<IMediaPlayer> NewPlayer = FindPlayerForUrl(Url, Options, OutPlayerName);
+
 	if (NewPlayer != Player)
 	{
 		if (Player.IsValid())
@@ -667,14 +695,10 @@ bool UMediaPlayer::Open(const FString& Url, const IMediaOptions& Options)
 		if (NewPlayer.IsValid())
 		{
 			NewPlayer->OnMediaEvent().AddUObject(this, &UMediaPlayer::HandlePlayerMediaEvent);
-			PlayerName = NewPlayerName;
-		}
-		else
-		{
-			PlayerName = NAME_None;
 		}
 
 		Player = NewPlayer;
+		PlayerName = OutPlayerName;
 	}
 
 	if (!Player.IsValid())
@@ -748,6 +772,7 @@ void UMediaPlayer::HandlePlayerMediaEvent(EMediaEvent Event)
 
 	case EMediaEvent::MediaOpened:
 		Player->GetControls().SetLooping(Loop);
+		Player->GetOutput().SetCaptionSink(CaptionSink);
 
 		if (ImageTexture != nullptr)
 		{
diff --git a/Engine/Source/Runtime/MediaAssets/Public/FileMediaSource.h b/Engine/Source/Runtime/MediaAssets/Public/FileMediaSource.h
--- a/Engine/Source/Runtime/MediaAssets/Public/FileMediaSource.h
+++ b/Engine/Source/Runtime/MediaAssets/Public/FileMediaSource.h
@@ -14,14 +14,43 @@ class MEDIAASSETS_API UFileMediaSource
 
 public:
 
-	/** The path to the media file to be played. */
+	/**
+	 * The path to the media file to be played.
+	 *
+	 * @see SetFilePath
+	 */
 	UPROPERTY(BlueprintReadWrite, EditAnywhere, Category=File, AssetRegistrySearchable)
 	FString FilePath;
 
 	/** Load entire media file into memory and play from there (if possible). */
-	UPROPERTY(BlueprintReadWrite, EditAnywhere, Category=File)
+	UPROPERTY(BlueprintReadWrite, EditAnywhere, Category=File, AdvancedDisplay)
 	bool PrecacheFile;
 
+public:
+
+	/**
+	 * Get the path to the media file to be played.
+	 *
+	 * @return The file path.
+	 * @see SetFilePath
+	 */
+	const FString& GetFilePath() const
+	{
+		return FilePath;
+	}
+
+	/**
+	 * Set the path to the media file that this source represents.
+	 *
+	 * Automatically converts full paths to media sources that reside in the
+	 * Engine's or project's /Content/Movies directory into relative paths.
+	 *
+	 * @param Path The path to set.
+	 * @see FilePath, GetFilePath
+	 */
+	UFUNCTION(BlueprintCallable, Category="Media|FileMediaSource")
+	void SetFilePath(const FString& Path);
+
 public:
 
 	//~ IMediaOptions interface
diff --git a/Engine/Source/Runtime/MediaAssets/Public/MediaPlayer.h b/Engine/Source/Runtime/MediaAssets/Public/MediaPlayer.h
--- a/Engine/Source/Runtime/MediaAssets/Public/MediaPlayer.h
+++ b/Engine/Source/Runtime/MediaAssets/Public/MediaPlayer.h
@@ -63,7 +63,7 @@ class MEDIAASSETS_API UMediaPlayer
 public:
 
 	/**
-	 * Checks whether media playback can be paused right now.
+	 * Check whether media playback can be paused right now.
 	 *
 	 * Playback can be paused if the media supports pausing and if it is currently playing.
 	 *
@@ -73,6 +73,30 @@ public:
 	UFUNCTION(BlueprintCallable, Category="Media|MediaPlayer")
 	bool CanPause() const;
 
+	/**
+	 * Check whether the specified media source can be played by this player.
+	 *
+	 * If a desired player name is set for this player (DesiredPlayerName), it will
+	 * only check whether that particular player type can play the specified source.
+	 *
+	 * @param MediaSource The media source to check.
+	 * @return true if the media source can be opened, false otherwise.
+	 * @see CanPlayUrl, DesiredPlayerName
+	 */
+	bool CanPlaySource(UMediaSource* MediaSource);
+
+	/**
+	 * Check whether the specified URL can be played by this player.
+	 *
+	 * If a desired player name is set for this player (DesiredPlayerName), it will
+	 * only check whether that particular player type can play the specified URL.
+	 *
+	 * @param Url The URL to check.
+	 * @see CanPlaySource, DesiredPlayerName
+	 */
+	UFUNCTION(BlueprintCallable, Category="Media|MediaPlayer")
+	bool CanPlayUrl(const FString& Url);
+
 	/**
 	 * Close the currently open media, if any.
 	 *
@@ -603,6 +627,16 @@ public:
 
 protected:
 
+	/**
+	 * Find a player that can play the specified media URL.
+	 *
+	 * @param Url The URL to play.
+	 * @param Options The media options for the URL.
+	 * @param OutPlayerName Will contain the name of the found player.
+	 * @return The player if found, or nullptr otherwise.
+	 */
+	TSharedPtr<IMediaPlayer> FindPlayerForUrl(const FString& Url, const IMediaOptions& Options, FName& OutPlayerName);
+
 	/**
 	 * Open a media source from a URL with optional parameters.
 	 *
diff --git a/Engine/Source/Runtime/MovieScene/Public/KeyParams.h b/Engine/Source/Runtime/MovieScene/Public/KeyParams.h
--- a/Engine/Source/Runtime/MovieScene/Public/KeyParams.h
+++ b/Engine/Source/Runtime/MovieScene/Public/KeyParams.h
@@ -21,37 +21,13 @@ enum class EMovieSceneKeyInterpolation
 	Constant UMETA(DisplayName="Constant"),
 };
 
-/**
- * Parameters for determining keying behavior
- */
-struct MOVIESCENE_API FKeyParams
+/** Defines different modes for adding keyframes in sequencer. */
+enum class ESequencerKeyMode
 {
-	FKeyParams()
-	{
-		bCreateHandleIfMissing = false;
-		bCreateTrackIfMissing = false;
-		bCreateKeyOnlyWhenAutoKeying = false;
-		bCreateKeyIfUnchanged = false;
-		bCreateKeyIfEmpty = false;
-	}
-
-	FKeyParams(const FKeyParams& InKeyParams)
-	{
-		bCreateHandleIfMissing = InKeyParams.bCreateHandleIfMissing;
-		bCreateTrackIfMissing = InKeyParams.bCreateTrackIfMissing;
-		bCreateKeyOnlyWhenAutoKeying = InKeyParams.bCreateKeyOnlyWhenAutoKeying;
-		bCreateKeyIfUnchanged = InKeyParams.bCreateKeyIfUnchanged;
-		bCreateKeyIfEmpty = InKeyParams.bCreateKeyIfEmpty;
-	}
-
-	/** Create handle if it doesn't exist. */
-	bool bCreateHandleIfMissing;
-	/** Create track if it doesn't exist. */
-	bool bCreateTrackIfMissing;
-	/** Create a new key only when currently autokeying.*/
-	bool bCreateKeyOnlyWhenAutoKeying;
-	/** Create a key even if it's unchanged. */
-	bool bCreateKeyIfUnchanged;
-	/** Create a key even if the track is empty */
-	bool bCreateKeyIfEmpty;
+	/** Keyframes are being generated automatically by the user interacting with object properties or transforms. */
+	AutoKey,
+	/** The user has requested to manually add a keyframe.  Keys will be added for unchanged and empty components of multi-component tracks depending on current settings. */
+	ManualKey,
+	/** The user has requested to manually add a keyframe.  Keys will be added for unchanged and empty components of multi-component tracks regardless of the current settings. */
+	ManualKeyForced
 };
diff --git a/Engine/Source/Runtime/MovieSceneTracks/Private/Sections/MovieScene3DTransformSection.cpp b/Engine/Source/Runtime/MovieSceneTracks/Private/Sections/MovieScene3DTransformSection.cpp
--- a/Engine/Source/Runtime/MovieSceneTracks/Private/Sections/MovieScene3DTransformSection.cpp
+++ b/Engine/Source/Runtime/MovieSceneTracks/Private/Sections/MovieScene3DTransformSection.cpp
@@ -100,27 +100,27 @@ UMovieScene3DTransformSection::UMovieScene3DTransformSection(const FObjectInitia
 
 /* UMovieScene3DTransformSection interface
  *****************************************************************************/
-void UMovieScene3DTransformSection::EvalTranslation(float Time, FVector& OutTranslation) const
+void UMovieScene3DTransformSection::EvalTranslation(float Time, FVector& InOutTranslation) const
 {
-	OutTranslation.X = Translation[0].Eval(Time);
-	OutTranslation.Y = Translation[1].Eval(Time);
-	OutTranslation.Z = Translation[2].Eval(Time);
+	InOutTranslation.X = Translation[0].Eval(Time, InOutTranslation.X);
+	InOutTranslation.Y = Translation[1].Eval(Time, InOutTranslation.Y);
+	InOutTranslation.Z = Translation[2].Eval(Time, InOutTranslation.Z);
 }
 
 
-void UMovieScene3DTransformSection::EvalRotation(float Time, FRotator& OutRotation) const
+void UMovieScene3DTransformSection::EvalRotation(float Time, FRotator& InOutRotation) const
 {
-	OutRotation.Roll = Rotation[0].Eval(Time);
-	OutRotation.Pitch = Rotation[1].Eval(Time);
-	OutRotation.Yaw = Rotation[2].Eval(Time);
+	InOutRotation.Roll = Rotation[0].Eval(Time, InOutRotation.Roll);
+	InOutRotation.Pitch = Rotation[1].Eval(Time, InOutRotation.Pitch);
+	InOutRotation.Yaw = Rotation[2].Eval(Time, InOutRotation.Yaw);
 }
 
 
-void UMovieScene3DTransformSection::EvalScale(float Time, FVector& OutScale) const
+void UMovieScene3DTransformSection::EvalScale(float Time, FVector& InOutScale) const
 {
-	OutScale.X = Scale[0].Eval(Time);
-	OutScale.Y = Scale[1].Eval(Time);
-	OutScale.Z = Scale[2].Eval(Time);
+	InOutScale.X = Scale[0].Eval(Time, InOutScale.X);
+	InOutScale.Y = Scale[1].Eval(Time, InOutScale.Y);
+	InOutScale.Z = Scale[2].Eval(Time, InOutScale.Z);
 }
 
 
@@ -549,3 +549,17 @@ void UMovieScene3DTransformSection::SetDefault(const FTransformKey& TransformKey
 	FRichCurve* KeyCurve = GetCurveForChannelAndAxis(TransformKey.Channel, TransformKey.Axis, Translation, Rotation, Scale);
 	SetCurveDefault(*KeyCurve, TransformKey.Value);
 }
+
+
+void UMovieScene3DTransformSection::ClearDefaults()
+{
+	Translation[0].ClearDefaultValue();
+	Translation[1].ClearDefaultValue();
+	Translation[2].ClearDefaultValue();
+	Rotation[0].ClearDefaultValue();
+	Rotation[1].ClearDefaultValue();
+	Rotation[2].ClearDefaultValue();
+	Scale[0].ClearDefaultValue();
+	Scale[1].ClearDefaultValue();
+	Scale[2].ClearDefaultValue();
+}
diff --git a/Engine/Source/Runtime/MovieSceneTracks/Private/Sections/MovieSceneActorReferenceSection.cpp b/Engine/Source/Runtime/MovieSceneTracks/Private/Sections/MovieSceneActorReferenceSection.cpp
--- a/Engine/Source/Runtime/MovieSceneTracks/Private/Sections/MovieSceneActorReferenceSection.cpp
+++ b/Engine/Source/Runtime/MovieSceneTracks/Private/Sections/MovieSceneActorReferenceSection.cpp
@@ -72,8 +72,39 @@ void UMovieSceneActorReferenceSection::SetKeyTime( FKeyHandle KeyHandle, float T
 
 void UMovieSceneActorReferenceSection::AddKey( float Time, const FGuid& Value, EMovieSceneKeyInterpolation KeyInterpolation )
 {
-	int32 NewActorGuidIndex = ActorGuids.Add(Value);
-	ActorGuidIndexCurve.AddKey(Time, NewActorGuidIndex);
+	AddKey(Time, Value);
+}
+
+
+FKeyHandle UMovieSceneActorReferenceSection::AddKey(float Time, const FGuid& Value)
+{
+	if (IsTimeWithinSection(Time))
+	{
+		if (TryModify())
+		{
+			FKeyHandle ExistingKeyHandle = ActorGuidIndexCurve.FindKey(Time);
+			if (ActorGuidIndexCurve.IsKeyHandleValid(ExistingKeyHandle))
+			{
+				int32 ActorGuidIndex = ActorGuidIndexCurve.GetKeyValue(ExistingKeyHandle);
+				if (ActorGuidIndex >= 0 && ActorGuidIndex < ActorGuids.Num())
+				{
+					ActorGuids[ActorGuidIndex] = Value;
+					return ExistingKeyHandle;
+				}
+				else
+				{
+					int32 NewActorGuidIndex = ActorGuids.Add(Value);
+					return ActorGuidIndexCurve.UpdateOrAddKey(Time, NewActorGuidIndex);
+				}
+			}
+			else
+			{
+				int32 NewActorGuidIndex = ActorGuids.Add(Value);
+				return ActorGuidIndexCurve.AddKey(Time, NewActorGuidIndex);
+			}
+		}
+	}
+	return FKeyHandle();
 }
 
 
@@ -96,6 +127,12 @@ void UMovieSceneActorReferenceSection::SetDefault( const FGuid& Value )
 }
 
 
+void UMovieSceneActorReferenceSection::ClearDefaults()
+{
+	ActorGuidIndexCurve.ClearDefaultValue();
+}
+
+
 void UMovieSceneActorReferenceSection::PreSave(const class ITargetPlatform* TargetPlatform)
 {
 	ActorGuidStrings.Empty();
diff --git a/Engine/Source/Runtime/MovieSceneTracks/Private/Sections/MovieSceneBoolSection.cpp b/Engine/Source/Runtime/MovieSceneTracks/Private/Sections/MovieSceneBoolSection.cpp
--- a/Engine/Source/Runtime/MovieSceneTracks/Private/Sections/MovieSceneBoolSection.cpp
+++ b/Engine/Source/Runtime/MovieSceneTracks/Private/Sections/MovieSceneBoolSection.cpp
@@ -20,9 +20,9 @@ void UMovieSceneBoolSection::PostLoad()
 	Super::PostLoad();
 }
 
-bool UMovieSceneBoolSection::Eval( float Position ) const
+bool UMovieSceneBoolSection::Eval( float Position, bool DefaultValue ) const
 {
-	return !!BoolCurve.Evaluate(Position);
+	return !!BoolCurve.Evaluate(Position, DefaultValue);
 }
 
 
@@ -99,10 +99,15 @@ void UMovieSceneBoolSection::SetDefault( const bool& Value )
 
 bool UMovieSceneBoolSection::NewKeyIsNewData( float Time, const bool& Value ) const
 {
-	return Eval(Time) != Value;
+	return Eval(Time, Value) != Value;
 }
 
 bool UMovieSceneBoolSection::HasKeys( const bool& Value ) const
 {
 	return BoolCurve.GetNumKeys() != 0;
 }
+
+void UMovieSceneBoolSection::ClearDefaults()
+{
+	BoolCurve.ClearDefaultValue();
+}
diff --git a/Engine/Source/Runtime/MovieSceneTracks/Private/Sections/MovieSceneByteSection.cpp b/Engine/Source/Runtime/MovieSceneTracks/Private/Sections/MovieSceneByteSection.cpp
--- a/Engine/Source/Runtime/MovieSceneTracks/Private/Sections/MovieSceneByteSection.cpp
+++ b/Engine/Source/Runtime/MovieSceneTracks/Private/Sections/MovieSceneByteSection.cpp
@@ -9,9 +9,9 @@ UMovieSceneByteSection::UMovieSceneByteSection( const FObjectInitializer& Object
 { }
 
 
-uint8 UMovieSceneByteSection::Eval( float Position ) const
+uint8 UMovieSceneByteSection::Eval( float Position, uint8 DefaultValue) const
 {
-	return ByteCurve.Evaluate(Position);
+	return ByteCurve.Evaluate(Position, DefaultValue);
 }
 
 
@@ -79,7 +79,7 @@ void UMovieSceneByteSection::AddKey( float Time, const uint8& Value, EMovieScene
 
 bool UMovieSceneByteSection::NewKeyIsNewData( float Time, const uint8& Value ) const
 {
-	return Eval( Time ) != Value;
+	return Eval( Time, Value ) != Value;
 }
 
 
@@ -96,3 +96,9 @@ void UMovieSceneByteSection::SetDefault( const uint8& Value )
 		ByteCurve.SetDefaultValue(Value);
 	}
 }
+
+
+void UMovieSceneByteSection::ClearDefaults()
+{
+	ByteCurve.ClearDefaultValue();
+}
diff --git a/Engine/Source/Runtime/MovieSceneTracks/Private/Sections/MovieSceneColorSection.cpp b/Engine/Source/Runtime/MovieSceneTracks/Private/Sections/MovieSceneColorSection.cpp
--- a/Engine/Source/Runtime/MovieSceneTracks/Private/Sections/MovieSceneColorSection.cpp
+++ b/Engine/Source/Runtime/MovieSceneTracks/Private/Sections/MovieSceneColorSection.cpp
@@ -252,3 +252,12 @@ void UMovieSceneColorSection::SetDefault(const FColorKey& Key)
 	FRichCurve* ChannelCurve = GetCurveForChannel(Key.Channel, &RedCurve, &GreenCurve, &BlueCurve, &AlphaCurve);
 	SetCurveDefault(*ChannelCurve, Key.ChannelValue);
 }
+
+
+void UMovieSceneColorSection::ClearDefaults()
+{
+	RedCurve.ClearDefaultValue();
+	GreenCurve.ClearDefaultValue();
+	BlueCurve.ClearDefaultValue();
+	AlphaCurve.ClearDefaultValue();
+}
diff --git a/Engine/Source/Runtime/MovieSceneTracks/Private/Sections/MovieSceneEventSection.cpp b/Engine/Source/Runtime/MovieSceneTracks/Private/Sections/MovieSceneEventSection.cpp
--- a/Engine/Source/Runtime/MovieSceneTracks/Private/Sections/MovieSceneEventSection.cpp
+++ b/Engine/Source/Runtime/MovieSceneTracks/Private/Sections/MovieSceneEventSection.cpp
@@ -16,7 +16,7 @@ UMovieSceneEventSection::UMovieSceneEventSection()
 /* UMovieSceneSection overrides
  *****************************************************************************/
 
-void UMovieSceneEventSection::AddKey(float Time, const FName& EventName, FKeyParams KeyParams)
+void UMovieSceneEventSection::AddKey(float Time, const FName& EventName)
 {
 	if (TryModify())
 	{
diff --git a/Engine/Source/Runtime/MovieSceneTracks/Private/Sections/MovieSceneFloatSection.cpp b/Engine/Source/Runtime/MovieSceneTracks/Private/Sections/MovieSceneFloatSection.cpp
--- a/Engine/Source/Runtime/MovieSceneTracks/Private/Sections/MovieSceneFloatSection.cpp
+++ b/Engine/Source/Runtime/MovieSceneTracks/Private/Sections/MovieSceneFloatSection.cpp
@@ -9,9 +9,9 @@ UMovieSceneFloatSection::UMovieSceneFloatSection( const FObjectInitializer& Obje
 { }
 
 
-float UMovieSceneFloatSection::Eval( float Position ) const
+float UMovieSceneFloatSection::Eval( float Position, float DefaultValue ) const
 {
-	return FloatCurve.Eval( Position );
+	return FloatCurve.Eval( Position, DefaultValue );
 }
 
 
@@ -89,3 +89,9 @@ void UMovieSceneFloatSection::SetDefault( const float& Value )
 {
 	SetCurveDefault( FloatCurve, Value );
 }
+
+
+void UMovieSceneFloatSection::ClearDefaults()
+{
+	FloatCurve.ClearDefaultValue();
+}
diff --git a/Engine/Source/Runtime/MovieSceneTracks/Private/Sections/MovieSceneStringSection.cpp b/Engine/Source/Runtime/MovieSceneTracks/Private/Sections/MovieSceneStringSection.cpp
--- a/Engine/Source/Runtime/MovieSceneTracks/Private/Sections/MovieSceneStringSection.cpp
+++ b/Engine/Source/Runtime/MovieSceneTracks/Private/Sections/MovieSceneStringSection.cpp
@@ -43,6 +43,11 @@ void UMovieSceneStringSection::SetDefault(const FString& Value)
 	StringCurve.SetDefaultValue(Value);
 }
 
+void UMovieSceneStringSection::ClearDefaults()
+{
+	StringCurve.ClearDefaultValue();
+}
+
 
 /* UMovieSceneSection interface
  *****************************************************************************/
diff --git a/Engine/Source/Runtime/MovieSceneTracks/Private/Sections/MovieSceneVectorSection.cpp b/Engine/Source/Runtime/MovieSceneTracks/Private/Sections/MovieSceneVectorSection.cpp
--- a/Engine/Source/Runtime/MovieSceneTracks/Private/Sections/MovieSceneVectorSection.cpp
+++ b/Engine/Source/Runtime/MovieSceneTracks/Private/Sections/MovieSceneVectorSection.cpp
@@ -222,3 +222,12 @@ void UMovieSceneVectorSection::SetDefault(const FVectorKey& Key)
 	FRichCurve* ChannelCurve = GetCurveForChannel(Key.Channel, Curves, ChannelsUsed);
 	ChannelCurve->SetDefaultValue(Key.Value);
 }
+
+
+void UMovieSceneVectorSection::ClearDefaults()
+{
+	for (auto Curve : Curves)
+	{
+		Curve.ClearDefaultValue();
+	}
+}
diff --git a/Engine/Source/Runtime/MovieSceneTracks/Private/TrackInstances/MovieScene3DAttachTrackInstance.cpp b/Engine/Source/Runtime/MovieSceneTracks/Private/TrackInstances/MovieScene3DAttachTrackInstance.cpp
--- a/Engine/Source/Runtime/MovieSceneTracks/Private/TrackInstances/MovieScene3DAttachTrackInstance.cpp
+++ b/Engine/Source/Runtime/MovieSceneTracks/Private/TrackInstances/MovieScene3DAttachTrackInstance.cpp
@@ -13,23 +13,37 @@ FMovieScene3DAttachTrackInstance::FMovieScene3DAttachTrackInstance( UMovieScene3
 { }
 
 
-void FMovieScene3DAttachTrackInstance::UpdateConstraint( float Position, const TArray<TWeakObjectPtr<UObject>>& RuntimeObjects, AActor* Actor, UMovieScene3DConstraintSection* ConstraintSection ) 
+void FMovieScene3DAttachTrackInstance::UpdateConstraint(EMovieSceneUpdatePass UpdatePass, float Position, const TArray<TWeakObjectPtr<UObject>>& RuntimeObjects, AActor* Actor, UMovieScene3DConstraintSection* ConstraintSection )
 {
-	FVector Translation;
-	FRotator Rotation;
-
-	UMovieScene3DAttachSection* AttachSection = CastChecked<UMovieScene3DAttachSection>(ConstraintSection);
-
-	for( int32 ObjIndex = 0; ObjIndex < RuntimeObjects.Num(); ++ObjIndex )
+	if (UpdatePass == MSUP_PreUpdate)
+	{
+		for (int32 ObjIndex = 0; ObjIndex < RuntimeObjects.Num(); ++ObjIndex)
+		{
+			USceneComponent* SceneComponent = MovieSceneHelpers::SceneComponentFromRuntimeObject(RuntimeObjects[ObjIndex].Get());
+			if (SceneComponent != nullptr)
+			{
+				SceneComponent->ResetRelativeTransform();
+			}
+		}
+	}
+	else if (UpdatePass == MSUP_PostUpdate)
 	{
-		USceneComponent* SceneComponent = MovieSceneHelpers::SceneComponentFromRuntimeObject(RuntimeObjects[ObjIndex].Get());
+		FVector Translation;
+		FRotator Rotation;
 
-		if (SceneComponent != nullptr)
+		UMovieScene3DAttachSection* AttachSection = CastChecked<UMovieScene3DAttachSection>(ConstraintSection);
+
+		for (int32 ObjIndex = 0; ObjIndex < RuntimeObjects.Num(); ++ObjIndex)
 		{
-			AttachSection->Eval(SceneComponent, Position, Actor, Translation, Rotation);
+			USceneComponent* SceneComponent = MovieSceneHelpers::SceneComponentFromRuntimeObject(RuntimeObjects[ObjIndex].Get());
+
+			if (SceneComponent != nullptr)
+			{
+				AttachSection->Eval(SceneComponent, Position, Actor, Translation, Rotation);
 
-			// Set the relative translation and rotation.  Note they are set once instead of individually to avoid a redundant component transform update.
-			SceneComponent->SetRelativeLocationAndRotation(Translation, Rotation);
+				// Set the relative translation and rotation.  Note they are set once instead of individually to avoid a redundant component transform update.
+				SceneComponent->SetRelativeLocationAndRotation(Translation, Rotation);
+			}
 		}
 	}
 }
diff --git a/Engine/Source/Runtime/MovieSceneTracks/Private/TrackInstances/MovieScene3DConstraintTrackInstance.cpp b/Engine/Source/Runtime/MovieSceneTracks/Private/TrackInstances/MovieScene3DConstraintTrackInstance.cpp
--- a/Engine/Source/Runtime/MovieSceneTracks/Private/TrackInstances/MovieScene3DConstraintTrackInstance.cpp
+++ b/Engine/Source/Runtime/MovieSceneTracks/Private/TrackInstances/MovieScene3DConstraintTrackInstance.cpp
@@ -89,7 +89,7 @@ void FMovieScene3DConstraintTrackInstance::Update(EMovieSceneUpdateData& UpdateD
 					AActor* Actor = Cast<AActor>(ConstraintObjects[ConstraintObjectIndex].Get());
 					if (Actor)
 					{
-						UpdateConstraint(UpdateData.Position, RuntimeObjects, Actor, ConstraintSection);
+						UpdateConstraint(UpdateData.UpdatePass, UpdateData.Position, RuntimeObjects, Actor, ConstraintSection);
 					}	
 				}
 			}
diff --git a/Engine/Source/Runtime/MovieSceneTracks/Private/TrackInstances/MovieScene3DPathTrackInstance.cpp b/Engine/Source/Runtime/MovieSceneTracks/Private/TrackInstances/MovieScene3DPathTrackInstance.cpp
--- a/Engine/Source/Runtime/MovieSceneTracks/Private/TrackInstances/MovieScene3DPathTrackInstance.cpp
+++ b/Engine/Source/Runtime/MovieSceneTracks/Private/TrackInstances/MovieScene3DPathTrackInstance.cpp
@@ -14,28 +14,42 @@ FMovieScene3DPathTrackInstance::FMovieScene3DPathTrackInstance( UMovieScene3DPat
 { }
 
 
-void FMovieScene3DPathTrackInstance::UpdateConstraint( float Position, const TArray<TWeakObjectPtr<UObject>>& RuntimeObjects, AActor* Actor, UMovieScene3DConstraintSection* ConstraintSection ) 
+void FMovieScene3DPathTrackInstance::UpdateConstraint(EMovieSceneUpdatePass UpdatePass, float Position, const TArray<TWeakObjectPtr<UObject>>& RuntimeObjects, AActor* Actor, UMovieScene3DConstraintSection* ConstraintSection )
 {
-	FVector Translation;
-	FRotator Rotation;
-
-	TArray<USplineComponent*> SplineComponents;
-	Actor->GetComponents(SplineComponents);
-
-	if (SplineComponents.Num())
+	if (UpdatePass == MSUP_PreUpdate)
 	{
-		UMovieScene3DPathSection* PathSection = CastChecked<UMovieScene3DPathSection>(ConstraintSection);
-
-		for( int32 ObjIndex = 0; ObjIndex < RuntimeObjects.Num(); ++ObjIndex )
+		for (int32 ObjIndex = 0; ObjIndex < RuntimeObjects.Num(); ++ObjIndex)
 		{
 			USceneComponent* SceneComponent = MovieSceneHelpers::SceneComponentFromRuntimeObject(RuntimeObjects[ObjIndex].Get());
-
 			if (SceneComponent != nullptr)
 			{
-				PathSection->Eval(SceneComponent, Position, SplineComponents[0], Translation, Rotation);
+				SceneComponent->ResetRelativeTransform();
+			}
+		}
+	}
+	else if (UpdatePass == MSUP_PostUpdate)
+	{
+		FVector Translation;
+		FRotator Rotation;
+
+		TArray<USplineComponent*> SplineComponents;
+		Actor->GetComponents(SplineComponents);
+
+		if (SplineComponents.Num())
+		{
+			UMovieScene3DPathSection* PathSection = CastChecked<UMovieScene3DPathSection>(ConstraintSection);
+
+			for (int32 ObjIndex = 0; ObjIndex < RuntimeObjects.Num(); ++ObjIndex)
+			{
+				USceneComponent* SceneComponent = MovieSceneHelpers::SceneComponentFromRuntimeObject(RuntimeObjects[ObjIndex].Get());
+
+				if (SceneComponent != nullptr)
+				{
+					PathSection->Eval(SceneComponent, Position, SplineComponents[0], Translation, Rotation);
 
-				// Set the relative translation and rotation.  Note they are set once instead of individually to avoid a redundant component transform update.
-				SceneComponent->SetRelativeLocationAndRotation(Translation, Rotation);
+					// Set the relative translation and rotation.  Note they are set once instead of individually to avoid a redundant component transform update.
+					SceneComponent->SetRelativeLocationAndRotation(Translation, Rotation);
+				}
 			}
 		}
 	}
diff --git a/Engine/Source/Runtime/MovieSceneTracks/Private/TrackInstances/MovieScene3DTransformTrackInstance.cpp b/Engine/Source/Runtime/MovieSceneTracks/Private/TrackInstances/MovieScene3DTransformTrackInstance.cpp
--- a/Engine/Source/Runtime/MovieSceneTracks/Private/TrackInstances/MovieScene3DTransformTrackInstance.cpp
+++ b/Engine/Source/Runtime/MovieSceneTracks/Private/TrackInstances/MovieScene3DTransformTrackInstance.cpp
@@ -90,33 +90,28 @@ void FMovieScene3DTransformTrackInstance::Update(EMovieSceneUpdateData& UpdateDa
 	{
 		UpdateRuntimeMobility(RuntimeObjects);
 	}
-
-	FVector Translation;
-	FRotator Rotation;
-	FVector Scale;
-
-	if( TransformTrack->Eval( UpdateData.Position, UpdateData.LastPosition, Translation, Rotation, Scale ) )
+	else // UpdateData.UpdatePass == MSUP_Update
 	{
-		for( int32 ObjIndex = 0; ObjIndex < RuntimeObjects.Num(); ++ObjIndex )
+		for (TWeakObjectPtr<UObject> ObjectPtr : RuntimeObjects)
 		{
-			UObject* Object = RuntimeObjects[ObjIndex].Get();
-			USceneComponent* SceneComponent = MovieSceneHelpers::SceneComponentFromRuntimeObject(Object);
-
-			if (SceneComponent != nullptr)
+			if (ObjectPtr.IsValid())
 			{
-				if (UpdateData.UpdatePass == MSUP_PreUpdate)
-				{		
-					// Set the transforms to identity explicitly instead of ResetRelativeTransform so that overlaps aren't evaluated at the origin
-					SceneComponent->RelativeLocation = FVector(FVector::ZeroVector);
-					SceneComponent->RelativeRotation = FRotator(FRotator::ZeroRotator);
-					SceneComponent->RelativeScale3D = FVector(1.f);
-					SceneComponent->UpdateComponentToWorld();
-				}
-				else if (UpdateData.UpdatePass == MSUP_Update)
+				USceneComponent* SceneComponent = MovieSceneHelpers::SceneComponentFromRuntimeObject(ObjectPtr.Get());
+				if (SceneComponent != nullptr)
 				{
-					SceneComponent->AddRelativeLocation(Translation);
-					SceneComponent->AddRelativeRotation(Rotation);
-					SceneComponent->SetRelativeScale3D(Scale);
+					FVector Translation = SceneComponent->GetRelativeTransform().GetTranslation();
+					FRotator Rotation = FRotator(SceneComponent->GetRelativeTransform().GetRotation());
+					FVector Scale = SceneComponent->GetRelativeTransform().GetScale3D();
+
+					if (TransformTrack->Eval(UpdateData.Position, UpdateData.LastPosition, Translation, Rotation, Scale))
+					{
+						SceneComponent->SetRelativeLocationAndRotation(Translation, Rotation);
+						SceneComponent->SetRelativeScale3D(Scale);
+
+						// Force the location and rotation values to avoid Rot->Quat->Rot conversions
+						SceneComponent->RelativeLocation = Translation;
+						SceneComponent->RelativeRotation = Rotation;
+					}
 				}
 			}
 		}
diff --git a/Engine/Source/Runtime/MovieSceneTracks/Private/TrackInstances/MovieSceneActorReferenceTrackInstance.cpp b/Engine/Source/Runtime/MovieSceneTracks/Private/TrackInstances/MovieSceneActorReferenceTrackInstance.cpp
--- a/Engine/Source/Runtime/MovieSceneTracks/Private/TrackInstances/MovieSceneActorReferenceTrackInstance.cpp
+++ b/Engine/Source/Runtime/MovieSceneTracks/Private/TrackInstances/MovieSceneActorReferenceTrackInstance.cpp
@@ -55,7 +55,7 @@ void FMovieSceneActorReferenceTrackInstance::RestoreState(const TArray<TWeakObje
 
 void FMovieSceneActorReferenceTrackInstance::Update(EMovieSceneUpdateData& UpdateData, const TArray<TWeakObjectPtr<UObject>>& RuntimeObjects, class IMovieScenePlayer& Player, FMovieSceneSequenceInstance& SequenceInstance ) 
 {
-	
+	// TODO: Support default values.
 	FGuid ActorReferenceGuid;
 	if( ActorReferenceTrack->Eval( UpdateData.Position, UpdateData.LastPosition, ActorReferenceGuid ) )
 	{
diff --git a/Engine/Source/Runtime/MovieSceneTracks/Private/TrackInstances/MovieSceneAudioTrackInstance.cpp b/Engine/Source/Runtime/MovieSceneTracks/Private/TrackInstances/MovieSceneAudioTrackInstance.cpp
--- a/Engine/Source/Runtime/MovieSceneTracks/Private/TrackInstances/MovieSceneAudioTrackInstance.cpp
+++ b/Engine/Source/Runtime/MovieSceneTracks/Private/TrackInstances/MovieSceneAudioTrackInstance.cpp
@@ -232,7 +232,13 @@ TWeakObjectPtr<UAudioComponent> FMovieSceneAudioTrackInstance::GetAudioComponent
 	{
 		USoundCue* TempPlaybackAudioCue = NewObject<USoundCue>();
 		UWorld* World = Actor ? Actor->GetWorld() : (Player.GetPlaybackContext() != nullptr) ? Player.GetPlaybackContext()->GetWorld() : nullptr;
+
 		UAudioComponent* AudioComponent = FAudioDevice::CreateComponent(TempPlaybackAudioCue, World, Actor, false, false);
+		if (AudioComponent && Actor)
+		{
+			AudioComponent->SetFlags(RF_Transient);
+			AudioComponent->AttachToComponent(Actor->GetRootComponent(), FAttachmentTransformRules::KeepRelativeTransform);
+		}
 
 		PlaybackAudioComponents[RowIndex].Add(Actor);
 		PlaybackAudioComponents[RowIndex][Actor] = TWeakObjectPtr<UAudioComponent>(AudioComponent);
diff --git a/Engine/Source/Runtime/MovieSceneTracks/Private/TrackInstances/MovieSceneBoolTrackInstance.cpp b/Engine/Source/Runtime/MovieSceneTracks/Private/TrackInstances/MovieSceneBoolTrackInstance.cpp
--- a/Engine/Source/Runtime/MovieSceneTracks/Private/TrackInstances/MovieSceneBoolTrackInstance.cpp
+++ b/Engine/Source/Runtime/MovieSceneTracks/Private/TrackInstances/MovieSceneBoolTrackInstance.cpp
@@ -57,14 +57,16 @@ void FMovieSceneBoolTrackInstance::RestoreState(const TArray<TWeakObjectPtr<UObj
 
 void FMovieSceneBoolTrackInstance::Update(EMovieSceneUpdateData& UpdateData, const TArray<TWeakObjectPtr<UObject>>& RuntimeObjects, class IMovieScenePlayer& Player, FMovieSceneSequenceInstance& SequenceInstance) 
 {
-	bool BoolValue = false;
-
-	if( BoolTrack->Eval( UpdateData.Position, UpdateData.LastPosition, BoolValue ) )
+	for (auto ObjectPtr : RuntimeObjects)
 	{
-		for (auto ObjectPtr : RuntimeObjects)
+		UObject* Object = ObjectPtr.Get();
+		if (Object != nullptr)
 		{
-			UObject* Object = ObjectPtr.Get();
-			PropertyBindings->CallFunction<bool>( Object, &BoolValue );
+			bool BoolValue = PropertyBindings->GetCurrentValue<bool>(Object);
+			if (BoolTrack->Eval(UpdateData.Position, UpdateData.LastPosition, BoolValue))
+			{
+				PropertyBindings->CallFunction<bool>(Object, &BoolValue);
+			}
 		}
 	}
 }
