UE - Disable real time gpu stats

diff --git a/Engine/Source/Runtime/Engine/Private/SceneUtils.cpp b/Engine/Source/Runtime/Engine/Private/SceneUtils.cpp
--- a/Engine/Source/Runtime/Engine/Private/SceneUtils.cpp
+++ b/Engine/Source/Runtime/Engine/Private/SceneUtils.cpp
@@ -8,7 +8,7 @@
 
 static TAutoConsoleVariable<int> CVarGPUStatsEnabled(
 	TEXT("r.GPUStatsEnabled"),
-	1,
+	0,
 	TEXT("Enables or disables GPU stat recording"));
 
 // If this is enabled, the child stat timings will be included in their parents' times.
