Copying //UE4/Dev-Framework to //UE4/Dev-Main (Source: //UE4/Dev-Framework @ 3136612)

diff --git a/Engine/Source/Runtime/Engine/Classes/Animation/SmartName.h b/Engine/Source/Runtime/Engine/Classes/Animation/SmartName.h
--- a/Engine/Source/Runtime/Engine/Classes/Animation/SmartName.h
+++ b/Engine/Source/Runtime/Engine/Classes/Animation/SmartName.h
@@ -2,10 +2,62 @@
 
 #pragma once
 
+#include "ReferenceSkeleton.h"
+#include "BoneContainer.h"
 #include "SmartName.generated.h"
 
 struct FSmartName;
 
+/** in the future if we need more bools, please convert to bitfield 
+ * These are not saved in asset but per skeleton. 
+ */
+USTRUCT()
+struct ENGINE_API FAnimCurveType
+{
+	GENERATED_USTRUCT_BODY()
+
+	bool bMaterial;
+	bool bMorphtarget;
+
+	FAnimCurveType(bool bInMorphtarget = false, bool bInMaterial = false)
+		: bMaterial(bInMaterial)
+		, bMorphtarget(bInMorphtarget)
+	{
+	}
+};
+
+/** Curve Meta Data for each name
+ * Unfortunately this should be linked to FName, but no GUID because we don't have GUID in run-time
+ * We only add this if anything changed, by default, it is attribute curve
+ */
+USTRUCT()
+struct FCurveMetaData
+{
+	GENERATED_USTRUCT_BODY()
+
+	struct FAnimCurveType Type;
+	TArray<struct FBoneReference> LinkedBones;
+
+	friend FArchive& operator<<(FArchive& Ar, FCurveMetaData& B)
+	{
+		Ar << B.Type.bMaterial;
+		Ar << B.Type.bMorphtarget;
+		Ar << B.LinkedBones;
+
+		return Ar;
+	}
+
+	bool Serialize(FArchive& Ar)
+	{
+		Ar << *this;
+		return true;
+	}
+
+	FCurveMetaData()
+	{
+	}
+};
+
 USTRUCT()
 struct ENGINE_API FSmartNameMapping
 {
@@ -13,61 +65,53 @@ struct ENGINE_API FSmartNameMapping
 
 	FSmartNameMapping();
 
-	// ID type, should be used to access names as fundamental type may change.
-	// NOTE: If changing this, also change the type of NextUid below. UHT ignores,
-	// this typedef and needs to know the fundamental type to parse properties.
-	// If this changes and UHT handles typedefs, set NextUid type to UID.
-	typedef uint16 UID;
-	// Max UID used for overflow checking
-	static const UID MaxUID = MAX_uint16;
-
 	// Add a name to the mapping, if it exists, get it.
 	// @param Name - The name to add/get
 	// @param OUT OutUid - The UID of the newly created or retrieved name
 	// @param OUT OutGuid - The GUID of the newly created or retrieved name
 	// @return bool - true if the name was added, false if it existed (OutUid will be correctly set still in this case)
-	bool AddOrFindName(FName Name, UID& OutUid, FGuid& OutGuid);
+	bool AddOrFindName(FName Name, SmartName::UID_Type& OutUid, FGuid& OutGuid);
 	// Add a name to the mapping if it doesn't exist
 	// @param Name - The name to add/get
 	// @param OUT OutUid - The UID of the newly created or retrieved name
 	// @param InGuid - the guid of the name
 	// @return bool - true if the name was added, false if it existed (OutUid will be correctly set still in this case)
-	bool AddName(FName Name, UID& OutUid, const FGuid& InGuid);
+	bool AddName(FName Name, SmartName::UID_Type& OutUid, const FGuid& InGuid);
 
 	// Get a name from the mapping
-	// @param Uid - UID of the name to retrieve
+	// @param Uid - SmartName::UID_Type of the name to retrieve
 	// @param OUT OutName - Retrieved name
 	// @return bool - true if name existed and OutName is valid
-	bool GetName(const UID& Uid, FName& OutName) const;
+	bool GetName(const SmartName::UID_Type& Uid, FName& OutName) const;
 	bool GetNameByGuid(const FGuid& Guid, FName& OutName) const;
 	// Fill an array with all used UIDs
 	// @param Array - Array to fill
-	void FillUidArray(TArray<UID>& Array) const;
+	void FillUidArray(TArray<SmartName::UID_Type>& Array) const;
 
 	// Fill an array with all used names
 	// @param Array - Array to fill
 	void FillNameArray(TArray<FName>& Array) const;
 
 	// Change a name
-	// @param Uid - UID of the name to change
+	// @param Uid - SmartName::UID_Type of the name to change
 	// @param NewName - New name to set 
 	// @return bool - true if the name was found and changed, false if the name wasn't present in the mapping
-	bool Rename(const UID& Uid, FName NewName);
+	bool Rename(const SmartName::UID_Type& Uid, FName NewName);
 
 	// Remove a name from the mapping
-	// @param Uid - UID of the name to remove
+	// @param Uid - SmartName::UID_Type of the name to remove
 	// @return bool - true if the name was found and removed, false if the name wasn't present in the mapping
-	bool Remove(const UID& Uid);
+	bool Remove(const SmartName::UID_Type& Uid);
 
-	// Return UID * if it finds it
+	// Return SmartName::UID_Type * if it finds it
 	//  @param NewName - New name to set 
-	// @return UID pointer - null if it doesn't find. pointer if it finds. 
-	const FSmartNameMapping::UID* FindUID(const FName& Name) const;
+	// @return SmartName::UID_Type pointer - null if it doesn't find. pointer if it finds. 
+	const SmartName::UID_Type* FindUID(const FName& Name) const;
 
 	// Check whether a name already exists in the mapping
-	// @param Uid - the UID to check
+	// @param Uid - the SmartName::UID_Type to check
 	// @return bool - whether the name was found
-	bool Exists(const UID& Uid) const;
+	bool Exists(const SmartName::UID_Type& Uid) const;
 
 	// Check whether a name already exists in the mapping
 	// @param Name - the name to check
@@ -83,10 +127,23 @@ struct ENGINE_API FSmartNameMapping
 	bool AddSmartName(FSmartName& OutName);
 #else
 	// in cooked build, you don't have GUID, so register without GUID
-	bool FindOrAddSmartName(FName Name, UID& OutUid);
+	bool FindOrAddSmartName(FName Name, SmartName::UID_Type& OutUid);
 #endif // WITH_EDITOR
 	bool FindSmartName(FName Name, FSmartName& OutName) const;
-	bool FindSmartNameByUID(FSmartNameMapping::UID UID, FSmartName& OutName) const;
+	bool FindSmartNameByUID(SmartName::UID_Type UID, FSmartName& OutName) const;
+
+	// Curve Meta Data Accessors
+	FCurveMetaData* GetCurveMetaData(FName CurveName)
+	{
+		checkSlow(Exists(CurveName));
+		return &CurveMetaDataMap.FindOrAdd(CurveName);
+	}
+
+	const FCurveMetaData* GetCurveMetaData(FName CurveName) const
+	{
+		checkSlow(Exists(CurveName));
+		return CurveMetaDataMap.Find(CurveName);
+	}
 
 	// Serialize this to the provided archive; required for TMap serialization
 	void Serialize(FArchive& Ar);
@@ -94,9 +151,10 @@ struct ENGINE_API FSmartNameMapping
 	friend FArchive& operator<<(FArchive& Ar, FSmartNameMapping& Elem);
 
 private:
-	UID NextUid;				// The next UID to use 
-	TMap<UID, FName> UidMap;	// Mapping of UIDs and names. This is transient and built upon serialize - this is run-time data, and searching from UID to FName is important
+	SmartName::UID_Type NextUid;				// The next SmartName::UID_Type to use 
+	TMap<SmartName::UID_Type, FName> UidMap;	// Mapping of UIDs and names. This is transient and built upon serialize - this is run-time data, and searching from SmartName::UID_Type to FName is important
 	TMap<FName, FGuid> GuidMap;	// Mapping of GUIDs and names. This is the one serializing
+	TMap<FName, FCurveMetaData> CurveMetaDataMap;
 };
 
 USTRUCT()
@@ -119,6 +177,7 @@ struct ENGINE_API FSmartNameContainer
 	friend class USkeleton;
 protected:
 	FSmartNameMapping* GetContainerInternal(const FName& ContainerName);
+	const FSmartNameMapping* GetContainerInternal(const FName& ContainerName) const;
 
 private:
 	TMap<FName, FSmartNameMapping> NameMappings;	// List of smartname mappings
@@ -134,8 +193,8 @@ struct ENGINE_API FSmartName
 	UPROPERTY()
 	FName DisplayName;
 
-	// UID - for faster access
-	FSmartNameMapping::UID	UID;
+	// SmartName::UID_Type - for faster access
+	SmartName::UID_Type	UID;
 
 #if WITH_EDITORONLY_DATA
 	UPROPERTY()
@@ -144,17 +203,17 @@ struct ENGINE_API FSmartName
 
 	FSmartName()
 		: DisplayName(NAME_None)
-		, UID(FSmartNameMapping::MaxUID)
+		, UID(SmartName::MaxUID)
 	{}
 
 #if WITH_EDITORONLY_DATA
-	FSmartName(const FName& InDisplayName, FSmartNameMapping::UID InUID, const FGuid& InGuid)
+	FSmartName(const FName& InDisplayName, SmartName::UID_Type InUID, const FGuid& InGuid)
 		: DisplayName(InDisplayName)
 		, UID(InUID)
 		, Guid(InGuid)
 	{}
 #else
-	FSmartName(const FName& InDisplayName, FSmartNameMapping::UID InUID)
+	FSmartName(const FName& InDisplayName, SmartName::UID_Type InUID)
 		: DisplayName(InDisplayName)
 		, UID(InUID)
 	{}
@@ -190,9 +249,9 @@ struct ENGINE_API FSmartName
 	bool IsValid() const
 	{
 #if WITH_EDITORONLY_DATA
-		return UID != FSmartNameMapping::MaxUID && Guid.IsValid(); 
+		return UID != SmartName::MaxUID && Guid.IsValid(); 
 #else
-		return UID != FSmartNameMapping::MaxUID;
+		return UID != SmartName::MaxUID;
 #endif  // WITH_EDITORONLY_DATA
 	}
 };
diff --git a/Engine/Source/Runtime/Engine/Classes/Components/ActorComponent.h b/Engine/Source/Runtime/Engine/Classes/Components/ActorComponent.h
--- a/Engine/Source/Runtime/Engine/Classes/Components/ActorComponent.h
+++ b/Engine/Source/Runtime/Engine/Classes/Components/ActorComponent.h
@@ -189,9 +189,13 @@ public:
 	uint32 bWantsInitializeComponent:1;
 
 	/** If true, we call the virtual BeginPlay */
-	UPROPERTY()
+	DEPRECATED(4.14, "bWantsBeginPlay was inconsistently enforced and is now unused")
 	uint32 bWantsBeginPlay:1;
 
+	/** If true, the component will be excluded from non-editor builds */
+	UPROPERTY(EditAnywhere, BlueprintReadOnly, Category = Cooking)
+	uint32 bIsEditorOnly:1;
+
 private:
 	/** Indicates that OnCreatedComponent has been called, but OnDestroyedComponent has not yet */
 	uint32 bHasBeenCreated:1;
@@ -267,6 +271,16 @@ public:
 	UFUNCTION(BlueprintCallable, Category="Components")
 	bool ComponentHasTag(FName Tag) const;
 
+	/**
+	* Called during saving to determine the load flags to save with the object.
+	*
+	* @return	true if this object should always be loaded for editor game
+	*/
+	virtual bool NeedsLoadForEditorGame() const override
+	{
+		return !IsEditorOnly() && Super::NeedsLoadForEditorGame();
+	}
+
 	//~ Begin Trigger/Activation Interface
 
 	UPROPERTY(BlueprintAssignable, Category = "Components|Activation")
@@ -348,7 +362,12 @@ public:
 	virtual bool GetComponentClassCanReplicate() const;
 
 	/** Returns whether this component is an editor-only object or not */
-	virtual bool IsEditorOnly() const { return false; }
+	virtual bool IsEditorOnly() const override { return bIsEditorOnly; }
+
+	virtual void MarkAsEditorOnlySubobject() override
+	{
+		bIsEditorOnly = true;
+	}
 
 	/** Returns net role of the owning actor */
 	/** Returns true if we are replicating and not authorative */
@@ -804,6 +823,11 @@ private:
 	virtual void Tick( float DeltaTime ) final { check(0); }
 
 #endif
+
+public:
+
+	/** Prefix used to identify template component instances */
+	static const FString ComponentTemplateNameSuffix;
 };
 
 //////////////////////////////////////////////////////////////////////////
diff --git a/Engine/Source/Runtime/Engine/Classes/Components/AudioComponent.h b/Engine/Source/Runtime/Engine/Classes/Components/AudioComponent.h
--- a/Engine/Source/Runtime/Engine/Classes/Components/AudioComponent.h
+++ b/Engine/Source/Runtime/Engine/Classes/Components/AudioComponent.h
@@ -9,11 +9,22 @@
 
 /** called when we finish playing audio, either because it played to completion or because a Stop() call turned it off early */
 DECLARE_DYNAMIC_MULTICAST_DELEGATE( FOnAudioFinished );
+
 /** shadow delegate declaration for above */
 DECLARE_MULTICAST_DELEGATE_OneParam( FOnAudioFinishedNative, class UAudioComponent* );
+
 /** Called when subtitles are sent to the SubtitleManager.  Set this delegate if you want to hijack the subtitles for other purposes */
 DECLARE_DYNAMIC_DELEGATE_TwoParams( FOnQueueSubtitles, const TArray<struct FSubtitleCue>&, Subtitles, float, CueDuration );
 
+/** Called as a sound plays on the audio component to allow BP to perform actions based on playback percentage.
+* Computed as samples played divided by total samples, taking into account pitch.
+* Not currently implemented on all platforms.
+*/
+DECLARE_DYNAMIC_MULTICAST_DELEGATE_TwoParams(FOnAudioPlaybackPercent, const class USoundWave*, PlayingSoundWave, const float, PlaybackPercent);
+
+/** shadow delegate declaration for above */
+DECLARE_MULTICAST_DELEGATE_ThreeParams(FOnAudioPlaybackPercentNative, const class UAudioComponent*, const class USoundWave*, const float);
+
 /**
  *	Struct used for storing one per-instance named parameter for this AudioComponent.
  *	Certain nodes in the SoundCue may reference parameters by name so they can be adjusted per-instance.
@@ -86,33 +97,68 @@ class ENGINE_API UAudioComponent : public USceneComponent
 
 	/** Auto destroy this component on completion */
 	UPROPERTY()
-	uint32 bAutoDestroy:1;
+	uint8 bAutoDestroy:1;
 
 	/** Stop sound when owner is destroyed */
 	UPROPERTY()
-	uint32 bStopWhenOwnerDestroyed:1;
+	uint8 bStopWhenOwnerDestroyed:1;
 
 	/** Whether the wave instances should remain active if they're dropped by the prioritization code. Useful for e.g. vehicle sounds that shouldn't cut out. */
 	UPROPERTY()
-	uint32 bShouldRemainActiveIfDropped:1;
+	uint8 bShouldRemainActiveIfDropped:1;
 
 	/** Is this audio component allowed to be spatialized? */
 	UPROPERTY(EditAnywhere, BlueprintReadWrite, Category=Attenuation)
-	uint32 bAllowSpatialization:1;
+	uint8 bAllowSpatialization:1;
 
 	/** Should the Attenuation Settings asset be used (false) or should the properties set directly on the component be used for attenuation properties */
 	UPROPERTY(EditAnywhere, BlueprintReadWrite, Category=Attenuation)
-	uint32 bOverrideAttenuation:1;
+	uint8 bOverrideAttenuation:1;
 
 	/** Whether or not this sound plays when the game is paused in the UI */
 	UPROPERTY(EditAnywhere, BlueprintReadWrite, Category = Sound)
-	uint32 bIsUISound : 1;
+	uint8 bIsUISound : 1;
 
 	/** Whether or not to apply a low-pass filter to the sound that plays in this audio component. */
 	UPROPERTY(EditAnywhere, BlueprintReadWrite, Category = Filter)
-	uint32 bEnableLowPassFilter : 1;
+	uint8 bEnableLowPassFilter : 1;
+
 	UPROPERTY(EditAnywhere, BlueprintReadWrite, Category = Sound)
-	uint32 bOverridePriority:1;
+	uint8 bOverridePriority:1;
+
+	/** If true, subtitles in the sound data will be ignored. */
+	UPROPERTY(EditAnywhere, BlueprintReadWrite, Category = Sound)
+	uint8 bSuppressSubtitles:1;
+
+	/** Whether this audio component is previewing a sound */
+	uint8 bPreviewComponent:1;
+
+	/** If true, this sound will not be stopped when flushing the audio device. */
+	uint8 bIgnoreForFlushing:1;
+
+	/** Whether audio effects are applied */
+	uint8 bEQFilterApplied:1;
+
+	/** Whether to artificially prioritize the component to play */
+	uint8 bAlwaysPlay:1;
+
+	/** Whether or not this audio component is a music clip */
+	uint8 bIsMusic:1;
+
+	/** Whether or not the audio component should be excluded from reverb EQ processing */
+	uint8 bReverb:1;
+
+	/** Whether or not this sound class forces sounds to the center channel */
+	uint8 bCenterChannelOnly:1;
+
+	/** Whether or not this sound is a preview sound */
+	uint8 bIsPreviewSound:1;
+
+	/** Whether or not this audio component has been paused */
+	uint8 bIsPaused:1;
+
+	/** The specific audio device to play this component on */
+	uint32 AudioDeviceHandle;
 
 	/** The lower bound to use when randomly determining a pitch multiplier */
 	UPROPERTY(EditAnywhere, BlueprintReadWrite, Category=Modulation)
@@ -138,6 +184,9 @@ class ENGINE_API UAudioComponent : public USceneComponent
 	UPROPERTY(EditAnywhere, BlueprintReadWrite, Category = Sound, meta = (ClampMin = "0.0", UIMin = "0.0", EditCondition = "bOverridePriority"))
 	float Priority;
 
+	/** Used by the subtitle manager to prioritize subtitles wave instances spawned by this component. */
+	float SubtitlePriority;
+
 	UPROPERTY()
 	float VolumeWeightedPriorityScale_DEPRECATED;
 
@@ -174,6 +223,13 @@ class ENGINE_API UAudioComponent : public USceneComponent
 	/** shadow delegate for non UObject subscribers */
 	FOnAudioFinishedNative OnAudioFinishedNative;
 
+	/** called when we finish playing audio, either because it played to completion or because a Stop() call turned it off early */
+	UPROPERTY(BlueprintAssignable)
+	FOnAudioPlaybackPercent OnAudioPlaybackPercent;
+
+	/** shadow delegate for non UObject subscribers */
+	FOnAudioPlaybackPercentNative OnAudioPlaybackPercentNative;
+
 	/** Called when subtitles are sent to the SubtitleManager.  Set this delegate if you want to hijack the subtitles for other purposes */
 	UPROPERTY()
 	FOnQueueSubtitles OnQueueSubtitles;
@@ -215,6 +271,10 @@ class ENGINE_API UAudioComponent : public USceneComponent
 	UFUNCTION(BlueprintCallable, Category="Audio|Components|Audio")
 	void Stop();
 
+	/** Pause an audio component playing its sound cue, issue any delegates if needed */
+	UFUNCTION(BlueprintCallable, Category = "Audio|Components|Audio")
+	void SetPaused(bool bPause);
+
 	/** @return true if this component is currently playing a SoundCue. */
 	UFUNCTION(BlueprintCallable, Category="Audio|Components|Audio")
 	bool IsPlaying() const;
@@ -266,38 +326,6 @@ public:
 	/** Sets the sound instance parameter. */
 	void SetSoundParameter(const FAudioComponentParam& Param);
 
-public:
-
-	/** If true, subtitles in the sound data will be ignored. */
-	uint32 bSuppressSubtitles:1;
-
-	/** Whether this audio component is previewing a sound */
-	uint32 bPreviewComponent:1;
-
-	/** If true, this sound will not be stopped when flushing the audio device. */
-	uint32 bIgnoreForFlushing:1;
-
-	/** Whether audio effects are applied */
-	uint32 bEQFilterApplied:1;
-
-	/** Whether to artificially prioritize the component to play */
-	uint32 bAlwaysPlay:1;
-
-	/** Whether or not this audio component is a music clip */
-	uint32 bIsMusic:1;
-
-	/** Whether or not the audio component should be excluded from reverb EQ processing */
-	uint32 bReverb:1;
-
-	/** Whether or not this sound class forces sounds to the center channel */
-	uint32 bCenterChannelOnly:1;
-
-	/** Whether or not this sound is a preview sound */
-	uint32 bIsPreviewSound:1;
-
-	/** Used by the subtitle manager to prioritize subtitles wave instances spawned by this component. */
-	float SubtitlePriority;
-
 	//~ Begin UObject Interface.
 #if WITH_EDITOR
 	virtual void PostEditChangeProperty(FPropertyChangedEvent& PropertyChangedEvent) override;
diff --git a/Engine/Source/Runtime/Engine/Classes/Components/BrushComponent.h b/Engine/Source/Runtime/Engine/Classes/Components/BrushComponent.h
--- a/Engine/Source/Runtime/Engine/Classes/Components/BrushComponent.h
+++ b/Engine/Source/Runtime/Engine/Classes/Components/BrushComponent.h
@@ -44,6 +44,7 @@ public:
 	virtual class UBodySetup* GetBodySetup() override { return BrushBodySetup; };
 	virtual void GetUsedMaterials( TArray<UMaterialInterface*>& OutMaterials ) const override;
 	virtual uint8 GetStaticDepthPriorityGroup() const override;
+	virtual bool IsEditorOnly() const override;
 #if WITH_EDITOR
 	virtual bool ComponentIsTouchingSelectionBox(const FBox& InSelBBox, const FEngineShowFlags& ShowFlags, const bool bConsiderOnlyBSP, const bool bMustEncompassEntireComponent) const override;
 	virtual bool ComponentIsTouchingSelectionFrustum(const FConvexVolume& InFrustum, const FEngineShowFlags& ShowFlags, const bool bConsiderOnlyBSP, const bool bMustEncompassEntireComponent) const override;
diff --git a/Engine/Source/Runtime/Engine/Classes/Components/ChildActorComponent.h b/Engine/Source/Runtime/Engine/Classes/Components/ChildActorComponent.h
--- a/Engine/Source/Runtime/Engine/Classes/Components/ChildActorComponent.h
+++ b/Engine/Source/Runtime/Engine/Classes/Components/ChildActorComponent.h
@@ -54,6 +54,10 @@ private:
 	UPROPERTY(Replicated, BlueprintReadOnly, Category=ChildActorComponent, TextExportTransient, NonPIEDuplicateTransient, meta=(AllowPrivateAccess="true"))
 	AActor*	ChildActor;
 
+	/** Property to point to the template child actor for details panel purposes */
+	UPROPERTY(VisibleDefaultsOnly, DuplicateTransient, Category=ChildActorComponent, meta=(ShowInnerProperties))
+	AActor* ChildActorTemplate;
+
 	/** We try to keep the child actor's name as best we can, so we store it off here when destroying */
 	FName ChildActorName;
 
@@ -65,9 +69,11 @@ public:
 	//~ Begin Object Interface.
 #if WITH_EDITOR
 	virtual void PostEditChangeProperty(FPropertyChangedEvent& PropertyChangedEvent) override;
+	virtual void PostEditChangeChainProperty(FPropertyChangedChainEvent& PropertyChangedEvent) override;
 	virtual void PostEditUndo() override;
 	virtual void PostLoad() override;
 #endif
+	virtual void Serialize(FArchive& Ar) override;
 	virtual void BeginDestroy() override;
 	virtual void GetLifetimeReplicatedProps(TArray<FLifetimeProperty>& OutLifetimeProps) const override;
 	virtual void PostRepNotifies() override;
@@ -90,6 +96,7 @@ public:
 	void CreateChildActor();
 
 	AActor* GetChildActor() const { return ChildActor; }
+	AActor* GetChildActorTemplate() const { return ChildActorTemplate; }
 
 	FName GetChildActorName() const { return ChildActorName; }
 
diff --git a/Engine/Source/Runtime/Engine/Classes/Components/PrimitiveComponent.h b/Engine/Source/Runtime/Engine/Classes/Components/PrimitiveComponent.h
--- a/Engine/Source/Runtime/Engine/Classes/Components/PrimitiveComponent.h
+++ b/Engine/Source/Runtime/Engine/Classes/Components/PrimitiveComponent.h
@@ -701,6 +701,7 @@ public:
 	 *	@note For collisions during physics simulation to generate hit events, 'Simulation Generates Hit Events' must be enabled for this component.
 	 *	@note When receiving a hit from another object's movement, the directions of 'Hit.Normal' and 'Hit.ImpactNormal'
 	 *	will be adjusted to indicate force from the other object against this object.
+	 *	@note NormalImpulse will be filled in for physics-simulating bodies, but will be zero for swept-component blocking collisions.
 	 */
 	UPROPERTY(BlueprintAssignable, Category="Collision")
 	FComponentHitSignature OnComponentHit;
@@ -1543,12 +1544,6 @@ public:
 	virtual void PostDuplicate(bool bDuplicateForPIE) override;
 	virtual SIZE_T GetResourceSize(EResourceSizeMode::Type Mode) override;
 
-	virtual void MarkAsEditorOnlySubobject() override
-	{
-		AlwaysLoadOnClient = false;
-		AlwaysLoadOnServer = false;
-	}
-
 #if WITH_EDITOR
 	/**
 	 * Called after importing property values for this object (paste, duplicate or .t3d import)
@@ -1682,6 +1677,14 @@ public:
 	 *	@param NewRot	New orienatation for the body
 	 */
 	virtual void SetAllPhysicsRotation(FRotator NewRot);
+
+	/**
+	 *	Set the rotation of all bodies in this component.
+	 *	If a SkeletalMeshComponent, the root body will be changed to the desired orientation, and the same delta is applied to all other bodies.
+	 *
+	 *	@param NewRot	New orienatation for the body
+	 */
+	virtual void SetAllPhysicsRotation(const FQuat& NewRot);
 	
 	/**
 	 *	Ensure simulation is running for all bodies in this component.
@@ -1772,7 +1775,7 @@ public:
 	 * @param       NewResponse  What the new response should be to the supplied Channel
 	 */
 	UFUNCTION(BlueprintCallable, Category="Collision")
-	void SetCollisionResponseToChannel(ECollisionChannel Channel, ECollisionResponse NewResponse);
+	virtual void SetCollisionResponseToChannel(ECollisionChannel Channel, ECollisionResponse NewResponse);
 	
 	/**
 	 *	Changes all ResponseToChannels container for this PrimitiveComponent. to be NewResponse
@@ -1780,14 +1783,14 @@ public:
 	 * @param       NewResponse  What the new response should be to the supplied Channel
 	 */
 	UFUNCTION(BlueprintCallable, Category="Collision")
-	void SetCollisionResponseToAllChannels(ECollisionResponse NewResponse);
+	virtual void SetCollisionResponseToAllChannels(ECollisionResponse NewResponse);
 	
 	/**
 	 *	Changes the whole ResponseToChannels container for this PrimitiveComponent.
 	 *
 	 * @param       NewResponses  New set of responses for this component
 	 */
-	void SetCollisionResponseToChannels(const FCollisionResponseContainer& NewReponses);
+	virtual void SetCollisionResponseToChannels(const FCollisionResponseContainer& NewReponses);
 	
 protected:
 
diff --git a/Engine/Source/Runtime/Engine/Classes/Components/SkeletalMeshComponent.h b/Engine/Source/Runtime/Engine/Classes/Components/SkeletalMeshComponent.h
--- a/Engine/Source/Runtime/Engine/Classes/Components/SkeletalMeshComponent.h
+++ b/Engine/Source/Runtime/Engine/Classes/Components/SkeletalMeshComponent.h
@@ -372,12 +372,18 @@ public:
 
 	/** The active animation graph program instance. */
 	UPROPERTY(transient)
-	class UAnimInstance* AnimScriptInstance;
+	UAnimInstance* AnimScriptInstance;
 
 	/** Any running sub anim instances that need to be updates on the game thread */
 	UPROPERTY(transient)
 	TArray<UAnimInstance*> SubInstances;
 
+	/** An instance created from the PostPhysicsBlueprint property of the skeletal mesh we're using,
+	 *  Runs after physics has been blended
+	 */
+	UPROPERTY(transient)
+	UAnimInstance* PostProcessAnimInstance;
+
 	UPROPERTY(EditAnywhere, BlueprintReadWrite, Category=Animation, meta=(ShowOnlyInnerProperties))
 	struct FSingleAnimationPlayData AnimationData;
 
@@ -586,6 +592,10 @@ public:
 	UPROPERTY()
 	uint32 bEnableLineCheckWithBounds:1;
 
+	/** Cache AnimCurveUidVersion from Skeleton and this will be used to identify if it needs to be updated */
+	UPROPERTY(transient)
+	uint16 CachedAnimCurveUidVersion;
+	
 	/** If bEnableLineCheckWithBounds is true, scale the bounds by this value before doing line check. */
 	UPROPERTY()
 	FVector LineCheckBoundsScale;
@@ -609,6 +619,13 @@ public:
 	UFUNCTION(BlueprintCallable, Category="Components|SkeletalMesh", meta=(Keywords = "AnimBlueprint", UnsafeDuringActorConstruction = "true"))
 	class UAnimInstance * GetAnimInstance() const;
 
+	/**
+	 * Returns the active post process instance is one is available. This is set on the mesh that this
+	 * component is using, and is evaluated immediately after the main instance.
+	 */
+	UFUNCTION(BlueprintCallable, Category = "Components|SkeletalMesh", meta = (Keywords = "AnimBlueprint"))
+	UAnimInstance* GetPostProcessInstance() const;
+
 	/** Below are the interface to control animation when animation mode, not blueprint mode **/
 	UFUNCTION(BlueprintCallable, Category = "Components|Animation", meta = (Keywords = "Animation"))
 	void SetAnimationMode(EAnimationMode::Type InAnimationMode);
@@ -982,6 +999,12 @@ public:
 	 */
 	void RecalcRequiredBones(int32 LODIndex);
 
+	/**
+	* Recalculates the AnimCurveUids array in RequiredBone of this SkeletalMeshComponent based on current required bone set
+	* Is called when Skeleton->IsRequiredCurvesUpToDate() = false
+	*/
+	void RecalcRequiredCurves();
+
 public:
 	//~ Begin UObject Interface.
 	virtual void Serialize(FArchive& Ar) override;
@@ -1060,6 +1083,7 @@ public:
 	virtual void SetAllPhysicsAngularVelocity(FVector const& NewVel, bool bAddToCurrent = false) override;
 	virtual void SetAllPhysicsPosition(FVector NewPos) override;
 	virtual void SetAllPhysicsRotation(FRotator NewRot) override;
+	virtual void SetAllPhysicsRotation(const FQuat& NewRot) override;
 	virtual void WakeAllRigidBodies() override;
 	virtual void PutAllRigidBodiesToSleep() override;
 	virtual bool IsAnyRigidBodyAwake() override;
@@ -1220,7 +1244,13 @@ public:
 	* @param	OutCurves				Blended Curve
 	*/
 	void PerformAnimationEvaluation(const USkeletalMesh* InSkeletalMesh, UAnimInstance* InAnimInstance, TArray<FTransform>& OutSpaceBases, TArray<FTransform>& OutBoneSpaceTransforms, FVector& OutRootBoneTranslation, FBlendedHeapCurve& OutCurve) const;
-	void PostAnimEvaluation( FAnimationEvaluationContext& EvaluationContext );
+
+	/**
+	 * Evaluates the post process instance from the skeletal mesh this component is using.
+	 */
+	void EvaluatePostProcessMeshInstance(TArray<FTransform>& OutBoneSpaceTransforms, FBlendedHeapCurve& OutCurve, const USkeletalMesh* InSkeletalMesh, FVector& OutRootBoneTranslation) const;
+
+	void PostAnimEvaluation(FAnimationEvaluationContext& EvaluationContext);
 
 	/**
 	 * Blend of Physics Bones with PhysicsWeight and Animated Bones with (1-PhysicsWeight)
@@ -1515,7 +1545,8 @@ public:
 
 protected:
 	bool NeedToSpawnAnimScriptInstance() const;
-	
+	bool NeedToSpawnPostPhysicsInstance() const;
+
 private:
 
 	FSkeletalMeshComponentEndPhysicsTickFunction EndPhysicsTickFunction;
@@ -1573,13 +1604,13 @@ public:
 
 	friend class FSkeletalMeshComponentDetails;
 
-	/** Returns array containing cachec animation curve mapping UIDs (which are copied over from USkeleton) */
-	TArray<FSmartNameMapping::UID> const * GetCachedAnimCurveMappingNameUids();
-
 	/** Apply animation curves to this component */
 	void ApplyAnimationCurvesToComponent(const TMap<FName, float>* InMaterialParameterCurves, const TMap<FName, float>* InAnimationMorphCurves);
 	
 private:
+	/** Override USkinnedMeshComponent */
+	virtual void AddSlavePoseComponent(USkinnedMeshComponent* SkinnedMeshComponent) override;
+
 	// Returns whether we need to run the Pre Cloth Tick or not
 	bool ShouldRunEndPhysicsTick() const;
 
@@ -1643,11 +1674,6 @@ private:
 	UPROPERTY()
 	float DefaultPlayRate_DEPRECATED;
 	
-	/** Caches the anim curve mapping smart name UIDs, by copying cached data from the Skeleton */
-	void UpdateCachedAnimCurveMappingNameUids();
-
-	/** Cached animation curves smart name mapping UIDs, only at runtime, not serialized. (used for FBlendedCurve::InitFrom) */
-	TArray<FSmartNameMapping::UID> CachedAnimCurveMappingNameUids;
 
 	/*
 	 * Update MorphTargetCurves - these are not animation curves, but SetMorphTarget and similar functions that can set to this mesh component
@@ -1695,4 +1721,15 @@ private:
 
 	/** Multicaster fired when this component creates physics state (in case external objects rely on physics state)*/
 	FOnSkelMeshPhysicsCreatedMultiCast OnSkelMeshPhysicsCreated;
+
+	/** Mark current anim UID version to up-to-date. Called when it's recalcualted */
+	void MarkRequiredCurveUpToDate();
+	
+	/* This will check if the required curves are up-to-date by checking version number with skeleton. 
+	 * Skeleton's curve list changes whenever newer is added or deleted. 
+	 * This still has to happen in editor as well as in game as 
+	 * There is no guarantee of Skeleton having all curves as we've seen over and over again. 
+	 * Cooking does not guarantee skeleton containing all names
+	 */
+	bool AreRequiredCurvesUpToDate() const;
 };
diff --git a/Engine/Source/Runtime/Engine/Classes/Components/SkinnedMeshComponent.h b/Engine/Source/Runtime/Engine/Classes/Components/SkinnedMeshComponent.h
--- a/Engine/Source/Runtime/Engine/Classes/Components/SkinnedMeshComponent.h
+++ b/Engine/Source/Runtime/Engine/Classes/Components/SkinnedMeshComponent.h
@@ -186,6 +186,14 @@ public:
 	/** update Recalculate Normal flag in matching section */
 	void UpdateRecomputeTangent(int32 MaterialIndex, int32 LodIndex, bool bRecomputeTangentValue);
 
+	/** 
+	 * Get CPU skinned vertices for the specified LOD level. Includes morph targets if they are enabled.
+	 * Note: This function is very SLOW as it needs to flush the render thread.
+	 * @param	OutVertices		The skinned vertices
+	 * @param	InLODIndex		The LOD we want to export
+	 */
+	void GetCPUSkinnedVertices(TArray<struct FFinalSkinVertex>& OutVertices, int32 InLODIndex);
+
 	/** 
 	 * When true, we will just using the bounds from our MasterPoseComponent.  This is useful for when we have a Mesh Parented
 	 * to the main SkelMesh (e.g. outline mesh or a full body overdraw effect that is toggled) that is always going to be the same
@@ -362,6 +370,10 @@ public:
 	UPROPERTY(EditAnywhere, AdvancedDisplay, BlueprintReadOnly, Category=Lighting, meta=(EditCondition="CastShadow", DisplayName = "Capsule Indirect Shadow"))
 	uint32 bCastCapsuleIndirectShadow:1;
 
+	/** CPU skinning rendering - only for previewing in Persona and conversion tools */
+	UPROPERTY(transient)
+	uint32 bCPUSkinning : 1;
+
 	/** 
 	 * Override the Physics Asset of the mesh. It uses SkeletalMesh.PhysicsAsset, but if you'd like to override use this function
 	 * 
@@ -729,6 +741,11 @@ public:
 	UFUNCTION(BlueprintCallable, Category="Components|SkinnedMesh")
 	void SetMasterPoseComponent(USkinnedMeshComponent* NewMasterBoneComponent);
 
+protected:
+	/** Add a slave component to the SlavePoseComponents array */
+	virtual void AddSlavePoseComponent(USkinnedMeshComponent* SkinnedMeshComponent);
+
+public:
 	/** 
 	 * Refresh Slave Components if exists
 	 * 
diff --git a/Engine/Source/Runtime/Engine/Classes/Components/TimelineComponent.h b/Engine/Source/Runtime/Engine/Classes/Components/TimelineComponent.h
--- a/Engine/Source/Runtime/Engine/Classes/Components/TimelineComponent.h
+++ b/Engine/Source/Runtime/Engine/Classes/Components/TimelineComponent.h
@@ -77,6 +77,10 @@ struct FTimelineVectorTrack
 	UPROPERTY()
 	FOnTimelineVector InterpFunc;
 
+	/** Name of track, usually set in Timeline Editor. Used by SetInterpVectorCurve function. */
+	UPROPERTY()
+	FName TrackName;
+
 	/** Name of property that we should update from this curve */
 	UPROPERTY()
 	FName VectorPropertyName;
@@ -111,6 +115,10 @@ struct FTimelineFloatTrack
 	UPROPERTY()
 	FOnTimelineFloat InterpFunc;
 
+	/** Name of track, usually set in Timeline Editor. Used by SetInterpFloatCurve function. */
+	UPROPERTY()
+	FName TrackName;
+
 	/** Name of property that we should update from this curve */
 	UPROPERTY()
 	FName FloatPropertyName;
@@ -146,6 +154,10 @@ struct FTimelineLinearColorTrack
 	UPROPERTY()
 	FOnTimelineLinearColor InterpFunc;
 
+	/** Name of track, usually set in Timeline Editor. Used by SetInterpLinearColorCurve function. */
+	UPROPERTY()
+	FName TrackName;
+
 	/** Name of property that we should update from this curve */
 	UPROPERTY()
 	FName LinearColorPropertyName;
@@ -306,6 +318,15 @@ public:
 	/** Set the length of the timeline */
 	ENGINE_API void SetTimelineLength(float NewLength);
 
+	/** Update a certain float track's curve */
+	ENGINE_API void SetFloatCurve(UCurveFloat* NewFloatCurve, FName FloatTrackName);
+
+	/** Update a certain vector track's curve */
+	ENGINE_API void SetVectorCurve(UCurveVector* NewVectorCurve, FName VectorTrackName);
+
+	/** Update a certain linear color track's curve */
+	ENGINE_API void SetLinearColorCurve(UCurveLinearColor* NewLinearColorCurve, FName LinearColorTrackName);
+
 	/** Optionally provide an object to automatically update properties on */
 	ENGINE_API void SetPropertySetObject(UObject* NewPropertySetObject);
 
@@ -322,19 +343,19 @@ public:
 	ENGINE_API void AddEvent(float Time, FOnTimelineEvent EventFunc);
 
 	/** Add a vector interpolation to the timeline */
-	ENGINE_API void AddInterpVector(UCurveVector* VectorCurve, FOnTimelineVector InterpFunc, FName PropertyName = NAME_None);
+	ENGINE_API void AddInterpVector(UCurveVector* VectorCurve, FOnTimelineVector InterpFunc, FName PropertyName = NAME_None, FName TrackName = NAME_None);
 
 	/** Add a vector interpolation to the timeline. Use a non-serializeable delegate. */
 	ENGINE_API void AddInterpVector(UCurveVector* VectorCurve, FOnTimelineVectorStatic InterpFunc);
 
 	/** Add a float interpolation to the timeline */
-	ENGINE_API void AddInterpFloat(UCurveFloat* FloatCurve, FOnTimelineFloat InterpFunc, FName PropertyName = NAME_None);
+	ENGINE_API void AddInterpFloat(UCurveFloat* FloatCurve, FOnTimelineFloat InterpFunc, FName PropertyName = NAME_None, FName TrackName = NAME_None);
 
 	/** Add a float interpolation to the timeline. Use a non-serializeable delegate. */
 	ENGINE_API void AddInterpFloat(UCurveFloat* FloatCurve, FOnTimelineFloatStatic InterpFunc );
 
 	/** Add a linear color interpolation to the timeline */
-	ENGINE_API void AddInterpLinearColor(UCurveLinearColor* LinearColorCurve, FOnTimelineLinearColor InterpFunc, FName PropertyName = NAME_None);
+	ENGINE_API void AddInterpLinearColor(UCurveLinearColor* LinearColorCurve, FOnTimelineLinearColor InterpFunc, FName PropertyName = NAME_None, FName TrackName = NAME_None);
 
 	/** Add a linear color interpolation to the timeline. Use a non-serializeable delegate. */
 	ENGINE_API void AddInterpLinearColor(UCurveLinearColor* LinearColorCurve, FOnTimelineLinearColorStatic InterpFunc);
@@ -453,6 +474,17 @@ public:
 	UFUNCTION(BlueprintCallable, Category = "Components|Timeline")
 	ENGINE_API bool GetIgnoreTimeDilation() const;
 
+	/** Update a certain float track's curve */
+	UFUNCTION(BlueprintCallable, Category = "Components|Timeline")
+	ENGINE_API void SetFloatCurve(UCurveFloat* NewFloatCurve, FName FloatTrackName);
+
+	/** Update a certain vector track's curve */
+	UFUNCTION(BlueprintCallable, Category = "Components|Timeline")
+	ENGINE_API void SetVectorCurve(UCurveVector* NewVectorCurve, FName VectorTrackName);
+
+	/** Update a certain linear color track's curve */
+	UFUNCTION(BlueprintCallable, Category = "Components|Timeline")
+	ENGINE_API void SetLinearColorCurve(UCurveLinearColor* NewLinearColorCurve, FName LinearColorTrackName);
 
 	UFUNCTION()
 	void OnRep_Timeline();
@@ -479,13 +511,13 @@ public:
 	ENGINE_API void AddEvent(float Time, FOnTimelineEvent EventFunc);
 	
 	/** Add a vector interpolation to the timeline */
-	ENGINE_API void AddInterpVector(UCurveVector* VectorCurve, FOnTimelineVector InterpFunc, FName PropertyName = NAME_None);
+	ENGINE_API void AddInterpVector(UCurveVector* VectorCurve, FOnTimelineVector InterpFunc, FName PropertyName = NAME_None, FName TrackName = NAME_None);
 	
 	/** Add a float interpolation to the timeline */
-	ENGINE_API void AddInterpFloat(UCurveFloat* FloatCurve, FOnTimelineFloat InterpFunc, FName PropertyName = NAME_None);
+	ENGINE_API void AddInterpFloat(UCurveFloat* FloatCurve, FOnTimelineFloat InterpFunc, FName PropertyName = NAME_None, FName TrackName = NAME_None);
 
 	/** Add a linear color interpolation to the timeline */
-	ENGINE_API void AddInterpLinearColor(UCurveLinearColor* LinearColorCurve, FOnTimelineLinearColor InterpFunc, FName PropertyName = NAME_None);
+	ENGINE_API void AddInterpLinearColor(UCurveLinearColor* LinearColorCurve, FOnTimelineLinearColor InterpFunc, FName PropertyName = NAME_None, FName TrackName = NAME_None);
 
 	/** Optionally provide an object to automatically update properties on */
 	ENGINE_API void SetPropertySetObject(UObject* NewPropertySetObject);
diff --git a/Engine/Source/Runtime/Engine/Classes/Engine/Engine.h b/Engine/Source/Runtime/Engine/Classes/Engine/Engine.h
--- a/Engine/Source/Runtime/Engine/Classes/Engine/Engine.h
+++ b/Engine/Source/Runtime/Engine/Classes/Engine/Engine.h
@@ -1347,13 +1347,6 @@ public:
 	UPROPERTY(globalconfig)
 	uint32 bSuppressMapWarnings:1;
 
-	/** if set, cook game classes into standalone packages (as defined in [Cooker.MPGameContentCookStandalone]) and load the appropriate
-	 * one at game time depending on the gametype specified on the URL
-	 * (the game class should then not be referenced in the maps themselves)
-	 */
-	UPROPERTY(globalconfig)
-	uint32 bCookSeparateSharedMPGameContent:1;
-
 	/** determines whether AI logging should be processed or not */
 	UPROPERTY(globalconfig)
 	uint32 bDisableAILogging:1;
@@ -2471,6 +2464,7 @@ public:
 #if WITH_SERVER_CODE
 	virtual bool HandleServerTravelCommand( const TCHAR* Cmd, FOutputDevice& Ar, UWorld* InWorld );
 
+	DEPRECATED(4.14, "Say Command moved to GameMode as an exec function")
 	virtual bool HandleSayCommand( const TCHAR* Cmd, FOutputDevice& Ar, UWorld* InWorld );
 #endif
 
diff --git a/Engine/Source/Runtime/Engine/Classes/Engine/EngineBaseTypes.h b/Engine/Source/Runtime/Engine/Classes/Engine/EngineBaseTypes.h
--- a/Engine/Source/Runtime/Engine/Classes/Engine/EngineBaseTypes.h
+++ b/Engine/Source/Runtime/Engine/Classes/Engine/EngineBaseTypes.h
@@ -10,6 +10,8 @@
 #include "TaskGraphInterfaces.h"
 #include "EngineBaseTypes.generated.h"
 
+struct FTickContext;
+
 //
 //	EInputEvent
 //
@@ -250,6 +252,13 @@ private:
 	  * relative to the element ahead of it in the cooling down list, remaining until the next time this function will tick 
 	  */
 	float RelativeTickCooldown;
+
+	/** 
+	* The last world game time at which we were ticked. Game time used is dependent on bTickEvenWhenPaused
+	* Valid only if we've been ticked at least once since having a tick interval; otherwise set to -1.f
+	*/
+	float LastTickGameTimeSeconds;
+
 public:
 
 	/** The frequency in seconds at which this tick function will be executed.  If less than or equal to 0 then it will tick every frame */
@@ -345,14 +354,17 @@ private:
 	 * Queues a tick function for execution from the game thread
 	 * @param TickContext - context to tick in
 	 */
-	void QueueTickFunction(class FTickTaskSequencer& TTS, const struct FTickContext& TickContext);
+	void QueueTickFunction(class FTickTaskSequencer& TTS, const FTickContext& TickContext);
 
 	/**
 	 * Queues a tick function for execution from the game thread
 	 * @param TickContext - context to tick in
 	 * @param StackForCycleDetection - Stack For Cycle Detection
 	 */
-	void QueueTickFunctionParallel(const struct FTickContext& TickContext, TArray<FTickFunction*, TInlineAllocator<8> >& StackForCycleDetection);
+	void QueueTickFunctionParallel(const FTickContext& TickContext, TArray<FTickFunction*, TInlineAllocator<8> >& StackForCycleDetection);
+
+	/** Returns the delta time to use when ticking this function given the TickContext */
+	float CalculateDeltaTime(const FTickContext& TickContext);
 
 	/** 
 	 * Logs the prerequisites
diff --git a/Engine/Source/Runtime/Engine/Classes/Engine/EngineTypes.h b/Engine/Source/Runtime/Engine/Classes/Engine/EngineTypes.h
--- a/Engine/Source/Runtime/Engine/Classes/Engine/EngineTypes.h
+++ b/Engine/Source/Runtime/Engine/Classes/Engine/EngineTypes.h
@@ -544,22 +544,42 @@ enum class ENetworkSmoothingMode : uint8
     If(QueryIgnoreMask & ShapeFilter != 0) filter out */
 typedef uint8 FMaskFilter;
 
+// NOTE!!Some of these values are used to index into FCollisionResponseContainers and must be kept in sync.
+// @see FCollisionResponseContainer::SetResponse().
+
+// @NOTE!!!! This DisplayName [DISPLAYNAME] SHOULD MATCH suffix of ECC_DISPLAYNAME
+// Otherwise it will mess up collision profile loading
+// If you change this, please also change FCollisionResponseContainers
+//
+// If you add any more TraceQuery="1", you also should change UCollsionProfile::LoadProfileConfig
+// Metadata doesn't work outside of editor, so you'll need to add manually
+
+// @NOTE : when you add more here for predefined engine channel
+// please change the max in the CollisionProfile
+// search ECC_Destructible
+
+// in order to use this custom channels
+// we recommend to define in your local file
+// - i.e. #define COLLISION_WEAPON		ECC_GameTraceChannel1
+// and make sure you customize these it in INI file by
+// 
+// in DefaultEngine.ini
+//
+// [/Script/Engine.CollisionProfile]
+// GameTraceChannel1="Weapon"
+// 
+// also in the INI file, you can override collision profiles that are defined by simply redefining
+// note that Weapon isn't defined in the BaseEngine.ini file, but "Trigger" is defined in Engine
+// +Profiles=(Name="Trigger",CollisionEnabled=QueryOnly,ObjectTypeName=WorldDynamic, DefaultResponse=ECR_Overlap, CustomResponses=((Channel=Visibility, Response=ECR_Ignore), (Channel=Weapon, Response=ECR_Ignore)))
+
+
 /** 
  * Enum indicating different type of objects for rigid-body collision purposes. 
- * NOTE!! Some of these values are used to index into FCollisionResponseContainers and must be kept in sync.
- * @see FCollisionResponseContainer::SetResponse().
  */
 UENUM(BlueprintType)
 enum ECollisionChannel
 {
-	/**
-	* @NOTE!!!! This DisplayName [DISPLAYNAME] SHOULD MATCH suffix of ECC_DISPLAYNAME
-	* Otherwise it will mess up collision profile loading
-	* If you change this, please also change FCollisionResponseContainers
-	* 
-	* If you add any more TraceQuery="1", you also should change UCollsionProfile::LoadProfileConfig
-	* Metadata doesn't work outside of editor, so you'll need to add manually 
-	*/
+
 	ECC_WorldStatic UMETA(DisplayName="WorldStatic"),
 	ECC_WorldDynamic UMETA(DisplayName="WorldDynamic"),
 	ECC_Pawn UMETA(DisplayName="Pawn"),
@@ -568,31 +588,16 @@ enum ECollisionChannel
 	ECC_PhysicsBody UMETA(DisplayName="PhysicsBody"),
 	ECC_Vehicle UMETA(DisplayName="Vehicle"),
 	ECC_Destructible UMETA(DisplayName="Destructible"),
-	// @NOTE : when you add more here for predefined engine channel
-	// please change the max in the CollisionProfile
-	// search ECC_Destructible
 
-	// Unspecified Engine Trace Channels
-	ECC_EngineTraceChannel1 UMETA(Hidden),		// IMPORTANT: This engine trace channel is reserved by the COLLISION_GIZMO definition
+	/** Reserved for gizmo collision */
+	ECC_EngineTraceChannel1 UMETA(Hidden),
+
 	ECC_EngineTraceChannel2 UMETA(Hidden),
 	ECC_EngineTraceChannel3 UMETA(Hidden),
 	ECC_EngineTraceChannel4 UMETA(Hidden), 
 	ECC_EngineTraceChannel5 UMETA(Hidden),
 	ECC_EngineTraceChannel6 UMETA(Hidden),
 
-	// in order to use this custom channels
-	// we recommend to define in your local file
-	// - i.e. #define COLLISION_WEAPON		ECC_GameTraceChannel1
-	// and make sure you customize these it in INI file by
-	// 
-	// in DefaultEngine.ini
-	//
-	// [/Script/Engine.CollisionProfile]
-	// GameTraceChannel1="Weapon"
-	// 
-	// also in the INI file, you can override collision profiles that are defined by simply redefining
-	// note that Weapon isn't defined in the BaseEngine.ini file, but "Trigger" is defined in Engine
-	// +Profiles=(Name="Trigger",CollisionEnabled=QueryOnly,ObjectTypeName=WorldDynamic, DefaultResponse=ECR_Overlap, CustomResponses=((Channel=Visibility, Response=ECR_Ignore), (Channel=Weapon, Response=ECR_Ignore)))
 	ECC_GameTraceChannel1 UMETA(Hidden),
 	ECC_GameTraceChannel2 UMETA(Hidden),
 	ECC_GameTraceChannel3 UMETA(Hidden),
@@ -616,13 +621,7 @@ enum ECollisionChannel
 	/** Add only nonserialized/transient flags below */
 
 	// NOTE!!!! THESE ARE BEING DEPRECATED BUT STILL THERE FOR BLUEPRINT. PLEASE DO NOT USE THEM IN CODE
-	/** 
-	 * This can be used to get all overlap event. If you trace with this channel, 
-	 * It will return everything except its own. Do not use this often as this is expensive operation
-	 */
-	 /**
-	  * can't add displaynames because then it will show up in the collision channel option
-	  */
+
 	ECC_OverlapAll_Deprecated UMETA(Hidden),
 	ECC_MAX,
 };
diff --git a/Engine/Source/Runtime/Engine/Classes/Engine/GameInstance.h b/Engine/Source/Runtime/Engine/Classes/Engine/GameInstance.h
--- a/Engine/Source/Runtime/Engine/Classes/Engine/GameInstance.h
+++ b/Engine/Source/Runtime/Engine/Classes/Engine/GameInstance.h
@@ -227,6 +227,15 @@ public:
 	virtual void HandleGameNetControlMessage(class UNetConnection* Connection, uint8 MessageByte, const FString& MessageStr)
 	{}
 	
+	/** Call to preload any content before loading a map URL, used during seamless travel as well as map loading */
+	virtual void PreloadContentForURL(FURL InURL);
+
+	/** Call to create the game mode for a given map URL */
+	virtual class AGameModeBase* CreateGameModeForURL(FURL InURL);
+
+	/** Return the game mode subclass to use for a given map, options, and portal. By default return passed in one */
+	virtual TSubclassOf<AGameModeBase> OverrideGameModeClass(TSubclassOf<AGameModeBase> GameModeClass, const FString& MapName, const FString& Options, const FString& Portal) const;
+
 	/** return true to delay an otherwise ready-to-join PendingNetGame performing LoadMap() and finishing up
 	 * useful to wait for content downloads, etc
 	 */
diff --git a/Engine/Source/Runtime/Engine/Classes/Engine/GameViewportClient.h b/Engine/Source/Runtime/Engine/Classes/Engine/GameViewportClient.h
--- a/Engine/Source/Runtime/Engine/Classes/Engine/GameViewportClient.h
+++ b/Engine/Source/Runtime/Engine/Classes/Engine/GameViewportClient.h
@@ -158,6 +158,7 @@ public:
 	virtual bool IsFocused(FViewport* Viewport) override;
 	virtual void Activated(FViewport* InViewport, const FWindowActivateEvent& InActivateEvent) override;
 	virtual void Deactivated(FViewport* InViewport, const FWindowActivateEvent& InActivateEvent) override;
+	virtual bool WindowCloseRequested() override;
 	virtual void CloseRequested(FViewport* Viewport) override;
 	virtual bool RequiresHitProxyStorage() override { return 0; }
 	virtual bool IsOrtho() const override;
@@ -488,6 +489,12 @@ public:
 		return CloseRequestedDelegate;
 	}
 
+	/** Accessor for the delegate called when the window owning the viewport is asked to close. */
+	FOnWindowCloseRequested& OnWindowCloseRequested()
+	{
+		return WindowCloseRequestedDelegate;
+	}
+
 	/** Accessor for the delegate called when the game viewport is created. */
 	static FSimpleMulticastDelegate& OnViewportCreated()
 	{
@@ -798,6 +805,9 @@ private:
 	/** Delegate called when a request to close the viewport is received */
 	FOnCloseRequested CloseRequestedDelegate;
 
+	/** Delegate called when the window owning the viewport is requested to close */
+	FOnWindowCloseRequested WindowCloseRequestedDelegate;
+
 	/** Delegate called when the game viewport is created. */
 	static FSimpleMulticastDelegate CreatedDelegate;
 
diff --git a/Engine/Source/Runtime/Engine/Classes/Engine/GameViewportDelegates.h b/Engine/Source/Runtime/Engine/Classes/Engine/GameViewportDelegates.h
--- a/Engine/Source/Runtime/Engine/Classes/Engine/GameViewportDelegates.h
+++ b/Engine/Source/Runtime/Engine/Classes/Engine/GameViewportDelegates.h
@@ -13,6 +13,20 @@
  */
 DECLARE_MULTICAST_DELEGATE_ThreeParams(FOnScreenshotCaptured, int32 /*Width*/, int32 /*Height*/, const TArray<FColor>& /*Colors*/);
 
+/**
+ * Delegate type used by UGameViewportClient when the top level window associated
+ * with the viewport has been requested to close.
+ * At this point, the viewport has not been closed and the operation may be canceled.
+ * This may not called from PIE, Editor Windows, on consoles, or before the game ends
+ * from other methods.
+ * This is only when the platform specific window is closed.
+ *
+ * Return indicates whether or not the window may be closed.
+ *
+ * @see UGameViewportClient.
+ */
+DECLARE_DELEGATE_RetVal(bool, FOnWindowCloseRequested);
+
 /**
  * Delegate type used by UGameViewportClient when call is made to close a viewport
  *
diff --git a/Engine/Source/Runtime/Engine/Classes/Engine/LocalPlayer.h b/Engine/Source/Runtime/Engine/Classes/Engine/LocalPlayer.h
--- a/Engine/Source/Runtime/Engine/Classes/Engine/LocalPlayer.h
+++ b/Engine/Source/Runtime/Engine/Classes/Engine/LocalPlayer.h
@@ -35,13 +35,13 @@ struct ENGINE_API FLocalPlayerContext
 		It tests against the APlayerController, APlayerState, and APawn. */
 	bool IsFromLocalPlayer(const AActor* ActorToTest) const;
 
-	/* Returns the world context. */
+	/** Returns the world context. */
 	UWorld* GetWorld() const;
 
-	/* Returns the local player. */
+	/** Returns the local player. */
 	class ULocalPlayer* GetLocalPlayer() const;
 
-	/* Returns the player controller. */
+	/** Returns the player controller. */
 	class APlayerController* GetPlayerController() const;
 
 	/** Templated version of GetPlayerController() */
@@ -58,8 +58,8 @@ struct ENGINE_API FLocalPlayerContext
 		}
 	}
 
-	/** Getter for the Game State */
-	class AGameState* GetGameState() const;
+	/** Getter for the Game State Base */
+	class AGameStateBase* GetGameState() const;
 
 	/** Templated Getter for the Game State */
 	template<class T>
diff --git a/Engine/Source/Runtime/Engine/Classes/Engine/MeshMerging.h b/Engine/Source/Runtime/Engine/Classes/Engine/MeshMerging.h
--- a/Engine/Source/Runtime/Engine/Classes/Engine/MeshMerging.h
+++ b/Engine/Source/Runtime/Engine/Classes/Engine/MeshMerging.h
@@ -336,6 +336,10 @@ struct FMeshMergingSettings
 	/** Whether or not vertex data such as vertex colours should be used when baking out materials */
 	UPROPERTY(EditAnywhere, Category = MaterialSettings, meta = (editcondition = "bMergeMaterials"))
 	bool bUseVertexDataForBakingMaterial;
+
+	// Whether or not to calculate varying output texture sizes according to their importance in the final atlas texture
+	UPROPERTY(Category = MaterialSettings, EditAnywhere)
+	bool bUseTextureBinning;
 			
 	UPROPERTY()
 	bool bCalculateCorrectLODModel_DEPRECATED;
@@ -382,6 +386,8 @@ struct FMeshMergingSettings
 		, bMergePhysicsData(false)
 		, bMergeMaterials(false)
 		, bBakeVertexDataToMesh(false)
+		, bUseVertexDataForBakingMaterial(true)
+		, bUseTextureBinning(false)
 		, bCalculateCorrectLODModel_DEPRECATED(false)
 		, LODSelectionType(EMeshLODSelectionType::CalculateLOD)
 		, ExportSpecificLOD_DEPRECATED(0)
diff --git a/Engine/Source/Runtime/Engine/Classes/Engine/SkeletalMesh.h b/Engine/Source/Runtime/Engine/Classes/Engine/SkeletalMesh.h
--- a/Engine/Source/Runtime/Engine/Classes/Engine/SkeletalMesh.h
+++ b/Engine/Source/Runtime/Engine/Classes/Engine/SkeletalMesh.h
@@ -19,6 +19,7 @@
 
 class UMorphTarget;
 class USkeleton;
+class UAnimInstance;
 
 #if WITH_APEX_CLOTHING
 struct FApexClothCollisionVolumeData;
@@ -122,42 +123,6 @@ struct FTriangleSortSettings
 
 };
 
-
-USTRUCT()
-struct FBoneReference
-{
-	GENERATED_USTRUCT_BODY()
-
-	/** Name of bone to control. This is the main bone chain to modify from. **/
-	UPROPERTY(EditAnywhere, Category = BoneReference)
-	FName BoneName;
-
-	/** Cached bone index for run time - right now bone index of skeleton **/
-	int32 BoneIndex;
-
-	FBoneReference()
-		: BoneIndex(INDEX_NONE)
-	{
-	}
-
-	bool operator==(const FBoneReference& Other) const
-	{
-		// faster to compare, and BoneName won't matter
-		return BoneIndex == Other.BoneIndex;
-	}
-	/** Initialize Bone Reference, return TRUE if success, otherwise, return false **/
-	ENGINE_API bool Initialize(const FBoneContainer& RequiredBones);
-
-	// @fixme laurent - only used by blendspace 'PerBoneBlend'. Fix this to support SkeletalMesh pose.
-	ENGINE_API bool Initialize(const USkeleton* Skeleton);
-
-	/** return true if valid. Otherwise return false **/
-	ENGINE_API bool IsValid(const FBoneContainer& RequiredBones) const;
-
-	FMeshPoseBoneIndex GetMeshPoseIndex() const { return FMeshPoseBoneIndex(BoneIndex); }
-	FCompactPoseBoneIndex GetCompactPoseIndex(const FBoneContainer& RequiredBones) const { return RequiredBones.MakeCompactPoseIndex(GetMeshPoseIndex()); }
-};
-
 /**
 * FSkeletalMeshOptimizationSettings - The settings used to optimize a skeletal mesh LOD.
 */
@@ -770,6 +735,13 @@ public:
 	UPROPERTY(EditAnywhere, editfixedsize, BlueprintReadOnly, Category=Clothing)
 	TArray<FClothingAssetData>		ClothingAssets;
 
+	/** Animation Blueprint class to run as a post process for this mesh.
+	 *  This blueprint will be ran before physics, but after the main
+	 *  anim instance for any skeletal mesh component using this mesh.
+	 */
+	UPROPERTY(EditAnywhere, BlueprintReadOnly, Category = SkeletalMesh)
+	TSubclassOf<UAnimInstance> PostProcessAnimBlueprint;
+
 protected:
 
 	/** Array of user data stored with the asset */
@@ -907,6 +879,13 @@ public:
 	 */
 	ENGINE_API TArray<USkeletalMeshSocket*>& GetMeshOnlySocketList();
 
+	/**
+	 * Const version
+	 * Returns the mesh only socket list - this ignores any sockets in the skeleton
+	 * Return value is a non-const reference so the socket list can be changed
+	 */
+	ENGINE_API const TArray<USkeletalMeshSocket*>& GetMeshOnlySocketList() const;
+
 	/**
 	* Returns the "active" socket list - all sockets from this mesh plus all non-duplicates from the skeleton
 	* Const ref return value as this cannot be modified externally
diff --git a/Engine/Source/Runtime/Engine/Classes/Engine/SkeletalMeshSocket.h b/Engine/Source/Runtime/Engine/Classes/Engine/SkeletalMeshSocket.h
--- a/Engine/Source/Runtime/Engine/Classes/Engine/SkeletalMeshSocket.h
+++ b/Engine/Source/Runtime/Engine/Classes/Engine/SkeletalMeshSocket.h
@@ -15,24 +15,24 @@ class USkeletalMeshSocket : public UObject
 	 *	everything explicitly to AttachComponent in the SkeletalMeshComponent.
 	 *	The Outer of a SkeletalMeshSocket should always be the USkeletalMesh.
 	 */
-	UPROPERTY(Category=SkeletalMeshSocket, VisibleAnywhere, BlueprintReadOnly)
+	UPROPERTY(Category="Socket Parameters", VisibleAnywhere, BlueprintReadOnly)
 	FName SocketName;
 
-	UPROPERTY(Category=SkeletalMeshSocket, VisibleAnywhere, BlueprintReadOnly)
+	UPROPERTY(Category="Socket Parameters", VisibleAnywhere, BlueprintReadOnly)
 	FName BoneName;
 
-	UPROPERTY(EditAnywhere, BlueprintReadOnly, Category=SkeletalMeshSocket)
+	UPROPERTY(EditAnywhere, BlueprintReadOnly, Category="Socket Parameters")
 	FVector RelativeLocation;
 
-	UPROPERTY(EditAnywhere, BlueprintReadOnly, Category=SkeletalMeshSocket)
+	UPROPERTY(EditAnywhere, BlueprintReadOnly, Category="Socket Parameters")
 	FRotator RelativeRotation;
 
-	UPROPERTY(EditAnywhere, BlueprintReadOnly, Category=SkeletalMeshSocket)
+	UPROPERTY(EditAnywhere, BlueprintReadOnly, Category="Socket Parameters")
 	FVector RelativeScale;
 
 	/** If true then the hierarchy of bones this socket is attached to will always be 
 	    evaluated, even if it had previously been removed due to the current lod setting */
-	UPROPERTY(EditAnywhere, BlueprintReadOnly, Category=SkeletalMeshSocket)
+	UPROPERTY(EditAnywhere, BlueprintReadOnly, Category="Socket Parameters")
 	bool bForceAlwaysAnimated;
 
 	UFUNCTION(BlueprintCallable, Category="Components|SkeletalMesh")
diff --git a/Engine/Source/Runtime/Engine/Classes/Engine/World.h b/Engine/Source/Runtime/Engine/Classes/Engine/World.h
--- a/Engine/Source/Runtime/Engine/Classes/Engine/World.h
+++ b/Engine/Source/Runtime/Engine/Classes/Engine/World.h
@@ -32,6 +32,8 @@ class AWorldSettings;
 class ACameraActor;
 class FUniqueNetId;
 class FWorldInGamePerformanceTrackers;
+class AGameModeBase;
+class AGameStateBase;
 struct FUniqueNetIdRepl;
 
 template<typename,typename> class TOctree;
@@ -557,10 +559,6 @@ class ENGINE_API UWorld : public UObject, public FNetworkNotify
 	UPROPERTY(Transient)
 	class ULineBatchComponent*					ForegroundLineBatcher;
 
-	/** The replicated actor which contains game state information that can be accessible to clients */
-	UPROPERTY(Transient)
-	class AGameState*							GameState;
-
 	/** Instance of this world's game-specific networking management */
 	UPROPERTY(Transient)
 	class AGameNetworkManager*					NetworkManager;
@@ -628,8 +626,13 @@ private:
 
 	/** The current GameMode, valid only on the server */
 	UPROPERTY(Transient)
-	class AGameMode*							AuthorityGameMode;
+	class AGameModeBase*						AuthorityGameMode;
+
+	/** The replicated actor which contains game state information that can be accessible to clients. Direct access is not allowed, use GetGameState<>() */
+	UPROPERTY(Transient)
+	class AGameStateBase*						GameState;
 
+	/** The AI System handles generating pathing information and AI behavior */
 	UPROPERTY(Transient)
 	class UAISystemBase*						AISystem;
 	
@@ -955,10 +958,13 @@ public:
 	/**  Time in seconds since level began play, but IS paused when the game is paused, and IS dilated/clamped. */
 	float TimeSeconds;
 
-	/** Time in seconds since level began play, but is NOT paused when the game is paused, and is NOT dilated/clamped. */
+	/**  Time in seconds since level began play, but IS NOT paused when the game is paused, and IS dilated/clamped. */
+	float UnpausedTimeSeconds;
+
+	/** Time in seconds since level began play, but IS NOT paused when the game is paused, and IS NOT dilated/clamped. */
 	float RealTimeSeconds;
 
-	/** Time in seconds since level began play, but IS paused when the game is paused, and is NOT dilated/clamped. */
+	/** Time in seconds since level began play, but IS paused when the game is paused, and IS NOT dilated/clamped. */
 	float AudioTimeSeconds;
 
 	/** Frame delta time in seconds adjusted by e.g. time dilation. */
@@ -1706,6 +1712,13 @@ public:
 	 */
 	float GetTimeSeconds() const;
 
+	/**
+	* Returns time in seconds since world was brought up for play, IS NOT stopped when game pauses, IS dilated/clamped
+	*
+	* @return time in seconds since world was brought up for play
+	*/
+	float GetUnpausedTimeSeconds() const;
+
 	/**
 	* Returns time in seconds since world was brought up for play, does NOT stop when game pauses, NOT dilated/clamped
 	*
@@ -2517,8 +2530,8 @@ public:
 	}
 
 	/** 
-	 *  Returns the current GameMode instance cast to the template type.
-	 *  This can only return a valid pointer on the server. Will always return null on a client
+	 * Returns the current Game Mode instance cast to the template type.
+	 * This can only return a valid pointer on the server and may be null if the cast fails. Will always return null on a client.
 	 */
 	template< class T >
 	T* GetAuthGameMode() const
@@ -2526,11 +2539,11 @@ public:
 		return Cast<T>(AuthorityGameMode);
 	}
 
-	/** 
-	 *  Returns the current GameMode instance.
-	 *  This can only return a valid pointer on the server. Will always return null on a client
+	/**
+	 * Returns the current Game Mode instance, which is always valid during gameplay on the server.
+	 * This will only return a valid pointer on the server. Will always return null on a client.
 	 */
-	AGameMode* GetAuthGameMode() const { return AuthorityGameMode; }
+	AGameModeBase* GetAuthGameMode() const { return AuthorityGameMode; }
 	
 	/** Returns the current GameState instance cast to the template type. */
 	template< class T >
@@ -2540,10 +2553,13 @@ public:
 	}
 
 	/** Returns the current GameState instance. */
-	AGameState* GetGameState() const { return GameState; }
+	AGameStateBase* GetGameState() const { return GameState; }
+
+	/** Sets the active GameState */
+	void SetGameState(AGameStateBase* NewGameState);
 
 	/** Copies GameState properties from the GameMode. */
-	void CopyGameState(AGameMode* FromGameMode, AGameState* FromGameState);
+	void CopyGameState(AGameModeBase* FromGameMode, AGameStateBase* FromGameState);
 
 
 	/** Spawns a Brush Actor in the World */
@@ -2792,8 +2808,8 @@ public:
 
 	/** seamlessly travels to the given URL by first loading the entry level in the background,
 	 * switching to it, and then loading the specified level. Does not disrupt network communication or disconnect clients.
-	 * You may need to implement GameMode::GetSeamlessTravelActorList(), PlayerController::GetSeamlessTravelActorList(),
-	 * GameMode::PostSeamlessTravel(), and/or GameMode::HandleSeamlessTravelPlayer() to handle preserving any information
+	 * You may need to implement GameModeBase::GetSeamlessTravelActorList(), PlayerController::GetSeamlessTravelActorList(),
+	 * GameModeBase::PostSeamlessTravel(), and/or GameModeBase::HandleSeamlessTravelPlayer() to handle preserving any information
 	 * that should be maintained (player teams, etc)
 	 * This codepath is designed for worlds that use little or no level streaming and GameModes where the game state
 	 * is reset/reloaded when transitioning. (like UT)
@@ -3020,6 +3036,11 @@ FORCEINLINE_DEBUGGABLE float UWorld::GetTimeSeconds() const
 	return TimeSeconds;
 }
 
+FORCEINLINE_DEBUGGABLE float UWorld::GetUnpausedTimeSeconds() const
+{
+	return UnpausedTimeSeconds;
+}
+
 FORCEINLINE_DEBUGGABLE float UWorld::GetRealTimeSeconds() const
 {
 	checkSlow(IsInGameThread());
diff --git a/Engine/Source/Runtime/Engine/Classes/GameFramework/Actor.h b/Engine/Source/Runtime/Engine/Classes/GameFramework/Actor.h
--- a/Engine/Source/Runtime/Engine/Classes/GameFramework/Actor.h
+++ b/Engine/Source/Runtime/Engine/Classes/GameFramework/Actor.h
@@ -1437,6 +1437,7 @@ public:
 	 * @note For collisions during physics simulation to generate hit events, 'Simulation Generates Hit Events' must be enabled.
 	 * @note When receiving a hit from another object's movement (bSelfMoved is false), the directions of 'Hit.Normal' and 'Hit.ImpactNormal'
 	 * will be adjusted to indicate force from the other object against this object.
+	 * @note NormalImpulse will be filled in for physics-simulating bodies, but will be zero for swept-component blocking collisions.
 	 */
 	UFUNCTION(BlueprintImplementableEvent, meta=(DisplayName = "Hit"), Category="Collision")
 	void ReceiveHit(class UPrimitiveComponent* MyComp, AActor* Other, class UPrimitiveComponent* OtherComp, bool bSelfMoved, FVector HitLocation, FVector HitNormal, FVector NormalImpulse, const FHitResult& Hit);
diff --git a/Engine/Source/Runtime/Engine/Classes/GameFramework/CharacterMovementComponent.h b/Engine/Source/Runtime/Engine/Classes/GameFramework/CharacterMovementComponent.h
--- a/Engine/Source/Runtime/Engine/Classes/GameFramework/CharacterMovementComponent.h
+++ b/Engine/Source/Runtime/Engine/Classes/GameFramework/CharacterMovementComponent.h
@@ -76,6 +76,14 @@ public:
 		HitResult.Reset(1.f, false);
 	}
 
+	/** Gets the distance to floor, either LineDist or FloorDist. */
+	float GetDistanceToFloor() const
+	{
+		// When the floor distance is set using SetFromSweep, the LineDist value will be reset.
+		// However, when SetLineFromTrace is used, there's no guarantee that FloorDist is set.
+		return bLineTrace ? LineDist : FloorDist;
+	}
+
 	void SetFromSweep(const FHitResult& InHit, const float InSweepFloorDist, const bool bIsWalkableFloor);
 	void SetFromLineTrace(const FHitResult& InHit, const float InSweepFloorDist, const float InLineDist, const bool bIsWalkableFloor);
 };
diff --git a/Engine/Source/Runtime/Engine/Classes/GameFramework/GameMode.h b/Engine/Source/Runtime/Engine/Classes/GameFramework/GameMode.h
--- a/Engine/Source/Runtime/Engine/Classes/GameFramework/GameMode.h
+++ b/Engine/Source/Runtime/Engine/Classes/GameFramework/GameMode.h
@@ -2,19 +2,10 @@
 
 #pragma once
 #include "TimerManager.h"
-#include "GameFramework/Info.h"
+#include "GameFramework/GameModeBase.h"
 #include "PlayerStart.h"
 #include "GameMode.generated.h"
 
-class FDebugDisplayInfo;
-class FUniqueNetId;
-struct FUniqueNetIdRepl;
-
-//~=============================================================================
-//  GameMode defines the rules and mechanics of the game.  It is only 
-//  instanced on the server and will never exist on the client.
-//~=============================================================================
-
 /** Possible state of the current match, where a match is all the gameplay that happens on a single map */
 namespace MatchState
 {
@@ -29,41 +20,13 @@ namespace MatchState
 	// Do not add any states before WaitingToStart or after WaitingPostMatch
 }
 
-/** Default delegate that provides an implementation for those that don't have special needs other than a toggle */
-DECLARE_DELEGATE_RetVal( bool, FCanUnpause );
-
-/** Configurable shortened aliases for GameMode classes.  For convenience when typing urls, for instance. */
-USTRUCT()
-struct FGameClassShortName
-{
-	GENERATED_USTRUCT_BODY()
-
-	/** Abbreviation that can be used as an alias for the class name */
-	UPROPERTY()
-	FString ShortName;
-
-	/** The class name to use when the alias is specified in a URL */
-	UPROPERTY()
-	FString GameClassName;
-};
-
 /**
- * The GameMode defines the game being played. It governs the game rules, scoring, what actors
- * are allowed to exist in this game type, and who may enter the game.
- *
- * It is only instanced on the server and will never exist on the client. 
- *
- * A GameMode actor is instantiated when the level is initialized for gameplay in
- * C++ UGameEngine::LoadMap().  
- * 
- * The class of this GameMode actor is determined by (in order) either the URL ?game=xxx, 
- * the GameMode Override value set in the World Settings, or the DefaultGameMode entry set 
- * in the game's Project Settings.
- *
- * @see https://docs.unrealengine.com/latest/INT/Gameplay/Framework/GameMode/index.html
+ * GameMode is a subclass of GameModeBase that behaves like a multiplayer match-based game.
+ * It has default behavior for picking spawn points and match state.
+ * If you want a simpler base, inherit from GameModeBase instead.
  */
-UCLASS(config=Game, notplaceable, BlueprintType, Blueprintable, Transient, hideCategories=(Info, Rendering, MovementReplication, Replication, Actor), meta=(ShortTooltip="Game Mode defines the game being played, its rules, scoring, and other facets of the game type."))
-class ENGINE_API AGameMode : public AInfo
+UCLASS()
+class ENGINE_API AGameMode : public AGameModeBase
 {
 	GENERATED_UCLASS_BODY()
 
@@ -73,10 +36,6 @@ class ENGINE_API AGameMode : public AInfo
 	UFUNCTION(BlueprintCallable, Category="Game")
 	FName GetMatchState() const { return MatchState; }
 
-	/** Returns true if the match state is InProgress or later */
-	UFUNCTION(BlueprintCallable, Category="Game")
-	virtual bool HasMatchStarted() const;
-
 	/** Returns true if the match state is InProgress or other gameplay state */
 	UFUNCTION(BlueprintCallable, Category="Game")
 	virtual bool IsMatchInProgress() const;
@@ -85,16 +44,6 @@ class ENGINE_API AGameMode : public AInfo
 	UFUNCTION(BlueprintCallable, Category="Game")
 	virtual bool HasMatchEnded() const;
 
-	/** Transitions to WaitingToStart and calls BeginPlay on actors. */
-	UFUNCTION(BlueprintCallable, Category="Game")
-	virtual void StartPlay();
-
-	/**
-	 * Called after a seamless level transition has been completed on the *new* GameMode.
-	 * Used to reinitialize players already in the game as they won't have *Login() called on them
-	 */
-	virtual void PostSeamlessTravel();
-
 	/** Transition from WaitingToStart to InProgress. You can call this manually, will also get called if ReadyToStartMatch returns true */
 	UFUNCTION(BlueprintCallable, Category="Game")
 	virtual void StartMatch();
@@ -103,17 +52,10 @@ class ENGINE_API AGameMode : public AInfo
 	UFUNCTION(BlueprintCallable, Category="Game")
 	virtual void EndMatch();
 
-	/** Transition to LeavingMap state. Start the transition out of the current map. Called at start of seamless travel, or right before map change for hard travel. */
-	virtual void StartToLeaveMap();
-
 	/** Restart the game, by default travel to the current map */
 	UFUNCTION(BlueprintCallable, Category="Game")
 	virtual void RestartGame();
 
-	/** Return to main menu, and disconnect any players */
-	UFUNCTION(BlueprintCallable, Category="Game")
-	virtual void ReturnToMainMenuHost();
-
 	/** Report that a match has failed due to unrecoverable error */
 	UFUNCTION(BlueprintCallable, Category="Game")
 	virtual void AbortMatch();
@@ -161,37 +103,10 @@ protected:
 
 public:
 
-	/**
-	 * perform map travels using SeamlessTravel() which loads in the background and doesn't disconnect clients
-	 * @see World::SeamlessTravel()
-	 */
-	UPROPERTY(EditAnywhere, BlueprintReadWrite, Category="GameMode")
-	uint32 bUseSeamlessTravel : 1;
-
-	/** Whether the game is pauseable. */
-	UPROPERTY(EditAnywhere, BlueprintReadWrite, Category="GameMode")
-	uint32 bPauseable:1;
-
-	/** Whether players should immediately spawn when logging in, or stay as spectators until they manually spawn */
-	UPROPERTY(EditAnywhere, BlueprintReadOnly, Category="GameMode")
-	uint32 bStartPlayersAsSpectators : 1;
-
 	/** Whether the game should immediately start when the first player logs in. Affects the default behavior of ReadyToStartMatch */
 	UPROPERTY(EditAnywhere, BlueprintReadOnly, Category="GameMode")
 	uint32 bDelayedStart : 1;
 
-	/** Save options string and parse it when needed */
-	UPROPERTY(BlueprintReadOnly, Category="GameMode")
-	FString OptionsString;
-
-	/** The default pawn class used by players. */
-	UPROPERTY(EditAnywhere, BlueprintReadWrite, Category=Classes)
-	TSubclassOf<class APawn> DefaultPawnClass;
-
-	/** HUD class this game uses. */
-	UPROPERTY(EditAnywhere, BlueprintReadWrite, Category=Classes, meta=(DisplayName="HUD Class"))
-	TSubclassOf<class AHUD> HUDClass;    
-
 	/** Current number of spectators. */
 	UPROPERTY(BlueprintReadOnly, Category=GameMode)
 	int32 NumSpectators;    
@@ -208,69 +123,20 @@ public:
 	UPROPERTY(EditAnywhere, BlueprintReadOnly, Category=GameMode, meta=(DisplayName="Minimum Respawn Delay"))
 	float MinRespawnDelay;
 
-	/** Game Session handles login approval, arbitration, online game interface */
-	UPROPERTY()
-	class AGameSession* GameSession;
-
 	/** Number of players that are still traveling from a previous map */
 	UPROPERTY(BlueprintReadOnly, Category=GameMode)
 	int32 NumTravellingPlayers;
 
-	/** Used to assign unique PlayerIDs to each PlayerState. */
-	int32 CurrentID;    
-
-	/** The default player name assigned to players that join with no name specified. */
-	UPROPERTY(EditAnywhere, Category=GameMode)
-	FText DefaultPlayerName;
-
 	/** Contains strings describing localized game agnostic messages. */
 	UPROPERTY()
 	TSubclassOf<class ULocalMessage> EngineMessageClass;
 
-	/** The class of PlayerController to spawn for players logging in. */
-	UPROPERTY(EditAnywhere, noclear, BlueprintReadOnly, Category=Classes)
-	TSubclassOf<class APlayerController> PlayerControllerClass;
-
-	/** The pawn class used by the PlayerController for players when spectating. */
-	UPROPERTY(EditAnywhere, noclear, BlueprintReadOnly, Category=Classes)
-	TSubclassOf<class ASpectatorPawn> SpectatorClass;
-
-	/** The PlayerController class used when spectating a network replay. */
-	UPROPERTY(EditAnywhere, noclear, BlueprintReadOnly, Category = Classes)
-	TSubclassOf<class APlayerController> ReplaySpectatorPlayerControllerClass;
-
-	/** A PlayerState of this class will be associated with every player to replicate relevant player information to all clients. */
-	UPROPERTY(EditAnywhere, noclear, BlueprintReadOnly, Category=Classes)
-	TSubclassOf<class APlayerState> PlayerStateClass;
-
-	/** Class of GameState associated with this GameMode. */
-	UPROPERTY(EditAnywhere, noclear, BlueprintReadOnly, Category=Classes)
-	TSubclassOf<class AGameState> GameStateClass;
-
-	/** GameState is used to replicate game state relevant properties to all clients. */
-	UPROPERTY(Transient)
-	class AGameState* GameState;
-
-	/** Helper template to returns the current GameState casted to the desired type. */
-	template< class T >
-	T* GetGameState() const
-	{
-		return Cast<T>(GameState);
-	}
-
 	/** PlayerStates of players who have disconnected from the server (saved in case they reconnect) */
 	UPROPERTY()
 	TArray<class APlayerState*> InactivePlayerArray;    
 
-	/** The list of delegates to check before unpausing a game */
-	TArray<FCanUnpause> Pausers;
-
 protected:
 
-	/** Handy alternate short names for GameMode classes (e.g. "DM" could be an alias for "MyProject.MyGameModeMP_DM". */
-	UPROPERTY(config)
-	TArray<struct FGameClassShortName> GameModeClassAliases;
-
 	/** Time a playerstate will stick around in an inactive state after a player logout */
 	UPROPERTY(EditAnywhere, Category=GameMode)
 	float InactivePlayerStateLifeSpan;
@@ -281,61 +147,6 @@ protected:
 
 public:
 
-	/** Alters the synthetic bandwidth limit for a running game. */
-	UFUNCTION(exec)
-	virtual void SetBandwidthLimit(float AsyncIOBandwidthLimit);
-
-	//~ Begin AActor Interface
-	virtual void PreInitializeComponents() override;
-	virtual void DisplayDebug(class UCanvas* Canvas, const FDebugDisplayInfo& DebugDisplay, float& YL, float& YPos) override;
-	virtual void Reset() override;
-	virtual void Tick(float DeltaSeconds) override;
-	//~ End AActor Interface
-
-	/** 
-	 * Overridable function to determine whether an Actor should have Reset called when the game has Reset called on it.
-	 * Default implementation returns true
-	 * @param ActorToReset The actor to make a determination for
-	 * @return true if ActorToReset should have Reset() called on it while restarting the game, 
-	 *		   false if the GameMode will manually reset it or if the actor does not need to be reset
-
-	 */
-	UFUNCTION(BlueprintNativeEvent, Category="Game")
-	bool ShouldReset(AActor* ActorToReset);
-
-	/** 
-	  * Overridable function called when resetting level.
-	  * Default implementation calls Reset() on all actors except GameMode and Controllers
-	  */
-	UFUNCTION(BlueprintCallable, Category="Game")
-	virtual void ResetLevel();
-
-	/**
-	 *	Check to see if we should start in cinematic mode (e.g. matinee movie capture).
-	 * 	@param	OutHidePlayer		Whether to hide the player
-	 *	@param	OutHideHud			Whether to hide the HUD		
-	 *	@param	OutDisableMovement	Whether to disable movement
-	 * 	@param	OutDisableTurning	Whether to disable turning
-	 *	@return	bool			true if we should turn on cinematic mode, 
-	 *							false if if we should not.
-	 */
-	bool ShouldStartInCinematicMode(APlayerController* Player, bool& OutHidePlayer, bool& OutHideHud, bool& OutDisableMovement, bool& OutDisableTurning);
-
-	/**
-	 * Initialize the game.
-	 * The GameMode's InitGame() event is called before any other functions (including PreInitializeComponents() )
-	 * and is used by the GameMode to initialize parameters and spawn its helper classes.
-	 * @warning: this is called before actors' PreInitializeComponents.
-	 */
-	virtual void InitGame(const FString& MapName, const FString& Options, FString& ErrorMessage);
-
-	/**
-	 * Initialize the GameState actor with default settings
-	 * called during PreInitializeComponents() of the GameMode after a GameState has been spawned
-	 * as well as during Reset()
-	 */
-	virtual void InitGameState();
-
 	/** Get local address */
 	virtual FString GetNetworkNumber();
 	
@@ -345,218 +156,31 @@ public:
 	/** Removes the passed in player controller from the correct count for player/spectator/tranistioning **/
 	void RemovePlayerControllerFromPlayerCount(APlayerController* PC);
 
-	/** Total number of players */
-	UFUNCTION(BlueprintCallable, Category="Game")
-	virtual int32 GetNumPlayers();
-
-	/**
-	 * Adds the delegate to the list if the player Controller has the right to pause
-	 * the game. The delegate is called to see if it is ok to unpause the game, e.g.
-	 * the reason the game was paused has been cleared.
-	 * @param PC the player Controller to check for admin privs
-	 * @param CanUnpauseDelegate the delegate to query when checking for unpause
-	 */
-	virtual bool SetPause(APlayerController* PC, FCanUnpause CanUnpauseDelegate = FCanUnpause());
-
-	/**
-	 * Checks the list of delegates to determine if the pausing can be cleared. If
-	 * the delegate says it's ok to unpause, that delegate is removed from the list
-	 * and the rest are checked. The game is considered unpaused when the list is
-	 * empty.
-	 */
-	virtual void ClearPause();
-
-	/**
-	 * Forcibly removes an object's CanUnpause delegates from the list of pausers.  If any of the object's CanUnpause delegate
-	 * handlers were in the list, triggers a call to ClearPause().
-	 *
-	 * Called when the player controller is being destroyed to prevent the game from being stuck in a paused state when a PC that
-	 * paused the game is destroyed before the game is unpaused.
-	 */
-	void ForceClearUnpauseDelegates( AActor* PauseActor );
-
-	//~=========================================================================
-	// URL Parsing
-	/** Grab the next option from a string. */
-	DEPRECATED(4.9, "Use UGameplayStatics::GrabOption instead")
-	static bool GrabOption( FString& Options, FString& ResultString );
-
-	/** Break up a key=value pair into its key and value. */
-	DEPRECATED(4.9, "Use UGameplayStatics::GetKeyValue instead")
-	static void GetKeyValue( const FString& Pair, FString& Key, FString& Value );
-
-	/* Find an option in the options string and return it. */
-	DEPRECATED(4.9, "Use UGameplayStatics::ParseOption instead")
-	static FString ParseOption( FString Options, const FString& InKey );
-
-	/** HasOption - return true if the option is specified on the command line. */
-	DEPRECATED(4.9, "Use UGameplayStatics::HasOption instead")
-	static bool HasOption( FString Options, const FString& InKey );
-
-	/** Search array of options for ParseString and return the index or CurrentValue if not found*/
-	DEPRECATED(4.9, "Use UGameplayStatics::GetIntOption instead")
-	static int32 GetIntOption( const FString& Options, const FString& ParseString, int32 CurrentValue);
-	//~=========================================================================
-
-	/** 
-	 * @return the full path to the optimal GameMode class to use for the specified map and options
-	 * this is used for preloading cooked packages, etc. and therefore doesn't need to include any fallbacks
-	 * as GetGameModeClass() will be called later to actually find/load the desired class
-	 */
+	DEPRECATED(4.14, "Deprecated in favor of PreloadContentForURL on GameInstance")
 	virtual FString GetDefaultGameClassPath(const FString& MapName, const FString& Options, const FString& Portal) const;
 
-	/** 
-	 * @return the class of GameMode to spawn for the game on the specified map and the specified options
-	 * this function should include any fallbacks in case the desired class can't be found
-	 */
+	DEPRECATED(4.14, "Deprecated in favor of OverrideGameModeClass on GameInstance")
 	virtual TSubclassOf<AGameMode> GetGameModeClass(const FString& MapName, const FString& Options, const FString& Portal) const;
 
-	DEPRECATED(4.5, "AGameMode::SetGameMode renamed AGameMode::GetGameModeClass")
-	TSubclassOf<AGameMode> SetGameMode(const FString& MapName, const FString& Options, const FString& Portal);
-
-
-	/** @return GameSession class to use for this game  */
-	virtual TSubclassOf<class AGameSession> GetGameSessionClass() const;
-
-	/** Called when a connection closes before getting to PostLogin() */
-	virtual void NotifyPendingConnectionLost();
+	DEPRECATED(4.14, "Deprecated in favor of GetGameModeForName on GameMapsSettings")
+	static FString StaticGetFullGameClassName(FString const& Str);
 
-	/** @return true if we want to travel_absolute */
+	/** Return true if we want to travel_absolute, used by RestartGame by default */
 	virtual bool GetTravelType();
 
-	/** Optional handling of ServerTravel for network games. */
-	virtual void ProcessServerTravel(const FString& URL, bool bAbsolute = false);
-
-	/**
-	 * Notifies all clients to travel to the specified URL.
-	 *
-	 * @param	URL				a string containing the mapname (or IP address) to travel to, along with option key/value pairs
-	 * @param	NextMapGuid		the GUID of the server's version of the next map
-	 * @param	bSeamless		indicates whether the travel should use seamless travel or not.
-	 * @param	bAbsolute		indicates which type of travel the server will perform (i.e. TRAVEL_Relative or TRAVEL_Absolute)
-	 */
-	virtual APlayerController* ProcessClientTravel( FString& URL, FGuid NextMapGuid, bool bSeamless, bool bAbsolute );
-
-	/** 
-	 * Accept or reject a player attempting to join the server.  Fails login if you set the ErrorMessage to a non-empty string.
-	 * PreLogin is called before Login.  Significant game time may pass before Login is called, especially if content is downloaded.
-	 *
-	 * @param	Options					The URL options (e.g. name/spectator) the player has passed
-	 * @param	Address					The network address of the player
-	 * @param	UniqueId				The unique id the player has passed to the server
-	 * @param	ErrorMessage			When set to a non-empty value, the player will be rejected using the error message set
-	 */
-	virtual void PreLogin(const FString& Options, const FString& Address, const TSharedPtr<const FUniqueNetId>& UniqueId, FString& ErrorMessage);
-	virtual void PreLogin(const FString& Options, const FString& Address, const FUniqueNetIdRepl& UniqueId, FString& ErrorMessage);
-
-	/** 
-	 * Called to login new players by creating a player controller, overridable by the game
-	 *
-	 * Sets up basic properties of the player (name, unique id, registers with backend, etc) and should not be used to do
-	 * more complicated game logic.  The player controller is not fully initialized within this function as far as networking is concerned.
-	 * Save "game logic" for PostLogin which is called shortly afterward.
-	 *
-	 * @param NewPlayer pointer to the UPlayer object that represents this player (either local or remote)
-	 * @param RemoteRole the remote role this controller has
-	 * @param Portal desired portal location specified by the client
-	 * @param Options game options passed in by the client at login
-	 * @param UniqueId platform specific unique identifier for the logging in player
-	 * @param ErrorMessage [out] error message, if any, why this login will be failing
-	 *
-	 * If login is successful, returns a new PlayerController to associate with this player. Login fails if ErrorMessage string is set. 
-	 *
-	 * @return a new player controller for the logged in player, NULL if login failed for any reason
-	 */
-	virtual APlayerController* Login(class UPlayer* NewPlayer, ENetRole InRemoteRole, const FString& Portal, const FString& Options, const TSharedPtr<const FUniqueNetId>& UniqueId, FString& ErrorMessage);
-	virtual APlayerController* Login(class UPlayer* NewPlayer, ENetRole InRemoteRole, const FString& Portal, const FString& Options, const FUniqueNetIdRepl& UniqueId, FString& ErrorMessage);
-
-	/** Called after a successful login.  This is the first place it is safe to call replicated functions on the PlayerAController. */
-	virtual void PostLogin( APlayerController* NewPlayer );
-
-	/** Notification that a player has successfully logged in, and has been given a player controller */
-	UFUNCTION(BlueprintImplementableEvent, Category="Game", meta=(DisplayName="OnPostLogin"))
-	void K2_PostLogin( APlayerController* NewPlayer );
-
-	/** 
-	 * Spawns a PlayerController at the specified location; split out from Login()/HandleSeamlessTravelPlayer() for easier overriding 
-	 *
-	 * @param RemoteRole the role this controller will play remotely
-	 * @param SpawnLocation location in the world to spawn
-	 * @param SpawnRotation rotation to set relative to the world
-	 *
-	 * @return PlayerController for the player, NULL if there is any reason this player shouldn't exist or due to some error
-	 */
-	virtual APlayerController* SpawnPlayerController(ENetRole InRemoteRole, FVector const& SpawnLocation, FRotator const& SpawnRotation);
-
-	/** @Returns true if NewPlayerController may only join the server as a spectator. */
-	UFUNCTION(BlueprintNativeEvent, Category="Game")
-	bool MustSpectate(APlayerController* NewPlayerController) const;
-
-	/** returns default pawn class for given controller */
-	UFUNCTION(BlueprintNativeEvent, Category="Game")
-	UClass* GetDefaultPawnClassForController(AController* InController);
-
-	/** Called when StartSpot is selected for spawning NewPlayer to allow optional initialization. */
-	UFUNCTION(BlueprintNativeEvent, Category="Game")
-	void InitStartSpot(class AActor* StartSpot, AController* NewPlayer);
-
-	/**
-	 * @param	NewPlayer - Controller for whom this pawn is spawned
-	 * @param	StartSpot - PlayerStart at which to spawn pawn
-	 * @return	a pawn of the default pawn class
-	 */
-	UFUNCTION(BlueprintNativeEvent, Category="Game")
-	APawn* SpawnDefaultPawnFor(AController* NewPlayer, class AActor* StartSpot);
-
-	/** replicates the current level streaming status to the given PlayerController */
-	virtual void ReplicateStreamingStatus(APlayerController* PC);
-
-	/**
-	 * handles all player initialization that is shared between the travel methods
-	 * (i.e. called from both PostLogin() and HandleSeamlessTravelPlayer())
-	 */
-	virtual void GenericPlayerInitialization(AController* C);
-
-	/** start match, or let player enter, immediately */
+	DEPRECATED(4.14, "SendPlayer is not in use, call ClientTravel directly instead")
+	virtual void SendPlayer( APlayerController* aPlayer, const FString& URL );
+	
+	DEPRECATED(4.14, "StartNewPlayer has been split into InitializeHUDForPlayer and HandleStartingNewPlayer")
 	virtual void StartNewPlayer(APlayerController* NewPlayer);
 
-	/** Called when a Controller with a PlayerState leaves the match. */
-	virtual void Logout( AController* Exiting );
-
-	/** Implementable event when a Controller with a PlayerState leaves the match. */
-	UFUNCTION(BlueprintImplementableEvent, Category="Game", meta=(DisplayName="OnLogout"))
-	void K2_OnLogout(AController* ExitingController);
-
-	/**
-	 * Make sure pawn properties are back to default
-	 * Also a good place to modify them on spawn
-	 */
-	virtual void SetPlayerDefaults(APawn* PlayerPawn);
-
-	/** Return whether Viewer is allowed to spectate from the point of view of ViewTarget. */
-	UFUNCTION(BlueprintNativeEvent, Category="Game")
-	bool CanSpectate( APlayerController* Viewer, APlayerState* ViewTarget );
-
-	/** 
-	 * Sets the name for a controller 
-	 * @param Controller	The controller of the player to change the name of
-	 * @param NewName		The name to set the player to
-	 * @param bNameChange	Whether the name is changing or if this is the first time it has been set
-	 */
-	UFUNCTION(BlueprintCallable, Category="Game")
-	virtual void ChangeName( AController* Controller, const FString& NewName, bool bNameChange );
-
-	/** 
-	 * Overridable event for GameMode blueprint to respond to a change name call
-	 * @param Controller	The controller of the player to change the name of
-	 * @param NewName		The name to set the player to
-	 * @param bNameChange	Whether the name is changing or if this is the first time it has been set
-	 */
-	UFUNCTION(BlueprintImplementableEvent,Category="Game",meta=(DisplayName="OnChangeName"))
-	void K2_OnChangeName( AController* Other, const FString& NewName, bool bNameChange );
+	/** Exec command to broadcast a string to all players */
+	UFUNCTION(Exec, BlueprintCallable, Category = AI)
+	virtual void Say(const FString& Msg);
 
-	/* Send a player to a URL.*/
-	virtual void SendPlayer( APlayerController* aPlayer, const FString& URL );
+	/** Alters the synthetic bandwidth limit for a running game. */
+	UFUNCTION(exec)
+	virtual void SetBandwidthLimit(float AsyncIOBandwidthLimit);
 
 	/** Broadcast a string to all players. */
 	virtual void Broadcast( AActor* Sender, const FString& Msg, FName Type = NAME_None );
@@ -568,61 +192,6 @@ public:
 	 */
 	virtual void BroadcastLocalized( AActor* Sender, TSubclassOf<ULocalMessage> Message, int32 Switch = 0, APlayerState* RelatedPlayerState_1 = NULL, APlayerState* RelatedPlayerState_2 = NULL, UObject* OptionalObject = NULL );
 
-	/** @return true if the given Controller StartSpot property should be used as the spawn location for its Pawn */
-	virtual bool ShouldSpawnAtStartSpot(AController* Player);
-
-	/** 
-	 * Return the 'best' player start for this player to start from.
-	 * @param Player The AController for whom we are choosing a Player Start
-	 * @param IncomingName Specifies the tag of a Player Start to use
-	 * @returns Actor chosen as player start (usually a PlayerStart)
-	 */
-	UFUNCTION(BlueprintNativeEvent, Category="Game")
-	class AActor* FindPlayerStart( AController* Player, const FString& IncomingName = TEXT("") );
-
-	/** Calls code to select the best player start for this player to start from. */
-	UFUNCTION(BlueprintPure, Category=Game, meta=(DisplayName="FindPlayerStart"))
-	AActor* K2_FindPlayerStart(AController* Player);
-
-	/** 
-	* Return the 'best' player start for this player to start from.  
-	* Default implementation just returns the first PlayerStart found.
-	* @param Player is the controller for whom we are choosing a playerstart
-	* @returns AActor chosen as player start (usually a PlayerStart)
-	 */
-	UFUNCTION(BlueprintNativeEvent, Category="Game")
-	class AActor* ChoosePlayerStart( AController* Player );
-
-	/** @return true if it's valid to call RestartPlayer. Will call Player->CanRestartPlayer */
-	UFUNCTION(BlueprintNativeEvent, Category="Game")
-	bool PlayerCanRestart( APlayerController* Player );
-
-	/**
-	 * Used to notify the game type that it is ok to update a player's gameplay
-	 * specific muting information now. The playercontroller needs to notify
-	 * the server when it is possible to do so or the unique net id will be
-	 * incorrect and the muting not work.
-	 *
-	 * @param aPlayer the playercontroller that is ready for updates
-	 */
-	virtual void UpdateGameplayMuteList( APlayerController* aPlayer );
-
-	/** @return true if player is allowed to access the cheats */
-	virtual bool AllowCheats(APlayerController* P);
-
-	/** @return	true if the player is allowed to pause the game. */
-	virtual bool AllowPausing( APlayerController* PC = NULL );
-
-	/**
-	 * Called from CommitMapChange before unloading previous level. Used for asynchronous level streaming
-	 * @param PreviousMapName - Name of the previous persistent level
-	 * @param NextMapName - Name of the persistent level being streamed to
-	 */
-	virtual void PreCommitMapChange(const FString& PreviousMapName, const FString& NextMapName);
-
-	/** Called from CommitMapChange after unloading previous level and loading new level+sublevels. Used for asynchronous level streaming */
-	virtual void PostCommitMapChange();
-
 	/** Add PlayerState to the inactive list, remove from the active list */
 	virtual void AddInactivePlayer(APlayerState* PlayerState, APlayerController* PC);
 
@@ -630,103 +199,51 @@ public:
 	  * @Returns true if a PlayerState was found and associated with PC. */
 	virtual bool FindInactivePlayer(APlayerController* PC);
 
-	/** Override PC's PlayerState with the values in OldPlayerState. */
+	/** Override PC's PlayerState with the values in OldPlayerState as part of the inactive player handling */
 	virtual void OverridePlayerState(APlayerController* PC, APlayerState* OldPlayerState);
 
-	/** 
-	 * called on server during seamless level transitions to get the list of Actors that should be moved into the new level
-	 * PlayerControllers, Role < ROLE_Authority Actors, and any non-Actors that are inside an Actor that is in the list
-	 * (i.e. Object.Outer == Actor in the list)
-	 * are all automatically moved regardless of whether they're included here
-	 * only dynamic actors in the PersistentLevel may be moved (this includes all actors spawned during gameplay)
-	 * this is called for both parts of the transition because actors might change while in the middle (e.g. players might join or leave the game)
-	 * @see also PlayerController::GetSeamlessTravelActorList() (the function that's called on clients)
-	 * @param bToTransition true if we are going from old level to transition map, false if we are going from transition map to new level
-	 * @param ActorList (out) list of actors to maintain
-	 */
-	virtual void GetSeamlessTravelActorList(bool bToTransition, TArray<AActor*>& ActorList);
-
-	/** Allow the game to specify a place for clients to download MapName */
-	DEPRECATED(4.11, "Override GameWelcomePlayer() instead")
-	virtual FString GetRedirectURL(const FString& MapName) const;
-
-	/** allows game to send network messages to provide more information to the client joining the game via NMT_GameSpecific
-	 * (for example required DLC)
-	 * the out string RedirectURL is built in and send automatically if only a simple URL is needed
-	 */
-	virtual void GameWelcomePlayer(UNetConnection* Connection, FString& RedirectURL);
-
-	/** 
-	 * used to swap a viewport/connection's PlayerControllers when seamless traveling and the new GameMode's
-	 * controller class is different than the previous
-	 * includes network handling
-	 * @param OldPC - the old PC that should be discarded
-	 * @param NewPC - the new PC that should be used for the player
-	 */
-	virtual void SwapPlayerControllers(APlayerController* OldPC, APlayerController* NewPC);
-
-	UFUNCTION(BlueprintImplementableEvent, Category="Game", meta=(DisplayName="OnSwapPlayerControllers"))
-	void K2_OnSwapPlayerControllers(APlayerController* OldPC, APlayerController* NewPC);
-
-	/** 
-	 * Handles reinitializing players that remained through a seamless level transition
-	 * called from C++ for players that finished loading after the server
-	 * @param C the Controller to handle
-	 */
-	virtual void HandleSeamlessTravelPlayer(AController*& C);
-
 	/** SetViewTarget of player control on server change */
 	virtual void SetSeamlessTravelViewTarget(APlayerController* PC);
 
 	/** Called when this PC is in cinematic mode, and its matinee is canceled by the user. */
 	virtual void MatineeCancelled();
 
-	/** Does end of game handling for the online layer */
-	virtual void RestartPlayer(class AController* NewPlayer);
-
-	UFUNCTION(BlueprintImplementableEvent, Category="Game", meta=(DisplayName="OnRestartPlayer"))
-	void K2_OnRestartPlayer(AController* NewPlayer);
-
-	/** Given a string, return a fully-qualified GameMode class name */
-	static FString StaticGetFullGameClassName(FString const& Str);
-
-	/** Returns true if replays will start/stop during gameplay starting/stopping */
-	virtual bool IsHandlingReplays();
-
-protected:
-
 	/**
-	 * Customize incoming player based on URL options
-	 *
-	 * @param NewPlayerController player logging in
-	 * @param UniqueId unique id for this player
-	 * @param Options URL options that came at login
-	 *
+	 * Called from CommitMapChange before unloading previous level. Used for asynchronous level streaming
+	 * @param PreviousMapName - Name of the previous persistent level
+	 * @param NextMapName - Name of the persistent level being streamed to
 	 */
-	virtual FString InitNewPlayer(class APlayerController* NewPlayerController, const TSharedPtr<const FUniqueNetId>& UniqueId, const FString& Options, const FString& Portal = TEXT(""));
-	virtual FString InitNewPlayer(class APlayerController* NewPlayerController, const FUniqueNetIdRepl& UniqueId, const FString& Options, const FString& Portal = TEXT(""));
+	virtual void PreCommitMapChange(const FString& PreviousMapName, const FString& NextMapName);
 
-private:
-	// Hidden functions that don't make sense to use on this class.
-	HIDE_ACTOR_TRANSFORM_FUNCTIONS();
-};
+	/** Called from CommitMapChange after unloading previous level and loading new level+sublevels. Used for asynchronous level streaming */
+	virtual void PostCommitMapChange();
 
-/** GameMode events, particularly for use by plugins */
-class ENGINE_API FGameModeEvents
-{
-public:
-	// Post login event, triggered when a player joins the game as well as after non-seamless ServerTravel
-	DECLARE_EVENT_TwoParams(AGameMode, FGameModePostLoginEvent, AGameMode* /*GameMode*/, APlayerController* /*NewPlayer*/);
+	/** Called when a connection closes before getting to PostLogin() */
+	virtual void NotifyPendingConnectionLost();
 
-	/**
-	* Logout event, triggered when a player leaves the game as well as during non-seamless ServerTravel
-	*
-	* Note that this is called before performing any cleanup of the specified AController
-	*/
-	DECLARE_EVENT_TwoParams(AGameMode, FGameModeLogoutEvent, AGameMode* /*GameMode*/, AController* /*Exiting*/);
-	
+	/** Handles when a player is disconnected, before the session does */
+	virtual void HandleDisconnect(UWorld* InWorld, UNetDriver* NetDriver);
 
-	static FGameModePostLoginEvent GameModePostLoginEvent;
-	static FGameModeLogoutEvent GameModeLogoutEvent;
-};
+	//~ Begin AActor Interface
+	virtual void Tick(float DeltaSeconds) override;
+	//~ End AActor Interface
 
+	//~ Begin AGameModeBase Interface
+	virtual void InitGame(const FString& MapName, const FString& Options, FString& ErrorMessage) override;
+	virtual void StartPlay() override;
+	virtual bool HasMatchStarted() const override;
+	virtual void PostLogin(APlayerController* NewPlayer) override;
+	virtual void Logout(AController* Exiting) override;
+	virtual int32 GetNumPlayers() override;
+	virtual int32 GetNumSpectators() override;
+	virtual bool IsHandlingReplays() override;
+	virtual void HandleStartingNewPlayer_Implementation(APlayerController* NewPlayer) override;
+	virtual bool PlayerCanRestart_Implementation(APlayerController* Player) override;
+	virtual void PostSeamlessTravel() override;
+	virtual void HandleSeamlessTravelPlayer(AController*& C) override;
+	virtual void InitSeamlessTravelPlayer(AController* NewController) override;
+	virtual bool CanServerTravel(const FString& URL, bool bAbsolute) override;
+	virtual void StartToLeaveMap() override;
+	virtual bool SpawnPlayerFromSimulate(const FVector& NewLocation, const FRotator& NewRotation) override;
+	//~ End AGameModeBase Interface
+};
