Copying //UE4/Dev-Framework to //UE4/Dev-Main (Source: //UE4/Dev-Framework @ 3136612)

diff --git a/Engine/Source/Runtime/Engine/Classes/GameFramework/GameNetworkManager.h b/Engine/Source/Runtime/Engine/Classes/GameFramework/GameNetworkManager.h
--- a/Engine/Source/Runtime/Engine/Classes/GameFramework/GameNetworkManager.h
+++ b/Engine/Source/Runtime/Engine/Classes/GameFramework/GameNetworkManager.h
@@ -184,6 +184,9 @@ class ENGINE_API AGameNetworkManager : public AInfo
 	/** Timer which calls UpdateNetSpeeds() once a second. */
 	virtual void UpdateNetSpeedsTimer();
 
+	/** Returns true if we should be in low bandwidth mode */
+	virtual bool IsInLowBandwidthMode();
+
 	//======================================================================================================================
 	// Player replication
 
diff --git a/Engine/Source/Runtime/Engine/Classes/GameFramework/GameState.h b/Engine/Source/Runtime/Engine/Classes/GameFramework/GameState.h
--- a/Engine/Source/Runtime/Engine/Classes/GameFramework/GameState.h
+++ b/Engine/Source/Runtime/Engine/Classes/GameFramework/GameState.h
@@ -1,38 +1,23 @@
 // Copyright 1998-2016 Epic Games, Inc. All Rights Reserved.
 
 #pragma once
-#include "TimerManager.h"
+#include "GameFramework/GameStateBase.h"
 #include "GameState.generated.h"
 
-//~=============================================================================
-// GameState is replicated and is valid on servers and clients.
-//~=============================================================================
-
-UCLASS(config=Game, notplaceable, BlueprintType, Blueprintable)
-class ENGINE_API AGameState : public AInfo
+/**
+ * GameState is a subclass of GameStateBase that behaves like a multiplayer match-based game.
+ * It is tied to functionality in GameMode.
+ */
+UCLASS()
+class ENGINE_API AGameState : public AGameStateBase
 {
 	GENERATED_UCLASS_BODY()
 
-	/** Class of the server's game mode, assigned by GameMode. */
-	UPROPERTY(replicatedUsing=OnRep_GameModeClass)
-	TSubclassOf<class AGameMode>  GameModeClass;
-
-	/** Instance of the current game mode, exists only on the server. For non-authority clients, this will be NULL. */
-	UPROPERTY(Transient, BlueprintReadOnly, Category=GameState)
-	class AGameMode* AuthorityGameMode;
-
-	/** Class used by spectators, assigned by GameMode. */
-	UPROPERTY(replicatedUsing=OnRep_SpectatorClass)
-	TSubclassOf<class ASpectatorPawn> SpectatorClass;
-
 	// Code to deal with the match state machine
 
 	/** Returns the current match state, this is an accessor to protect the state machine flow */
 	FName GetMatchState() const { return MatchState; }
 
-	/** Returns true if the match state is InProgress or later */
-	virtual bool HasMatchStarted() const;
-
 	/** Returns true if we're in progress */
 	virtual bool IsMatchInProgress() const;
 
@@ -70,18 +55,6 @@ public:
 	UPROPERTY(replicatedUsing=OnRep_ElapsedTime, BlueprintReadOnly, Category = GameState)
 	int32 ElapsedTime;
 
-	/** Array of all PlayerStates, maintained on both server and clients (PlayerStates are always relevant) */
-	UPROPERTY(BlueprintReadOnly, Category=GameState)
-	TArray<class APlayerState*> PlayerArray;
-
-	/** GameMode class notification callback. */
-	UFUNCTION()
-	virtual void OnRep_GameModeClass();
-
-	/** Callback when we receive the spectator class */
-	UFUNCTION()
-	virtual void OnRep_SpectatorClass();
-
 	/** Match state has changed */
 	UFUNCTION()
 	virtual void OnRep_MatchState();
@@ -90,70 +63,28 @@ public:
 	UFUNCTION()
 	virtual void OnRep_ElapsedTime();
 
-	//~ Begin AActor Interface
-	virtual void PostInitializeComponents() override;
-	//~ End AActor Interface
-
-	/** Helper to return the default object of the GameMode class corresponding to this GameState */
-	AGameMode* GetDefaultGameMode() const;
-	
-	/** Called when the GameClass property is set (at startup for the server, after the variable has been replicated on clients) */
-	virtual void ReceivedGameModeClass();
-
-	/** Called when the SpectatorClass property is set (at startup for the server, after the variable has been replicated on clients) */
-	virtual void ReceivedSpectatorClass();
-
-	/** Called during seamless travel transition twice (once when the transition map is loaded, once when destination map is loaded) */
-	virtual void SeamlessTravelTransitionCheckpoint(bool bToTransitionMap);
-
-	/** Add PlayerState to the PlayerArray */
-	virtual void AddPlayerState(class APlayerState* PlayerState);
-
-	/** Remove PlayerState from the PlayerArray. */
-	virtual void RemovePlayerState(class APlayerState* PlayerState);
-
 	/** Called periodically, overridden by subclasses */
 	virtual void DefaultTimer();
 
-	/** Should players show gore? */
+	DEPRECATED(4.14, "ShouldShowGore is deprecated, it is not supported by newer assets")
 	virtual bool ShouldShowGore() const;
 
-	/** Returns the simulated TimeSeconds on the server */
-	UFUNCTION(BlueprintCallable, Category=GameState)
-	virtual float GetServerWorldTimeSeconds() const;
-
-protected:
-
-	/** Called periodically to update ReplicatedWorldTimeSeconds */
-	virtual void UpdateServerTimeSeconds();
-
-	/** Allows clients to calculate ServerWorldTimeSecondsDelta */
-	UFUNCTION()
-	virtual void OnRep_ReplicatedWorldTimeSeconds();
-
-	/** Server TimeSeconds. Useful for syncing up animation and gameplay. */
-	UPROPERTY(replicatedUsing=OnRep_ReplicatedWorldTimeSeconds, transient)
-	float ReplicatedWorldTimeSeconds;
+	//~ Begin AActor Interface
+	virtual void PostInitializeComponents() override;
+	//~ End AActor Interface
 
-	/** The difference from the local world's TimeSeconds and the server world's TimeSeconds. */
-	UPROPERTY(transient)
-	float ServerWorldTimeSecondsDelta;
+	//~ Begin AGameStateBase Interface
+	virtual bool HasMatchStarted() const override;
+	virtual void HandleBeginPlay() override;
+	virtual float GetPlayerStartTime(class AController* Controller) const override;
+	virtual float GetPlayerRespawnDelay(class AController* Controller) const override;
+	//~ End AGameStateBase Interface
 
-	/** Frequency that the server updates the replicated TimeSeconds from the world. Set to zero to disable periodic updates. */
-	UPROPERTY(EditDefaultsOnly, Category=GameState)
-	float ServerWorldTimeSecondsUpdateFrequency;
+protected:
 
 	/** Handle for efficient management of DefaultTimer timer */
 	FTimerHandle TimerHandle_DefaultTimer;
 
-	/** Handle for efficient management of the UpdateServerTimeSeconds timer */
-	FTimerHandle TimerHandle_UpdateServerTimeSeconds;
-
-private:
-	// Hidden functions that don't make sense to use on this class.
-	HIDE_ACTOR_TRANSFORM_FUNCTIONS();
-
-	friend class UDemoNetDriver;
 };
 
 
diff --git a/Engine/Source/Runtime/Engine/Classes/GameFramework/Pawn.h b/Engine/Source/Runtime/Engine/Classes/GameFramework/Pawn.h
--- a/Engine/Source/Runtime/Engine/Classes/GameFramework/Pawn.h
+++ b/Engine/Source/Runtime/Engine/Classes/GameFramework/Pawn.h
@@ -307,9 +307,6 @@ public:
 	/** Tell client that the Pawn is begin restarted. Calls Restart(). */
 	virtual void PawnClientRestart();
 
-	/** Replicated function to set the pawn rotation, allowing the server to force. */
-	virtual void ClientSetRotation( FRotator NewRotation );
-
 	/** Updates Pawn's rotation to the given rotation, assumed to be the Controller's ControlRotation. Respects the bUseControllerRotation* settings. */
 	virtual void FaceRotation(FRotator NewControlRotation, float DeltaTime = 0.f);
 
@@ -450,6 +447,10 @@ public:
 
 	// DEPRECATED FUNCTIONS
 
+	/** Deprecated, misleading name and redundant */
+	DEPRECATED(4.14, "ClientSetRotation is deprecated. Call ClientSetRotation on the PlayerController directly")
+	virtual void ClientSetRotation(FRotator NewRotation);
+
 	/** (Deprecated) Launch Character with LaunchVelocity  */
 	DEPRECATED(4.8, "LaunchPawn is deprecated. For Characters, use LaunchCharacter() instead.")
 	UFUNCTION(BlueprintCallable, Category="Pawn", meta=(DeprecatedFunction, DeprecationMessage="Use Character.LaunchCharacter instead"))
diff --git a/Engine/Source/Runtime/Engine/Classes/GameFramework/PlayerController.h b/Engine/Source/Runtime/Engine/Classes/GameFramework/PlayerController.h
--- a/Engine/Source/Runtime/Engine/Classes/GameFramework/PlayerController.h
+++ b/Engine/Source/Runtime/Engine/Classes/GameFramework/PlayerController.h
@@ -315,7 +315,7 @@ class ENGINE_API APlayerController : public AController
 	/** this is set on the OLD PlayerController when performing a swap over a network connection
 	 * so we know what connection we're waiting on acknowledgment from to finish destroying this PC
 	 * (or when the connection is closed)
-	 * @see GameMode::SwapPlayerControllers()
+	 * @see GameModeBase::SwapPlayerControllers()
 	 */
 	UPROPERTY(DuplicateTransient)
 	class UNetConnection* PendingSwapConnection;
@@ -416,7 +416,7 @@ public:
 	virtual void ClientRepObjRef(UObject* Object);
 
 	/**
-	 * Locally try to pause game (call serverpause to pause network game); returns success indicator.  Calls GameMode's SetPause().
+	 * Locally try to pause game (call serverpause to pause network game); returns success indicator.  Calls GameModeBase's SetPause().
 	 * @return true if succeeded to pause
 	 */
 	virtual bool SetPause(bool bPause, FCanUnpause CanUnpauseDelegate = FCanUnpause());
@@ -542,7 +542,7 @@ public:
 	 * are all automatically moved regardless of whether they're included here
 	 * only dynamic actors in the PersistentLevel may be moved (this includes all actors spawned during gameplay)
 	 * this is called for both parts of the transition because actors might change while in the middle (e.g. players might join or leave the game)
-	 * @see also GameMode::GetSeamlessTravelActorList() (the function that's called on servers)
+	 * @see also GameModeBase::GetSeamlessTravelActorList() (the function that's called on servers)
 	 * @param bToEntry true if we are going from old level -> entry, false if we are going from entry -> new level
 	 * @param ActorList (out) list of actors to maintain
 	 */
@@ -550,16 +550,22 @@ public:
 
 	/** Called when seamless traveling and we are being replaced by the specified PC
 	 * clean up any persistent state (post process chains on LocalPlayers, for example)
-	 * (not called if PlayerControllerClass is the same for the from and to GameModes)
+	 * (not called if PlayerController is the same for the from and to GameModes)
 	 */
 	virtual void SeamlessTravelTo(class APlayerController* NewPC);
 
 	/** Called when seamless traveling and the specified PC is being replaced by this one
 	 * copy over data that should persist
-	 * (not called if PlayerControllerClass is the same for the from and to GameModes)
+	 * (not called if PlayerController is the same for the from and to GameModes)
 	 */
 	virtual void SeamlessTravelFrom(class APlayerController* OldPC);
 
+	/** 
+	 * Called after this player controller has transitioned through seamless travel, but before that player is initialized
+	 * This is called both when a new player controller is created, and when it is maintained
+	 */
+	virtual void PostSeamlessTravel();
+
 	/** 
 	 * Tell the client to enable or disable voice chat (not muting)
 	 * @param bEnable enable or disable voice chat
@@ -1433,7 +1439,7 @@ public:
 	 */
 	virtual void ViewAPlayer(int32 dir);
 
-	/** @return true if this controller thinks it's able to restart. Called from GameMode::PlayerCanRestart */
+	/** @return true if this controller thinks it's able to restart. Called from GameModeBase::PlayerCanRestart */
 	virtual bool CanRestartPlayer();
 
 	/**
@@ -1485,7 +1491,7 @@ public:
 	 *
 	 * @Param GameModeClass - The Class of the game that was replicated
 	 */
-	virtual void ReceivedGameModeClass(TSubclassOf<class AGameMode> GameModeClass);
+	virtual void ReceivedGameModeClass(TSubclassOf<class AGameModeBase> GameModeClass);
 
 	/** Notify the server that client data was received on the Pawn.
 	 * @return true if InPawn is acknowledged on the server, false otherwise. */
@@ -1630,7 +1636,7 @@ public:
 	virtual void SendClientAdjustment();
 
 	/**
-	 * Designate this player controller as local (public for GameMode to use, not expected to be called anywhere else)
+	 * Designate this player controller as local (public for GameModeBase to use, not expected to be called anywhere else)
 	 */
 	void SetAsLocalPlayerController() { bIsLocalPlayerController = true; }
 
@@ -1647,7 +1653,7 @@ public:
 	UPROPERTY()
 	uint16		SeamlessTravelCount;
 
-	/** The value of SeamlessTravelCount, upon the last call to GameMode::HandleSeamlessTravelPlayer; used to detect seamless travel */
+	/** The value of SeamlessTravelCount, upon the last call to GameModeBase::HandleSeamlessTravelPlayer; used to detect seamless travel */
 	UPROPERTY()
 	uint16		LastCompletedSeamlessTravelCount;
 
diff --git a/Engine/Source/Runtime/Engine/Classes/GameFramework/PlayerStart.h b/Engine/Source/Runtime/Engine/Classes/GameFramework/PlayerStart.h
--- a/Engine/Source/Runtime/Engine/Classes/GameFramework/PlayerStart.h
+++ b/Engine/Source/Runtime/Engine/Classes/GameFramework/PlayerStart.h
@@ -12,7 +12,7 @@
 UCLASS(Blueprintable, ClassGroup=Common, hidecategories=Collision)
 class ENGINE_API APlayerStart : public ANavigationObjectBase
 {
-	/** To take more control over PlayerStart selection, you can override the virtual AGameMode::FindPlayerStart and AGameMode::ChoosePlayerStart function */
+	/** To take more control over PlayerStart selection, you can override the virtual AGameModeBase::FindPlayerStart and AGameModeBase::ChoosePlayerStart function */
 
 	GENERATED_UCLASS_BODY()
 
diff --git a/Engine/Source/Runtime/Engine/Classes/GameFramework/WorldSettings.h b/Engine/Source/Runtime/Engine/Classes/GameFramework/WorldSettings.h
--- a/Engine/Source/Runtime/Engine/Classes/GameFramework/WorldSettings.h
+++ b/Engine/Source/Runtime/Engine/Classes/GameFramework/WorldSettings.h
@@ -17,21 +17,6 @@ enum EVisibilityAggressiveness
 	VIS_Max,
 };
 
-/** Helper structure, used to associate GameModes for a map via its filename prefix. */
-USTRUCT()
-struct FGameModePrefix
-{
-	GENERATED_USTRUCT_BODY()
-
-	/** map prefix, e.g. "DM" */
-	UPROPERTY()
-	FString Prefix;
-
-	/** GameMode used if none specified on the URL */
-	UPROPERTY()
-	FString GameMode;
-};
-
 USTRUCT()
 struct FLightmassWorldInfoSettings
 {
@@ -359,11 +344,7 @@ class ENGINE_API AWorldSettings : public AInfo, public IInterface_AssetUserData
 	
 	/** The default GameMode to use when starting this map in the game. If this value is NULL, the INI setting for default game type is used. */
 	UPROPERTY(EditAnywhere, BlueprintReadOnly, Category=GameMode, meta=(DisplayName="GameMode Override"))
-	TSubclassOf<class AGameMode> DefaultGameMode;
-
-	/** Used for loading appropriate game type if non-specified in URL */
-	UPROPERTY(config)
-	TArray<struct FGameModePrefix> DefaultMapPrefixes;
+	TSubclassOf<class AGameModeBase> DefaultGameMode;
 
 	/** Class of GameNetworkManager to spawn for network games */
 	UPROPERTY()
@@ -590,12 +571,12 @@ public:
 	virtual float SetTimeDilation(float NewTimeDilation);
 
 	/**
-	 * Called by GameMode.HandleMatchIsWaitingToStart, calls BeginPlay on all actors
+	 * Called from GameStateBase, calls BeginPlay on all actors
 	 */
 	virtual void NotifyBeginPlay();
 
 	/** 
-	 * Called by GameMode.HandleMatchHasStarted, used to notify native classes of match startup (such as level scripting)
+	 * Called from GameStateBase, used to notify native classes of match startup (such as level scripting)
 	 */	
 	virtual void NotifyMatchStarted();
 
diff --git a/Engine/Source/Runtime/Engine/Classes/Kismet/GameplayStatics.h b/Engine/Source/Runtime/Engine/Classes/Kismet/GameplayStatics.h
--- a/Engine/Source/Runtime/Engine/Classes/Kismet/GameplayStatics.h
+++ b/Engine/Source/Runtime/Engine/Classes/Kismet/GameplayStatics.h
@@ -72,7 +72,7 @@ class ENGINE_API UGameplayStatics : public UBlueprintFunctionLibrary
 	 *	@param	OutActors	Output array of Actors of the specified class.
 	 */
 	UFUNCTION(BlueprintCallable, Category="Utilities",  meta=(WorldContext="WorldContextObject", DeterminesOutputType="ActorClass", DynamicOutputParam="OutActors"))
-	static void GetAllActorsOfClass(UObject* WorldContextObject, TSubclassOf<AActor> ActorClass, TArray<AActor*>& OutActors);
+	static void GetAllActorsOfClass(const UObject* WorldContextObject, TSubclassOf<AActor> ActorClass, TArray<AActor*>& OutActors);
 
 	/** 
 	 *	Find all Actors in the world with the specified interface.
@@ -81,7 +81,7 @@ class ENGINE_API UGameplayStatics : public UBlueprintFunctionLibrary
 	 *	@param	OutActors	Output array of Actors of the specified interface.
 	 */
 	UFUNCTION(BlueprintCallable, Category="Utilities",  meta=(WorldContext="WorldContextObject", DeterminesOutputType="Interface", DynamicOutputParam="OutActors"))
-	static void GetAllActorsWithInterface(UObject* WorldContextObject, TSubclassOf<UInterface> Interface, TArray<AActor*>& OutActors);
+	static void GetAllActorsWithInterface(const UObject* WorldContextObject, TSubclassOf<UInterface> Interface, TArray<AActor*>& OutActors);
 
 	/**
 	*	Find all Actors in the world with the specified tag.
@@ -90,36 +90,36 @@ class ENGINE_API UGameplayStatics : public UBlueprintFunctionLibrary
 	*	@param	OutActors	Output array of Actors of the specified tag.
 	*/
 	UFUNCTION(BlueprintCallable, Category="Utilities",  meta=(WorldContext="WorldContextObject"))
-	static void GetAllActorsWithTag(UObject* WorldContextObject, FName Tag, TArray<AActor*>& OutActors);
+	static void GetAllActorsWithTag(const UObject* WorldContextObject, FName Tag, TArray<AActor*>& OutActors);
 
 	// --- Player functions ------------------------------
 
 	/** Returns the game instance object  */
 	UFUNCTION(BlueprintPure, Category="Game", meta=(WorldContext="WorldContextObject"))
-	static class UGameInstance* GetGameInstance(UObject* WorldContextObject);
+	static class UGameInstance* GetGameInstance(const UObject* WorldContextObject);
 
 	/** Returns the player controller at the specified player index */
 	UFUNCTION(BlueprintPure, Category="Game", meta=(WorldContext="WorldContextObject", UnsafeDuringActorConstruction="true"))
-	static class APlayerController* GetPlayerController(UObject* WorldContextObject, int32 PlayerIndex);
+	static class APlayerController* GetPlayerController(const UObject* WorldContextObject, int32 PlayerIndex);
 
 	/** Returns the player pawn at the specified player index */
 	UFUNCTION(BlueprintPure, Category="Game", meta=(WorldContext="WorldContextObject", UnsafeDuringActorConstruction="true"))
-	static class APawn* GetPlayerPawn(UObject* WorldContextObject, int32 PlayerIndex);
+	static class APawn* GetPlayerPawn(const UObject* WorldContextObject, int32 PlayerIndex);
 
 	/** Returns the player character (NULL if the player pawn doesn't exist OR is not a character) at the specified player index */
 	UFUNCTION(BlueprintPure, Category="Game", meta=(WorldContext="WorldContextObject", UnsafeDuringActorConstruction="true"))
-	static class ACharacter* GetPlayerCharacter(UObject* WorldContextObject, int32 PlayerIndex);
+	static class ACharacter* GetPlayerCharacter(const UObject* WorldContextObject, int32 PlayerIndex);
 
 	/** Returns the player's camera manager for the specified player index */
 	UFUNCTION(BlueprintPure, Category="Game", meta=(WorldContext="WorldContextObject", UnsafeDuringActorConstruction="true"))
-	static class APlayerCameraManager* GetPlayerCameraManager(UObject* WorldContextObject, int32 PlayerIndex);
+	static class APlayerCameraManager* GetPlayerCameraManager(const UObject* WorldContextObject, int32 PlayerIndex);
 
 	/** Create a new player for this game.  
 	 *  @param ControllerId		The ID of the controller that the should control the newly created player.  A value of -1 specifies to use the next available ID
 	 *  @param bSpawnPawn		Whether a pawn should be spawned immediately. If false a pawn will not be created until transition to the next map.
 	 */
 	UFUNCTION(BlueprintCallable, Category="Game", meta=(WorldContext="WorldContextObject", AdvancedDisplay="2", UnsafeDuringActorConstruction="true"))
-	static class APlayerController* CreatePlayer(UObject* WorldContextObject, int32 ControllerId = -1, bool bSpawnPawn = true);
+	static class APlayerController* CreatePlayer(const UObject* WorldContextObject, int32 ControllerId = -1, bool bSpawnPawn = true);
 
 	/** Removes a player from this game.  
 	 *  @param Player			The player controller of the player to be removed
@@ -148,19 +148,19 @@ class ENGINE_API UGameplayStatics : public UBlueprintFunctionLibrary
 	
 	/** Stream the level with the LevelName ; Calling again before it finishes has no effect */
 	UFUNCTION(BlueprintCallable, meta=(WorldContext="WorldContextObject", Latent = "", LatentInfo = "LatentInfo"), Category="Game")
-	static void LoadStreamLevel(UObject* WorldContextObject, FName LevelName, bool bMakeVisibleAfterLoad, bool bShouldBlockOnLoad, FLatentActionInfo LatentInfo);
+	static void LoadStreamLevel(const UObject* WorldContextObject, FName LevelName, bool bMakeVisibleAfterLoad, bool bShouldBlockOnLoad, FLatentActionInfo LatentInfo);
 
 	/** Unload a streamed in level */
 	UFUNCTION(BlueprintCallable, meta=(WorldContext="WorldContextObject", Latent = "", LatentInfo = "LatentInfo"), Category="Game")
-	static void UnloadStreamLevel(UObject* WorldContextObject, FName LevelName, FLatentActionInfo LatentInfo);
+	static void UnloadStreamLevel(const UObject* WorldContextObject, FName LevelName, FLatentActionInfo LatentInfo);
 	
 	/** Returns level streaming object with specified level package name */
 	UFUNCTION(BlueprintPure, meta=(WorldContext="WorldContextObject"), Category="Game")
-	static class ULevelStreaming* GetStreamingLevel(UObject* WorldContextObject, FName PackageName);
+	static class ULevelStreaming* GetStreamingLevel(const UObject* WorldContextObject, FName PackageName);
 
 	/** Flushes level streaming in blocking fashion and returns when all sub-levels are loaded / visible / hidden */
 	UFUNCTION(BlueprintCallable, meta=(WorldContext="WorldContextObject"), Category = "Game")
-	static void FlushLevelStreaming(UObject* WorldContextObject);
+	static void FlushLevelStreaming(const UObject* WorldContextObject);
 		
 	/** Cancels all currently queued streaming packages */
 	UFUNCTION(BlueprintCallable, Category = "Game")
@@ -175,7 +175,7 @@ class ENGINE_API UGameplayStatics : public UBlueprintFunctionLibrary
 	 * @param	Options				a string of options to use for the travel URL
 	 */
 	UFUNCTION(BlueprintCallable, meta=(WorldContext="WorldContextObject", AdvancedDisplay = "2"), Category="Game")
-	static void OpenLevel(UObject* WorldContextObject, FName LevelName, bool bAbsolute = true, FString Options = FString(TEXT("")));
+	static void OpenLevel(const UObject* WorldContextObject, FName LevelName, bool bAbsolute = true, FString Options = FString(TEXT("")));
 
 	/**
 	* Get the name of the currently-open level.
@@ -183,17 +183,19 @@ class ENGINE_API UGameplayStatics : public UBlueprintFunctionLibrary
 	* @param bRemovePrefixString	remove any streaming- or editor- added prefixes from the level name.
 	*/
 	UFUNCTION(BlueprintCallable, meta = (WorldContext = "WorldContextObject", AdvancedDisplay = "1"), Category = "Game")
-	static FString GetCurrentLevelName(UObject* WorldContextObject, bool bRemovePrefixString = true);
+	static FString GetCurrentLevelName(const UObject* WorldContextObject, bool bRemovePrefixString = true);
 
 	// --- Global functions ------------------------------
 
-	/** Returns the current GameMode or NULL if the GameMode can't be retrieved */
-	UFUNCTION(BlueprintPure, Category="Game", meta=(WorldContext="WorldContextObject") )
-	static class AGameMode* GetGameMode(UObject* WorldContextObject);
+	/** Returns the current GameModeBase or Null if it can't be retrieved, such as on the client */
+	UFUNCTION(BlueprintPure, Category="Game", meta=(WorldContext="WorldContextObject"))
+	static class AGameModeBase* GetGameMode(const UObject* WorldContextObject);
 
-	UFUNCTION(BlueprintPure, Category="Game", meta=(WorldContext="WorldContextObject") )
-	static class AGameState* GetGameState(UObject* WorldContextObject);
+	/** Returns the current GameStateBase or Null if it can't be retrieved */
+	UFUNCTION(BlueprintPure, Category="Game", meta=(WorldContext="WorldContextObject"))
+	static class AGameStateBase* GetGameState(const UObject* WorldContextObject);
 
+	/** Returns the class of a passed in Object, will always be valid if Object is not NULL */
 	UFUNCTION(BlueprintPure, meta=(DisplayName = "GetClass", DeterminesOutputType = "Object"), Category="Utilities")
 	static class UClass *GetObjectClass(const UObject *Object);
 
@@ -202,14 +204,14 @@ class ENGINE_API UGameplayStatics : public UBlueprintFunctionLibrary
 	 * @return Current time dilation.
 	 */
 	UFUNCTION(BlueprintPure, Category="Utilities|Time", meta=(WorldContext="WorldContextObject") )
-	static float GetGlobalTimeDilation(UObject* WorldContextObject);
+	static float GetGlobalTimeDilation(const UObject* WorldContextObject);
 
 	/**
 	 * Sets the global time dilation.
 	 * @param	TimeDilation	value to set the global time dilation to
 	 */
 	UFUNCTION(BlueprintCallable, Category="Utilities|Time", meta=(WorldContext="WorldContextObject") )
-	static void SetGlobalTimeDilation(UObject* WorldContextObject, float TimeDilation);
+	static void SetGlobalTimeDilation(const UObject* WorldContextObject, float TimeDilation);
 
 	/**
 	 * Sets the game's paused state
@@ -217,14 +219,14 @@ class ENGINE_API UGameplayStatics : public UBlueprintFunctionLibrary
 	 * @return	Whether the game was successfully paused/unpaused
 	 */
 	UFUNCTION(BlueprintCallable, Category="Game", meta=(WorldContext="WorldContextObject") )
-	static bool SetGamePaused(UObject* WorldContextObject, bool bPaused);
+	static bool SetGamePaused(const UObject* WorldContextObject, bool bPaused);
 
 	/**
 	 * Returns the game's paused state
 	 * @return	Whether the game is currently paused or not
 	 */
 	UFUNCTION(BlueprintPure, Category="Game", meta=(WorldContext="WorldContextObject") )
-	static bool IsGamePaused(UObject* WorldContextObject);
+	static bool IsGamePaused(const UObject* WorldContextObject);
 
 
 	/** Hurt locally authoritative actors within the radius. Will only hit components that block the Visibility channel.
@@ -239,7 +241,7 @@ class ENGINE_API UGameplayStatics : public UBlueprintFunctionLibrary
 	 * @return true if damage was applied to at least one actor.
 	 */
 	UFUNCTION(BlueprintCallable, BlueprintAuthorityOnly, Category="Game|Damage", meta=(WorldContext="WorldContextObject", AutoCreateRefTerm="IgnoreActors"))
-	static bool ApplyRadialDamage(UObject* WorldContextObject, float BaseDamage, const FVector& Origin, float DamageRadius, TSubclassOf<class UDamageType> DamageTypeClass, const TArray<AActor*>& IgnoreActors, AActor* DamageCauser = NULL, AController* InstigatedByController = NULL, bool bDoFullDamage = false, ECollisionChannel DamagePreventionChannel = ECC_Visibility);
+	static bool ApplyRadialDamage(const UObject* WorldContextObject, float BaseDamage, const FVector& Origin, float DamageRadius, TSubclassOf<class UDamageType> DamageTypeClass, const TArray<AActor*>& IgnoreActors, AActor* DamageCauser = NULL, AController* InstigatedByController = NULL, bool bDoFullDamage = false, ECollisionChannel DamagePreventionChannel = ECC_Visibility);
 	
 	/** Hurt locally authoritative actors within the radius. Will only hit components that block the Visibility channel.
 	 * @param BaseDamage - The base damage to apply, i.e. the damage at the origin.
@@ -255,7 +257,7 @@ class ENGINE_API UGameplayStatics : public UBlueprintFunctionLibrary
 	 * @return true if damage was applied to at least one actor.
 	 */
 	UFUNCTION(BlueprintCallable, BlueprintAuthorityOnly, Category="Game|Damage", meta=(WorldContext="WorldContextObject", AutoCreateRefTerm="IgnoreActors"))
-	static bool ApplyRadialDamageWithFalloff(UObject* WorldContextObject, float BaseDamage, float MinimumDamage, const FVector& Origin, float DamageInnerRadius, float DamageOuterRadius, float DamageFalloff, TSubclassOf<class UDamageType> DamageTypeClass, const TArray<AActor*>& IgnoreActors, AActor* DamageCauser = NULL, AController* InstigatedByController = NULL, ECollisionChannel DamagePreventionChannel = ECC_Visibility);
+	static bool ApplyRadialDamageWithFalloff(const UObject* WorldContextObject, float BaseDamage, float MinimumDamage, const FVector& Origin, float DamageInnerRadius, float DamageOuterRadius, float DamageFalloff, TSubclassOf<class UDamageType> DamageTypeClass, const TArray<AActor*>& IgnoreActors, AActor* DamageCauser = NULL, AController* InstigatedByController = NULL, ECollisionChannel DamagePreventionChannel = ECC_Visibility);
 	
 
 	/** Hurts the specified actor with the specified impact.
@@ -292,7 +294,7 @@ class ENGINE_API UGameplayStatics : public UBlueprintFunctionLibrary
 	 * @param bOrientShakeTowardsEpicenter - Changes the rotation of shake to point towards epicenter instead of forward
 	 */
 	UFUNCTION(BlueprintCallable, Category="Game|Feedback", meta=(WorldContext="WorldContextObject", UnsafeDuringActorConstruction = "true"))
-	static void PlayWorldCameraShake(UObject* WorldContextObject, TSubclassOf<class UCameraShake> Shake, FVector Epicenter, float InnerRadius, float OuterRadius, float Falloff = 1.f, bool bOrientShakeTowardsEpicenter = false);
+	static void PlayWorldCameraShake(const UObject* WorldContextObject, TSubclassOf<class UCameraShake> Shake, FVector Epicenter, float InnerRadius, float OuterRadius, float Falloff = 1.f, bool bOrientShakeTowardsEpicenter = false);
 
 	// --- Particle functions ------------------------------
 
@@ -304,7 +306,7 @@ class ENGINE_API UGameplayStatics : public UBlueprintFunctionLibrary
 	 * @param bAutoDestroy - Whether the component will automatically be destroyed when the particle system completes playing or whether it can be reactivated
 	 */
 	UFUNCTION(BlueprintCallable, Category="Effects|Components|ParticleSystem", meta=(Keywords = "particle system", WorldContext="WorldContextObject", UnsafeDuringActorConstruction = "true"))
-	static UParticleSystemComponent* SpawnEmitterAtLocation(UObject* WorldContextObject, class UParticleSystem* EmitterTemplate, FVector Location, FRotator Rotation = FRotator::ZeroRotator, bool bAutoDestroy = true);
+	static UParticleSystemComponent* SpawnEmitterAtLocation(const UObject* WorldContextObject, class UParticleSystem* EmitterTemplate, FVector Location, FRotator Rotation = FRotator::ZeroRotator, bool bAutoDestroy = true);
 
 	/** Plays the specified effect at the given location and rotation, fire and forget. The system will go away when the effect is complete. Does not replicate.
 	 * @param World - The World to spawn in
@@ -318,8 +320,8 @@ class ENGINE_API UGameplayStatics : public UBlueprintFunctionLibrary
 	* @param EmitterTemplate - particle system to create
 	 * @param AttachComponent - Component to attach to.
 	 * @param AttachPointName - Optional named point within the AttachComponent to spawn the emitter at
-	 * @param Location - Depending on the value of Location Type this is either a relative offset from the attach component/point or an absolute world position that will be translated to a relative offset
-	 * @param Rotation - Depending on the value of LocationType this is either a relative offset from the attach component/point or an absolute world rotation that will be translated to a realative offset
+	 * @param Location - Depending on the value of LocationType this is either a relative offset from the attach component/point or an absolute world location that will be translated to a relative offset (if LocationType is KeepWorldPosition).
+	 * @param Rotation - Depending on the value of LocationType this is either a relative offset from the attach component/point or an absolute world rotation that will be translated to a relative offset (if LocationType is KeepWorldPosition).
 	 * @param LocationType - Specifies whether Location is a relative offset or an absolute world position
 	 * @param bAutoDestroy - Whether the component will automatically be destroyed when the particle system completes playing or whether it can be reactivated
 	 */
@@ -335,7 +337,7 @@ class ENGINE_API UGameplayStatics : public UBlueprintFunctionLibrary
 	 * @note This will always return false if there is no audio device, or the audio device is disabled.
 	 */
 	UFUNCTION(BlueprintCallable, Category = "Audio", meta = (WorldContext = "WorldContextObject"))
-	static bool AreAnyListenersWithinRange(UObject* WorldContextObject, FVector Location, float MaximumRange);
+	static bool AreAnyListenersWithinRange(const UObject* WorldContextObject, FVector Location, float MaximumRange);
 	
 
 	/**
@@ -347,7 +349,7 @@ class ENGINE_API UGameplayStatics : public UBlueprintFunctionLibrary
 	* @param TimeSec - A time value to linearly interpolate the global modulation pitch over from it's current value.
 	*/
 	UFUNCTION(BlueprintCallable, BlueprintCosmetic, Category = "Audio", meta = (WorldContext = "WorldContextObject"))
-	static void SetGlobalPitchModulation(UObject* WorldContextObject, float PitchModulation, float TimeSec);
+	static void SetGlobalPitchModulation(const UObject* WorldContextObject, float PitchModulation, float TimeSec);
 
 	/**
 	* Sets the global listener focus parameters which will scale focus behavior of sounds based on their focus azimuth settings in their attenuation settings. 
@@ -364,7 +366,7 @@ class ENGINE_API UGameplayStatics : public UBlueprintFunctionLibrary
 	* @param NonFocusPriorityScale - A priority scale value (> 0.0) to use for sounds which are out-of-focus. Values < 1.0 will reduce the priority of sounds out-of-focus sounds, values > 1.0 will increase the priority of out-of-focus sounds.
 	*/
 	UFUNCTION(BlueprintCallable, BlueprintCosmetic, Category = "Audio", meta = (WorldContext = "WorldContextObject"))
-	static void SetGlobalListenerFocusParameters(UObject* WorldContextObject, float FocusAzimuthScale = 1.0f, float NonFocusAzimuthScale = 1.0f, float FocusDistanceScale = 1.0f, float NonFocusDistanceScale = 1.0f, float FocusVolumeScale = 1.0f, float NonFocusVolumeScale = 1.0f, float FocusPriorityScale = 1.0f, float NonFocusPriorityScale = 1.0f);
+	static void SetGlobalListenerFocusParameters(const UObject* WorldContextObject, float FocusAzimuthScale = 1.0f, float NonFocusAzimuthScale = 1.0f, float FocusDistanceScale = 1.0f, float NonFocusDistanceScale = 1.0f, float FocusVolumeScale = 1.0f, float NonFocusVolumeScale = 1.0f, float FocusPriorityScale = 1.0f, float NonFocusPriorityScale = 1.0f);
 
 	/**
 	 * Plays a sound directly with no attenuation, perfect for UI sounds.
@@ -378,7 +380,7 @@ class ENGINE_API UGameplayStatics : public UBlueprintFunctionLibrary
 	 * @param StartTime - How far in to the sound to begin playback at
 	 */
 	UFUNCTION(BlueprintCallable, BlueprintCosmetic, Category="Audio", meta=( WorldContext="WorldContextObject", AdvancedDisplay = "2", UnsafeDuringActorConstruction = "true" ))
-	static void PlaySound2D(UObject* WorldContextObject, class USoundBase* Sound, float VolumeMultiplier = 1.f, float PitchMultiplier = 1.f, float StartTime = 0.f, class USoundConcurrency* ConcurrencySettings = nullptr);
+	static void PlaySound2D(const UObject* WorldContextObject, class USoundBase* Sound, float VolumeMultiplier = 1.f, float PitchMultiplier = 1.f, float StartTime = 0.f, class USoundConcurrency* ConcurrencySettings = nullptr);
 
 	/**
 	 * Spawns a sound with no attenuation, perfect for UI sounds.
@@ -392,7 +394,7 @@ class ENGINE_API UGameplayStatics : public UBlueprintFunctionLibrary
 	 * @return An audio component to manipulate the spawned sound
 	 */
 	UFUNCTION(BlueprintCallable, BlueprintCosmetic, Category="Audio", meta=( WorldContext="WorldContextObject", AdvancedDisplay = "2", UnsafeDuringActorConstruction = "true", Keywords = "play" ))
-	static UAudioComponent* SpawnSound2D(UObject* WorldContextObject, class USoundBase* Sound, float VolumeMultiplier = 1.f, float PitchMultiplier = 1.f, float StartTime = 0.f, class USoundConcurrency* ConcurrencySettings = nullptr);
+	static UAudioComponent* SpawnSound2D(const UObject* WorldContextObject, class USoundBase* Sound, float VolumeMultiplier = 1.f, float PitchMultiplier = 1.f, float StartTime = 0.f, class USoundConcurrency* ConcurrencySettings = nullptr, bool bPersistAcrossLevelTransition = false);
 
 	/**
 	 * Creates a sound with no attenuation, perfect for UI sounds. This does NOT play the sound
@@ -405,7 +407,7 @@ class ENGINE_API UGameplayStatics : public UBlueprintFunctionLibrary
 	 * @return An audio component to manipulate the spawned sound
 	 */
 	UFUNCTION(BlueprintCallable, BlueprintCosmetic, Category="Audio", meta=( WorldContext="WorldContextObject", AdvancedDisplay = "2", UnsafeDuringActorConstruction = "true", Keywords = "play" ))
-		static UAudioComponent* CreateSound2D(UObject* WorldContextObject, class USoundBase* Sound, float VolumeMultiplier = 1.f, float PitchMultiplier = 1.f, float StartTime = 0.f, class USoundConcurrency* ConcurrencySettings = nullptr);
+	static UAudioComponent* CreateSound2D(const UObject* WorldContextObject, class USoundBase* Sound, float VolumeMultiplier = 1.f, float PitchMultiplier = 1.f, float StartTime = 0.f, class USoundConcurrency* ConcurrencySettings = nullptr, bool bPersistAcrossLevelTransition = false);
 
 	/**
 	 * Plays a sound at the given location. This is a fire and forget sound and does not travel with any actor. Replication is also not handled at this point.
@@ -419,9 +421,9 @@ class ENGINE_API UGameplayStatics : public UBlueprintFunctionLibrary
 	 * @param ConcurrencySettings - Override concurrency settings package to play sound with
 	 */
 	UFUNCTION(BlueprintCallable, Category="Audio", meta=(WorldContext="WorldContextObject", AdvancedDisplay = "3", UnsafeDuringActorConstruction = "true"))
-	static void PlaySoundAtLocation(UObject* WorldContextObject, class USoundBase* Sound, FVector Location, FRotator Rotation, float VolumeMultiplier = 1.f, float PitchMultiplier = 1.f, float StartTime = 0.f, class USoundAttenuation* AttenuationSettings = nullptr, class USoundConcurrency* ConcurrencySettings = nullptr);
+	static void PlaySoundAtLocation(const UObject* WorldContextObject, class USoundBase* Sound, FVector Location, FRotator Rotation, float VolumeMultiplier = 1.f, float PitchMultiplier = 1.f, float StartTime = 0.f, class USoundAttenuation* AttenuationSettings = nullptr, class USoundConcurrency* ConcurrencySettings = nullptr);
 
-	static void PlaySoundAtLocation(UObject* WorldContextObject, class USoundBase* Sound, FVector Location, float VolumeMultiplier = 1.f, float PitchMultiplier = 1.f, float StartTime = 0.f, class USoundAttenuation* AttenuationSettings = nullptr, class USoundConcurrency* ConcurrencySettings = nullptr)
+	static void PlaySoundAtLocation(const UObject* WorldContextObject, class USoundBase* Sound, FVector Location, float VolumeMultiplier = 1.f, float PitchMultiplier = 1.f, float StartTime = 0.f, class USoundAttenuation* AttenuationSettings = nullptr, class USoundConcurrency* ConcurrencySettings = nullptr)
 	{
 		PlaySoundAtLocation(WorldContextObject, Sound, Location, FRotator::ZeroRotator, VolumeMultiplier, PitchMultiplier, StartTime, AttenuationSettings, ConcurrencySettings);
 	}
@@ -439,7 +441,7 @@ class ENGINE_API UGameplayStatics : public UBlueprintFunctionLibrary
 	 * @return An audio component to manipulate the spawned sound
 	 */
 	UFUNCTION(BlueprintCallable, Category="Audio", meta=(WorldContext="WorldContextObject", AdvancedDisplay = "3", UnsafeDuringActorConstruction = "true", Keywords = "play"))
-	static UAudioComponent* SpawnSoundAtLocation(UObject* WorldContextObject, class USoundBase* Sound, FVector Location, FRotator Rotation = FRotator::ZeroRotator, float VolumeMultiplier = 1.f, float PitchMultiplier = 1.f, float StartTime = 0.f, class USoundAttenuation* AttenuationSettings = nullptr, class USoundConcurrency* ConcurrencySettings = nullptr);
+	static UAudioComponent* SpawnSoundAtLocation(const UObject* WorldContextObject, class USoundBase* Sound, FVector Location, FRotator Rotation = FRotator::ZeroRotator, float VolumeMultiplier = 1.f, float PitchMultiplier = 1.f, float StartTime = 0.f, class USoundAttenuation* AttenuationSettings = nullptr, class USoundConcurrency* ConcurrencySettings = nullptr);
 
 	/** Plays a sound attached to and following the specified component. This is a fire and forget sound. Replication is also not handled at this point.
 	 * @param Sound - sound to play
@@ -482,7 +484,7 @@ class ENGINE_API UGameplayStatics : public UBlueprintFunctionLibrary
 	 * @param StartTime - How far in to the dialogue to begin playback at
 	 */
 	UFUNCTION(BlueprintCallable, BlueprintCosmetic, Category="Audio", meta=( WorldContext="WorldContextObject", AdvancedDisplay = "3", UnsafeDuringActorConstruction = "true" ))
-	static void PlayDialogue2D(UObject* WorldContextObject, class UDialogueWave* Dialogue, const struct FDialogueContext& Context, float VolumeMultiplier = 1.f, float PitchMultiplier = 1.f, float StartTime = 0.f);
+	static void PlayDialogue2D(const UObject* WorldContextObject, class UDialogueWave* Dialogue, const struct FDialogueContext& Context, float VolumeMultiplier = 1.f, float PitchMultiplier = 1.f, float StartTime = 0.f);
 
 	/**
 	 * Spawns a dialogue with no attenuation, perfect for UI.
@@ -496,7 +498,7 @@ class ENGINE_API UGameplayStatics : public UBlueprintFunctionLibrary
 	 * @return An audio component to manipulate the spawned sound
 	 */
 	UFUNCTION(BlueprintCallable, BlueprintCosmetic, Category="Audio", meta=( WorldContext="WorldContextObject", AdvancedDisplay = "3", UnsafeDuringActorConstruction = "true", Keywords = "play" ))
-	static UAudioComponent* SpawnDialogue2D(UObject* WorldContextObject, class UDialogueWave* Dialogue, const struct FDialogueContext& Context, float VolumeMultiplier = 1.f, float PitchMultiplier = 1.f, float StartTime = 0.f);
+	static UAudioComponent* SpawnDialogue2D(const UObject* WorldContextObject, class UDialogueWave* Dialogue, const struct FDialogueContext& Context, float VolumeMultiplier = 1.f, float PitchMultiplier = 1.f, float StartTime = 0.f);
 
 	/** Plays a dialogue at the given location. This is a fire and forget sound and does not travel with any actor. Replication is also not handled at this point.
 	 * @param Dialogue - dialogue to play
@@ -509,9 +511,9 @@ class ENGINE_API UGameplayStatics : public UBlueprintFunctionLibrary
 	 * @param AttenuationSettings - Override attenuation settings package to play sound with
 	 */
 	UFUNCTION(BlueprintCallable, Category="Audio", meta=(WorldContext="WorldContextObject", AdvancedDisplay = "4", UnsafeDuringActorConstruction = "true"))
-	static void PlayDialogueAtLocation(UObject* WorldContextObject, class UDialogueWave* Dialogue, const struct FDialogueContext& Context, FVector Location, FRotator Rotation, float VolumeMultiplier = 1.f, float PitchMultiplier = 1.f, float StartTime = 0.f, class USoundAttenuation* AttenuationSettings = NULL);
+	static void PlayDialogueAtLocation(const UObject* WorldContextObject, class UDialogueWave* Dialogue, const struct FDialogueContext& Context, FVector Location, FRotator Rotation, float VolumeMultiplier = 1.f, float PitchMultiplier = 1.f, float StartTime = 0.f, class USoundAttenuation* AttenuationSettings = NULL);
 
-	static void PlayDialogueAtLocation(UObject* WorldContextObject, class UDialogueWave* Dialogue, const struct FDialogueContext& Context, FVector Location, float VolumeMultiplier = 1.f, float PitchMultiplier = 1.f, float StartTime = 0.f, class USoundAttenuation* AttenuationSettings = NULL)
+	static void PlayDialogueAtLocation(const UObject* WorldContextObject, class UDialogueWave* Dialogue, const struct FDialogueContext& Context, FVector Location, float VolumeMultiplier = 1.f, float PitchMultiplier = 1.f, float StartTime = 0.f, class USoundAttenuation* AttenuationSettings = NULL)
 	{
 		PlayDialogueAtLocation(WorldContextObject, Dialogue, Context, Location, FRotator::ZeroRotator, VolumeMultiplier, PitchMultiplier, StartTime, AttenuationSettings);
 	}
@@ -528,7 +530,7 @@ class ENGINE_API UGameplayStatics : public UBlueprintFunctionLibrary
 	 * @return Audio Component to manipulate the playing dialogue with
 	 */
 	UFUNCTION(BlueprintCallable, Category="Audio", meta=(WorldContext="WorldContextObject", AdvancedDisplay = "4", UnsafeDuringActorConstruction = "true", Keywords = "play"))
-	static UAudioComponent* SpawnDialogueAtLocation(UObject* WorldContextObject, class UDialogueWave* Dialogue, const struct FDialogueContext& Context, FVector Location, FRotator Rotation = FRotator::ZeroRotator, float VolumeMultiplier = 1.f, float PitchMultiplier = 1.f, float StartTime = 0.f, class USoundAttenuation* AttenuationSettings = NULL);
+	static UAudioComponent* SpawnDialogueAtLocation(const UObject* WorldContextObject, class UDialogueWave* Dialogue, const struct FDialogueContext& Context, FVector Location, FRotator Rotation = FRotator::ZeroRotator, float VolumeMultiplier = 1.f, float PitchMultiplier = 1.f, float StartTime = 0.f, class USoundAttenuation* AttenuationSettings = NULL);
 
 	/** Plays a dialogue attached to and following the specified component. This is a fire and forget sound. Replication is also not handled at this point.
 	 * @param Dialogue - dialogue to play
@@ -576,7 +578,7 @@ class ENGINE_API UGameplayStatics : public UBlueprintFunctionLibrary
 	// --- Audio Functions ----------------------------
 	/** Set the sound mix of the audio system for special EQing **/
 	UFUNCTION(BlueprintCallable, Category="Audio", meta=(WorldContext = "WorldContextObject"))
-	static void SetBaseSoundMix(UObject* WorldContextObject, class USoundMix* InSoundMix);
+	static void SetBaseSoundMix(const UObject* WorldContextObject, class USoundMix* InSoundMix);
 
 	/** Overrides the sound class adjuster in the given sound mix. If the sound class does not exist in the input sound mix, the sound class adjustment will be added to the sound mix.
 	 * @param InSoundMixModifier The sound mix to modify.
@@ -587,7 +589,7 @@ class ENGINE_API UGameplayStatics : public UBlueprintFunctionLibrary
 	 * @param bApplyToChildren Whether or not to apply this override to the sound class' children or to just the specified sound class.
 	 */
 	UFUNCTION(BlueprintCallable, Category = "Audio", meta = (WorldContext = "WorldContextObject"))
-	static void SetSoundMixClassOverride(UObject* WorldContextObject, class USoundMix* InSoundMixModifier, class USoundClass* InSoundClass, float Volume = 1.0f, float Pitch = 1.0f, float FadeInTime = 1.0f, bool bApplyToChildren = true);
+	static void SetSoundMixClassOverride(const UObject* WorldContextObject, class USoundMix* InSoundMixModifier, class USoundClass* InSoundClass, float Volume = 1.0f, float Pitch = 1.0f, float FadeInTime = 1.0f, bool bApplyToChildren = true);
 
 	/** Clears the override of the sound class adjuster in the given sound mix. If the override did not exist in the sound mix, this will do nothing.
 	 * @param InSoundMixModifier The sound mix to modify.
@@ -595,19 +597,19 @@ class ENGINE_API UGameplayStatics : public UBlueprintFunctionLibrary
 	 * @param FadeOutTime The interpolation time to use to go from the current sound class adjuster override values to the non-override values.
 	*/
 	UFUNCTION(BlueprintCallable, Category = "Audio", meta = (WorldContext = "WorldContextObject"))
-	static void ClearSoundMixClassOverride(UObject* WorldContextObject, class USoundMix* InSoundMixModifier, class USoundClass* InSoundClass, float FadeOutTime = 1.0f);
+	static void ClearSoundMixClassOverride(const UObject* WorldContextObject, class USoundMix* InSoundMixModifier, class USoundClass* InSoundClass, float FadeOutTime = 1.0f);
 
 	/** Push a sound mix modifier onto the audio system **/
 	UFUNCTION(BlueprintCallable, Category="Audio", meta=(WorldContext = "WorldContextObject"))
-	static void PushSoundMixModifier(UObject* WorldContextObject, class USoundMix* InSoundMixModifier);
+	static void PushSoundMixModifier(const UObject* WorldContextObject, class USoundMix* InSoundMixModifier);
 
 	/** Pop a sound mix modifier from the audio system **/
 	UFUNCTION(BlueprintCallable, Category="Audio", meta=(WorldContext = "WorldContextObject"))
-	static void PopSoundMixModifier(UObject* WorldContextObject, class USoundMix* InSoundMixModifier);
+	static void PopSoundMixModifier(const UObject* WorldContextObject, class USoundMix* InSoundMixModifier);
 
 	/** Clear all sound mix modifiers from the audio system **/
 	UFUNCTION(BlueprintCallable, Category="Audio", meta=(WorldContext = "WorldContextObject"))
-	static void ClearSoundMixModifiers(UObject* WorldContextObject);
+	static void ClearSoundMixModifiers(const UObject* WorldContextObject);
 
 	/** Activates a Reverb Effect without the need for a volume
 	 * @param ReverbEffect Reverb Effect to use
@@ -617,7 +619,7 @@ class ENGINE_API UGameplayStatics : public UBlueprintFunctionLibrary
 	 * @param FadeTime Time before Reverb Effect is fully active
 	 */
 	UFUNCTION(BlueprintCallable, Category="Audio", meta=(WorldContext = "WorldContextObject", AdvancedDisplay = "2"))
-	static void ActivateReverbEffect(UObject* WorldContextObject, class UReverbEffect* ReverbEffect, FName TagName, float Priority = 0.f, float Volume = 0.5f, float FadeTime = 2.f);
+	static void ActivateReverbEffect(const UObject* WorldContextObject, class UReverbEffect* ReverbEffect, FName TagName, float Priority = 0.f, float Volume = 0.5f, float FadeTime = 2.f);
 
 	/**
 	 * Deactivates a Reverb Effect not applied by a volume
@@ -625,13 +627,13 @@ class ENGINE_API UGameplayStatics : public UBlueprintFunctionLibrary
 	 * @param TagName Tag associated with Reverb Effect to remove
 	 */
 	UFUNCTION(BlueprintCallable, Category="Audio", meta=(WorldContext = "WorldContextObject"))
-	static void DeactivateReverbEffect(UObject* WorldContextObject, FName TagName);
+	static void DeactivateReverbEffect(const UObject* WorldContextObject, FName TagName);
 
 	/** 
 	 * Returns the highest priority reverb settings currently active from any source (volumes or manual setting).
 	 */
 	UFUNCTION(BlueprintCallable, Category = "Audio", meta = (WorldContext = "WorldContextObject"))
-	static class UReverbEffect* GetCurrentReverbEffect(UObject* WorldContextObject);
+	static class UReverbEffect* GetCurrentReverbEffect(const UObject* WorldContextObject);
 
 	// --- Decal functions ------------------------------
 
@@ -643,7 +645,7 @@ class ENGINE_API UGameplayStatics : public UBlueprintFunctionLibrary
 	 * @param LifeSpan - destroy decal component after time runs out (0 = infinite)
 	 */
 	UFUNCTION(BlueprintCallable, Category="Rendering|Components|Decal", meta=(WorldContext="WorldContextObject", UnsafeDuringActorConstruction = "true"))
-	static UDecalComponent* SpawnDecalAtLocation(UObject* WorldContextObject, class UMaterialInterface* DecalMaterial, FVector DecalSize, FVector Location, FRotator Rotation = FRotator(-90, 0, 0), float LifeSpan = 0);
+	static UDecalComponent* SpawnDecalAtLocation(const UObject* WorldContextObject, class UMaterialInterface* DecalMaterial, FVector DecalSize, FVector Location, FRotator Rotation = FRotator(-90, 0, 0), float LifeSpan = 0);
 
 	/** Spawns a decal attached to and following the specified component. Does not replicate.
 	 * @param DecalMaterial - decal's material
@@ -766,18 +768,26 @@ class ENGINE_API UGameplayStatics : public UBlueprintFunctionLibrary
 
 	/** Returns the frame delta time in seconds, adjusted by time dilation. */
 	UFUNCTION(BlueprintPure, Category = "Utilities|Time", meta = (WorldContext="WorldContextObject"))
-	static float GetWorldDeltaSeconds(UObject* WorldContextObject);
+	static float GetWorldDeltaSeconds(const UObject* WorldContextObject);
+
+	/** Returns time in seconds since world was brought up for play, adjusted by time dilation and IS stopped when game pauses */
+	UFUNCTION(BlueprintPure, Category="Utilities|Time", meta=(WorldContext="WorldContextObject"))
+	static float GetTimeSeconds(const UObject* WorldContextObject);
+
+	/** Returns time in seconds since world was brought up for play, adjusted by time dilation and IS NOT stopped when game pauses */
+	UFUNCTION(BlueprintPure, Category="Utilities|Time", meta=(WorldContext="WorldContextObject"))
+	static float GetUnpausedTimeSeconds(const UObject* WorldContextObject);
 
 	/** Returns time in seconds since world was brought up for play, does NOT stop when game pauses, NOT dilated/clamped */
 	UFUNCTION(BlueprintPure, Category="Utilities|Time", meta=(WorldContext="WorldContextObject"))
-	static float GetRealTimeSeconds(UObject* WorldContextObject);
+	static float GetRealTimeSeconds(const UObject* WorldContextObject);
 	
 	/** Returns time in seconds since world was brought up for play, IS stopped when game pauses, NOT dilated/clamped. */
 	UFUNCTION(BlueprintPure, Category="Utilities|Time", meta=(WorldContext="WorldContextObject"))
-	static float GetAudioTimeSeconds(UObject* WorldContextObject);
+	static float GetAudioTimeSeconds(const UObject* WorldContextObject);
 
 	UFUNCTION(BlueprintPure, Category="Utilities|Time", meta=(WorldContext="WorldContextObject"))
-	static void GetAccurateRealTime(UObject* WorldContextObject, int32& Seconds, float& PartialSeconds);
+	static void GetAccurateRealTime(const UObject* WorldContextObject, int32& Seconds, float& PartialSeconds);
 
 	/*~ DVRStreaming API */
 	
@@ -808,10 +818,10 @@ class ENGINE_API UGameplayStatics : public UBlueprintFunctionLibrary
 	 * @return					Returns false if there is no valid solution or the valid solutions are blocked.  Returns true otherwise.
 	 */
 	UFUNCTION(BlueprintCallable, Category="Game|Components|ProjectileMovement", DisplayName="SuggestProjectileVelocity", meta=(WorldContext="WorldContextObject"))
-	static bool BlueprintSuggestProjectileVelocity(UObject* WorldContextObject, FVector& TossVelocity, FVector StartLocation, FVector EndLocation, float LaunchSpeed, float OverrideGravityZ, ESuggestProjVelocityTraceOption::Type TraceOption, float CollisionRadius, bool bFavorHighArc, bool bDrawDebug);
+	static bool BlueprintSuggestProjectileVelocity(const UObject* WorldContextObject, FVector& TossVelocity, FVector StartLocation, FVector EndLocation, float LaunchSpeed, float OverrideGravityZ, ESuggestProjVelocityTraceOption::Type TraceOption, float CollisionRadius, bool bFavorHighArc, bool bDrawDebug);
 
 	/** Native version, has more options than the Blueprint version. */
-	static bool SuggestProjectileVelocity(UObject* WorldContextObject, FVector& TossVelocity, FVector StartLocation, FVector EndLocation, float TossSpeed, bool bHighArc = false, float CollisionRadius = 0.f, float OverrideGravityZ = 0, ESuggestProjVelocityTraceOption::Type TraceOption = ESuggestProjVelocityTraceOption::TraceFullPath, const FCollisionResponseParams& ResponseParam = FCollisionResponseParams::DefaultResponseParam, const TArray<AActor*>& ActorsToIgnore = TArray<AActor*>(), bool bDrawDebug = false);
+	static bool SuggestProjectileVelocity(const UObject* WorldContextObject, FVector& TossVelocity, FVector StartLocation, FVector EndLocation, float TossSpeed, bool bHighArc = false, float CollisionRadius = 0.f, float OverrideGravityZ = 0, ESuggestProjVelocityTraceOption::Type TraceOption = ESuggestProjVelocityTraceOption::TraceFullPath, const FCollisionResponseParams& ResponseParam = FCollisionResponseParams::DefaultResponseParam, const TArray<AActor*>& ActorsToIgnore = TArray<AActor*>(), bool bDrawDebug = false);
 
 	/**
 	* Predict the arc of a virtual projectile affected by gravity with collision checks along the arc. Returns a list of positions of the simulated arc and the destination reached by the simulation.
@@ -834,7 +844,7 @@ class ENGINE_API UGameplayStatics : public UBlueprintFunctionLibrary
 	* @param OverrideGravityZ			Optional override of Gravity (default uses WorldGravityZ
 	*/
 	UFUNCTION(BlueprintCallable, Category = "Game", meta = (WorldContext = "WorldContextObject", AutoCreateRefTerm = "ActorsToIgnore", AdvancedDisplay = "DrawDebugTime, DrawDebugType, SimFrequency, MaxSimTime, OverrideGravityZ", TraceChannel = ECC_WorldDynamic, bTracePath = true))
-	static bool PredictProjectilePath(UObject* WorldContextObject, FHitResult& OutHit, TArray<FVector>& OutPathPositions, FVector& OutLastTraceDestination, FVector StartPos, FVector LaunchVelocity, bool bTracePath, float ProjectileRadius, const TArray<TEnumAsByte<EObjectTypeQuery> >& ObjectTypes, bool bTraceComplex, const TArray<AActor*>& ActorsToIgnore, EDrawDebugTrace::Type DrawDebugType, float DrawDebugTime, float SimFrequency = 30.f, float MaxSimTime = 2.f, float OverrideGravityZ = 0);
+	static bool PredictProjectilePath(const UObject* WorldContextObject, FHitResult& OutHit, TArray<FVector>& OutPathPositions, FVector& OutLastTraceDestination, FVector StartPos, FVector LaunchVelocity, bool bTracePath, float ProjectileRadius, const TArray<TEnumAsByte<EObjectTypeQuery> >& ObjectTypes, bool bTraceComplex, const TArray<AActor*>& ActorsToIgnore, EDrawDebugTrace::Type DrawDebugType, float DrawDebugTime, float SimFrequency = 30.f, float MaxSimTime = 2.f, float OverrideGravityZ = 0);
 
 	/**
 	* Returns the launch velocity needed for a projectile at rest at StartPos to land on EndPos.
@@ -848,15 +858,15 @@ class ENGINE_API UGameplayStatics : public UBlueprintFunctionLibrary
 	* @param ArcParam					Change height of arc between 0.0-1.0 where 0.5 is the default medium arc
 	*/
 	UFUNCTION(BlueprintCallable, Category = "Game", DisplayName = "SuggestProjectileVelocity Custom Arc", meta = (WorldContext = "WorldContextObject", AdvancedDisplay = "OverrideGravityZ, ArcParam"))
-	static bool SuggestProjectileVelocity_CustomArc(UObject* WorldContextObject, FVector& OutLaunchVelocity, FVector StartPos, FVector EndPos, float OverrideGravityZ = 0, float ArcParam = 0.5f);
+	static bool SuggestProjectileVelocity_CustomArc(const UObject* WorldContextObject, FVector& OutLaunchVelocity, FVector StartPos, FVector EndPos, float OverrideGravityZ = 0, float ArcParam = 0.5f);
 
 	/** Returns world origin current location. */
 	UFUNCTION(BlueprintPure, Category="Game", meta=(WorldContext="WorldContextObject") )
-	static FIntVector GetWorldOriginLocation(UObject* WorldContextObject);
+	static FIntVector GetWorldOriginLocation(const UObject* WorldContextObject);
 	
 	/** Requests a new location for a world origin. */
 	UFUNCTION(BlueprintCallable, Category="Game", meta=(WorldContext="WorldContextObject"))
-	static void SetWorldOriginLocation(UObject* WorldContextObject, FIntVector NewLocation);
+	static void SetWorldOriginLocation(const UObject* WorldContextObject, FIntVector NewLocation);
 
 	/**
 	* Counts how many grass foliage instances overlap a given sphere.
@@ -868,7 +878,7 @@ class ENGINE_API UGameplayStatics : public UBlueprintFunctionLibrary
 	* @return Number of foliage instances with their mesh set to Mesh that overlap the sphere.
 	*/
 	UFUNCTION(BlueprintCallable, Category = "Foliage", meta = (WorldContext = "WorldContextObject", UnsafeDuringActorConstruction = "true"))
-	static int32 GrassOverlappingSphereCount(UObject* WorldContextObject, const UStaticMesh* StaticMesh, FVector CenterPosition, float Radius);
+	static int32 GrassOverlappingSphereCount(const UObject* WorldContextObject, const UStaticMesh* StaticMesh, FVector CenterPosition, float Radius);
 
 	/** 
 	 * Transforms the given 2D screen space coordinate into a 3D world-space point and direction.
diff --git a/Engine/Source/Runtime/Engine/Classes/Particles/EmitterCameraLensEffectBase.h b/Engine/Source/Runtime/Engine/Classes/Particles/EmitterCameraLensEffectBase.h
--- a/Engine/Source/Runtime/Engine/Classes/Particles/EmitterCameraLensEffectBase.h
+++ b/Engine/Source/Runtime/Engine/Classes/Particles/EmitterCameraLensEffectBase.h
@@ -20,8 +20,8 @@ protected:
 	class UParticleSystem* PS_CameraEffect;
 
 	/** The effect to use for non extreme content */
-	UPROPERTY(EditDefaultsOnly, Category = EmitterCameraLensEffectBase)
-	class UParticleSystem* PS_CameraEffectNonExtremeContent;
+	UPROPERTY()
+	class UParticleSystem* PS_CameraEffectNonExtremeContent_DEPRECATED;
 
 	/** Camera this emitter is attached to, will be notified when emitter is destroyed */
 	UPROPERTY(transient)
diff --git a/Engine/Source/Runtime/Engine/Classes/PhysicsEngine/BodySetup.h b/Engine/Source/Runtime/Engine/Classes/PhysicsEngine/BodySetup.h
--- a/Engine/Source/Runtime/Engine/Classes/PhysicsEngine/BodySetup.h
+++ b/Engine/Source/Runtime/Engine/Classes/PhysicsEngine/BodySetup.h
@@ -50,7 +50,7 @@ class UBodySetup : public UObject
 	GENERATED_UCLASS_BODY()
 
 	/** Simplified collision representation of this  */
-	UPROPERTY(EditAnywhere, Category = BodySetup, meta=(DisplayName = "Primitives"))
+	UPROPERTY(EditAnywhere, Category = BodySetup, meta=(DisplayName = "Primitives", NoResetToDefault))
 	struct FKAggregateGeom AggGeom;
 
 	/** Used in the PhysicsAsset case. Associates this Body with Bone in a skeletal mesh. */
diff --git a/Engine/Source/Runtime/Engine/Classes/PhysicsEngine/ConstraintTypes.h b/Engine/Source/Runtime/Engine/Classes/PhysicsEngine/ConstraintTypes.h
--- a/Engine/Source/Runtime/Engine/Classes/PhysicsEngine/ConstraintTypes.h
+++ b/Engine/Source/Runtime/Engine/Classes/PhysicsEngine/ConstraintTypes.h
@@ -42,7 +42,7 @@ struct ENGINE_API FConstraintBaseParams
 	float Restitution;
 
 	/** Determines how close to the limit we have to get before turning the joint on. Larger value will be more expensive, but will do a better job not violating constraints. A smaller value will be more efficient, but easier to violate. */
-	UPROPERTY(EditAnywhere, AdvancedDisplay, Category = Constraint, meta = (ClampMin = "0.0", ClampMax = "1.0"))
+	UPROPERTY(EditAnywhere, AdvancedDisplay, Category = Constraint, meta = (ClampMin = "1.0"))
 	float ContactDistance;
 
 	/** Whether we want to use a soft constraint (spring). */
diff --git a/Engine/Source/Runtime/Engine/Classes/PhysicsEngine/PhysicsSettings.h b/Engine/Source/Runtime/Engine/Classes/PhysicsEngine/PhysicsSettings.h
--- a/Engine/Source/Runtime/Engine/Classes/PhysicsEngine/PhysicsSettings.h
+++ b/Engine/Source/Runtime/Engine/Classes/PhysicsEngine/PhysicsSettings.h
@@ -184,7 +184,7 @@ class ENGINE_API UPhysicsSettings : public UDeveloperSettings
 	bool bSuppressFaceRemapTable;
 
 	/** If true, store extra information to allow FindCollisionUV to derive UV info from a line trace hit result, using the FindCollisionUV utility */
-	UPROPERTY(config, EditAnywhere, Category = Optimization, meta = (DisplayName = "Support UV From Hit Results"))
+	UPROPERTY(config, EditAnywhere, Category = Optimization, meta = (DisplayName = "Support UV From Hit Results", ConfigRestartRequired = true))
 	bool bSupportUVFromHitResults;
 
 	/**
diff --git a/Engine/Source/Runtime/Engine/Classes/PhysicsEngine/RadialForceComponent.h b/Engine/Source/Runtime/Engine/Classes/PhysicsEngine/RadialForceComponent.h
--- a/Engine/Source/Runtime/Engine/Classes/PhysicsEngine/RadialForceComponent.h
+++ b/Engine/Source/Runtime/Engine/Classes/PhysicsEngine/RadialForceComponent.h
@@ -66,6 +66,8 @@ protected:
 protected:
 	//~ Begin UActorComponent Interface.
 	virtual void TickComponent(float DeltaTime, enum ELevelTick TickType, FActorComponentTickFunction *ThisTickFunction) override;
+	
+	virtual void BeginPlay() override;
 	//~ End UActorComponent Interface.
 
 	//~ Begin UObject Interface.
diff --git a/Engine/Source/Runtime/Engine/Classes/Sound/AudioSettings.h b/Engine/Source/Runtime/Engine/Classes/Sound/AudioSettings.h
--- a/Engine/Source/Runtime/Engine/Classes/Sound/AudioSettings.h
+++ b/Engine/Source/Runtime/Engine/Classes/Sound/AudioSettings.h
@@ -45,6 +45,10 @@ class ENGINE_API UAudioSettings : public UDeveloperSettings
 	UPROPERTY(config, EditAnywhere, Category = "Audio", meta = (AllowedClasses = "SoundConcurrency", DisplayName = "Default Sound Concurrency"))
 	FStringAssetReference DefaultSoundConcurrencyName;
 
+	/** The SoundSubmix assigned to newly created sounds */
+	UPROPERTY(config, EditAnywhere, Category = "Audio", meta = (AllowedClasses = "SoundSubmix", DisplayName = "Default Sound Submix"))
+	FStringAssetReference DefaultSoundSubmixName;
+
 	/** The SoundMix to use as base when no other system has specified a Base SoundMix */
 	UPROPERTY(config, EditAnywhere, Category="Audio", meta=(AllowedClasses="SoundMix"))
 	FStringAssetReference DefaultBaseSoundMix;
diff --git a/Engine/Source/Runtime/Engine/Classes/Sound/SoundBase.h b/Engine/Source/Runtime/Engine/Classes/Sound/SoundBase.h
--- a/Engine/Source/Runtime/Engine/Classes/Sound/SoundBase.h
+++ b/Engine/Source/Runtime/Engine/Classes/Sound/SoundBase.h
@@ -8,6 +8,8 @@
 
 #include "Audio.h"
 #include "Sound/SoundConcurrency.h"
+#include "Sound/SoundEffectSource.h"
+#include "Sound/SoundSubmix.h"
 #include "SoundBase.generated.h"
 
 
@@ -26,23 +28,30 @@ class ENGINE_API USoundBase : public UObject
 {
 	GENERATED_UCLASS_BODY()
 
-private:
+public:
 	static USoundClass* DefaultSoundClassObject;
 	static USoundConcurrency* DefaultSoundConcurrencyObject;
+	static USoundSubmix* DefaultSoundSubmixObject;
 
 protected:
 	/** Sound class this sound belongs to */
 	UPROPERTY(EditAnywhere, Category=Sound, meta=(DisplayName = "Sound Class"))
 	USoundClass* SoundClassObject;
 
+	/** Sound submix this sound belongs to */
+	UPROPERTY(EditAnywhere, Category=Sound, meta = (DisplayName = "Sound Submix"))
+	USoundSubmix* SoundSubmixObject;
+
 public:
 	/** When "stat sounds -debug" has been specified, draw this sound's attenuation shape when the sound is audible. For debugging purpose only. */
 	UPROPERTY(EditAnywhere, Category = Debug)
 	uint32 bDebug:1;
 
+	/** Whether or not to override the sound concurrency object with local concurrency settings. */
 	UPROPERTY(EditAnywhere, BlueprintReadWrite, Category = Concurrency)
 	uint32 bOverrideConcurrency:1;
 
+	/** Whether or not to ignore focus on this sound. */
 	UPROPERTY(EditAnywhere, BlueprintReadWrite, Category = Attenuation)
 	uint32 bIgnoreFocus:1;
 
@@ -73,6 +82,10 @@ public:
 	UPROPERTY(EditAnywhere, BlueprintReadWrite, Category = Concurrency, meta = (ClampMin = "0.0", UIMin = "0.0", ClampMax = "100.0", UIMax = "100.0") )
 	float Priority;
 
+	/** The source effect chain to use for this sound. */
+	UPROPERTY(EditAnywhere, Category = Effects)
+	TArray<USoundEffectSourcePreset*> SourceEffectChain;
+
 public:	
 	/** Number of times this cue is currently being played. */
 	int32 CurrentPlayCount;
@@ -101,9 +114,7 @@ public:
 	virtual float GetVolumeMultiplier();
 	virtual float GetPitchMultiplier();
 
-	/**
-	  * Returns the subtitle priority
-	  */
+	/** Returns the subtitle priority */
 	virtual float GetSubtitlePriority() const { return DEFAULT_SUBTITLE_PRIORITY; };
 	
 	/** Returns whether or not any part of this sound wants interior volumes applied to it */
@@ -112,20 +123,19 @@ public:
 	/** Returns whether or not this sound is looping. */
 	bool IsLooping();
 
-	/** 
-	 * Parses the Sound to generate the WaveInstances to play
-	 */
+	/** Parses the Sound to generate the WaveInstances to play. */
 	virtual void Parse( class FAudioDevice* AudioDevice, const UPTRINT NodeWaveInstanceHash, FActiveSound& ActiveSound, const FSoundParseParameters& ParseParams, TArray<FWaveInstance*>& WaveInstances ) { }
 
-	/**
-	 * Returns the SoundClass used for this sound
-	 */
+	/** Returns the SoundClass used for this sound. */
 	virtual USoundClass* GetSoundClass() const;
 
-	/** Returns the FSoundConcurrencySettings struct to use */
+	/** Returns the SoundSubmix used for this sound. */
+	virtual USoundSubmix* GetSoundSubmix() const;
+
+	/** Returns the FSoundConcurrencySettings struct to use. */
 	const FSoundConcurrencySettings* GetSoundConcurrencySettingsToApply();
 
-	/** Returns the priority to use when evaluating concurrency */
+	/** Returns the priority to use when evaluating concurrency. */
 	float GetPriority() const;
 
 	/** Returns the sound concurrency object ID if it exists. If it doesn't exist, returns 0. */
diff --git a/Engine/Source/Runtime/Engine/Classes/Sound/SoundCue.h b/Engine/Source/Runtime/Engine/Classes/Sound/SoundCue.h
--- a/Engine/Source/Runtime/Engine/Classes/Sound/SoundCue.h
+++ b/Engine/Source/Runtime/Engine/Classes/Sound/SoundCue.h
@@ -99,7 +99,8 @@ public:
 	template<class T>
 	T* ConstructSoundNode(TSubclassOf<USoundNode> SoundNodeClass = T::StaticClass(), bool bSelectNewNode = true)
 	{
-		T* SoundNode = NewObject<T>(this, SoundNodeClass);
+		// Set flag to be transactional so it registers with undo system
+		T* SoundNode = NewObject<T>(this, SoundNodeClass, NAME_None, RF_Transactional);
 #if WITH_EDITOR
 		AllNodes.Add(SoundNode);
 		SetupSoundNode(SoundNode, bSelectNewNode);
diff --git a/Engine/Source/Runtime/Engine/Classes/Sound/SoundEffectBase.h b/Engine/Source/Runtime/Engine/Classes/Sound/SoundEffectBase.h
new file mode 100644
--- /dev/null
+++ b/Engine/Source/Runtime/Engine/Classes/Sound/SoundEffectBase.h
@@ -0,0 +1,48 @@
+// Copyright 1998-2016 Epic Games, Inc. All Rights Reserved.
+#pragma once
+
+#include "SoundEffectPreset.h"
+#include "SoundEffectBase.generated.h"
+
+class USoundEffectPreset;
+
+UCLASS()
+class ENGINE_API USoundEffectBase : public UObject
+{
+public:
+	GENERATED_UCLASS_BODY()
+
+	virtual ~USoundEffectBase();
+
+	/** Returns if the submix is active or bypassing audio. */
+	bool IsActive() const { return bIsActive; }
+
+	/** Disables the submix effect. */
+	void Disable() { bIsActive = false; }
+
+	/** Enables the submix effect. */
+	void Enable() { bIsActive = true; }
+
+	/** Sets the sound effect's preset, overriding any properties to the presets properties. */
+	void SetPreset(USoundEffectPreset* InPreset);
+
+protected:
+
+	/** Return the class of the effect. */
+	virtual UClass* GetEffectClass() const 
+	{
+		return nullptr;
+	};
+
+	USoundEffectPreset* SoundEffectPreset;
+	uint32 PresetSettingsSize;
+
+	TArray<uint8> RawPresetDataScratchInputBuffer;
+	TArray<uint8> RawPresetDataScratchOutputBuffer;
+	TQueue<TArray<uint8>> EffectSettingsQueue;
+	TArray<uint8> CurrentAudioThreadSettingsData;
+
+	FThreadSafeBool bIsRunning;
+	FThreadSafeBool bIsActive;
+};
+
diff --git a/Engine/Source/Runtime/Engine/Classes/Sound/SoundEffectPreset.h b/Engine/Source/Runtime/Engine/Classes/Sound/SoundEffectPreset.h
new file mode 100644
--- /dev/null
+++ b/Engine/Source/Runtime/Engine/Classes/Sound/SoundEffectPreset.h
@@ -0,0 +1,83 @@
+// Copyright 1998-2016 Epic Games, Inc. All Rights Reserved.
+#pragma once
+
+#if WITH_EDITOR
+#include "Developer/AssetTools/Public/AssetToolsModule.h"
+#endif
+#include "SoundEffectPreset.generated.h"
+
+#if WITH_EDITOR
+class USoundEffectPreset;
+
+class FAssetTypeActions_SoundEffectPreset : public IAssetTypeActions
+{
+public:
+	FAssetTypeActions_SoundEffectPreset(USoundEffectPreset* InEffectPreset);
+
+	// IAssetTypeActions Implementation
+	virtual bool HasActions(const TArray<UObject*>& InObjects) const override;
+	virtual void GetActions(const TArray<UObject*>& InObjects, FMenuBuilder& MenuBuilder) override;
+	virtual void OpenAssetEditor(const TArray<UObject*>& InObjects, TSharedPtr<class IToolkitHost> EditWithinLevelEditor = TSharedPtr<IToolkitHost>()) override;
+	virtual void AssetsActivated(const TArray<UObject*>& InObjects, EAssetTypeActivationMethod::Type ActivationType) override;
+	virtual bool CanFilter() override;
+	virtual bool CanLocalize() const override;
+	virtual bool CanMerge() const override;
+	virtual void Merge(UObject* InObject) override;
+	virtual void Merge(UObject* BaseAsset, UObject* RemoteAsset, UObject* LocalAsset, const FOnMergeResolved& ResolutionCallback) override;
+	virtual bool ShouldForceWorldCentric() override;
+	virtual void PerformAssetDiff(UObject* OldAsset, UObject* NewAsset, const struct FRevisionInfo& OldRevision, const struct FRevisionInfo& NewRevision) const override;
+	virtual class UThumbnailInfo* GetThumbnailInfo(UObject* Asset) const override;
+	virtual TSharedPtr<class SWidget> GetThumbnailOverlay(const FAssetData& AssetData) const override;
+	virtual FText GetAssetDescription(const class FAssetData& AssetData) const override;
+	virtual bool IsImportedAsset() const override;
+	virtual void GetResolvedSourceFilePaths(const TArray<UObject*>& TypeAssets, TArray<FString>& OutSourceFilePaths) const override;
+
+	virtual FText GetName() const override;
+	virtual FColor GetTypeColor() const override { return FColor(20.0f, 20.0f, 20.0f); }
+	virtual UClass* GetSupportedClass() const override;
+	virtual uint32 GetCategories() override { return EAssetTypeCategories::Sounds; }
+
+private:
+
+	USoundEffectPreset* EffectPreset;
+};
+#endif
+
+/** Raw data container for arbitrary preset UStructs */
+struct FPresetSettings
+{
+	void* Data;
+	uint32 Size;
+	UScriptStruct* Struct;
+
+	FPresetSettings()
+		: Data(nullptr)
+		, Size(0)
+		, Struct(nullptr)
+	{}
+};
+
+UCLASS(config = Engine, abstract, editinlinenew, BlueprintType)
+class ENGINE_API USoundEffectPreset : public UObject
+{
+	GENERATED_UCLASS_BODY()
+
+public:
+	virtual void* GetSettings() { return nullptr; }
+	virtual uint32 GetSettingsSize() const { return 0; }
+	virtual UScriptStruct* GetSettingsStruct() const { return nullptr; }
+	virtual UClass* GetEffectClass() const { return nullptr; }
+	virtual FText GetAssetActionName() const { return FText(); }
+
+	void RegisterAssetTypeAction();
+	
+	const FPresetSettings& GetPresetSettings() const
+	{
+		return PresetSettings;
+	}
+
+private:
+	FPresetSettings PresetSettings;
+};
+
+
diff --git a/Engine/Source/Runtime/Engine/Classes/Sound/SoundEffectSource.h b/Engine/Source/Runtime/Engine/Classes/Sound/SoundEffectSource.h
new file mode 100644
--- /dev/null
+++ b/Engine/Source/Runtime/Engine/Classes/Sound/SoundEffectSource.h
@@ -0,0 +1,87 @@
+// Copyright 1998-2016 Epic Games, Inc. All Rights Reserved.
+#pragma once
+
+#if WITH_EDITOR
+#include "Developer/AssetTools/Public/AssetToolsModule.h"
+#endif
+#include "SoundEffectPreset.h"
+#include "SoundEffectBase.h"
+#include "SoundEffectSource.generated.h"
+
+class USoundEffectSource;
+class USoundEffectBase;
+
+/** Derived class for source effects. */
+UCLASS(config = Engine, abstract, editinlinenew, BlueprintType)
+class ENGINE_API USoundEffectSourcePreset : public USoundEffectPreset
+{
+	GENERATED_UCLASS_BODY()
+
+	/** Create a new source sound effect instance. */
+	virtual USoundEffectSource* CreateNewEffect() const { return nullptr; }
+};
+
+/** Struct which has data needed to initialize the source effect. */
+struct FSoundEffectSourceInitData
+{
+	void* PresetSettings;
+	float SampleRate;
+	int32 NumSourceChannels;
+	float SourceDuration;
+	double AudioClock;
+	FVector SourcePosition;
+	FVector ListenerPosition;
+};
+
+/** Struct which has data to initialize the source effect. */
+struct FSoundEffectSourceInputData
+{
+	void* PresetData;
+	TArray<float> AudioBuffer;
+	int32 NumSourceChannels;
+	FVector SourcePosition;
+	FVector LeftChannelPosition;
+	FVector RightChannelPosition;
+	FVector ListenerPosition;
+	float CurrentVolume;
+	float CurrentPitch;
+	float CurrentPlayTime;
+	float Duration;
+	float Distance;
+	double AudioClock;
+	int32 CurrentLoopCount;
+	int32 MaxLoopCount;
+	uint8 bLooping : 1;
+	uint8 bIsSpatialized : 1;
+};
+
+struct FSoundEffectSourceOutputData
+{
+	TArray<float> AudioBuffer;
+	int32 OutputChannels;
+	uint8 bIsSpatialized:1;
+};
+
+UCLASS(config = Engine, hidecategories = Object, abstract, editinlinenew, BlueprintType)
+class ENGINE_API USoundEffectSource : public USoundEffectBase
+{
+	GENERATED_UCLASS_BODY()
+
+		/** Called on an audio effect at initialization on main thread before audio processing begins. */
+		virtual void Init(const FSoundEffectSourceInitData& InSampleRate) PURE_VIRTUAL(USoundEffectSource::Init, ;);
+
+	/** Process the input block of audio. Called on audio thread. */
+	virtual void OnProcessAudio(const FSoundEffectSourceInputData& InData, FSoundEffectSourceOutputData& OutData) PURE_VIRTUAL(USoundEffectSource::OnProcessAudio, ;);
+
+private:
+	/** Processes audio in the source effect. */
+	void ProcessAudio(FSoundEffectSourceInputData& InData, FSoundEffectSourceOutputData& OutData);
+
+protected:
+	UClass* GetEffectClass() const override { return GetClass(); };
+
+	// Allow FAudioMixerSubmix to call ProcessAudio
+	friend class FMixerSubmix;
+};
+
+
diff --git a/Engine/Source/Runtime/Engine/Classes/Sound/SoundEffectSubmix.h b/Engine/Source/Runtime/Engine/Classes/Sound/SoundEffectSubmix.h
new file mode 100644
--- /dev/null
+++ b/Engine/Source/Runtime/Engine/Classes/Sound/SoundEffectSubmix.h
@@ -0,0 +1,63 @@
+// Copyright 1998-2016 Epic Games, Inc. All Rights Reserved.
+#pragma once
+
+#include "SoundEffectPreset.h"
+#include "SoundEffectBase.h"
+#include "SoundEffectSubmix.generated.h"
+
+class USoundEffectPreset;
+class USoundEffectSubmix;
+
+/** Derived class for source effects. */
+UCLASS(config = Engine, hidecategories = Object, abstract, editinlinenew, BlueprintType)
+class ENGINE_API USoundEffectSubmixPreset : public USoundEffectPreset
+{
+	GENERATED_UCLASS_BODY()
+
+	/** Create a new submix effect instance. */
+	virtual USoundEffectSubmix* CreateNewEffect() const { return nullptr; }
+};
+
+/** Struct which has data needed to initialize the submix effect. */
+struct FSoundEffectSubmixInitData
+{
+	void* PresetSettings;
+	float SampleRate;
+	int32 NumOutputChannels;
+};
+
+/** Struct which supplies audio data to submix effects on game thread. */
+struct FSoundEffectSubmixInputData
+{
+	void* PresetData;
+	TArray<float>* AudioBuffer;
+	int32 NumChannels;
+	double AudioClock;
+};
+
+struct FSoundEffectSubmixOutputData
+{
+	TArray<float>* AudioBuffer;
+};
+
+UCLASS(config = Engine, hidecategories = Object, abstract, editinlinenew, BlueprintType)
+class ENGINE_API USoundEffectSubmix : public USoundEffectBase
+{
+	GENERATED_UCLASS_BODY()
+
+	/** Called on an audio effect at initialization on main thread before audio processing begins. */
+	virtual void Init(const FSoundEffectSubmixInitData& InSampleRate) PURE_VIRTUAL(USoundEffectSubmix::Init, ;);
+
+	/** Called on game thread to allow submix effect to query game data if needed. */
+	virtual void Tick() {}
+
+	/** Process the input block of audio. Called on audio thread. */
+	virtual void OnProcessAudio(const FSoundEffectSubmixInputData& InData, FSoundEffectSubmixOutputData& OutData) PURE_VIRTUAL(USoundEffectSubmix::OnProcessAudio, ;);
+
+	/** Processes audio in the source effect. */
+	void ProcessAudio(FSoundEffectSubmixInputData& InData, FSoundEffectSubmixOutputData& OutData);
+
+protected:
+	UClass* GetEffectClass() const { return GetClass(); }
+
+};
diff --git a/Engine/Source/Runtime/Engine/Classes/Sound/SoundNode.h b/Engine/Source/Runtime/Engine/Classes/Sound/SoundNode.h
--- a/Engine/Source/Runtime/Engine/Classes/Sound/SoundNode.h
+++ b/Engine/Source/Runtime/Engine/Classes/Sound/SoundNode.h
@@ -62,6 +62,7 @@ public:
 	//~ Begin UObject Interface
 #if WITH_EDITOR
 	virtual void PostEditChangeProperty(FPropertyChangedEvent& PropertyChangedEvent) override;
+	virtual void PostLoad() override;
 	static void AddReferencedObjects(UObject* InThis, FReferenceCollector& Collector);
 #endif //WITH_EDITOR
 	virtual void Serialize(FArchive& Ar) override;
diff --git a/Engine/Source/Runtime/Engine/Classes/Sound/SoundSubmix.h b/Engine/Source/Runtime/Engine/Classes/Sound/SoundSubmix.h
new file mode 100644
--- /dev/null
+++ b/Engine/Source/Runtime/Engine/Classes/Sound/SoundSubmix.h
@@ -0,0 +1,43 @@
+// Copyright 1998-2016 Epic Games, Inc. All Rights Reserved.
+#pragma once
+
+#include "SoundEffectSubmix.h"
+#include "SoundSubmix.generated.h"
+
+class USoundEffectSubmixPreset;
+
+
+UCLASS()
+class ENGINE_API USoundSubmix : public UObject
+{
+	GENERATED_UCLASS_BODY()
+
+	// Child submixes to this sound mix
+	UPROPERTY(EditAnywhere, Category = SoundSubmix)
+	TArray<USoundSubmix*> ChildSubmixes;
+
+	UPROPERTY()
+	USoundSubmix* ParentSubmix;
+
+	UPROPERTY(EditAnywhere, Category = SoundSubmix)
+	TArray<USoundEffectSubmixPreset*> SubmixEffectChain;
+
+	// The output wet level to use for the output of this submix in parent submixes
+	UPROPERTY(EditAnywhere, Category = SoundSubmix)
+	float OutputWetLevel;
+
+protected:
+
+	//~ Begin UObject Interface.
+	virtual void Serialize( FArchive& Ar ) override;
+	virtual FString GetDesc( void ) override;
+	virtual void BeginDestroy() override;
+	virtual void PostLoad() override;
+
+#if WITH_EDITOR
+	virtual void PreEditChange(UProperty* PropertyAboutToChange) override;
+	virtual void PostEditChangeProperty(struct FPropertyChangedEvent& PropertyChangedEvent) override;
+#endif
+	//~ End UObject Interface.
+};
+
diff --git a/Engine/Source/Runtime/Engine/Classes/Sound/SoundWave.h b/Engine/Source/Runtime/Engine/Classes/Sound/SoundWave.h
--- a/Engine/Source/Runtime/Engine/Classes/Sound/SoundWave.h
+++ b/Engine/Source/Runtime/Engine/Classes/Sound/SoundWave.h
@@ -121,7 +121,7 @@ class ENGINE_API USoundWave : public USoundBase
 	int32 CompressionQuality;
 
 	/** If set, when played directly (not through a sound cue) the wave will be played looping. */
-	UPROPERTY(EditAnywhere, Category=SoundWave )
+	UPROPERTY(EditAnywhere, Category=SoundWave, AssetRegistrySearchable)
 	uint32 bLooping:1;
 
 	/** Whether this sound can be streamed to avoid increased memory usage */
diff --git a/Engine/Source/Runtime/Engine/Private/AI/Navigation/AvoidanceManager.cpp b/Engine/Source/Runtime/Engine/Private/AI/Navigation/AvoidanceManager.cpp
--- a/Engine/Source/Runtime/Engine/Private/AI/Navigation/AvoidanceManager.cpp
+++ b/Engine/Source/Runtime/Engine/Private/AI/Navigation/AvoidanceManager.cpp
@@ -51,7 +51,7 @@ UAvoidanceManager::UAvoidanceManager(const FObjectInitializer& ObjectInitializer
 {
 	DefaultTimeToLive = 1.5f;
 	LockTimeAfterAvoid = 0.2f;
-	LockTimeAfterClean = 0.01f;
+	LockTimeAfterClean = 0.001f;
 	DeltaTimeToPredict = 0.5f;
 	ArtificialRadiusExpansion = 1.5f;
 	TestHeightDifference_DEPRECATED = 500.0f;
diff --git a/Engine/Source/Runtime/Engine/Private/AI/Navigation/NavLinkProxy.cpp b/Engine/Source/Runtime/Engine/Private/AI/Navigation/NavLinkProxy.cpp
--- a/Engine/Source/Runtime/Engine/Private/AI/Navigation/NavLinkProxy.cpp
+++ b/Engine/Source/Runtime/Engine/Private/AI/Navigation/NavLinkProxy.cpp
@@ -56,13 +56,10 @@ ANavLinkProxy::ANavLinkProxy(const FObjectInitializer& ObjectInitializer) : Supe
 	}
 #endif
 
-	if (HasAnyFlags(RF_ClassDefaultObject) == false)
-	{
-		SmartLinkComp = CreateDefaultSubobject<UNavLinkCustomComponent>(TEXT("SmartLinkComp"));
-		SmartLinkComp->SetNavigationRelevancy(false);
-		SmartLinkComp->SetMoveReachedLink(this, &ANavLinkProxy::NotifySmartLinkReached);
-		bSmartLinkIsRelevant = false;
-	}
+	SmartLinkComp = CreateDefaultSubobject<UNavLinkCustomComponent>(TEXT("SmartLinkComp"));
+	SmartLinkComp->SetNavigationRelevancy(false);
+	SmartLinkComp->SetMoveReachedLink(this, &ANavLinkProxy::NotifySmartLinkReached);
+	bSmartLinkIsRelevant = false;
 
 	PointLinks.Add(FNavigationLink());
 	SetActorEnableCollision(false);
diff --git a/Engine/Source/Runtime/Engine/Private/AI/Navigation/NavLinkRenderingComponent.cpp b/Engine/Source/Runtime/Engine/Private/AI/Navigation/NavLinkRenderingComponent.cpp
--- a/Engine/Source/Runtime/Engine/Private/AI/Navigation/NavLinkRenderingComponent.cpp
+++ b/Engine/Source/Runtime/Engine/Private/AI/Navigation/NavLinkRenderingComponent.cpp
@@ -24,8 +24,7 @@ UNavLinkRenderingComponent::UNavLinkRenderingComponent(const FObjectInitializer&
 
 	BodyInstance.SetCollisionProfileName(UCollisionProfile::NoCollision_ProfileName);
 
-	AlwaysLoadOnClient = false;
-	AlwaysLoadOnServer = false;
+	bIsEditorOnly = true;
 
 	bGenerateOverlapEvents = false;
 }
diff --git a/Engine/Source/Runtime/Engine/Private/AI/Navigation/PImplRecastNavMesh.cpp b/Engine/Source/Runtime/Engine/Private/AI/Navigation/PImplRecastNavMesh.cpp
--- a/Engine/Source/Runtime/Engine/Private/AI/Navigation/PImplRecastNavMesh.cpp
+++ b/Engine/Source/Runtime/Engine/Private/AI/Navigation/PImplRecastNavMesh.cpp
@@ -718,6 +718,9 @@ void FPImplRecastNavMesh::SetRecastMesh(dtNavMesh* NavMesh)
 	{
 		NavMeshOwner->UpdateNavObject();
 	}
+
+	// reapply area sort order in new detour navmesh
+	OnAreaCostChanged();
 }
 
 void FPImplRecastNavMesh::Raycast(const FVector& StartLoc, const FVector& EndLoc, const FNavigationQueryFilter& InQueryFilter, const UObject* Owner, 
diff --git a/Engine/Source/Runtime/Engine/Private/ActiveSound.cpp b/Engine/Source/Runtime/Engine/Private/ActiveSound.cpp
--- a/Engine/Source/Runtime/Engine/Private/ActiveSound.cpp
+++ b/Engine/Source/Runtime/Engine/Private/ActiveSound.cpp
@@ -21,12 +21,14 @@ FActiveSound::FActiveSound()
 	, ConcurrencyGeneration(0)
 	, ConcurrencySettings(nullptr)
 	, SoundClassOverride(nullptr)
+	, SoundSubmixOverride(nullptr)
 	, bHasCheckedOcclusion(false)
 	, bAllowSpatialization(true)
 	, bHasAttenuationSettings(false)
 	, bShouldRemainActiveIfDropped(false)
 	, bFadingOut(false)
 	, bFinished(false)
+	, bIsPaused(false)
 	, bShouldStopDueToMaxConcurrency(false)
 	, bRadioFilterSelected(false)
 	, bApplyRadioFilter(false)
@@ -46,6 +48,7 @@ FActiveSound::FActiveSound()
 	, bWarnedAboutOrphanedLooping(false)
 #endif
 	, bEnableLowPassFilter(false)
+	, bUpdatePlayPercentage(false)
 	, UserIndex(0)
 	, bIsOccluded(false)
 	, bAsyncOcclusionPending(false)
@@ -69,6 +72,8 @@ FActiveSound::FActiveSound()
 	, OcclusionCheckInterval(0.f)
 	, LastOcclusionCheckTime(TNumericLimits<float>::Lowest())
 	, MaxDistance(WORLD_MAX)
+	, Azimuth(0.0f)
+	, AbsoluteAzimuth(0.0f)
 	, LastLocation(FVector::ZeroVector)
 	, AudioVolumeID(0)
 	, LastUpdateTime(0.f)
@@ -187,12 +192,25 @@ USoundClass* FActiveSound::GetSoundClass() const
 		return SoundClassOverride;
 	}
 	else if (Sound)
-		{
+	{
 		return Sound->GetSoundClass();
-		}
+	}
 
 	return nullptr;
+}
+
+USoundSubmix* FActiveSound::GetSoundSubmix() const
+{
+	if (SoundSubmixOverride)
+	{
+		return SoundSubmixOverride;
+	}
+	else if (Sound)
+	{
+		return Sound->GetSoundSubmix();
 	}
+	return nullptr;
+}
 
 int32 FActiveSound::FindClosestListener( const TArray<FListener>& InListeners ) const
 {
@@ -224,7 +242,6 @@ uint32 FActiveSound::GetSoundConcurrencyObjectID() const
 	}
 }
 
-
 void FActiveSound::UpdateWaveInstances( TArray<FWaveInstance*> &InWaveInstances, const float DeltaTime )
 {
 	check(AudioDevice);
@@ -282,6 +299,8 @@ void FActiveSound::UpdateWaveInstances( TArray<FWaveInstance*> &InWaveInstances,
 	ParseParams.bEnableLowPassFilter = bEnableLowPassFilter;
 	ParseParams.LowPassFilterFrequency = LowPassFilterFrequency;
 	ParseParams.SoundClass = GetSoundClass();
+	ParseParams.SoundSubmix = GetSoundSubmix();
+	ParseParams.bIsPaused = bIsPaused;
 
 	if (bApplyInteriorVolumes)
 	{
@@ -818,29 +837,39 @@ void FActiveSound::ApplyAttenuation(FSoundParseParameters& ParseParams, const FL
 
 	FAttenuationListenerData ListenerData;
 
-	// Get the current focus factor
-	const float FocusFactor = AudioDevice->GetFocusFactor(ListenerData, Sound, SoundTransform, *Settings, &Listener.Transform);
-
 	// Reset distance and priority scale to 1.0 in case changed in editor
 	FocusDistanceScale = 1.0f;
 	FocusPriorityScale = 1.0f;
-	
-	// Computer the focus factor if needed
+
 	check(Sound);
-	if (Settings->bSpatialize && Settings->bEnableListenerFocus && !Sound->bIgnoreFocus)
+
+	if (Settings->bSpatialize)
 	{
+		// Compute the azimuth of the active sound
 		const FGlobalFocusSettings& FocusSettings = AudioDevice->GetGlobalFocusSettings();
 
-		// Get the volume scale to apply the volume calculation based on the focus factor
-		const float FocusVolumeAttenuation = Settings->GetFocusAttenuation(FocusSettings, FocusFactor);
-		Volume *= FocusVolumeAttenuation;
+		AudioDevice->GetAzimuth(ListenerData, Sound, SoundTransform, *Settings, Listener.Transform, Azimuth, AbsoluteAzimuth);
+
+		ParseParams.AttenuationDistance = ListenerData.AttenuationDistance;
 
-		// Scale the volume-weighted priority scale value we use for sorting this sound for voice-stealing
-		FocusPriorityScale = Settings->GetFocusPriorityScale(FocusSettings, FocusFactor);
-		ParseParams.Priority *= FocusPriorityScale;
+		ParseParams.AbsoluteAzimuth = AbsoluteAzimuth;
 
-		// Get the distance scale to use when computing distance-calculations for 3d attenuation
-		FocusDistanceScale = Settings->GetFocusDistanceScale(FocusSettings, FocusFactor);
+		if (Settings->bEnableListenerFocus && !Sound->bIgnoreFocus)
+		{
+			// Get the current focus factor
+			const float FocusFactor = AudioDevice->GetFocusFactor(ListenerData, Sound, Azimuth, *Settings);
+
+			// Get the volume scale to apply the volume calculation based on the focus factor
+			const float FocusVolumeAttenuation = Settings->GetFocusAttenuation(FocusSettings, FocusFactor);
+			Volume *= FocusVolumeAttenuation;
+
+			// Scale the volume-weighted priority scale value we use for sorting this sound for voice-stealing
+			FocusPriorityScale = Settings->GetFocusPriorityScale(FocusSettings, FocusFactor);
+			ParseParams.Priority *= FocusPriorityScale;
+
+			// Get the distance scale to use when computing distance-calculations for 3d attenuation
+			FocusDistanceScale = Settings->GetFocusDistanceScale(FocusSettings, FocusFactor);
+		}
 	}
 
 	// Attenuate the volume based on the model
diff --git a/Engine/Source/Runtime/Engine/Private/Actor.cpp b/Engine/Source/Runtime/Engine/Private/Actor.cpp
--- a/Engine/Source/Runtime/Engine/Private/Actor.cpp
+++ b/Engine/Source/Runtime/Engine/Private/Actor.cpp
@@ -238,7 +238,7 @@ void AActor::ResetOwnedComponents()
 	for (UObject* Child : ActorChildren)
 	{
 		UActorComponent* Component = Cast<UActorComponent>(Child);
-		if (Component)
+		if (Component && Component->GetOwner() == this)
 		{
 			OwnedComponents.Add(Component);
 
@@ -333,8 +333,10 @@ bool AActor::TeleportTo( const FVector& DestLocation, const FRotator& DestRotati
 		return false;
 	}
 
+	UWorld* MyWorld = GetWorld();
+
 	// Can't move non-movable actors during play
-	if( (RootComponent->Mobility == EComponentMobility::Static) && GetWorld()->AreActorsInitialized() )
+	if( (RootComponent->Mobility == EComponentMobility::Static) && MyWorld->AreActorsInitialized() )
 	{
 		return false;
 	}
@@ -357,7 +359,7 @@ bool AActor::TeleportTo( const FVector& DestLocation, const FRotator& DestRotati
 			NewLocation = NewLocation + Offset;
 
 			// check if able to find an acceptable destination for this actor that doesn't embed it in world geometry
-			bTeleportSucceeded = GetWorld()->FindTeleportSpot(this, NewLocation, DestRotation);
+			bTeleportSucceeded = MyWorld->FindTeleportSpot(this, NewLocation, DestRotation);
 			NewLocation = NewLocation - Offset;
 		}
 
@@ -834,27 +836,30 @@ void AActor::Tick( float DeltaSeconds )
 {
 	// Blueprint code outside of the construction script should not run in the editor
 	// Allow tick if we are not a dedicated server, or we allow this tick on dedicated servers
-	if (GetWorldSettings() != NULL && (bAllowReceiveTickEventOnDedicatedServer || !IsRunningDedicatedServer()))
+	if (GetWorldSettings() != nullptr && (bAllowReceiveTickEventOnDedicatedServer || !IsRunningDedicatedServer()))
 	{
 		ReceiveTick(DeltaSeconds);
 	}
 
 
 	// Update any latent actions we have for this actor
-	GetWorld()->GetLatentActionManager().ProcessLatentActions(this, DeltaSeconds);
+
+	// If this tick is skipped on a frame because we've got a TickInterval, our latent actions will be ticked
+	// anyway by UWorld::Tick(). Given that, our latent actions don't need to be passed a larger
+	// DeltaSeconds to make up the frames that they missed (because they wouldn't have missed any).
+	// So pass in the world's DeltaSeconds value rather than our specific DeltaSeconds value.
+	UWorld* MyWorld = GetWorld();
+	MyWorld->GetLatentActionManager().ProcessLatentActions(this, MyWorld->GetDeltaSeconds());
 
 	if (bAutoDestroyWhenFinished)
 	{
 		bool bOKToDestroy = true;
 
 		// @todo: naive implementation, needs improved
-		TInlineComponentArray<UActorComponent*> Components;
-		GetComponents(Components);
+		TInlineComponentArray<UActorComponent*> Components(this);
 
-		for (int32 CompIdx=0; CompIdx<Components.Num(); ++CompIdx)
+		for (UActorComponent* const Comp : Components)
 		{
-			UActorComponent* const Comp = Components[CompIdx];
-
 			UParticleSystemComponent* const PSC = Cast<UParticleSystemComponent>(Comp);
 			if ( PSC && (PSC->bIsActive || !PSC->bWasCompleted) )
 			{
@@ -3020,10 +3025,7 @@ void AActor::BeginPlay()
 		if (Component->IsRegistered() && !Component->HasBegunPlay())
 		{
 			Component->RegisterAllComponentTickFunctions(true);
-			//if (Component->bWantsBeginPlay) // TODO: this should be enforced, but we need to address an upgrade path first to not quietly break code. Not checking this flag is the old behavior.
-			{
-				Component->BeginPlay();
-			}
+			Component->BeginPlay();
 		}
 		else
 		{
@@ -4471,9 +4473,9 @@ void AActor::K2_SetActorRelativeTransform(const FTransform& NewRelativeTransform
 
 float AActor::GetGameTimeSinceCreation()
 {
-	if(GetWorld() != nullptr)
+	if (UWorld* MyWorld = GetWorld())
 	{
-		return GetWorld()->GetTimeSeconds() - CreationTime;		
+		return MyWorld->GetTimeSeconds() - CreationTime;		
 	}
 	// return 0.f if GetWorld return's null
 	else
diff --git a/Engine/Source/Runtime/Engine/Private/AdvancedPreviewScene.cpp b/Engine/Source/Runtime/Engine/Private/AdvancedPreviewScene.cpp
--- a/Engine/Source/Runtime/Engine/Private/AdvancedPreviewScene.cpp
+++ b/Engine/Source/Runtime/Engine/Private/AdvancedPreviewScene.cpp
@@ -30,7 +30,7 @@ FAdvancedPreviewScene::FAdvancedPreviewScene(ConstructionValues CVS, float InFlo
 	
 	// Add and set up sky light using the set cube map texture
 	SkyLightComponent = NewObject<USkyLightComponent>();
-	SkyLightComponent->Cubemap = Profile.EnvironmentCubeMap;
+	SkyLightComponent->Cubemap = Profile.EnvironmentCubeMap.Get();
 	SkyLightComponent->SourceType = ESkyLightSourceType::SLS_SpecifiedCubemap;
 	SkyLightComponent->Mobility = EComponentMobility::Movable;
 	SkyLightComponent->bLowerHemisphereIsBlack = false;
@@ -54,7 +54,7 @@ FAdvancedPreviewScene::FAdvancedPreviewScene(ConstructionValues CVS, float InFlo
 	InstancedSkyMaterial->Parent = SkyMaterial;		
 
 	UTextureCube* DefaultTexture = LoadObject<UTextureCube>(NULL, TEXT("/Engine/MapTemplates/Sky/SunsetAmbientCubemap.SunsetAmbientCubemap"));
-	InstancedSkyMaterial->SetTextureParameterValueEditorOnly(FName("SkyBox"), ( Profile.EnvironmentCubeMap != nullptr ) ? Profile.EnvironmentCubeMap : DefaultTexture );
+	InstancedSkyMaterial->SetTextureParameterValueEditorOnly(FName("SkyBox"), ( Profile.EnvironmentCubeMap.Get() != nullptr ) ? Profile.EnvironmentCubeMap.Get() : DefaultTexture );
 	InstancedSkyMaterial->SetScalarParameterValueEditorOnly(FName("CubemapRotation"), Profile.LightingRigRotation / 360.0f);
 	InstancedSkyMaterial->SetScalarParameterValueEditorOnly(FName("Intensity"), Profile.SkyLightIntensity);
 	InstancedSkyMaterial->PostLoad();
@@ -66,7 +66,7 @@ FAdvancedPreviewScene::FAdvancedPreviewScene(ConstructionValues CVS, float InFlo
 	PostProcessComponent->bUnbound = true;
 	AddComponent(PostProcessComponent, Transform);
 
-	UStaticMesh* FloorMesh = LoadObject<UStaticMesh>(NULL, TEXT("/Engine/EditorMeshes/PhAT_FloorBox.PhAT_FloorBox"), NULL, LOAD_None, NULL);
+	UStaticMesh* FloorMesh = LoadObject<UStaticMesh>(NULL, TEXT("/Engine/EditorMeshes/AssetViewer/Floor_Mesh.Floor_Mesh"), NULL, LOAD_None, NULL);
 	check(FloorMesh);
 	FloorMeshComponent = NewObject<UStaticMeshComponent>(GetTransientPackage());
 	FloorMeshComponent->SetStaticMesh(FloorMesh);
@@ -106,12 +106,14 @@ void FAdvancedPreviewScene::UpdateScene(FPreviewSceneProfile& Profile, bool bUpd
 		static const FName SkyBoxName("SkyBox");
 		static const FName CubeMapRotationName("CubemapRotation");
 
-		UTexture* Texture = Profile.EnvironmentCubeMap;
+		UTextureCube* EnvironmentTexture = Profile.EnvironmentCubeMap.LoadSynchronous();
+		UTexture* Texture = Cast<UTexture>(EnvironmentTexture);
 		InstancedSkyMaterial->GetTextureParameterValue(SkyBoxName, Texture);
-		if (Texture != Profile.EnvironmentCubeMap)
+
+		if (Texture != EnvironmentTexture)
 		{
-			InstancedSkyMaterial->SetTextureParameterValueEditorOnly(SkyBoxName, Profile.EnvironmentCubeMap);
-			SkyLightComponent->Cubemap = Profile.EnvironmentCubeMap;
+			InstancedSkyMaterial->SetTextureParameterValueEditorOnly(SkyBoxName, EnvironmentTexture);
+			SkyLightComponent->Cubemap = EnvironmentTexture;
 			bSkyChanged = true;
 		}
 		
@@ -276,24 +278,42 @@ const bool FAdvancedPreviewScene::HandleInputKey(FViewport* InViewport, int32 Co
 	return bResult;
 }
 
-void FAdvancedPreviewScene::SetFloorVisibility(const bool bVisible)
+void FAdvancedPreviewScene::SetFloorVisibility(const bool bVisible, const bool bDirect)
 {
-	FName PropertyName("bShowFloor");
+	// If not direct set visibility in profile and refresh the scene
+	if (!bDirect)
+	{
+		FName PropertyName("bShowFloor");
 
-	UProperty* FloorProperty = FindField<UProperty>(FPreviewSceneProfile::StaticStruct(), PropertyName);
-	DefaultSettings->Profiles[CurrentProfileIndex].bShowFloor = bVisible;
+		UProperty* FloorProperty = FindField<UProperty>(FPreviewSceneProfile::StaticStruct(), PropertyName);
+		DefaultSettings->Profiles[CurrentProfileIndex].bShowFloor = bVisible;
 
-	FPropertyChangedEvent PropertyEvent(FloorProperty);
-	DefaultSettings->PostEditChangeProperty(PropertyEvent);
+		FPropertyChangedEvent PropertyEvent(FloorProperty);
+		DefaultSettings->PostEditChangeProperty(PropertyEvent);
+	}
+	else
+	{
+		// Otherwise set visiblity directly on the component
+		FloorMeshComponent->SetVisibility(bVisible ? DefaultSettings->Profiles[CurrentProfileIndex].bShowFloor : bVisible);
+	}
 }
 
-void FAdvancedPreviewScene::SetEnvironmentVisibility(const bool bVisible)
+void FAdvancedPreviewScene::SetEnvironmentVisibility(const bool bVisible, const bool bDirect)
 {
-	UProperty* EnvironmentProperty = FindField<UProperty>(FPreviewSceneProfile::StaticStruct(), GET_MEMBER_NAME_CHECKED(FPreviewSceneProfile, bShowEnvironment));
-	DefaultSettings->Profiles[CurrentProfileIndex].bShowEnvironment = bVisible;
+	// If not direct set visibility in profile and refresh the scene
+	if (!bDirect)
+	{
+		UProperty* EnvironmentProperty = FindField<UProperty>(FPreviewSceneProfile::StaticStruct(), GET_MEMBER_NAME_CHECKED(FPreviewSceneProfile, bShowEnvironment));
+		DefaultSettings->Profiles[CurrentProfileIndex].bShowEnvironment = bVisible;
 
-	FPropertyChangedEvent PropertyEvent(EnvironmentProperty);
-	DefaultSettings->PostEditChangeProperty(PropertyEvent);
+		FPropertyChangedEvent PropertyEvent(EnvironmentProperty);
+		DefaultSettings->PostEditChangeProperty(PropertyEvent);
+	}
+	else
+	{
+		// Otherwise set visiblity directly on the component
+		SkyComponent->SetVisibility(bVisible ? DefaultSettings->Profiles[CurrentProfileIndex].bShowEnvironment : bVisible);
+	}
 }
 
 const float FAdvancedPreviewScene::GetSkyRotation() const
