Copying //UE4/Dev-Framework to //UE4/Dev-Main (Source: //UE4/Dev-Framework @ 3136612)

diff --git a/Engine/Source/Runtime/Engine/Private/Animation/AnimCurveTypes.cpp b/Engine/Source/Runtime/Engine/Private/Animation/AnimCurveTypes.cpp
--- a/Engine/Source/Runtime/Engine/Private/Animation/AnimCurveTypes.cpp
+++ b/Engine/Source/Runtime/Engine/Private/Animation/AnimCurveTypes.cpp
@@ -13,7 +13,7 @@ DECLARE_CYCLE_STAT(TEXT("AnimSeq EvalCurveData"), STAT_AnimSeq_EvalCurveData, ST
 
 void FAnimCurveBase::PostSerialize(FArchive& Ar)
 {
-	FSmartNameMapping::UID CurveUid = FSmartNameMapping::MaxUID;
+	SmartName::UID_Type CurveUid = SmartName::MaxUID;
 	Ar.UsingCustomVersion(FFrameworkObjectVersion::GUID);
 
 	if ((Ar.IsLoading()))
@@ -35,7 +35,7 @@ void FAnimCurveBase::PostSerialize(FArchive& Ar)
 	}
 }
 
-void FAnimCurveBase::SetCurveTypeFlag(EAnimCurveFlags InFlag, bool bValue)
+void FAnimCurveBase::SetCurveTypeFlag(EAnimAssetCurveFlags InFlag, bool bValue)
 {
 	if (bValue)
 	{
@@ -47,13 +47,13 @@ void FAnimCurveBase::SetCurveTypeFlag(EAnimCurveFlags InFlag, bool bValue)
 	}
 }
 
-void FAnimCurveBase::ToggleCurveTypeFlag(EAnimCurveFlags InFlag)
+void FAnimCurveBase::ToggleCurveTypeFlag(EAnimAssetCurveFlags InFlag)
 {
 	bool Current = GetCurveTypeFlag(InFlag);
 	SetCurveTypeFlag(InFlag, !Current);
 }
 
-bool FAnimCurveBase::GetCurveTypeFlag(EAnimCurveFlags InFlag) const
+bool FAnimCurveBase::GetCurveTypeFlag(EAnimAssetCurveFlags InFlag) const
 {
 	return (CurveTypeFlags & InFlag) != 0;
 }
@@ -189,7 +189,7 @@ void FRawCurveTracks::EvaluateCurveData( FBlendedCurve& Curves, float CurrentTim
 	for(auto CurveIter = FloatCurves.CreateConstIterator(); CurveIter; ++CurveIter)
 	{
 		const FFloatCurve& Curve = *CurveIter;
-		Curves.Set(Curve.Name.UID, Curve.Evaluate(CurrentTime), Curve.GetCurveTypeFlags());
+		Curves.Set(Curve.Name.UID, Curve.Evaluate(CurrentTime));
 	}
 }
 
@@ -490,6 +490,6 @@ void FAnimCurveParam::Initialize(USkeleton* Skeleton)
 	else
 	{
 		// invalidate current UID
-		UID = FSmartNameMapping::MaxUID;
+		UID = SmartName::MaxUID;
 	}
 }
\ No newline at end of file
diff --git a/Engine/Source/Runtime/Engine/Private/Animation/AnimInstance.cpp b/Engine/Source/Runtime/Engine/Private/Animation/AnimInstance.cpp
--- a/Engine/Source/Runtime/Engine/Private/Animation/AnimInstance.cpp
+++ b/Engine/Source/Runtime/Engine/Private/Animation/AnimInstance.cpp
@@ -253,7 +253,7 @@ void UAnimInstance::UninitializeAnimation()
 		TArray<FName> ParamsToClearCopy = MaterialParamatersToClear;
 		for(int i = 0; i < ParamsToClearCopy.Num(); ++i)
 		{
-			AddCurveValue(ParamsToClearCopy[i], 0.0f, ACF_DriveMaterial);
+			AddCurveValue(ParamsToClearCopy[i], 0.0f);
 		}
 	}
 
@@ -982,6 +982,11 @@ void UAnimInstance::RecalcRequiredBones()
 	}
 }
 
+void UAnimInstance::RecalcRequiredCurves()
+{
+	GetProxyOnGameThread<FAnimInstanceProxy>().RecalcRequiredCurves();
+}
+
 void UAnimInstance::Serialize(FArchive& Ar)
 {
 	Super::Serialize(Ar);
@@ -1040,58 +1045,60 @@ void UAnimInstance::AddAnimNotifyFromGeneratedClass(int32 NotifyIndex)
 	GetProxyOnGameThread<FAnimInstanceProxy>().AddAnimNotifyFromGeneratedClass(NotifyIndex);
 }
 
-void UAnimInstance::AddCurveValue(const FName& CurveName, float Value, int32 CurveTypeFlags)
+void UAnimInstance::AddCurveValue(const FName& CurveName, float Value)
 {
 	FAnimInstanceProxy& Proxy = GetProxyOnGameThread<FAnimInstanceProxy>();
 
 	// save curve value, it will overwrite if same exists, 
 	//CurveValues.Add(CurveName, Value);
-	if (CurveTypeFlags & ACF_DriveAttribute)
+	float* CurveValPtr = AnimationCurves[(uint8)EAnimCurveType::AttributeCurve].Find(CurveName);
+	if ( CurveValPtr )
 	{
-		float *CurveValPtr = AnimationCurves[(uint8)EAnimCurveType::AttributeCurve].Find(CurveName);
-		if ( CurveValPtr )
-		{
-			// sum up, in the future we might normalize, but for now this just sums up
-			// this won't work well if all of them have full weight - i.e. additive 
-			*CurveValPtr = Value;
-		}
-		else
-		{
-			AnimationCurves[(uint8)EAnimCurveType::AttributeCurve].Add(CurveName, Value);
-		}
+		// sum up, in the future we might normalize, but for now this just sums up
+		// this won't work well if all of them have full weight - i.e. additive 
+		*CurveValPtr = Value;
 	}
-
-	if (CurveTypeFlags & ACF_DriveMorphTarget)
+	else
 	{
-		float *CurveValPtr = AnimationCurves[(uint8)EAnimCurveType::MorphTargetCurve].Find(CurveName);
-		if ( CurveValPtr )
-		{
-			// sum up, in the future we might normalize, but for now this just sums up
-			// this won't work well if all of them have full weight - i.e. additive 
-			*CurveValPtr = Value;
-		}
-		else
-		{
-			AnimationCurves[(uint8)EAnimCurveType::MorphTargetCurve].Add(CurveName, Value);
-		}
+		AnimationCurves[(uint8)EAnimCurveType::AttributeCurve].Add(CurveName, Value);
 	}
 
-	if (CurveTypeFlags & ACF_DriveMaterial)
+	check(CurrentSkeleton);
+
+	const FCurveMetaData* CurveMetaData = CurrentSkeleton->GetCurveMetaData(CurveName);
+	if (CurveMetaData)
 	{
-		MaterialParamatersToClear.RemoveSwap(CurveName);
-		float* CurveValPtr = AnimationCurves[(uint8)EAnimCurveType::MaterialCurve].Find(CurveName);
-		if( CurveValPtr)
+		if (CurveMetaData->Type.bMorphtarget)
 		{
-			*CurveValPtr = Value;
+			CurveValPtr = AnimationCurves[(uint8)EAnimCurveType::MorphTargetCurve].Find(CurveName);
+			if (CurveValPtr)
+			{
+				// sum up, in the future we might normalize, but for now this just sums up
+				// this won't work well if all of them have full weight - i.e. additive 
+				*CurveValPtr = Value;
+			}
+			else
+			{
+				AnimationCurves[(uint8)EAnimCurveType::MorphTargetCurve].Add(CurveName, Value);
+			}
 		}
-		else
+		if (CurveMetaData->Type.bMaterial)
 		{
-			AnimationCurves[(uint8)EAnimCurveType::MaterialCurve].Add(CurveName, Value);
+			MaterialParamatersToClear.RemoveSwap(CurveName);
+			CurveValPtr = AnimationCurves[(uint8)EAnimCurveType::MaterialCurve].Find(CurveName);
+			if (CurveValPtr)
+			{
+				*CurveValPtr = Value;
+			}
+			else
+			{
+				AnimationCurves[(uint8)EAnimCurveType::MaterialCurve].Add(CurveName, Value);
+			}
 		}
 	}
 }
 
-void UAnimInstance::AddCurveValue(const USkeleton::AnimCurveUID Uid, float Value, int32 CurveTypeFlags)
+void UAnimInstance::AddCurveValue(const USkeleton::AnimCurveUID Uid, float Value)
 {
 	FName CurrentCurveName;
 	// Grab the smartname mapping from our current skeleton and resolve the curve name. We cannot cache
@@ -1100,7 +1107,7 @@ void UAnimInstance::AddCurveValue(const USkeleton::AnimCurveUID Uid, float Value
 	{
 		NameMapping->GetName(Uid, CurrentCurveName);
 	}
-	AddCurveValue(CurrentCurveName, Value, CurveTypeFlags);
+	AddCurveValue(CurrentCurveName, Value);
 }
 
 int32 UAnimInstance::GetSyncGroupReadIndex() const
@@ -1127,23 +1134,14 @@ void UAnimInstance::UpdateCurvesToComponents(USkeletalMeshComponent* Component /
 	}
 }
 
-void UAnimInstance::GetAnimationCurveList(int32 CurveFlags, TMap<FName, float>& OutCurveList) const
+void UAnimInstance::GetAnimationCurveList(EAnimCurveType Type, TMap<FName, float>& OutCurveList) const
 {
-	OutCurveList.Reset();
-
-	if (CurveFlags & ACF_DriveAttribute)
-	{
-		OutCurveList.Append(AnimationCurves[(uint8)EAnimCurveType::AttributeCurve]);
-	}
+	uint8 ArrayIndex = (uint8)Type;
 
-	if (CurveFlags & ACF_DriveMorphTarget)
+	if (ArrayIndex < (uint8)EAnimCurveType::MaxAnimCurveType)
 	{
-		OutCurveList.Append(AnimationCurves[(uint8)EAnimCurveType::MorphTargetCurve]);
-	}
-
-	if (CurveFlags & ACF_DriveMaterial)
-	{
-		OutCurveList.Append(AnimationCurves[(uint8)EAnimCurveType::MaterialCurve]);
+		// add unique only
+		OutCurveList.Append(AnimationCurves[ArrayIndex]);
 	}
 }
 
@@ -1180,11 +1178,11 @@ void UAnimInstance::UpdateCurves(const FBlendedHeapCurve& InCurve)
 
 	if (InCurve.UIDList != nullptr)
 	{
-		const TArray<FSmartNameMapping::UID>& UIDList = *InCurve.UIDList;
+		const TArray<SmartName::UID_Type>& UIDList = *InCurve.UIDList;
 		for (int32 CurveId = 0; CurveId < InCurve.UIDList->Num(); ++CurveId)
 		{
 			// had to add to another data type
-			AddCurveValue(UIDList[CurveId], InCurve.Elements[CurveId].Value, InCurve.Elements[CurveId].Flags);
+			AddCurveValue(UIDList[CurveId], InCurve.Elements[CurveId].Value);
 		}
 	}	
 
@@ -1193,7 +1191,7 @@ void UAnimInstance::UpdateCurves(const FBlendedHeapCurve& InCurve)
 	TArray<FName> ParamsToClearCopy = MaterialParamatersToClear;
 	for(int i = 0; i < ParamsToClearCopy.Num(); ++i)
 	{
-		AddCurveValue(ParamsToClearCopy[i], 0.0f, ACF_DriveMaterial);
+		AddCurveValue(ParamsToClearCopy[i], 0.0f);
 	}
 
 	// @todo: delete me later when james g's change goes in
@@ -1703,7 +1701,7 @@ float UAnimInstance::PlaySlotAnimation(UAnimSequenceBase* Asset, FName SlotNodeN
 	return Montage_Play(NewMontage, InPlayRate);
 }
 
-UAnimMontage* UAnimInstance::PlaySlotAnimationAsDynamicMontage(UAnimSequenceBase* Asset, FName SlotNodeName, float BlendInTime, float BlendOutTime, float InPlayRate, int32 LoopCount, float BlendOutTriggerTime)
+UAnimMontage* UAnimInstance::PlaySlotAnimationAsDynamicMontage(UAnimSequenceBase* Asset, FName SlotNodeName, float BlendInTime, float BlendOutTime, float InPlayRate, int32 LoopCount, float BlendOutTriggerTime, float InTimeToStartMontageAt)
 {
 	// create temporary montage and play
 	bool bValidAsset = Asset && !Asset->IsA(UAnimMontage::StaticClass());
@@ -1763,7 +1761,7 @@ UAnimMontage* UAnimInstance::PlaySlotAnimationAsDynamicMontage(UAnimSequenceBase
 	NewMontage->SlotAnimTracks.Add(NewTrack);
 
 	// if playing is successful, return the montage to allow more control if needed
-	float PlayTime = Montage_Play(NewMontage, InPlayRate);
+	float PlayTime = Montage_Play(NewMontage, InPlayRate, EMontagePlayReturnType::MontageLength, InTimeToStartMontageAt);
 	return PlayTime > 0.0f ? NewMontage : NULL;
 }
 
@@ -1837,7 +1835,7 @@ bool UAnimInstance::IsPlayingSlotAnimation(const UAnimSequenceBase* Asset, FName
 }
 
 /** Play a Montage. Returns Length of Montage in seconds. Returns 0.f if failed to play. */
-float UAnimInstance::Montage_Play(UAnimMontage* MontageToPlay, float InPlayRate/*= 1.f*/, EMontagePlayReturnType ReturnValueType)
+float UAnimInstance::Montage_Play(UAnimMontage* MontageToPlay, float InPlayRate/*= 1.f*/, EMontagePlayReturnType ReturnValueType, float InTimeToStartMontageAt)
 {
 	if (MontageToPlay && (MontageToPlay->SequenceLength > 0.f) && MontageToPlay->HasValidSlotSetup())
 	{
@@ -1860,8 +1858,11 @@ float UAnimInstance::Montage_Play(UAnimMontage* MontageToPlay, float InPlayRate/
 			FAnimMontageInstance* NewInstance = new FAnimMontageInstance(this);
 			check(NewInstance);
 
+			const float MontageLength = MontageToPlay->SequenceLength;
+
 			NewInstance->Initialize(MontageToPlay);
 			NewInstance->Play(InPlayRate);
+			NewInstance->SetPosition(FMath::Clamp(InTimeToStartMontageAt, 0.f, MontageLength));
 			MontageInstances.Add(NewInstance);
 			ActiveMontagesMap.Add(MontageToPlay, NewInstance);
 
@@ -1876,8 +1877,6 @@ float UAnimInstance::Montage_Play(UAnimMontage* MontageToPlay, float InPlayRate/
 			UE_LOG(LogAnimMontage, Verbose, TEXT("Montage_Play: AnimMontage: %s,  (DesiredWeight:%0.2f, Weight:%0.2f)"),
 						*NewInstance->Montage->GetName(), NewInstance->GetDesiredWeight(), NewInstance->GetWeight());
 			
-			const float MontageLength = NewInstance->Montage->SequenceLength;
-			
 			return (ReturnValueType == EMontagePlayReturnType::MontageLength) ? MontageLength : (MontageLength/(InPlayRate*MontageToPlay->RateScale));
 		}
 		else
@@ -2454,6 +2453,12 @@ void UAnimInstance::ClearMontageInstanceReferences(FAnimMontageInstance& InMonta
 	InMontageInstance.MontageSync_StopLeading();
 }
 
+FAnimNode_SubInput* UAnimInstance::GetSubInputNode() const
+{
+	const FAnimInstanceProxy& Proxy = GetProxyOnAnyThread<FAnimInstanceProxy>();
+	return Proxy.SubInstanceInputNode;
+}
+
 FAnimMontageInstance* UAnimInstance::GetActiveInstanceForMontage(UAnimMontage const& Montage) const
 {
 	return GetActiveInstanceForMontage(&Montage);
diff --git a/Engine/Source/Runtime/Engine/Private/Animation/AnimInstanceProxy.cpp b/Engine/Source/Runtime/Engine/Private/Animation/AnimInstanceProxy.cpp
--- a/Engine/Source/Runtime/Engine/Private/Animation/AnimInstanceProxy.cpp
+++ b/Engine/Source/Runtime/Engine/Private/Animation/AnimInstanceProxy.cpp
@@ -729,10 +729,21 @@ void FAnimInstanceProxy::RecalcRequiredBones(USkeletalMeshComponent* Component,
 {
 	RequiredBones.InitializeTo(Component->RequiredBones, *Asset);
 
+	// If this instance can accept input poses, initialise the input pose container
+	if(SubInstanceInputNode)
+	{
+		SubInstanceInputNode->InputPose.SetBoneContainer(&RequiredBones);
+	}
+
 	// When RequiredBones mapping has changed, AnimNodes need to update their bones caches. 
 	bBoneCachesInvalidated = true;
 }
 
+void FAnimInstanceProxy::RecalcRequiredCurves()
+{
+	RequiredBones.CacheRequiredAnimCurveUids();
+}
+
 void FAnimInstanceProxy::UpdateAnimation()
 {
 	ANIM_MT_SCOPE_CYCLE_COUNTER(ProxyUpdateAnimation, !IsInGameThread());
@@ -833,7 +844,7 @@ void FAnimInstanceProxy::SlotEvaluatePose(FName SlotNodeName, const FCompactPose
 			
 			// Bone array has to be allocated prior to calling GetPoseFromAnimTrack
 			NewPose.Pose.SetBoneContainer(&RequiredBones);
-			NewPose.Curve.InitFrom(SkeletalMeshComponent->GetCachedAnimCurveMappingNameUids());
+			NewPose.Curve.InitFrom(RequiredBones);
 
 			// Extract pose from Track
 			FAnimExtractContext ExtractionContext(EvalState.MontagePosition, EvalState.Montage->HasRootMotion() && RootMotionMode != ERootMotionMode::NoRootMotionExtraction);
@@ -841,7 +852,7 @@ void FAnimInstanceProxy::SlotEvaluatePose(FName SlotNodeName, const FCompactPose
 
 			// add montage curves 
 			FBlendedCurve MontageCurve;
-			MontageCurve.InitFrom(SkeletalMeshComponent->GetCachedAnimCurveMappingNameUids());
+			MontageCurve.InitFrom(RequiredBones);
 			EvalState.Montage->EvaluateCurveData(MontageCurve, EvalState.MontagePosition);
 			NewPose.Curve.Combine(MontageCurve);
 
diff --git a/Engine/Source/Runtime/Engine/Private/Animation/AnimNodeBase.cpp b/Engine/Source/Runtime/Engine/Private/Animation/AnimNodeBase.cpp
--- a/Engine/Source/Runtime/Engine/Private/Animation/AnimNodeBase.cpp
+++ b/Engine/Source/Runtime/Engine/Private/Animation/AnimNodeBase.cpp
@@ -59,8 +59,9 @@ UAnimBlueprint* FAnimationBaseContext::GetAnimBlueprint() const
 void FPoseContext::Initialize(FAnimInstanceProxy* InAnimInstanceProxy)
 {
 	checkSlow(AnimInstanceProxy && AnimInstanceProxy->GetRequiredBones().IsValid());
-	Pose.SetBoneContainer(&AnimInstanceProxy->GetRequiredBones());
-	Curve.InitFrom(AnimInstanceProxy->GetSkelMeshComponent()->GetCachedAnimCurveMappingNameUids());
+	const FBoneContainer& RequiredBone = AnimInstanceProxy->GetRequiredBones();
+	Pose.SetBoneContainer(&RequiredBone);
+	Curve.InitFrom(RequiredBone);
 }
 
 /////////////////////////////////////////////////////
@@ -69,8 +70,9 @@ void FPoseContext::Initialize(FAnimInstanceProxy* InAnimInstanceProxy)
 void FComponentSpacePoseContext::ResetToRefPose()
 {
 	checkSlow(AnimInstanceProxy && AnimInstanceProxy->GetRequiredBones().IsValid());
-	Pose.InitPose(&AnimInstanceProxy->GetRequiredBones());
-	Curve.InitFrom(AnimInstanceProxy->GetSkelMeshComponent()->GetCachedAnimCurveMappingNameUids());
+	const FBoneContainer& RequiredBone = AnimInstanceProxy->GetRequiredBones();
+	Pose.InitPose(&RequiredBone);
+	Curve.InitFrom(RequiredBone);
 }
 
 /////////////////////////////////////////////////////
diff --git a/Engine/Source/Runtime/Engine/Private/Animation/AnimNode_TransitionPoseEvaluator.cpp b/Engine/Source/Runtime/Engine/Private/Animation/AnimNode_TransitionPoseEvaluator.cpp
--- a/Engine/Source/Runtime/Engine/Private/Animation/AnimNode_TransitionPoseEvaluator.cpp
+++ b/Engine/Source/Runtime/Engine/Private/Animation/AnimNode_TransitionPoseEvaluator.cpp
@@ -33,8 +33,9 @@ void FAnimNode_TransitionPoseEvaluator::Initialize(const FAnimationInitializeCon
 
 void FAnimNode_TransitionPoseEvaluator::CacheBones(const FAnimationCacheBonesContext& Context) 
 {
-	CachedPose.SetBoneContainer(&Context.AnimInstanceProxy->GetRequiredBones());
-	CachedCurve.InitFrom(Context.AnimInstanceProxy->GetSkelMeshComponent()->GetCachedAnimCurveMappingNameUids());
+	const FBoneContainer& RequiredBone = Context.AnimInstanceProxy->GetRequiredBones();
+	CachedPose.SetBoneContainer(&RequiredBone);
+	CachedCurve.InitFrom(RequiredBone);
 }
 
 void FAnimNode_TransitionPoseEvaluator::Update(const FAnimationUpdateContext& Context)
diff --git a/Engine/Source/Runtime/Engine/Private/Animation/AnimSequence.cpp b/Engine/Source/Runtime/Engine/Private/Animation/AnimSequence.cpp
--- a/Engine/Source/Runtime/Engine/Private/Animation/AnimSequence.cpp
+++ b/Engine/Source/Runtime/Engine/Private/Animation/AnimSequence.cpp
@@ -536,7 +536,6 @@ void UAnimSequence::VerifyTrackMap(USkeleton* MySkeleton)
 	{
 		ShowResaveMessage(this);
 
-		const TArray<FBoneNode>& BoneTree = UseSkeleton->GetBoneTree();
 		AnimationTrackNames.Empty();
 		AnimationTrackNames.AddUninitialized(TrackToSkeletonMapTable.Num());
 		for(int32 I=0; I<TrackToSkeletonMapTable.Num(); ++I)
@@ -572,7 +571,6 @@ void UAnimSequence::VerifyTrackMap(USkeleton* MySkeleton)
 			{
 				ShowResaveMessage(this);
 
-				const TArray<FBoneNode>& BoneTree = UseSkeleton->GetBoneTree();
 				for(int32 I=NumTracks-1; I>=0; --I)
 				{
 					int32 BoneTreeIndex = UseSkeleton->GetReferenceSkeleton().FindBoneIndex(AnimationTrackNames[I]);
@@ -1146,7 +1144,6 @@ void UAnimSequence::GetBonePose(FCompactPose& OutPose, FBlendedCurve& OutCurve,
 	}
 #endif // WITH_EDITOR
 
-	TArray<FBoneNode> const& BoneTree = MySkeleton->GetBoneTree();
 	TArray<int32> const& SkeletonToPoseBoneIndexArray = RequiredBones.GetSkeletonToPoseBoneIndexArray();
 
 	BoneTrackArray& RotationScalePairs = FGetBonePoseScratchArea::Get().RotationScalePairs;
@@ -1183,7 +1180,7 @@ void UAnimSequence::GetBonePose(FCompactPose& OutPose, FBlendedCurve& OutCurve,
 					RotationScalePairs.Add(BoneTrackPair(CompactPoseBoneIndex, TrackIndex));
 
 					// Skip extracting translation component for EBoneTranslationRetargetingMode::Skeleton.
-					switch (BoneTree[SkeletonBoneIndex].TranslationRetargetingMode)
+					switch (MySkeleton->GetBoneTranslationRetargetingMode(SkeletonBoneIndex))
 					{
 					case EBoneTranslationRetargetingMode::Animation:
 						TranslationPairs.Add(BoneTrackPair(CompactPoseBoneIndex, TrackIndex));
@@ -2221,10 +2218,10 @@ void UAnimSequence::BakeOutAdditiveIntoRawData()
 	{
 		// Initialise curve data from Skeleton
 		FBlendedCurve Curve;
-		Curve.InitFrom(&(MySkeleton->GetCachedAnimCurveMappingNameUids()));
+		Curve.InitFrom(RequiredBones);
 
 		FBlendedCurve DummyBaseCurve;
-		DummyBaseCurve.InitFrom(&(MySkeleton->GetCachedAnimCurveMappingNameUids()));
+		DummyBaseCurve.InitFrom(RequiredBones);
 
 		//Grab pose for this frame
 		const float CurrentFrameTime = Frame * IntervalTime;
@@ -2240,10 +2237,10 @@ void UAnimSequence::BakeOutAdditiveIntoRawData()
 		}
 
 		//Write out curve data for this frame
-		const TArray<FSmartNameMapping::UID>& UIDList = *Curve.UIDList;
+		const TArray<SmartName::UID_Type>& UIDList = *Curve.UIDList;
 		for (int32 CurveIndex = 0; CurveIndex < UIDList.Num(); ++CurveIndex)
 		{
-			FSmartNameMapping::UID CurveUID = UIDList[CurveIndex];
+			SmartName::UID_Type CurveUID = UIDList[CurveIndex];
 			FCurveElement& CurveEL = Curve.Elements[CurveIndex];
 			FFloatCurve* RawCurve = GetFloatCurve(NewCurveTracks, CurveUID);
 			if (!RawCurve && CurveEL.Value > 0.f) //Only make a new curve if we are going to give it data
@@ -2251,7 +2248,8 @@ void UAnimSequence::BakeOutAdditiveIntoRawData()
 				FSmartName NewCurveName;
 				// if we don't have name, there is something wrong here. 
 				ensureAlways(MySkeleton->GetSmartNameByUID(USkeleton::AnimCurveMappingName, CurveUID, NewCurveName));
-				NewCurveTracks.AddCurveData(NewCurveName, CurveEL.Flags, FRawCurveTracks::FloatType);
+				// curve flags don't matter much for compressed curves
+				NewCurveTracks.AddCurveData(NewCurveName, 0, FRawCurveTracks::FloatType);
 				RawCurve = GetFloatCurve(NewCurveTracks, CurveUID);
 			}
 
diff --git a/Engine/Source/Runtime/Engine/Private/Animation/AnimSequenceBase.cpp b/Engine/Source/Runtime/Engine/Private/Animation/AnimSequenceBase.cpp
--- a/Engine/Source/Runtime/Engine/Private/Animation/AnimSequenceBase.cpp
+++ b/Engine/Source/Runtime/Engine/Private/Animation/AnimSequenceBase.cpp
@@ -8,6 +8,7 @@
 #include "Animation/AnimSequenceBase.h"
 #include "Animation/AnimInstance.h"
 #include "MessageLog.h"
+#include "FrameworkObjectVersion.h"
 
 DEFINE_LOG_CATEGORY(LogAnimMarkerSync);
 
@@ -64,6 +65,25 @@ void UAnimSequenceBase::PostLoad()
 		VerifyCurveNames<FTransformCurve>(MySkeleton, USkeleton::AnimTrackCurveMappingName, RawCurveData.TransformCurves);
 #endif
 
+		// this should continue to add if skeleton hasn't been saved either 
+		// we don't wipe out data, so make sure you add back in if required
+		if (GetLinkerCustomVersion(FFrameworkObjectVersion::GUID) < FFrameworkObjectVersion::MoveCurveTypesToSkeleton
+			|| MySkeleton->GetLinkerCustomVersion(FFrameworkObjectVersion::GUID) < FFrameworkObjectVersion::MoveCurveTypesToSkeleton)
+		{
+			// fix up curve flags to skeleton
+			for (const FFloatCurve& Curve : RawCurveData.FloatCurves)
+			{
+				bool bMorphtargetSet = Curve.GetCurveTypeFlag(ACF_DriveMorphTarget_DEPRECATED);
+				bool bMaterialSet = Curve.GetCurveTypeFlag(ACF_DriveMaterial_DEPRECATED);
+
+				// only add this if that has to 
+				if (bMorphtargetSet || bMaterialSet)
+				{
+					MySkeleton->AccumulateCurveMetaData(Curve.Name.DisplayName, bMaterialSet, bMorphtargetSet);
+				}
+			}
+		}
+
 		RawCurveData.SortFloatCurvesByUID();
 	}
 }
@@ -550,6 +570,8 @@ void UAnimSequenceBase::EvaluateCurveData(FBlendedCurve& OutCurve, float Current
 
 void UAnimSequenceBase::Serialize(FArchive& Ar)
 {
+	Ar.UsingCustomVersion(FFrameworkObjectVersion::GUID);
+
 	Super::Serialize(Ar);
 
 	// fix up version issue and so on
diff --git a/Engine/Source/Runtime/Engine/Private/Animation/AnimSingleNodeInstance.cpp b/Engine/Source/Runtime/Engine/Private/Animation/AnimSingleNodeInstance.cpp
--- a/Engine/Source/Runtime/Engine/Private/Animation/AnimSingleNodeInstance.cpp
+++ b/Engine/Source/Runtime/Engine/Private/Animation/AnimSingleNodeInstance.cpp
@@ -375,10 +375,10 @@ void UAnimSingleNodeInstance::StepForward()
 	if (UAnimSequence* Sequence = Cast<UAnimSequence>(CurrentAsset))
 	{
 		FAnimSingleNodeInstanceProxy& Proxy = GetProxyOnGameThread<FAnimSingleNodeInstanceProxy>();
-
-		float KeyLength = Sequence->SequenceLength/Sequence->NumFrames+SMALL_NUMBER;
+		const FAnimKeyHelper Helper(Sequence->SequenceLength, Sequence->NumFrames);
+		float KeyLength = Helper.TimePerKey() + SMALL_NUMBER;
 		float Fraction = (Proxy.GetCurrentTime()+KeyLength)/Sequence->SequenceLength;
-		int32 Frames = FMath::Clamp<int32>((float)(Sequence->NumFrames*Fraction), 0, Sequence->NumFrames);
+		int32 Frames = FMath::Clamp<int32>((float)(Helper.LastKey()*Fraction), 0, Helper.LastKey());
 		SetPosition(Frames*KeyLength);
 	}	
 }
@@ -389,9 +389,10 @@ void UAnimSingleNodeInstance::StepBackward()
 	{
 		FAnimSingleNodeInstanceProxy& Proxy = GetProxyOnGameThread<FAnimSingleNodeInstanceProxy>();
 
-		float KeyLength = Sequence->SequenceLength/Sequence->NumFrames+SMALL_NUMBER;
+		const FAnimKeyHelper Helper(Sequence->SequenceLength, Sequence->NumFrames);
+		float KeyLength = Helper.TimePerKey() + SMALL_NUMBER;
 		float Fraction = (Proxy.GetCurrentTime()-KeyLength)/Sequence->SequenceLength;
-		int32 Frames = FMath::Clamp<int32>((float)(Sequence->NumFrames*Fraction), 0, Sequence->NumFrames);
+		int32 Frames = FMath::Clamp<int32>((float)(Helper.LastKey()*Fraction), 0, Helper.LastKey());
 		SetPosition(Frames*KeyLength);
 	}
 }
diff --git a/Engine/Source/Runtime/Engine/Private/Animation/AnimSingleNodeInstanceProxy.cpp b/Engine/Source/Runtime/Engine/Private/Animation/AnimSingleNodeInstanceProxy.cpp
--- a/Engine/Source/Runtime/Engine/Private/Animation/AnimSingleNodeInstanceProxy.cpp
+++ b/Engine/Source/Runtime/Engine/Private/Animation/AnimSingleNodeInstanceProxy.cpp
@@ -199,7 +199,7 @@ bool FAnimSingleNodeInstanceProxy::Evaluate(FPoseContext& Output)
 
 				for (const auto& PoseName : PoseNames)
 				{
-					if (PoseName.UID != FSmartNameMapping::MaxUID)
+					if (PoseName.UID != SmartName::MaxUID)
 					{
 						int32 PoseIndex = PoseNames.Find(PoseName);
 						if (PoseIndex != INDEX_NONE)
@@ -261,8 +261,7 @@ void FAnimSingleNodeInstanceProxy::PropagatePreviewCurve(FPoseContext& Output)
 
 		if (MySkeleton->GetSmartNameByName(USkeleton::AnimCurveMappingName, Name, PreviewCurveName))
 		{
-			Output.Curve.Set(PreviewCurveName.UID, Value, ACF_EditorPreviewCurves);
-
+			Output.Curve.Set(PreviewCurveName.UID, Value);
 		}
 	}
 }
diff --git a/Engine/Source/Runtime/Engine/Private/Animation/AnimationRuntime.cpp b/Engine/Source/Runtime/Engine/Private/Animation/AnimationRuntime.cpp
--- a/Engine/Source/Runtime/Engine/Private/Animation/AnimationRuntime.cpp
+++ b/Engine/Source/Runtime/Engine/Private/Animation/AnimationRuntime.cpp
@@ -1515,9 +1515,7 @@ void FAnimationRuntime::RetargetBoneTransform(const USkeleton* MySkeleton, const
 {
 	if (MySkeleton)
 	{
-		const TArray<FBoneNode>& BoneTree = MySkeleton->GetBoneTree();
-
-		switch (BoneTree[SkeletonBoneIndex].TranslationRetargetingMode)
+		switch (MySkeleton->GetBoneTranslationRetargetingMode(SkeletonBoneIndex))
 		{
 		case EBoneTranslationRetargetingMode::AnimationScaled:
 		{
diff --git a/Engine/Source/Runtime/Engine/Private/Animation/AssetMappingTable.cpp b/Engine/Source/Runtime/Engine/Private/Animation/AssetMappingTable.cpp
--- a/Engine/Source/Runtime/Engine/Private/Animation/AssetMappingTable.cpp
+++ b/Engine/Source/Runtime/Engine/Private/Animation/AssetMappingTable.cpp
@@ -20,7 +20,7 @@ bool FAssetMapping::IsValidMapping() const
 bool FAssetMapping::IsValidMapping(UAnimationAsset* InSourceAsset, UAnimationAsset* InTargetAsset) 
 {
 	// for now we only allow same class
-	return ( InSourceAsset && InTargetAsset &&
+	return ( InSourceAsset && InTargetAsset && InSourceAsset != InTargetAsset && 
 			InSourceAsset->StaticClass() == InTargetAsset->StaticClass() &&
 			InSourceAsset->GetSkeleton() == InTargetAsset->GetSkeleton() &&
 			InSourceAsset->IsValidAdditive() == InTargetAsset->IsValidAdditive()
@@ -33,24 +33,11 @@ bool FAssetMapping::SetTargetAsset(UAnimationAsset* InTargetAsset)
 	if (SourceAsset && InTargetAsset)
 	{
 		// if source and target is same, we clear target asset
-		if (SourceAsset == InTargetAsset)
+		if (IsValidMapping(SourceAsset, InTargetAsset))
 		{
-			TargetAsset = nullptr;
-			return false;
+			TargetAsset = InTargetAsset;
+			return true;
 		}
-		else
-		{
-			if (IsValidMapping(SourceAsset, InTargetAsset))
-			{
-				TargetAsset = InTargetAsset;
-				return true;
-			}
-		}
-	}
-	else if (SourceAsset && !InTargetAsset)
-	{
-		TargetAsset = nullptr;
-		return false;
 	}
 
 	return false;
diff --git a/Engine/Source/Runtime/Engine/Private/Animation/PoseAsset.cpp b/Engine/Source/Runtime/Engine/Private/Animation/PoseAsset.cpp
--- a/Engine/Source/Runtime/Engine/Private/Animation/PoseAsset.cpp
+++ b/Engine/Source/Runtime/Engine/Private/Animation/PoseAsset.cpp
@@ -6,6 +6,8 @@
 #include "Animation/PoseAsset.h"
 #include "FrameworkObjectVersion.h"
 
+#define LOCTEXT_NAMESPACE "PoseAsset"
+
 // utility function 
 #if WITH_EDITOR
 FSmartName GetUniquePoseName(USkeleton* Skeleton)
@@ -13,14 +15,14 @@ FSmartName GetUniquePoseName(USkeleton* Skeleton)
 	check(Skeleton);
 	int32 NameIndex = 0;
 
-	FSmartNameMapping::UID NewUID;
+	SmartName::UID_Type NewUID;
 	FName NewName;
 
 	do
 	{
 		NewName = FName(*FString::Printf(TEXT("Pose_%d"), NameIndex++));
 		NewUID = Skeleton->GetUIDByName(USkeleton::AnimCurveMappingName, NewName);
-	} while (NewUID != FSmartNameMapping::MaxUID);
+	} while (NewUID != SmartName::MaxUID);
 
 	// if found, 
 	FSmartName NewPoseName;
@@ -108,7 +110,7 @@ void FPoseDataContainer::GetPoseCurve(const FPoseData* PoseData, FBlendedCurve&
 		for (int32 CurveIndex = 0; CurveIndex < Curves.Num(); ++CurveIndex)
 		{
 			const FAnimCurveBase& Curve = Curves[CurveIndex];
-			OutCurve.Set(Curve.Name.UID, CurveValues[CurveIndex], Curve.GetCurveTypeFlags());
+			OutCurve.Set(Curve.Name.UID, CurveValues[CurveIndex]);
 		}
 	}
 }
@@ -576,6 +578,25 @@ void UPoseAsset::PostLoad()
 		{
 			MySkeleton->VerifySmartName(USkeleton::AnimCurveMappingName, Curve.Name);
 		}
+
+		// double loop but this check only should happen once per asset
+		// this should continue to add if skeleton hasn't been saved either 
+		if (GetLinkerCustomVersion(FFrameworkObjectVersion::GUID) < FFrameworkObjectVersion::MoveCurveTypesToSkeleton 
+			|| MySkeleton->GetLinkerCustomVersion(FFrameworkObjectVersion::GUID) < FFrameworkObjectVersion::MoveCurveTypesToSkeleton)
+		{
+			// fix up curve flags to skeleton
+			for (auto& Curve : PoseContainer.Curves)
+			{
+				bool bMorphtargetSet = Curve.GetCurveTypeFlag(ACF_DriveMorphTarget_DEPRECATED);
+				bool bMaterialSet = Curve.GetCurveTypeFlag(ACF_DriveMaterial_DEPRECATED);
+
+				// only add this if that has to 
+				if (bMorphtargetSet || bMaterialSet)
+				{
+					MySkeleton->AccumulateCurveMetaData(Curve.Name.DisplayName, bMaterialSet, bMorphtargetSet);
+				}
+			}
+		}
 	}
 
 	// I  have to fix pose names
@@ -717,9 +738,38 @@ bool UPoseAsset::ContainsPose(const FName& InPoseName) const
 }
 
 #if WITH_EDITOR
-void UPoseAsset::AddOrUpdatePoseWithUniqueName(USkeletalMeshComponent* MeshComponent)
+bool UPoseAsset::AddOrUpdatePoseWithUniqueName(USkeletalMeshComponent* MeshComponent, FSmartName* OutPoseName /*= nullptr*/)
 {
-	AddOrUpdatePose(GetUniquePoseName(GetSkeleton()), MeshComponent);
+	bool bSuccess = true;
+	bool bSavedAdditivePose = bAdditivePose;
+
+	// if it's already additive, convert to full pose first
+	if (bAdditivePose)
+	{
+		// convert to full pose
+		bSuccess = ConvertToFullPose();
+	}
+
+	if (bSuccess)
+	{
+		FSmartName NewPoseName = GetUniquePoseName(GetSkeleton());
+		AddOrUpdatePose(NewPoseName, MeshComponent);
+
+		if (OutPoseName)
+		{
+			*OutPoseName = NewPoseName;
+		}
+
+		// convert back to additive if it was that way
+		if (bSavedAdditivePose)
+		{
+			ConvertToAdditivePose(bSavedAdditivePose);
+		}
+
+		OnPoseListChanged.Broadcast();
+	}
+
+	return bSuccess;
 }
 
 void UPoseAsset::AddOrUpdatePose(const FSmartName& PoseName, USkeletalMeshComponent* MeshComponent)
@@ -746,12 +796,33 @@ void UPoseAsset::AddOrUpdatePose(const FSmartName& PoseName, USkeletalMeshCompon
 			}
 		}
 
-		AddOrUpdatePose(PoseName, TrackNames, BoneTransform);
+		const USkeleton* MeshSkeleton = MeshComponent->SkeletalMesh->Skeleton;
+		const FSmartNameMapping* Mapping = MeshSkeleton->GetSmartNameContainer(USkeleton::AnimCurveMappingName);
+		
+		TArray<float> NewCurveValues;
+		NewCurveValues.AddZeroed(PoseContainer.Curves.Num());
+
+		if(Mapping)
+		{
+			const FBlendedHeapCurve& MeshCurves = MeshComponent->GetAnimationCurves();
+ 
+			for (int32 NewCurveIndex = 0; NewCurveIndex < NewCurveValues.Num(); ++NewCurveIndex)
+			{
+				FAnimCurveBase& Curve = PoseContainer.Curves[NewCurveIndex];
+				if (const SmartName::UID_Type* CurveUID = Mapping->FindUID(Curve.Name.DisplayName))
+				{
+					const float MeshCurveValue = MeshCurves.Get(*CurveUID);
+					NewCurveValues[NewCurveIndex] = MeshCurveValue;
+				}
+			}
+		}
+
+		AddOrUpdatePose(PoseName, TrackNames, BoneTransform, NewCurveValues);
 		PoseContainer.Shrink(MySkeleton, RetargetSource);
 	}
 }
 
-void UPoseAsset::AddOrUpdatePose(const FSmartName& PoseName, TArray<FName> TrackNames, TArray<FTransform>& LocalTransform)
+void UPoseAsset::AddOrUpdatePose(const FSmartName& PoseName, const TArray<FName>& TrackNames, const TArray<FTransform>& LocalTransform, const TArray<float>& CurveValues)
 {
 	USkeleton* MySkeleton = GetSkeleton();
 	if (MySkeleton)
@@ -767,6 +838,8 @@ void UPoseAsset::AddOrUpdatePose(const FSmartName& PoseName, TArray<FName> Track
 		PoseData->LocalSpacePose.AddUninitialized(TotalTracks);
 		PoseData->LocalSpacePoseMask.SetNumZeroed(TotalTracks, true);
 		PoseContainer.FillUpDefaultPose(PoseData, MySkeleton, RetargetSource);
+		check(CurveValues.Num() == PoseContainer.Curves.Num());
+		PoseData->CurveData = CurveValues;
 
 		const FReferenceSkeleton& RefSkeleton = MySkeleton->GetReferenceSkeleton();
 		for (int32 Index = 0; Index < TrackNames.Num(); ++Index)
@@ -848,13 +921,13 @@ void UPoseAsset::CreatePoseFromAnimation(class UAnimSequence* AnimSequence, cons
 
 				// @Todo fill up curve data
 				TArray<float> CurveData;
-				float IntervalBetweenKeys = AnimSequence->SequenceLength / NumPoses;
+				float IntervalBetweenKeys = (NumPoses > 1)? AnimSequence->SequenceLength / (NumPoses -1 ) : 0.f;
 
 				// add curves - only float curves
 				const int32 TotalFloatCurveCount = AnimSequence->RawCurveData.FloatCurves.Num();
 				
 				// have to construct own UIDList;
-				TArray<FSmartNameMapping::UID> UIDList;
+				TArray<SmartName::UID_Type> UIDList;
 
 				if (TotalFloatCurveCount > 0)
 				{
@@ -937,7 +1010,7 @@ void UPoseAsset::UpdatePoseFromAnimation(class UAnimSequence* AnimSequence)
 }
 #endif // WITH_EDITOR
 
-bool UPoseAsset::ModifyPoseName(FName OldPoseName, FName NewPoseName, const FSmartNameMapping::UID* NewUID)
+bool UPoseAsset::ModifyPoseName(FName OldPoseName, FName NewPoseName, const SmartName::UID_Type* NewUID)
 {
 	USkeleton* MySkeleton = GetSkeleton();
 
@@ -961,6 +1034,7 @@ bool UPoseAsset::ModifyPoseName(FName OldPoseName, FName NewPoseName, const FSma
 		else
 		{
 			// we're renaming current one
+			MySkeleton->Modify();
 			MySkeleton->RenameSmartName(USkeleton::AnimCurveMappingName, OldPoseSmartName.UID, NewPoseName);
 			NewPoseSmartName = OldPoseSmartName;
 			NewPoseSmartName.DisplayName = NewPoseName;
@@ -1199,3 +1273,4 @@ bool UPoseAsset::GetBasePoseTransform(TArray<FTransform>& OutBasePose, TArray<fl
 	return false;
 }
 
+#undef LOCTEXT_NAMESPACE 
\ No newline at end of file
diff --git a/Engine/Source/Runtime/Engine/Private/Animation/PreviewAssetAttachComponent.cpp b/Engine/Source/Runtime/Engine/Private/Animation/PreviewAssetAttachComponent.cpp
--- a/Engine/Source/Runtime/Engine/Private/Animation/PreviewAssetAttachComponent.cpp
+++ b/Engine/Source/Runtime/Engine/Private/Animation/PreviewAssetAttachComponent.cpp
@@ -11,6 +11,22 @@ void FPreviewAssetAttachContainer::AddAttachedObject( UObject* AttachObject, FNa
 	AttachedObjects.Add( Pair );
 }
 
+void FPreviewAssetAttachContainer::AddUniqueAttachedObject(UObject* AttachObject, FName AttachPointName)
+{
+	for (const FPreviewAttachedObjectPair& AttachedObject : AttachedObjects)
+	{
+		if (AttachedObject.AttachedTo == AttachPointName)
+		{
+			return;
+		}
+	}
+
+	FPreviewAttachedObjectPair Pair;
+	Pair.AttachedTo = AttachPointName;
+	Pair.SetAttachedObject(AttachObject);
+	AttachedObjects.Add(Pair);
+}
+
 void FPreviewAssetAttachContainer::RemoveAttachedObject( UObject* ObjectToRemove, FName AttachName )
 {
 	for( int i = AttachedObjects.Num() - 1; i >= 0; --i )
@@ -52,6 +68,11 @@ FPreviewAttachedObjectPair& FPreviewAssetAttachContainer::operator []( int32 i )
 	return AttachedObjects[i];
 }
 
+const FPreviewAttachedObjectPair& FPreviewAssetAttachContainer::operator [](int32 i) const
+{
+	return AttachedObjects[i];
+}
+
 TConstIterator FPreviewAssetAttachContainer::CreateConstIterator() const
 {
 	return AttachedObjects.CreateConstIterator();
diff --git a/Engine/Source/Runtime/Engine/Private/Animation/Skeleton.cpp b/Engine/Source/Runtime/Engine/Private/Animation/Skeleton.cpp
--- a/Engine/Source/Runtime/Engine/Private/Animation/Skeleton.cpp
+++ b/Engine/Source/Runtime/Engine/Private/Animation/Skeleton.cpp
@@ -12,6 +12,8 @@
 #include "Animation/Rig.h"
 #include "Animation/BlendProfile.h"
 #include "MessageLog.h"
+#include "Engine/SkeletalMeshSocket.h"
+#include "Engine/PreviewMeshCollection.h"
 
 #define LOCTEXT_NAMESPACE "Skeleton"
 #define ROOT_BONE_PARENT	INDEX_NONE
@@ -52,6 +54,8 @@ USkeleton::USkeleton(const FObjectInitializer& ObjectInitializer)
 {
 	// Make sure we have somewhere for curve names.
 	SmartNames.AddContainer(AnimCurveMappingName);
+
+	AnimCurveUidVersion = 0;
 }
 
 void USkeleton::PostInitProperties()
@@ -80,7 +84,7 @@ void USkeleton::PostLoad()
 	check(Guid.IsValid());
 
 	// Cache smart name uids for animation curve names
-	CacheAnimCurveMappingNameUids();
+	IncreaseAnimCurveUidVersion();
 }
 
 void USkeleton::PostDuplicate(bool bDuplicateForPIE)
@@ -115,6 +119,7 @@ void USkeleton::Serialize( FArchive& Ar )
 
 			FName RetargetSourceName;
 			FReferencePose RetargetSource;
+			AnimRetargetSources.Empty();
 			for (int32 Index=0; Index<NumOfRetargetSources; ++Index)
 			{
 				Ar << RetargetSourceName;
@@ -454,13 +459,13 @@ void USkeleton::RebuildLinkup(const USkeletalMesh* InSkelMesh)
 
 void USkeleton::UpdateReferencePoseFromMesh(const USkeletalMesh* InSkelMesh)
 {
-	check (InSkelMesh);
+	check(InSkelMesh);
 
-	for (int32 BoneIndex=0; BoneIndex<ReferenceSkeleton.GetNum(); BoneIndex++)
+	for (int32 BoneIndex = 0; BoneIndex < ReferenceSkeleton.GetNum(); BoneIndex++)
 	{
 		// find index from ref pose array
 		const int32 MeshBoneIndex = InSkelMesh->RefSkeleton.FindBoneIndex(ReferenceSkeleton.GetBoneName(BoneIndex));
-		if( MeshBoneIndex != INDEX_NONE )
+		if (MeshBoneIndex != INDEX_NONE)
 		{
 			ReferenceSkeleton.UpdateRefPoseTransform(BoneIndex, InSkelMesh->RefSkeleton.GetRefBonePose()[MeshBoneIndex]);
 		}
@@ -848,6 +853,23 @@ void USkeleton::SetPreviewMesh(USkeletalMesh* PreviewMesh, bool bMarkAsDirty/*=t
 	PreviewSkeletalMesh = PreviewMesh;
 }
 
+void USkeleton::LoadAdditionalPreviewSkeletalMeshes()
+{
+	AdditionalPreviewSkeletalMeshes.LoadSynchronous();
+}
+
+UPreviewMeshCollection* USkeleton::GetAdditionalPreviewSkeletalMeshes() const
+{
+	return AdditionalPreviewSkeletalMeshes.Get();
+}
+
+void USkeleton::SetAdditionalPreviewSkeletalMeshes(UPreviewMeshCollection* PreviewMeshCollection)
+{
+	Modify();
+
+	AdditionalPreviewSkeletalMeshes = PreviewMeshCollection;
+}
+
 int32 USkeleton::ValidatePreviewAttachedObjects()
 {
 	int32 NumBrokenAssets = PreviewAttachedAssetContainer.ValidatePreviewAttachedObjects();
@@ -959,7 +981,15 @@ FAnimSlotGroup* USkeleton::FindAnimSlotGroup(const FName& InGroupName)
 	});
 }
 
-bool USkeleton::ContainsSlotName(const FName& InSlotName) 
+const FAnimSlotGroup* USkeleton::FindAnimSlotGroup(const FName& InGroupName) const
+{
+	return SlotGroups.FindByPredicate([&InGroupName](const FAnimSlotGroup& Item)
+	{
+		return Item.GroupName == InGroupName;
+	});
+}
+
+bool USkeleton::ContainsSlotName(const FName& InSlotName) const
 {
 	return SlotToGroupNameMap.Contains(InSlotName);
 }
@@ -1082,12 +1112,12 @@ bool USkeleton::AddSmartNameAndModify(FName ContainerName, FName NewDisplayName,
 		{
 			Modify(true);
 			Successful = true;
-			CacheAnimCurveMappingNameUids();
+			IncreaseAnimCurveUidVersion();
 		}
 	}
 	return Successful;
 }
-bool USkeleton::RenameSmartnameAndModify(FName ContainerName, FSmartNameMapping::UID Uid, FName NewName)
+bool USkeleton::RenameSmartnameAndModify(FName ContainerName, SmartName::UID_Type Uid, FName NewName)
 {
 	bool Successful = false;
 	FSmartNameMapping* RequestedMapping = SmartNames.GetContainerInternal(ContainerName);
@@ -1099,13 +1129,13 @@ bool USkeleton::RenameSmartnameAndModify(FName ContainerName, FSmartNameMapping:
 		{
 			Modify();
 			Successful = RequestedMapping->Rename(Uid, NewName);
-			CacheAnimCurveMappingNameUids();
+			IncreaseAnimCurveUidVersion();
 		}
 	}
 	return Successful;
 }
 
-void USkeleton::RemoveSmartnameAndModify(FName ContainerName, FSmartNameMapping::UID Uid)
+void USkeleton::RemoveSmartnameAndModify(FName ContainerName, SmartName::UID_Type Uid)
 {
 	FSmartNameMapping* RequestedMapping = SmartNames.GetContainerInternal(ContainerName);
 	if (RequestedMapping)
@@ -1113,18 +1143,18 @@ void USkeleton::RemoveSmartnameAndModify(FName ContainerName, FSmartNameMapping:
 		Modify();
 		if (RequestedMapping->Remove(Uid))
 		{
-			CacheAnimCurveMappingNameUids();
+			IncreaseAnimCurveUidVersion();
 		}
 	}
 }
 
-void USkeleton::RemoveSmartnamesAndModify(FName ContainerName, const TArray<FSmartNameMapping::UID>& Uids)
+void USkeleton::RemoveSmartnamesAndModify(FName ContainerName, const TArray<SmartName::UID_Type>& Uids)
 {
 	FSmartNameMapping* RequestedMapping = SmartNames.GetContainerInternal(ContainerName);
 	if (RequestedMapping)
 	{
 		bool bModified = false;
-		for (const FSmartNameMapping::UID& UID : Uids)
+		for (const SmartName::UID_Type& UID : Uids)
 		{
 			if (RequestedMapping->Exists(UID))
 			{
@@ -1136,13 +1166,13 @@ void USkeleton::RemoveSmartnamesAndModify(FName ContainerName, const TArray<FSma
 		if (bModified)
 		{
 			Modify();
-			CacheAnimCurveMappingNameUids();
+			IncreaseAnimCurveUidVersion();
 		}
 	}
 }
 #endif // WITH_EDITOR
 
-bool USkeleton::GetSmartNameByUID(const FName& ContainerName, FSmartNameMapping::UID UID, FSmartName& OutSmartName)
+bool USkeleton::GetSmartNameByUID(const FName& ContainerName, SmartName::UID_Type UID, FSmartName& OutSmartName)
 {
 	const FSmartNameMapping* RequestedMapping = SmartNames.GetContainerInternal(ContainerName);
 	if (RequestedMapping)
@@ -1164,12 +1194,12 @@ bool USkeleton::GetSmartNameByName(const FName& ContainerName, const FName& InNa
 	return false;
 }
 
-FSmartNameMapping::UID USkeleton::GetUIDByName(const FName& ContainerName, const FName& Name)
+SmartName::UID_Type USkeleton::GetUIDByName(const FName& ContainerName, const FName& Name) const
 {
 	const FSmartNameMapping* RequestedMapping = SmartNames.GetContainerInternal(ContainerName);
 	if (RequestedMapping)
 	{
-		if (const FSmartNameMapping::UID* UID = RequestedMapping->FindUID(Name))
+		if (const SmartName::UID_Type* UID = RequestedMapping->FindUID(Name))
 		{
 			if (UID)
 			{
@@ -1178,7 +1208,7 @@ FSmartNameMapping::UID USkeleton::GetUIDByName(const FName& ContainerName, const
 		}
 	}
 
-	return FSmartNameMapping::MaxUID;
+	return SmartName::MaxUID;
 }
 
 // this does very simple thing.
@@ -1189,7 +1219,7 @@ void USkeleton::VerifySmartName(const FName& ContainerName, FSmartName& InOutSma
 	VerifySmartNameInternal(ContainerName, InOutSmartName);
 	if (ContainerName == USkeleton::AnimCurveMappingName)
 	{
-		CacheAnimCurveMappingNameUids();
+		IncreaseAnimCurveUidVersion();
 	}
 }
 
@@ -1285,7 +1315,7 @@ void USkeleton::VerifySmartNames(const FName&  ContainerName, TArray<FSmartName>
 
 	if (ContainerName == USkeleton::AnimCurveMappingName)
 	{
-		CacheAnimCurveMappingNameUids();
+		IncreaseAnimCurveUidVersion();
 	}
 }
 
@@ -1295,7 +1325,7 @@ FSmartNameMapping* USkeleton::GetOrAddSmartNameContainer(const FName& ContainerN
 	if (Mapping == nullptr)
 	{
 		Modify();
-		CacheAnimCurveMappingNameUids();
+		IncreaseAnimCurveUidVersion();
 		SmartNames.AddContainer(ContainerName);
 		Mapping = SmartNames.GetContainerInternal(ContainerName);		
 	}
@@ -1314,13 +1344,13 @@ void USkeleton::RegenerateGuid()
 	check(Guid.IsValid());
 }
 
-bool USkeleton::RenameSmartName(FName ContainerName, const FSmartNameMapping::UID& Uid, FName NewName)
+bool USkeleton::RenameSmartName(FName ContainerName, const SmartName::UID_Type& Uid, FName NewName)
 {
 	FSmartNameMapping* Mapping = SmartNames.GetContainerInternal(ContainerName);
 	if (Mapping->Exists(Uid))
 	{
 		Modify();
-		CacheAnimCurveMappingNameUids();
+		IncreaseAnimCurveUidVersion();
 		Mapping->Rename(Uid, NewName);
 		return true;
 	}
@@ -1328,18 +1358,81 @@ bool USkeleton::RenameSmartName(FName ContainerName, const FSmartNameMapping::UI
 	return false;
 }
 
-void USkeleton::CacheAnimCurveMappingNameUids()
+void USkeleton::IncreaseAnimCurveUidVersion()
+{
+	// this will be checked by skeletalmeshcomponent and if it's same, it won't care. If it's different, it will rebuild UID list
+	++AnimCurveUidVersion;
+}
+
+FCurveMetaData* USkeleton::GetCurveMetaData(const FName& CurveName)
+{
+	FSmartNameMapping* Mapping = SmartNames.GetContainerInternal(USkeleton::AnimCurveMappingName);
+	if (ensureAlways(Mapping))
+	{
+		return Mapping->GetCurveMetaData(CurveName);
+	}
+
+	return nullptr;
+}
+
+const FCurveMetaData* USkeleton::GetCurveMetaData(const FName& CurveName) const
 {
 	const FSmartNameMapping* Mapping = SmartNames.GetContainerInternal(USkeleton::AnimCurveMappingName);
-	if (Mapping != nullptr)
+	if (ensureAlways(Mapping))
 	{
-		Mapping->FillUidArray(CachedAnimCurveMappingNameUids);
+		return Mapping->GetCurveMetaData(CurveName);
 	}
+
+	return nullptr;
 }
 
-const TArray<FSmartNameMapping::UID>& USkeleton::GetCachedAnimCurveMappingNameUids()
+FCurveMetaData* USkeleton::GetCurveMetaData(const FSmartName& CurveName)
 {
-	return CachedAnimCurveMappingNameUids;
+	FSmartNameMapping* Mapping = SmartNames.GetContainerInternal(USkeleton::AnimCurveMappingName);
+	if (ensureAlways(Mapping))
+	{
+		// the name might have changed, make sure it's up-to-date
+		FName DisplayName;
+		Mapping->GetName(CurveName.UID, DisplayName);
+		return Mapping->GetCurveMetaData(DisplayName);
+	}
+
+	return nullptr;
+}
+
+const FCurveMetaData* USkeleton::GetCurveMetaData(const FSmartName& CurveName) const
+{
+	const FSmartNameMapping* Mapping = SmartNames.GetContainerInternal(USkeleton::AnimCurveMappingName);
+	if (ensureAlways(Mapping))
+	{
+		// the name might have changed, make sure it's up-to-date
+		FName DisplayName;
+		Mapping->GetName(CurveName.UID, DisplayName);
+		return Mapping->GetCurveMetaData(DisplayName);
+	}
+
+	return nullptr;
+}
+
+// this is called when you know both flags - called by post serialize
+void USkeleton::AccumulateCurveMetaData(FName CurveName, bool bMaterialSet, bool bMorphtargetSet)
+{
+	if (bMaterialSet || bMorphtargetSet)
+	{
+		const FSmartNameMapping* Mapping = SmartNames.GetContainerInternal(USkeleton::AnimCurveMappingName);
+		if (ensureAlways(Mapping))
+		{
+			// if we don't have name, add one
+			if (Mapping->Exists(CurveName))
+			{
+				FCurveMetaData* CurveMetaData = GetCurveMetaData(CurveName);
+				// we don't want to undo previous flags, if it was true, we just alolw more to it. 
+				CurveMetaData->Type.bMaterial |= bMaterialSet;
+				CurveMetaData->Type.bMorphtarget |= bMorphtargetSet;
+				MarkPackageDirty();
+			}
+		}
+	}
 }
 
 #if WITH_EDITOR
@@ -1506,7 +1599,7 @@ void USkeleton::GetAssetRegistryTags(TArray<FAssetRegistryTag>& OutTags) const
 UBlendProfile* USkeleton::CreateNewBlendProfile(const FName& InProfileName)
 {
 	Modify();
-	UBlendProfile* NewProfile = NewObject<UBlendProfile>(this, InProfileName, RF_Public);
+	UBlendProfile* NewProfile = NewObject<UBlendProfile>(this, InProfileName, RF_Public | RF_Transactional);
 	BlendProfiles.Add(NewProfile);
 
 	return NewProfile;
@@ -1527,4 +1620,32 @@ UBlendProfile* USkeleton::GetBlendProfile(const FName& InProfileName)
 }
 
 #endif //WITH_EDITOR
+
+USkeletalMeshSocket* USkeleton::FindSocket(FName InSocketName) const
+{
+	int32 DummyIndex;
+	return FindSocketAndIndex(InSocketName, DummyIndex);
+}
+
+USkeletalMeshSocket* USkeleton::FindSocketAndIndex(FName InSocketName, int32& OutIndex) const
+{
+	OutIndex = INDEX_NONE;
+	if (InSocketName == NAME_None)
+	{
+		return nullptr;
+	}
+
+	for (int32 i = 0; i < Sockets.Num(); ++i)
+	{
+		USkeletalMeshSocket* Socket = Sockets[i];
+		if (Socket && Socket->SocketName == InSocketName)
+		{
+			OutIndex = i;
+			return Socket;
+		}
+	}
+
+	return nullptr;
+}
+
 #undef LOCTEXT_NAMESPACE 
\ No newline at end of file
diff --git a/Engine/Source/Runtime/Engine/Private/Animation/SmartName.cpp b/Engine/Source/Runtime/Engine/Private/Animation/SmartName.cpp
--- a/Engine/Source/Runtime/Engine/Private/Animation/SmartName.cpp
+++ b/Engine/Source/Runtime/Engine/Private/Animation/SmartName.cpp
@@ -12,12 +12,12 @@ FSmartNameMapping::FSmartNameMapping()
 : NextUid(0)
 {}
 
-bool FSmartNameMapping::AddOrFindName(FName Name, UID& OutUid, FGuid& OutGuid)
+bool FSmartNameMapping::AddOrFindName(FName Name, SmartName::UID_Type& OutUid, FGuid& OutGuid)
 {
 	check(Name.IsValid());
 
-	// Check for UID overflow
-	const UID* ExistingUid = UidMap.FindKey(Name);
+	// Check for SmartName::UID_Type overflow
+	const SmartName::UID_Type* ExistingUid = UidMap.FindKey(Name);
 	const FGuid* ExistingGuid = GuidMap.Find(Name);
 
 	// make sure they both exists and same 
@@ -35,13 +35,13 @@ bool FSmartNameMapping::AddOrFindName(FName Name, UID& OutUid, FGuid& OutGuid)
 	return AddName(Name, OutUid, OutGuid);
 }
 
-bool FSmartNameMapping::AddName(FName Name, UID& OutUid, const FGuid& InGuid)
+bool FSmartNameMapping::AddName(FName Name, SmartName::UID_Type& OutUid, const FGuid& InGuid)
 {
 	check(Name.IsValid());
 	if (GuidMap.Find(Name) == nullptr && GuidMap.FindKey(InGuid) == nullptr)
 	{
 		// make sure we didn't reach till end
-		check(NextUid != MaxUID);
+		check(NextUid != SmartName::MaxUID);
 
 		OutUid = NextUid;
 		UidMap.Add(OutUid, Name);
@@ -55,7 +55,7 @@ bool FSmartNameMapping::AddName(FName Name, UID& OutUid, const FGuid& InGuid)
 	return false;
 }
 
-bool FSmartNameMapping::GetName(const UID& Uid, FName& OutName) const
+bool FSmartNameMapping::GetName(const SmartName::UID_Type& Uid, FName& OutName) const
 {
 	const FName* FoundName = UidMap.Find(Uid);
 	if(FoundName)
@@ -78,7 +78,7 @@ bool FSmartNameMapping::GetNameByGuid(const FGuid& Guid, FName& OutName) const
 	return false;
 }
 
-bool FSmartNameMapping::Rename(const UID& Uid, FName NewName)
+bool FSmartNameMapping::Rename(const SmartName::UID_Type& Uid, FName NewName)
 {
 	FName* ExistingName = UidMap.Find(Uid);
 	if(ExistingName)
@@ -90,6 +90,17 @@ bool FSmartNameMapping::Rename(const UID& Uid, FName NewName)
 		GuidMap.Remove(*ExistingName);
 		GuidMap.Add(NewName, *Guid);
 
+		// fix up meta data
+		FCurveMetaData* MetaDataToCopy = CurveMetaDataMap.Find(*ExistingName);
+		if (MetaDataToCopy)
+		{
+			FCurveMetaData& NewMetaData = CurveMetaDataMap.Add(NewName);
+			NewMetaData = *MetaDataToCopy;
+			
+			// remove old one
+			CurveMetaDataMap.Remove(*ExistingName);			
+		}
+
 		check(GuidMap.Num() == UidMap.Num());
 
 		*ExistingName = NewName;
@@ -98,7 +109,7 @@ bool FSmartNameMapping::Rename(const UID& Uid, FName NewName)
 	return false;
 }
 
-bool FSmartNameMapping::Remove(const UID& Uid)
+bool FSmartNameMapping::Remove(const SmartName::UID_Type& Uid)
 {
 	FName* ExistingName = UidMap.Find(Uid);
 	if (ExistingName)
@@ -108,6 +119,7 @@ bool FSmartNameMapping::Remove(const UID& Uid)
 
 		// re add new value
 		GuidMap.Remove(*ExistingName);
+		CurveMetaDataMap.Remove(*ExistingName);
 		UidMap.Remove(Uid);
 
 		check(GuidMap.Num() == UidMap.Num());
@@ -152,11 +164,11 @@ void FSmartNameMapping::Serialize(FArchive& Ar)
 		Ar << NextUid;
 		Ar << UidMap;
 
-		TMap<UID, FName> NewUidMap;
+		TMap<SmartName::UID_Type, FName> NewUidMap;
 		// convert to GuidMap
 		NextUid = 0;
 		GuidMap.Empty(UidMap.Num());
-		for (TPair<UID, FName>& UidPair : UidMap)
+		for (TPair<SmartName::UID_Type, FName>& UidPair : UidMap)
 		{
 			GuidMap.Add(UidPair.Value, FGuid::NewGuid());
 			NewUidMap.Add(NextUid++, UidPair.Value);
@@ -164,6 +176,11 @@ void FSmartNameMapping::Serialize(FArchive& Ar)
 
 		UidMap = NewUidMap;
 	}
+
+	if (Ar.CustomVer(FFrameworkObjectVersion::GUID) >= FFrameworkObjectVersion::MoveCurveTypesToSkeleton)
+	{
+		Ar << CurveMetaDataMap;
+	}
 }
 
 int32 FSmartNameMapping::GetNumNames() const
@@ -171,7 +188,7 @@ int32 FSmartNameMapping::GetNumNames() const
 	return UidMap.Num();
 }
 
-void FSmartNameMapping::FillUidArray(TArray<UID>& Array) const
+void FSmartNameMapping::FillUidArray(TArray<SmartName::UID_Type>& Array) const
 {
 	UidMap.GenerateKeyArray(Array);
 }
@@ -181,7 +198,7 @@ void FSmartNameMapping::FillNameArray(TArray<FName>& Array) const
 	UidMap.GenerateValueArray(Array);
 }
 
-bool FSmartNameMapping::Exists(const UID& Uid) const
+bool FSmartNameMapping::Exists(const SmartName::UID_Type& Uid) const
 {
 	return UidMap.Find(Uid) != nullptr;
 }
@@ -191,7 +208,7 @@ bool FSmartNameMapping::Exists(const FName& Name) const
 	return UidMap.FindKey(Name) != nullptr;
 }
 
-const FSmartNameMapping::UID* FSmartNameMapping::FindUID(const FName& Name) const
+const SmartName::UID_Type* FSmartNameMapping::FindUID(const FName& Name) const
 {
 	return UidMap.FindKey(Name);
 }
@@ -206,7 +223,7 @@ FArchive& operator<<(FArchive& Ar, FSmartNameMapping& Elem)
 #if WITH_EDITOR
 bool FSmartNameMapping::FindOrAddSmartName(FName Name, FSmartName& OutName)
 {
-	FSmartNameMapping::UID NewUID;
+	SmartName::UID_Type NewUID;
 	FGuid NewGuid;
 	bool bNewlyAdded = AddOrFindName(Name, NewUID, NewGuid);
 	OutName = FSmartName(Name, NewUID, NewGuid);
@@ -223,7 +240,7 @@ bool FSmartNameMapping::AddSmartName(FSmartName& OutName)
 
 // in cooked build, we don't have Guid, so just register with empty guid. 
 // you only should come here if it it hasn't found yet. 
-bool FSmartNameMapping::FindOrAddSmartName(FName Name, UID& OutUid)
+bool FSmartNameMapping::FindOrAddSmartName(FName Name, SmartName::UID_Type& OutUid)
 {
 	FSmartName FoundName;
 	if (FindSmartName(Name, FoundName))
@@ -241,7 +258,7 @@ bool FSmartNameMapping::FindOrAddSmartName(FName Name, UID& OutUid)
 
 bool FSmartNameMapping::FindSmartName(FName Name, FSmartName& OutName) const
 {
-	const FSmartNameMapping::UID* ExistingUID = FindUID(Name);
+	const SmartName::UID_Type* ExistingUID = FindUID(Name);
 	if (ExistingUID)
 	{
 #if WITH_EDITOR
@@ -256,7 +273,7 @@ bool FSmartNameMapping::FindSmartName(FName Name, FSmartName& OutName) const
 	return false;
 }
 
-bool FSmartNameMapping::FindSmartNameByUID(FSmartNameMapping::UID UID, FSmartName& OutName) const
+bool FSmartNameMapping::FindSmartNameByUID(SmartName::UID_Type UID, FSmartName& OutName) const
 {
 	FName ExistingName;
 	if (GetName(UID, ExistingName))
@@ -299,6 +316,11 @@ FSmartNameMapping* FSmartNameContainer::GetContainerInternal(const FName& Contai
 	return NameMappings.Find(ContainerName);
 }
 
+const FSmartNameMapping* FSmartNameContainer::GetContainerInternal(const FName& ContainerName) const
+{
+	return NameMappings.Find(ContainerName);
+}
+
 ////////////////////////////////////////////////////////////////////////
 //
 // FSmartName
@@ -318,4 +340,4 @@ bool FSmartName::Serialize(FArchive& Ar)
 #endif // WITH_EDITORONLY_DATA
 
 	return true;
-}
\ No newline at end of file
+}
diff --git a/Engine/Source/Runtime/Engine/Private/AnimationAsset.cpp b/Engine/Source/Runtime/Engine/Private/AnimationAsset.cpp
--- a/Engine/Source/Runtime/Engine/Private/AnimationAsset.cpp
+++ b/Engine/Source/Runtime/Engine/Private/AnimationAsset.cpp
@@ -318,7 +318,7 @@ void UAnimationAsset::ReplaceReferredAnimations(const TMap<UAnimationAsset*, UAn
 // 	}
 }
 
-USkeletalMesh* UAnimationAsset::GetPreviewMesh() 
+USkeletalMesh* UAnimationAsset::GetPreviewMesh()
 {
 	USkeletalMesh* PreviewMesh = PreviewSkeletalMesh.Get();
 	if(!PreviewMesh)
@@ -341,6 +341,11 @@ USkeletalMesh* UAnimationAsset::GetPreviewMesh()
 	return PreviewMesh;
 }
 
+USkeletalMesh* UAnimationAsset::GetPreviewMesh() const
+{
+	return PreviewSkeletalMesh.Get();
+}
+
 void UAnimationAsset::SetPreviewMesh(USkeletalMesh* PreviewMesh)
 {
 	Modify();
@@ -499,253 +504,11 @@ const TArray<UAssetUserData*>* UAnimationAsset::GetAssetUserDataArray() const
 {
 	return &AssetUserData;
 }
-
-//////////////////////////////////////////////////////////////////////////
-// FBoneContainer
-
-FBoneContainer::FBoneContainer()
-: Asset(NULL)
-, AssetSkeletalMesh(NULL)
-, AssetSkeleton(NULL)
-, RefSkeleton(NULL)
-, bDisableRetargeting(false)
-, bUseRAWData(false)
-, bUseSourceData(false)
-{
-	BoneIndicesArray.Empty();
-	BoneSwitchArray.Empty();
-	SkeletonToPoseBoneIndexArray.Empty();
-	PoseToSkeletonBoneIndexArray.Empty();
-}
-
-FBoneContainer::FBoneContainer(const TArray<FBoneIndexType>& InRequiredBoneIndexArray, UObject& InAsset)
-: BoneIndicesArray(InRequiredBoneIndexArray)
-, Asset(&InAsset)
-, AssetSkeletalMesh(NULL)
-, AssetSkeleton(NULL)
-, RefSkeleton(NULL)
-, bDisableRetargeting(false)
-, bUseRAWData(false)
-, bUseSourceData(false)
-{
-	Initialize();
-}
-
-void FBoneContainer::InitializeTo(const TArray<FBoneIndexType>& InRequiredBoneIndexArray, UObject& InAsset)
-{
-	BoneIndicesArray = InRequiredBoneIndexArray;
-	Asset = &InAsset;
-
-	Initialize();
-}
-
-struct FBoneContainerScratchArea : public TThreadSingleton<FBoneContainerScratchArea>
-{
-	TArray<int32> MeshIndexToCompactPoseIndex;
-};
-
-void FBoneContainer::Initialize()
-{
-	RefSkeleton = NULL;
-	UObject* AssetObj = Asset.Get();
-	USkeletalMesh* AssetSkeletalMeshObj = Cast<USkeletalMesh>(AssetObj);
-	USkeleton* AssetSkeletonObj = nullptr;
-
-	if (AssetSkeletalMeshObj)
-	{
-		RefSkeleton = &AssetSkeletalMeshObj->RefSkeleton;
-		AssetSkeletonObj = AssetSkeletalMeshObj->Skeleton;
-	}
-	else
-	{
-		AssetSkeletonObj = Cast<USkeleton>(AssetObj);
-		if (AssetSkeletonObj)
-		{
-			RefSkeleton = &AssetSkeletonObj->GetReferenceSkeleton();
-		}
-	}
-
-	// Only supports SkeletalMeshes or Skeletons.
-	check( AssetSkeletalMeshObj || AssetSkeletonObj );
-	// Skeleton should always be there.
-	checkf( AssetSkeletonObj, TEXT("%s missing skeleton"), *GetNameSafe(AssetSkeletalMeshObj));
-	check( RefSkeleton );
-
-	AssetSkeleton = AssetSkeletonObj;
-	AssetSkeletalMesh = AssetSkeletalMeshObj;
-
-	// Take biggest amount of bones between SkeletalMesh and Skeleton for BoneSwitchArray.
-	// SkeletalMesh can have less, but AnimSequences tracks will map to Skeleton which can have more.
-	const int32 MaxBones = AssetSkeletonObj ? FMath::Max<int32>(RefSkeleton->GetNum(), AssetSkeletonObj->GetReferenceSkeleton().GetNum()) : RefSkeleton->GetNum();
-
-	// Initialize BoneSwitchArray.
-	BoneSwitchArray.Init(false, MaxBones);
-	const int32 NumRequiredBones = BoneIndicesArray.Num();
-	for(int32 Index=0; Index<NumRequiredBones; Index++)
-	{
-		const FBoneIndexType BoneIndex = BoneIndicesArray[Index];
-		checkSlow( BoneIndex < MaxBones );
-		BoneSwitchArray[BoneIndex] = true;
-	}
-
-	// Clear remapping table
-	SkeletonToPoseBoneIndexArray.Reset();
-
-	// Cache our mapping tables
-	// Here we create look up tables between our target asset and its USkeleton's refpose.
-	// Most times our Target is a SkeletalMesh
-	if (AssetSkeletalMeshObj)
-	{
-		RemapFromSkelMesh(*AssetSkeletalMeshObj, *AssetSkeletonObj);
-	}
-	// But we also support a Skeleton's RefPose.
-	else
-	{
-		// Right now we only support a single Skeleton. Skeleton hierarchy coming soon!
-		RemapFromSkeleton(*AssetSkeletonObj);
-	}
-
-	//Set up compact pose data
-	int32 NumReqBones = BoneIndicesArray.Num();
-	CompactPoseParentBones.Reset(NumReqBones);
-
-	CompactPoseRefPoseBones.Reset(NumReqBones);
-	CompactPoseRefPoseBones.AddUninitialized(NumReqBones);
-
-	CompactPoseToSkeletonIndex.Reset(NumReqBones);
-	CompactPoseToSkeletonIndex.AddUninitialized(NumReqBones);
-
-	SkeletonToCompactPose.Reset(SkeletonToPoseBoneIndexArray.Num());
-
-	const TArray<FTransform>& RefPoseArray = RefSkeleton->GetRefBonePose();
-	TArray<int32>& MeshIndexToCompactPoseIndex = FBoneContainerScratchArea::Get().MeshIndexToCompactPoseIndex;
-	MeshIndexToCompactPoseIndex.Reset(PoseToSkeletonBoneIndexArray.Num());
-	MeshIndexToCompactPoseIndex.AddUninitialized(PoseToSkeletonBoneIndexArray.Num());
-
-	for (int32& Item : MeshIndexToCompactPoseIndex)
-	{
-		Item = -1;
-	}
-		
-	for (int32 CompactBoneIndex = 0; CompactBoneIndex < NumReqBones; ++CompactBoneIndex)
-	{
-		FBoneIndexType MeshPoseIndex = BoneIndicesArray[CompactBoneIndex];
-		MeshIndexToCompactPoseIndex[MeshPoseIndex] = CompactBoneIndex;
-
-		//Parent Bone
-		const int32 ParentIndex = GetParentBoneIndex(MeshPoseIndex);
-		const int32 CompactParentIndex = ParentIndex == INDEX_NONE ? INDEX_NONE : MeshIndexToCompactPoseIndex[ParentIndex];
-
-		CompactPoseParentBones.Add(FCompactPoseBoneIndex(CompactParentIndex));
-	}
-
-	//Ref Pose
-	for (int32 CompactBoneIndex = 0; CompactBoneIndex < NumReqBones; ++CompactBoneIndex)
-	{
-		FBoneIndexType MeshPoseIndex = BoneIndicesArray[CompactBoneIndex];
-		CompactPoseRefPoseBones[CompactBoneIndex] = RefPoseArray[MeshPoseIndex];
-	}
-
-	for (int32 CompactBoneIndex = 0; CompactBoneIndex < NumReqBones; ++CompactBoneIndex)
-	{
-		FBoneIndexType MeshPoseIndex = BoneIndicesArray[CompactBoneIndex];
-		CompactPoseToSkeletonIndex[CompactBoneIndex] = PoseToSkeletonBoneIndexArray[MeshPoseIndex];
-	}
-
-
-	for (int32 SkeletonBoneIndex = 0; SkeletonBoneIndex < SkeletonToPoseBoneIndexArray.Num(); ++SkeletonBoneIndex)
-	{
-		int32 PoseBoneIndex = SkeletonToPoseBoneIndexArray[SkeletonBoneIndex];
-		int32 CompactIndex = (PoseBoneIndex != INDEX_NONE) ? MeshIndexToCompactPoseIndex[PoseBoneIndex] : INDEX_NONE;
-		SkeletonToCompactPose.Add(FCompactPoseBoneIndex(CompactIndex));
-	}
-}
-
-int32 FBoneContainer::GetPoseBoneIndexForBoneName(const FName& BoneName) const
-{
-	checkSlow( IsValid() );
-	return RefSkeleton->FindBoneIndex(BoneName);
-}
-
-int32 FBoneContainer::GetParentBoneIndex(const int32& BoneIndex) const
-{
-	checkSlow( IsValid() );
-	checkSlow(BoneIndex != INDEX_NONE);
-	return RefSkeleton->GetParentIndex(BoneIndex);
-}
-
-FCompactPoseBoneIndex FBoneContainer::GetParentBoneIndex(const FCompactPoseBoneIndex& BoneIndex) const
-{
-	checkSlow(IsValid());
-	checkSlow(BoneIndex != INDEX_NONE);
-	return CompactPoseParentBones[BoneIndex.GetInt()];
-}
-
-int32 FBoneContainer::GetDepthBetweenBones(const int32& BoneIndex, const int32& ParentBoneIndex) const
-{
-	checkSlow( IsValid() );
-	checkSlow( BoneIndex != INDEX_NONE );
-	return RefSkeleton->GetDepthBetweenBones(BoneIndex, ParentBoneIndex);
-}
-
-bool FBoneContainer::BoneIsChildOf(const int32& BoneIndex, const int32& ParentBoneIndex) const
-{
-	checkSlow( IsValid() );
-	checkSlow( (BoneIndex != INDEX_NONE) && (ParentBoneIndex != INDEX_NONE) );
-	return RefSkeleton->BoneIsChildOf(BoneIndex, ParentBoneIndex);
-}
-
-bool FBoneContainer::BoneIsChildOf(const FCompactPoseBoneIndex& BoneIndex, const FCompactPoseBoneIndex& ParentBoneIndex) const
-{
-	checkSlow(IsValid());
-	checkSlow((BoneIndex != INDEX_NONE) && (ParentBoneIndex != INDEX_NONE));
-
-	// Bones are in strictly increasing order.
-	// So child must have an index greater than his parent.
-	if (BoneIndex > ParentBoneIndex)
-	{
-		FCompactPoseBoneIndex SearchBoneIndex = GetParentBoneIndex(BoneIndex);
-		do
-		{
-			if (SearchBoneIndex == ParentBoneIndex)
-			{
-				return true;
-			}
-			SearchBoneIndex = GetParentBoneIndex(SearchBoneIndex);
-
-		} while (SearchBoneIndex != INDEX_NONE);
-	}
-
-	return false;
-}
-
-void FBoneContainer::RemapFromSkelMesh(USkeletalMesh const & SourceSkeletalMesh, USkeleton& TargetSkeleton)
-{
-	int32 const SkelMeshLinkupIndex = TargetSkeleton.GetMeshLinkupIndex(&SourceSkeletalMesh);
-	check(SkelMeshLinkupIndex != INDEX_NONE);
-
-	FSkeletonToMeshLinkup const & LinkupTable = TargetSkeleton.LinkupCache[SkelMeshLinkupIndex];
-
-	// Copy LinkupTable arrays for now.
-	// @laurent - Long term goal is to trim that down based on LOD, so we can get rid of the BoneIndicesArray and branch cost of testing if PoseBoneIndex is in that required bone index array.
-	SkeletonToPoseBoneIndexArray = LinkupTable.SkeletonToMeshTable;
-	PoseToSkeletonBoneIndexArray = LinkupTable.MeshToSkeletonTable;
-}
-
-void FBoneContainer::RemapFromSkeleton(USkeleton const & SourceSkeleton)
-{
-	// Map SkeletonBoneIndex to the SkeletalMesh Bone Index, taking into account the required bone index array.
-	SkeletonToPoseBoneIndexArray.Init(INDEX_NONE, SourceSkeleton.GetRefLocalPoses().Num());
-	for(int32 Index=0; Index<BoneIndicesArray.Num(); Index++)
-	{
-		int32 const & PoseBoneIndex = BoneIndicesArray[Index];
-		SkeletonToPoseBoneIndexArray[PoseBoneIndex] = PoseBoneIndex;
-	}
-
-	// Skeleton to Skeleton mapping...
-	PoseToSkeletonBoneIndexArray = SkeletonToPoseBoneIndexArray;
-}
-
+///////////////////////////////////////////////////////////////////////////////////////
+//
+// FBlendSampleData
+//
+////////////////////////////////////////////////////////////////////////////////////////
 void FBlendSampleData::NormalizeDataWeight(TArray<FBlendSampleData>& SampleDataList)
 {
 	float TotalSum = 0.f;
diff --git a/Engine/Source/Runtime/Engine/Private/Audio.cpp b/Engine/Source/Runtime/Engine/Private/Audio.cpp
--- a/Engine/Source/Runtime/Engine/Private/Audio.cpp
+++ b/Engine/Source/Runtime/Engine/Private/Audio.cpp
@@ -37,7 +37,6 @@ DEFINE_STAT(STAT_AudioFinishedDelegates);
 DEFINE_STAT(STAT_AudioBufferTime);
 DEFINE_STAT(STAT_AudioBufferTimeChannels);
 
-DEFINE_STAT(STAT_OggWaveInstances);
 DEFINE_STAT(STAT_VorbisDecompressTime);
 DEFINE_STAT(STAT_VorbisPrepareDecompressionTime);
 DEFINE_STAT(STAT_AudioDecompressTime);
@@ -178,79 +177,103 @@ FString FSoundSource::Describe(bool bUseLongName)
 		WaveInstance->ActiveSound ? *WaveInstance->ActiveSound->GetOwnerName() : TEXT("None"));
 }
 
-void FSoundSource::Stop( void )
+void FSoundSource::Stop()
 {
 	IStreamingManager::Get().GetAudioStreamingManager().RemoveStreamingSoundSource(this);
-	if( WaveInstance )
+	if (WaveInstance)
 	{
-		check( AudioDevice );
-		AudioDevice->FreeSources.AddUnique( this );
-		AudioDevice->WaveInstanceSourceMap.Remove( WaveInstance );
+		check(AudioDevice);
+		AudioDevice->FreeSources.AddUnique(this);
+		AudioDevice->WaveInstanceSourceMap.Remove(WaveInstance);
 		WaveInstance->NotifyFinished(true);
-		WaveInstance = NULL;
+		WaveInstance = nullptr;
 	}
 	else
 	{
-		check( AudioDevice->FreeSources.Find( this ) != INDEX_NONE );
+		check(AudioDevice->FreeSources.Find(this) != INDEX_NONE);
 	}
 }
 
-bool FSoundSource::IsGameOnly( void )
+void FSoundSource::SetPauseByGame(bool bInIsPauseByGame)
+{
+	bIsPausedByGame = bInIsPauseByGame;
+	UpdatePause();
+}
+
+void FSoundSource::SetPauseManually(bool bInIsPauseManually)
+{
+	bIsManuallyPaused = bInIsPauseManually;
+	UpdatePause();
+}
+
+void FSoundSource::UpdatePause()
+{
+	if (IsPaused() && !bIsPausedByGame && !bIsManuallyPaused)
+	{
+		Play();
+	}
+	else if (!IsPaused() && (bIsManuallyPaused || bIsPausedByGame))
+	{
+		Pause();
+	}
+}
+
+bool FSoundSource::IsGameOnly() const
 {
 	return (WaveInstance && !WaveInstance->bIsUISound);
 }
 
-bool FSoundSource::SetReverbApplied( bool bHardwareAvailable )
+bool FSoundSource::SetReverbApplied(bool bHardwareAvailable)
 {
 	// Do not apply reverb if it is explicitly disallowed
 	bReverbApplied = WaveInstance->bReverb && bHardwareAvailable;
 
 	// Do not apply reverb to music
-	if( WaveInstance->bIsMusic )
+	if (WaveInstance->bIsMusic)
 	{
 		bReverbApplied = false;
 	}
 
 	// Do not apply reverb to multichannel sounds
-	if( WaveInstance->WaveData->NumChannels > 2 )
+	if (WaveInstance->WaveData->NumChannels > 2)
 	{
 		bReverbApplied = false;
 	}
 
-	return( bReverbApplied );
+	return(bReverbApplied);
 }
 
-float FSoundSource::SetStereoBleed( void )
+float FSoundSource::SetStereoBleed()
 {
 	StereoBleed = 0.0f;
 
 	// All stereo sounds bleed by default
-	if( WaveInstance->WaveData->NumChannels == 2 )
+	if (WaveInstance->WaveData->NumChannels == 2)
 	{
 		StereoBleed = WaveInstance->StereoBleed;
 
-		if( AudioDevice->GetMixDebugState() == DEBUGSTATE_TestStereoBleed )
+		if (AudioDevice->GetMixDebugState() == DEBUGSTATE_TestStereoBleed)
 		{
 			StereoBleed = 1.0f;
 		}
 	}
 
-	return( StereoBleed );
+	return StereoBleed;
 }
 
-float FSoundSource::SetLFEBleed( void )
+float FSoundSource::SetLFEBleed()
 {
 	LFEBleed = WaveInstance->LFEBleed;
 
-	if( AudioDevice->GetMixDebugState() == DEBUGSTATE_TestLFEBleed )
+	if (AudioDevice->GetMixDebugState() == DEBUGSTATE_TestLFEBleed)
 	{
 		LFEBleed = 10.0f;
 	}
 
-	return( LFEBleed );
+	return LFEBleed;
 }
 
-void FSoundSource::SetFilterFrequency(void)
+void FSoundSource::SetFilterFrequency()
 {
 	LPFFrequency = MAX_FILTER_FREQUENCY;
 
@@ -440,7 +463,6 @@ FSpatializationParams FSoundSource::GetSpatializationParams()
 {
 	FSpatializationParams Params;
 
-		// Calculate direction from listener to sound, where the sound is at the origin if unspatialized.
 	if (WaveInstance->bUseSpatialization)
 	{
 		FVector EmitterPosition = AudioDevice->GetListenerTransformedDirection(WaveInstance->Location, &Params.Distance);
@@ -448,6 +470,7 @@ FSpatializationParams FSoundSource::GetSpatializationParams()
 		// If we are using the OmniRadius feature
 		if (WaveInstance->OmniRadius > 0.0f)
 		{
+			// Initialize to full omni-directionality (bigger value, more omni)
 			static const float MaxNormalizedRadius = 1000000.0f;
 			Params.NormalizedOmniRadius = MaxNormalizedRadius;
 
@@ -478,6 +501,7 @@ FSpatializationParams FSoundSource::GetSpatializationParams()
 		Params.Distance = 0.0f;
 		Params.EmitterPosition = FVector::ZeroVector;
 	}
+	Params.EmitterWorldPosition = WaveInstance->Location;
 
 	// We are currently always computing spatialization for XAudio2 relative to the listener!
 	Params.ListenerOrientation = FVector::UpVector;
@@ -486,6 +510,33 @@ FSpatializationParams FSoundSource::GetSpatializationParams()
 	return Params;
 }
 
+void FSoundSource::NotifyPlaybackPercent()
+{
+	if (WaveInstance->ActiveSound->bUpdatePlayPercentage)
+	{
+		const uint64 AudioComponentID = WaveInstance->ActiveSound->GetAudioComponentID();
+		if (AudioComponentID > 0)
+		{
+			const float PlaybackPercent = GetPlaybackPercent();
+			const USoundWave* SoundWave = WaveInstance->WaveData;
+			FAudioThread::RunCommandOnGameThread([AudioComponentID, SoundWave, PlaybackPercent]()
+			{
+				if (UAudioComponent* AudioComponent = UAudioComponent::GetAudioComponentFromID(AudioComponentID))
+				{
+					if (AudioComponent->OnAudioPlaybackPercent.IsBound())
+					{
+						AudioComponent->OnAudioPlaybackPercent.Broadcast(SoundWave, PlaybackPercent);
+					}
+
+					if (AudioComponent->OnAudioPlaybackPercentNative.IsBound())
+					{
+						AudioComponent->OnAudioPlaybackPercentNative.Broadcast(AudioComponent, SoundWave, PlaybackPercent);
+					}
+				}
+			});
+		}
+	}
+}
 
 /*-----------------------------------------------------------------------------
 	FNotifyBufferFinishedHooks implementation.
@@ -556,44 +607,48 @@ uint32 FWaveInstance::TypeHashCounter = 0;
  * @param InActiveSound		ActiveSound this wave instance belongs to.
  */
 FWaveInstance::FWaveInstance( FActiveSound* InActiveSound )
-:	WaveData( NULL )
-,   SoundClass( NULL )
-,	ActiveSound( InActiveSound )
-,	Volume( 0.0f )
-,	VolumeMultiplier( 1.0f )
-,	VolumeApp( 1.0f )
-,	Priority( 1.0f )
-,	VoiceCenterChannelVolume( 0.0f )
-,	RadioFilterVolume( 0.0f )
-,	RadioFilterVolumeThreshold( 0.0f )
-,	StereoBleed( 0.0f )
-,	LFEBleed( 0.0f )
-,	LoopingMode( LOOP_Never )
-,	StartTime( -1.f )
-,	bApplyRadioFilter( false )
-,	bIsStarted(false)
-,	bIsFinished( false )
-,	bAlreadyNotifiedHook( false )
-,	bUseSpatialization( false )
-,	bEnableLowPassFilter(false)
-,	bIsOccluded(false)
-,	bEQFilterApplied(false)
-,	bIsUISound( false )
-,	bIsMusic( false )
-,	bReverb( true )
-,	bCenterChannelOnly( false )
-,	bReportedSpatializationWarning( false )
-,	SpatializationAlgorithm(SPATIALIZATION_Default)
-,	OutputTarget(EAudioOutputTarget::Speaker)
-,	LowPassFilterFrequency(MAX_FILTER_FREQUENCY)
-,	OcclusionFilterFrequency(MAX_FILTER_FREQUENCY)
-, AmbientZoneFilterFrequency(MAX_FILTER_FREQUENCY)
-,	AttenuationFilterFrequency(MAX_FILTER_FREQUENCY), Pitch(0.0f)
-,	Velocity( FVector::ZeroVector )
-,	Location( FVector::ZeroVector )
-,	OmniRadius(0.0f)
-,	StereoSpread(0.0f)
-,	UserIndex( 0 )
+	: WaveData(nullptr)
+	, SoundClass(nullptr)
+	, SoundSubmix(nullptr)
+	, ActiveSound(InActiveSound)
+	, Volume(0.0f)
+	, VolumeMultiplier(1.0f)
+	, VolumeApp(1.0f)
+	, Priority(1.0f)
+	, VoiceCenterChannelVolume(0.0f)
+	, RadioFilterVolume(0.0f)
+	, RadioFilterVolumeThreshold(0.0f)
+	, StereoBleed(0.0f)
+	, LFEBleed(0.0f)
+	, LoopingMode(LOOP_Never)
+	, StartTime(-1.f)
+	, bApplyRadioFilter(false)
+	, bIsStarted(false)
+	, bIsFinished(false)
+	, bAlreadyNotifiedHook(false)
+	, bUseSpatialization(false)
+	, bEnableLowPassFilter(false)
+	, bIsOccluded(false)
+	, bEQFilterApplied(false)
+	, bIsUISound(false)
+	, bIsMusic(false)
+	, bReverb(true)
+	, bCenterChannelOnly(false)
+	, bReportedSpatializationWarning(false)
+	, SpatializationAlgorithm(SPATIALIZATION_Default)
+	, OutputTarget(EAudioOutputTarget::Speaker)
+	, LowPassFilterFrequency(MAX_FILTER_FREQUENCY)
+	, OcclusionFilterFrequency(MAX_FILTER_FREQUENCY)
+	, AmbientZoneFilterFrequency(MAX_FILTER_FREQUENCY)
+	, AttenuationFilterFrequency(MAX_FILTER_FREQUENCY)
+	, Pitch(0.0f)
+	, Velocity(FVector::ZeroVector)
+	, Location(FVector::ZeroVector)
+	, OmniRadius(0.0f)
+	, StereoSpread(0.0f)
+	, AttenuationDistance(0.0f)
+	, AbsoluteAzimuth(0.0f)
+	, UserIndex(0)
 {
 	TypeHash = ++TypeHashCounter;
 }
@@ -718,6 +773,10 @@ FString FWaveInstance::GetName() const
 #define UE_mmioFOURCC(ch0, ch1, ch2, ch3)\
 	UE_MAKEFOURCC(ch0, ch1, ch2, ch3)
 
+#if PLATFORM_SUPPORTS_PRAGMA_PACK
+#pragma pack(push, 2)
+#endif
+
 // Main Riff-Wave header.
 struct FRiffWaveHeader
 {
@@ -745,6 +804,44 @@ struct FFormatChunk
 	uint16   cbSize;            // The count in bytes of the size of extra information (after cbSize).
 };
 
+// FExtendedFormatChunk subformat GUID.
+struct FSubformatGUID
+{
+	uint32 Data1;				// Format type, corresponds to a wFormatTag in WaveFormatEx.
+
+								// Fixed values for all extended wave formats.
+	uint16 Data2 = 0x0000;
+	uint16 Data3 = 0x0010;
+	uint8 Data4[8];
+
+	FSubformatGUID()
+	{
+		Data4[0] = 0x80;
+		Data4[1] = 0x00;
+		Data4[2] = 0x00;
+		Data4[3] = 0xaa;
+		Data4[4] = 0x00;
+		Data4[5] = 0x38;
+		Data4[6] = 0x9b;
+		Data4[7] = 0x71;
+	}
+};
+
+// ChunkID: 'fmt ' ("WaveFormatExtensible" structure)
+struct FExtendedFormatChunk
+{
+	FFormatChunk Format;			// Standard WaveFormatEx ('fmt ') chunk, with
+									// wFormatTag == WAVE_FORMAT_EXTENSIBLE and cbSize == 22
+	union
+	{
+		uint16 wValidBitsPerSample;	// Actual bits of precision. Can be less than wBitsPerSample.
+		uint16 wSamplesPerBlock;	// Valid if wValidBitsPerSample == 0. Used by compressed formats.
+		uint16 wReserved;			// If neither applies, set to 0.
+	} Samples;
+	uint32 dwChannelMask;			// Which channels are present in the stream.
+	FSubformatGUID SubFormat;		// Subformat identifier.
+};
+
 // ChunkID: 'smpl'
 struct FSampleChunk
 {
@@ -759,12 +856,17 @@ struct FSampleChunk
 	uint32   cbSamplerData;		//
 };
 
+#if PLATFORM_SUPPORTS_PRAGMA_PACK
+#pragma pack(pop)
+#endif
+
 //
 //	Figure out the WAVE file layout.
 //
 bool FWaveModInfo::ReadWaveInfo( uint8* WaveData, int32 WaveDataSize, FString* ErrorReason )
 {
 	FFormatChunk* FmtChunk;
+	FExtendedFormatChunk* FmtChunkEx = nullptr;
 	FRiffWaveHeader* RiffHdr = ( FRiffWaveHeader* )WaveData;
 	WaveDataEnd = WaveData + WaveDataSize;
 
@@ -839,6 +941,56 @@ bool FWaveModInfo::ReadWaveInfo( uint8* WaveData, int32 WaveDataSize, FString* E
 	pChannels = &FmtChunk->nChannels;
 	pFormatTag = &FmtChunk->wFormatTag;
 
+	// If we have an extended fmt chunk, the format tag won't be a wave format. Instead we need to read the subformat ID.
+	if (INTEL_ORDER32(RiffChunk->ChunkLen) >= 40 && FmtChunk->wFormatTag == 0xFFFE) // WAVE_FORMAT_EXTENSIBLE
+	{
+		FmtChunkEx = (FExtendedFormatChunk*)((uint8*)RiffChunk + 8);
+
+#if !PLATFORM_LITTLE_ENDIAN
+		if (!AlreadySwapped)
+		{
+			FmtChunkEx->Samples.wValidBitsPerSample = INTEL_ORDER16(FmtChunkEx->Samples.wValidBitsPerSample);
+			FmtChunkEx->SubFormat.Data1 = INTEL_ORDER32(FmtChunkEx->SubFormat.Data1);
+			FmtChunkEx->SubFormat.Data2 = INTEL_ORDER16(FmtChunkEx->SubFormat.Data2);
+			FmtChunkEx->SubFormat.Data3 = INTEL_ORDER16(FmtChunkEx->SubFormat.Data3);
+			*((uint64*)FmtChunkEx->SubFormat.Data4) = INTEL_ORDER64(*(uint64*)FmtChunkEx->SubFormat.Data4);
+		}
+#endif
+
+		bool bValid = true;
+		static const FSubformatGUID GUID;
+
+		if (FmtChunkEx->SubFormat.Data1 == 0x00000001 /* PCM */ &&
+			FmtChunkEx->Samples.wValidBitsPerSample > 0 && FmtChunkEx->Samples.wValidBitsPerSample != FmtChunk->wBitsPerSample)
+		{
+			bValid = false;
+			if (ErrorReason) *ErrorReason = TEXT("Unsupported WAVE file format: actual bit rate does not match the container size.");
+		}
+		else if (FMemory::Memcmp((uint8*)&FmtChunkEx->SubFormat + 4, (uint8*)&GUID + 4, sizeof(GUID) - 4) != 0)
+		{
+			bValid = false;
+			if (ErrorReason) *ErrorReason = TEXT("Unsupported WAVE file format: subformat identifier not recognized.");
+		}
+
+		if (!bValid)
+		{
+#if !PLATFORM_LITTLE_ENDIAN // swap them back just in case.
+			if (!AlreadySwapped)
+			{
+				FmtChunkEx->Samples.wValidBitsPerSample = INTEL_ORDER16(FmtChunkEx->Samples.wValidBitsPerSample);
+				FmtChunkEx->SubFormat.Data1 = INTEL_ORDER32(FmtChunkEx->SubFormat.Data1);
+				FmtChunkEx->SubFormat.Data2 = INTEL_ORDER16(FmtChunkEx->SubFormat.Data2);
+				FmtChunkEx->SubFormat.Data3 = INTEL_ORDER16(FmtChunkEx->SubFormat.Data3);
+				*((uint64*)FmtChunkEx->SubFormat.Data4) = INTEL_ORDER64(*(uint64*)FmtChunkEx->SubFormat.Data4);
+			}
+#endif
+			return (false);
+		}
+
+		// Set the format tag pointer to the subformat GUID.
+		pFormatTag = reinterpret_cast<uint16*>(&FmtChunkEx->SubFormat.Data1);
+	}
+
 	// re-initalize the RiffChunk pointer
 	RiffChunk = ( FRiffChunkOld* )&WaveData[3 * 4];
 
@@ -862,6 +1014,14 @@ bool FWaveModInfo::ReadWaveInfo( uint8* WaveData, int32 WaveDataSize, FString* E
 				FmtChunk->nAvgBytesPerSec = INTEL_ORDER32( FmtChunk->nAvgBytesPerSec );
 				FmtChunk->nBlockAlign = INTEL_ORDER16( FmtChunk->nBlockAlign );
 				FmtChunk->wBitsPerSample = INTEL_ORDER16( FmtChunk->wBitsPerSample );
+				if (FmtChunkEx != nullptr)
+				{
+					FmtChunkEx->Samples.wValidBitsPerSample = INTEL_ORDER16(FmtChunkEx->Samples.wValidBitsPerSample);
+					FmtChunkEx->SubFormat.Data1 = INTEL_ORDER32(FmtChunkEx->SubFormat.Data1);
+					FmtChunkEx->SubFormat.Data2 = INTEL_ORDER16(FmtChunkEx->SubFormat.Data2);
+					FmtChunkEx->SubFormat.Data3 = INTEL_ORDER16(FmtChunkEx->SubFormat.Data3);
+					*((uint64*)FmtChunkEx->SubFormat.Data4) = INTEL_ORDER64(*(uint64*)FmtChunkEx->SubFormat.Data4);
+				}
 			}
 		#endif
 		if (ErrorReason) *ErrorReason = TEXT("Invalid WAVE file.");
@@ -878,7 +1038,6 @@ bool FWaveModInfo::ReadWaveInfo( uint8* WaveData, int32 WaveDataSize, FString* E
 	SampleDataStart = ( uint8* )RiffChunk + 8;
 	pWaveDataSize = &RiffChunk->ChunkLen;
 	SampleDataSize = INTEL_ORDER32( RiffChunk->ChunkLen );
-	OldBitsPerSample = FmtChunk->wBitsPerSample;
 	SampleDataEnd = SampleDataStart + SampleDataSize;
 
 	if( ( uint8* )SampleDataEnd > ( uint8* )WaveDataEnd )
@@ -893,9 +1052,9 @@ bool FWaveModInfo::ReadWaveInfo( uint8* WaveData, int32 WaveDataSize, FString* E
 
 	NewDataSize = SampleDataSize;
 
-	if(    FmtChunk->wFormatTag != 0x0001 // WAVE_FORMAT_PCM  
-		&& FmtChunk->wFormatTag != 0x0002 // WAVE_FORMAT_ADPCM
-		&& FmtChunk->wFormatTag != 0x0011) // WAVE_FORMAT_DVI_ADPCM
+	if (   *pFormatTag != 0x0001 // WAVE_FORMAT_PCM  
+		&& *pFormatTag != 0x0002 // WAVE_FORMAT_ADPCM
+		&& *pFormatTag != 0x0011) // WAVE_FORMAT_DVI_ADPCM
 	{
 		ReportImportFailure();
 		if (ErrorReason) *ErrorReason = TEXT("Unsupported wave file format.  Only PCM, ADPCM, and DVI ADPCM can be imported.");
