Copying //UE4/Dev-Framework to //UE4/Dev-Main (Source: //UE4/Dev-Framework @ 3136612)

diff --git a/Engine/Source/Runtime/Engine/Private/AudioDevice.cpp b/Engine/Source/Runtime/Engine/Private/AudioDevice.cpp
--- a/Engine/Source/Runtime/Engine/Private/AudioDevice.cpp
+++ b/Engine/Source/Runtime/Engine/Private/AudioDevice.cpp
@@ -11,10 +11,13 @@
 #include "Sound/SoundCue.h"
 #include "Sound/SoundNodeWavePlayer.h"
 #include "Sound/SoundWave.h"
+#include "GameFramework/GameUserSettings.h"
 #include "IAudioExtensionPlugin.h"
 
 #if WITH_EDITOR
 #include "Editor/UnrealEd/Classes/Editor/EditorEngine.h"
+#include "AssetRegistryModule.h"
+#include "Internationalization.h"
 #endif
 
 /*-----------------------------------------------------------------------------
@@ -79,7 +82,7 @@ void FDynamicParameter::Update(float DeltaTime)
 FAudioDevice::FAudioDevice()
 	: MaxChannels(0)
 	, SampleRate(0)
-	, CommonAudioPoolSize(0)
+	, BufferLength(0)
 	, CommonAudioPool(nullptr)
 	, CommonAudioPoolFreeBytes(0)
 	, DeviceHandle(INDEX_NONE)
@@ -101,6 +104,7 @@ FAudioDevice::FAudioDevice()
 	, bGameWasTicking(true)
 	, bDisableAudioCaching(false)
 	, bIsAudioDeviceHardwareInitialized(false)
+	, bAudioMixerModuleLoaded(false)
 	, bStartupSoundsPreCached(false)
 	, bSpatializationExtensionEnabled(false)
 	, bHRTFEnabledForAll(false)
@@ -122,6 +126,12 @@ FAudioEffectsManager* FAudioDevice::CreateEffectsManager()
 	return new FAudioEffectsManager(this);
 }
 
+FAudioQualitySettings FAudioDevice::GetQualityLevelSettings()
+{
+	const UAudioSettings* AudioSettings = GetDefault<UAudioSettings>();
+	return AudioSettings->GetQualityLevelSettings(GEngine->GetGameUserSettings()->GetAudioQualityLevel());
+}
+
 bool FAudioDevice::Init(int32 InMaxChannels)
 {
 	if (bIsInitialized)
@@ -130,8 +140,14 @@ bool FAudioDevice::Init(int32 InMaxChannels)
 	}
 
 	bool bDeferStartupPrecache = false;
+
 	// initialize config variables
 	MaxChannels = InMaxChannels;
+
+	// Setup the desired sample rate and buffer length
+	SampleRate = 44100;
+	BufferLength = 1024;
+
 	verify(GConfig->GetInt(TEXT("Audio"), TEXT("CommonAudioPoolSize"), CommonAudioPoolSize, GEngineIni));
 
 	// If this is true, skip the initial startup precache so we can do it later in the flow
@@ -157,6 +173,8 @@ bool FAudioDevice::Init(int32 InMaxChannels)
 
 	// Parses sound classes.
 	InitSoundClasses();
+	InitSoundSubmixes();
+	InitSoundEffectPresets();
 
 	// allow the platform to startup
 	if (InitializeHardware() == false)
@@ -351,7 +369,10 @@ void FAudioDevice::ResetInterpolation()
 	BaseSoundMix = nullptr;
 
 	// reset audio effects
-	Effects->ResetInterpolation();
+	if (Effects)
+	{
+		Effects->ResetInterpolation();
+	}
 }
 
 void FAudioDevice::EnableRadioEffect(bool bEnable)
@@ -881,7 +902,6 @@ bool FAudioDevice::HandlePlayAllPIEAudioCommand(const TCHAR* Cmd, FOutputDevice&
 	return true;
 }
 
-
 bool FAudioDevice::HandleAudio3dVisualizeCommand(const TCHAR* Cmd, FOutputDevice& Ar)
 {
 	FAudioDeviceManager* DeviceManager = GEngine->GetAudioDeviceManager();
@@ -899,7 +919,6 @@ bool FAudioDevice::HandleAudioSoloSoundClass(const TCHAR* Cmd, FOutputDevice& Ar
 	{
 		DeviceManager->SetDebugSoloSoundClass(Cmd);
 	}
-
 	return true;
 }
 
@@ -910,7 +929,6 @@ bool FAudioDevice::HandleAudioSoloSoundWave(const TCHAR* Cmd, FOutputDevice& Ar)
 	{
 		DeviceManager->SetDebugSoloSoundWave(Cmd);
 	}
-
 	return true;
 }
 
@@ -921,7 +939,16 @@ bool FAudioDevice::HandleAudioSoloSoundCue(const TCHAR* Cmd, FOutputDevice& Ar)
 	{
 		DeviceManager->SetDebugSoloSoundCue(Cmd);
 	}
+	return true;
+}
 
+bool FAudioDevice::HandleAudioMixerDebugSound(const TCHAR* Cmd, FOutputDevice& Ar)
+{
+	FAudioDeviceManager* DeviceManager = GEngine->GetAudioDeviceManager();
+	if (DeviceManager)
+	{
+		DeviceManager->SetAudioMixerDebugSound(Cmd);
+	}
 	return true;
 }
 
@@ -1306,6 +1333,10 @@ bool FAudioDevice::Exec(UWorld* InWorld, const TCHAR* Cmd, FOutputDevice& Ar)
 	{
 		return HandleAudioMemoryInfo(Cmd, Ar);
 	}
+	else if (FParse::Command(&Cmd, TEXT("AudioMixerDebugSound")))
+	{
+		return HandleAudioMixerDebugSound(Cmd, Ar);
+	}
 #endif // !UE_BUILD_SHIPPING
 
 	return false;
@@ -1559,7 +1590,10 @@ void FAudioDevice::UpdateSoundMix(USoundMix* SoundMix, FSoundMixState* SoundMixS
 			}
 
 			// Might have already started fading EQ effect so try setting again
-			Effects->SetMixSettings(SoundMix);
+			if (Effects)
+			{
+				Effects->SetMixSettings(SoundMix);
+			}
 		}
 		else if (SoundMixState->CurrentState == ESoundMixState::Active)
 		{
@@ -1671,7 +1705,7 @@ bool FAudioDevice::TryClearingSoundMix(USoundMix* SoundMix, FSoundMixState* Soun
 
 bool FAudioDevice::TryClearingEQSoundMix(USoundMix* SoundMix)
 {
-	if (SoundMix && Effects->GetCurrentEQMix() == SoundMix)
+	if (SoundMix && Effects && Effects->GetCurrentEQMix() == SoundMix)
 	{
 		USoundMix* NextEQMix = FindNextHighestEQPrioritySoundMix(SoundMix);
 		if (NextEQMix)
@@ -2196,7 +2230,10 @@ void FAudioDevice::SetBaseSoundMix(USoundMix* NewMix)
 			ApplySoundMix(NewMix, ExistingState);
 
 			// Use it to set EQ Settings, which will check its priority
-			Effects->SetMixSettings(NewMix);
+			if (Effects)
+			{
+				Effects->SetMixSettings(NewMix);
+			}
 		}
 
 		ExistingState->IsBaseSoundMix = true;
@@ -2232,7 +2269,10 @@ void FAudioDevice::PushSoundMixModifier(USoundMix* SoundMix, bool bIsPassive)
 			ApplySoundMix(SoundMix, SoundMixState);
 
 			// Use it to set EQ Settings, which will check its priority
-			Effects->SetMixSettings(SoundMix);
+			if (Effects)
+			{
+				Effects->SetMixSettings(SoundMix);
+			}
 		}
 		else
 		{
@@ -2508,7 +2548,11 @@ void FAudioDevice::DeactivateReverbEffect(FName TagName)
 void* FAudioDevice::InitEffect(FSoundSource* Source)
 {
 	check(IsInAudioThread());
-	return(Effects->InitEffect(Source));
+	if (Effects)
+	{
+		return Effects->InitEffect(Source);
+	}
+	return nullptr;
 }
 
 
@@ -2517,55 +2561,69 @@ void* FAudioDevice::UpdateEffect(FSoundSource* Source)
 	SCOPE_CYCLE_COUNTER(STAT_AudioUpdateEffects);
 
 	check(IsInAudioThread());
-	return(Effects->UpdateEffect(Source));
+	if (Effects)
+	{
+		return Effects->UpdateEffect(Source);
+	}
+	return nullptr;
 }
 
 
 void FAudioDevice::DestroyEffect(FSoundSource* Source)
 {
 	check(IsInAudioThread());
-	return(Effects->DestroyEffect(Source));
+	if (Effects)
+	{
+		Effects->DestroyEffect(Source);
+	}
 }
 
-
 void FAudioDevice::HandlePause(bool bGameTicking, bool bGlobalPause)
 {
 	DECLARE_CYCLE_STAT(TEXT("FAudioThreadTask.HandlePause"), STAT_AudioHandlePause, STATGROUP_AudioThreadCommands);
 
-	FAudioDevice* AudioDevice = this;
-	FAudioThread::RunCommandOnAudioThread([AudioDevice, bGameTicking, bGlobalPause]()
+	// Run this command on the audio thread if this is getting called on game thread
+	if (!IsInAudioThread())
+	{
+		FAudioDevice* AudioDevice = this;
+		FAudioThread::RunCommandOnAudioThread([AudioDevice, bGameTicking, bGlobalPause]()
+		{
+			AudioDevice->HandlePause(bGameTicking, bGlobalPause);
+		}, GET_STATID(STAT_AudioHandlePause));
+
+		return;
+	}
+
+	// Handles the global pause/unpause feature
+
+	// Pause all sounds if transitioning to pause mode.
+	if (!bGameTicking && (bGameWasTicking || bGlobalPause))
 	{
-		// Pause all sounds if transitioning to pause mode.
-		if (!bGameTicking && (AudioDevice->bGameWasTicking || bGlobalPause))
+		for (int32 i = 0; i < Sources.Num(); i++)
 		{
-			for (int32 i = 0; i < AudioDevice->Sources.Num(); i++)
+			FSoundSource* Source = Sources[ i ];
+			if (!Source->IsPausedByGame() && (bGlobalPause || Source->IsGameOnly()))
 			{
-				FSoundSource* Source = AudioDevice->Sources[ i ];
-				if (!Source->IsPaused() && (bGlobalPause || Source->IsGameOnly()))
-				{
-					Source->Pause();
-				}
+				Source->SetPauseByGame(true);
 			}
 		}
-		// Unpause all sounds if transitioning back to game.
-		else if (bGameTicking && (!AudioDevice->bGameWasTicking || bGlobalPause))
+	}
+	// Unpause all sounds if transitioning back to game.
+	else if (bGameTicking && (!bGameWasTicking || bGlobalPause))
+	{
+		for (int32 i = 0; i < Sources.Num(); i++)
 		{
-			for (int32 i = 0; i < AudioDevice->Sources.Num(); i++)
+			FSoundSource* Source = Sources[ i ];
+			if (Source->IsPausedByGame() && (bGlobalPause || Source->IsGameOnly()))
 			{
-				FSoundSource* Source = AudioDevice->Sources[ i ];
-				if (Source->IsPaused() && (bGlobalPause || Source->IsGameOnly()))
-				{
-					Source->Play();
-				}
+				Source->SetPauseByGame(false);
 			}
 		}
+	}
 
-		AudioDevice->bGameWasTicking = bGameTicking;
-	}, GET_STATID(STAT_AudioHandlePause));
-
+	bGameWasTicking = bGameTicking;
 }
 
-
 int32 FAudioDevice::GetSortedActiveWaveInstances(TArray<FWaveInstance*>& WaveInstances, const ESortedActiveWaveGetType::Type GetType)
 {
 	check(IsInAudioThread());
@@ -2628,9 +2686,19 @@ int32 FAudioDevice::GetSortedActiveWaveInstances(TArray<FWaveInstance*>& WaveIns
 		SCOPE_CYCLE_COUNTER(STAT_AudioEvaluateConcurrency);
 		ConcurrencyManager.StopQuietSoundsDueToMaxConcurrency();
 	}
+	
+	// Remove all wave instances from the wave instance list that are stopping due to max concurrency
+	for (int32 i = WaveInstances.Num() - 1; i >= 0; --i)
+	{
+		if (WaveInstances[i]->ShouldStopDueToMaxConcurrency())
+		{
+			WaveInstances.RemoveAtSwap(i, 1, false);
+		}
+	}
 
-	int FirstActiveIndex = 0;
-	if (WaveInstances.Num() >= MaxChannels)
+	int32 FirstActiveIndex = 0;
+
+	if (WaveInstances.Num() >= 0)
 	{
 		// Helper function for "Sort" (higher priority sorts last).
 		struct FCompareFWaveInstanceByPlayPriority
@@ -2647,6 +2715,7 @@ int32 FAudioDevice::GetSortedActiveWaveInstances(TArray<FWaveInstance*>& WaveIns
 		// Get the first index that will result in a active source voice
 		FirstActiveIndex = FMath::Max(WaveInstances.Num() - MaxChannels, 0);
 	}
+
 	return(FirstActiveIndex);
 }
 
@@ -2679,12 +2748,6 @@ void FAudioDevice::StopSources(TArray<FWaveInstance*>& WaveInstances, int32 Firs
 
 		if (Source->WaveInstance)
 		{
-#if STATS && WITH_EDITORONLY_DATA
-			if (Source->UsesCPUDecompression())
-			{
-				INC_DWORD_STAT(STAT_OggWaveInstances);
-			}
-#endif
 			// If we need to stop this sound due to max concurrency (i.e. it was quietest in a concurrency group)
 			if (Source->WaveInstance->ShouldStopDueToMaxConcurrency())
 			{
@@ -2697,6 +2760,9 @@ void FAudioDevice::StopSources(TArray<FWaveInstance*>& WaveInstances, int32 Firs
 			}
 			else
 			{
+				// Update the pause state of the source.
+				Source->SetPauseManually(Source->WaveInstance->bIsPaused);
+
 				// Need to update the source still so that it gets any volume settings applied to
 				// otherwise the source may play at a very quiet volume and not actually set to 0.0
 				Source->Update();
@@ -2767,13 +2833,17 @@ void FAudioDevice::StartSources(TArray<FWaveInstance*>& WaveInstances, int32 Fir
 						// If we succeeded then play and update the source
 						if (bSuccess)
 						{
+							// Set the pause before updating it
+							Source->SetPauseManually(Source->WaveInstance->bIsPaused);
+
 							check(Source->IsInitialized());
+							Source->Update();
+
 							// If the source didn't get paused while initializing, then play it
 							if (!Source->IsPaused())
 							{
 								Source->Play();
 							}
-							Source->Update();
 						}
 					}
 				}
@@ -2800,6 +2870,7 @@ void FAudioDevice::StartSources(TArray<FWaveInstance*>& WaveInstances, int32 Fir
 				// If we've already been initialized, then just update the voice
 				if (Source->IsInitialized())
 				{
+					Source->NotifyPlaybackPercent();
 					Source->Update();
 				}
 				// Otherwise, we need still need to initialize
@@ -2811,6 +2882,7 @@ void FAudioDevice::StartSources(TArray<FWaveInstance*>& WaveInstances, int32 Fir
 						// Note: if we succeeded in starting to prepare to init, we already added the wave instance map to the source so don't need to add here.
 						check(Source->IsInitialized());
 						Source->Play();
+
 						Source->Update();
 					}
 					else
@@ -2898,10 +2970,13 @@ void FAudioDevice::Update(bool bGameTicking)
 		}
 	}
 
-	Effects->SetReverbSettings(ReverbSettings);
+	if (Effects)
+	{
+		Effects->SetReverbSettings(ReverbSettings);
 
-	// Update the audio effects - reverb, EQ etc
-	Effects->Update();
+		// Update the audio effects - reverb, EQ etc
+		Effects->Update();
+	}
 
 	// Gets the current state of the sound classes accounting for sound mix
 	UpdateSoundClassProperties(UpdateDeltaTime);
@@ -3050,8 +3125,12 @@ void FAudioDevice::SendUpdateResultsToGameThread(const int32 FirstActiveIndex)
 	DECLARE_CYCLE_STAT(TEXT("FGameThreadAudioTask.AudioSendResults"), STAT_AudioSendResults, STATGROUP_TaskGraphTasks);
 
 	const uint32 AudioDeviceID = DeviceHandle;
+	UReverbEffect* ReverbEffect = nullptr;
+	if (Effects)
+	{
+		ReverbEffect = Effects->GetCurrentReverbEffect();
+	}
 
-	UReverbEffect* ReverbEffect = Effects->GetCurrentReverbEffect();
 	FAudioThread::RunCommandOnGameThread([AudioDeviceID, ReverbEffect
 #if !UE_BUILD_SHIPPING
 											, StatSoundInfos, StatSoundMixes, bStatsStale
@@ -3247,6 +3326,16 @@ void FAudioDevice::StopActiveSound(FActiveSound* ActiveSound)
 	AddSoundToStop(ActiveSound);
 }
 
+void FAudioDevice::PauseActiveSound(const uint64 AudioComponentID, const bool bInIsPaused)
+{
+	check(IsInAudioThread());
+	FActiveSound* ActiveSound = FindActiveSound(AudioComponentID);
+	if (ActiveSound)
+	{
+		ActiveSound->bIsPaused = bInIsPaused;
+	}
+}
+
 FActiveSound* FAudioDevice::FindActiveSound(const uint64 AudioComponentID)
 {
 	check(IsInAudioThread());
@@ -3333,9 +3422,22 @@ void FAudioDevice::GetMaxDistanceAndFocusFactor(USoundBase* Sound, const UWorld*
 
 		OutMaxDistance = AttenuationSettingsToApply->GetMaxDimension();
 
-		// Now scale the max distance based on the focus settings in the attenuation settings
-		FAttenuationListenerData ListenerData;
-		OutFocusFactor = GetFocusFactor(ListenerData, Sound, SoundTransform, *AttenuationSettingsToApply);
+		if (AttenuationSettingsToApply->bSpatialize && AttenuationSettingsToApply->bEnableListenerFocus && !Sound->bIgnoreFocus)
+		{
+			// Now scale the max distance based on the focus settings in the attenuation settings
+			FAttenuationListenerData ListenerData;
+			float Azimuth = 0.0f;
+			float AbsoluteAzimuth = 0.0f;
+			const int32 ClosestListenerIndex = FindClosestListenerIndex(SoundTransform);
+			const FTransform& ListenerTransform = ListenerTransforms[ClosestListenerIndex];	
+			GetAzimuth(ListenerData, Sound, SoundTransform, *AttenuationSettingsToApply, ListenerTransform, Azimuth, AbsoluteAzimuth);
+
+			OutFocusFactor = GetFocusFactor(ListenerData, Sound, Azimuth, *AttenuationSettingsToApply);
+		}
+		else
+		{
+			OutFocusFactor = 1.0f;
+		}
 	}
 	else
 	{
@@ -3452,94 +3554,185 @@ void FAudioDevice::GetAttenuationListenerData(FAttenuationListenerData& OutListe
 	}
 }
 
-float FAudioDevice::GetFocusFactor(FAttenuationListenerData& OutListenerData, const USoundBase* Sound, const FTransform& SoundTransform, const FAttenuationSettings& AttenuationSettings, const FTransform* InListenerTransform) const
+void FAudioDevice::GetAzimuth(FAttenuationListenerData& OutListenerData, const USoundBase* Sound, const FTransform& SoundTransform, const FAttenuationSettings& AttenuationSettings, const FTransform& ListenerTransform, float& OutAzimuth, float& OutAbsoluteAzimuth) const
+{
+	GetAttenuationListenerData(OutListenerData, SoundTransform, AttenuationSettings, &ListenerTransform);
+
+	const FVector& ListenerForwardDir = OutListenerData.ListenerTransform.GetUnitAxis(EAxis::X);
+
+	const float SoundToListenerForwardDotProduct = FVector::DotProduct(ListenerForwardDir, OutListenerData.ListenerToSoundDir);
+	const float SoundListenerAngleRadians = FMath::Acos(SoundToListenerForwardDotProduct);
+
+	// Normal azimuth only goes to 180 (0 is in front, 180 is behind).
+	OutAzimuth = FMath::RadiansToDegrees(SoundListenerAngleRadians);
+
+	const FVector& ListenerRightDir = OutListenerData.ListenerTransform.GetUnitAxis(EAxis::Y);
+	const float SoundToListenerRightDotProduct = FVector::DotProduct(ListenerRightDir, OutListenerData.ListenerToSoundDir);
+
+	FVector AbsAzimuthVector2D = FVector(SoundToListenerForwardDotProduct, SoundToListenerRightDotProduct, 0.0f);
+	AbsAzimuthVector2D.Normalize();
+
+	OutAbsoluteAzimuth = FMath::Atan(AbsAzimuthVector2D.Y / AbsAzimuthVector2D.X);
+	OutAbsoluteAzimuth = FMath::RadiansToDegrees(OutAbsoluteAzimuth);
+	OutAbsoluteAzimuth = FMath::Abs(OutAbsoluteAzimuth);
+
+	if (AbsAzimuthVector2D.X > 0.0f && AbsAzimuthVector2D.Y < 0.0f)
+	{
+		OutAbsoluteAzimuth = 360 - OutAbsoluteAzimuth;
+	}
+	else if (AbsAzimuthVector2D.X < 0.0f && AbsAzimuthVector2D.Y < 0.0f)
+	{
+		OutAbsoluteAzimuth += 180.0f;
+	}
+	else if (AbsAzimuthVector2D.X < 0.0f && AbsAzimuthVector2D.Y > 0.0f)
+	{
+		OutAbsoluteAzimuth = 180 - OutAbsoluteAzimuth;
+	}
+}
+
+float FAudioDevice::GetFocusFactor(FAttenuationListenerData& OutListenerData, const USoundBase* Sound, const float Azimuth, const FAttenuationSettings& AttenuationSettings) const
 {
 	check(Sound);
 
 	// 0.0f means we are in focus, 1.0f means we are out of focus
 	float FocusFactor = 0.0f;
 
-	if (AttenuationSettings.bSpatialize && AttenuationSettings.bEnableListenerFocus && !Sound->bIgnoreFocus)
+	const float FocusAzimuth = FMath::Clamp(GlobalFocusSettings.FocusAzimuthScale * AttenuationSettings.FocusAzimuth, 0.0f, 180.0f);
+	const float NonFocusAzimuth = FMath::Clamp(GlobalFocusSettings.NonFocusAzimuthScale * AttenuationSettings.NonFocusAzimuth, 0.0f, 180.0f);
+
+	if (FocusAzimuth != NonFocusAzimuth)
+	{
+		FocusFactor = (Azimuth - FocusAzimuth) / (NonFocusAzimuth - FocusAzimuth);
+		FocusFactor = FMath::Clamp(FocusFactor, 0.0f, 1.0f);
+	}
+	else if (Azimuth >= FocusAzimuth)
 	{
-		// Compute the focus factor based on listener emitter geometry (SoundTransform)
-		GetAttenuationListenerData(OutListenerData, SoundTransform, AttenuationSettings, InListenerTransform);
+		FocusFactor = 1.0f;
+	}
 
-		const FVector& ListenerForwardDir = OutListenerData.ListenerTransform.GetUnitAxis(EAxis::X);
+	return FocusFactor;
+}
 
-		const float FocusDotProduct = FVector::DotProduct(ListenerForwardDir, OutListenerData.ListenerToSoundDir);
-		const float FocusAngleRadians = FMath::Acos(FocusDotProduct);
-		const float FocusAngle = FMath::RadiansToDegrees(FocusAngleRadians);
+FAudioDevice::FCreateComponentParams::FCreateComponentParams()
+	: World(nullptr)
+	, Actor(nullptr)
+{
+	AudioDevice = (GEngine ? GEngine->GetMainAudioDevice() : nullptr);
+	CommonInit();
+}
 
-		const FGlobalFocusSettings& FocusSettings = GetGlobalFocusSettings();
+FAudioDevice::FCreateComponentParams::FCreateComponentParams(UWorld* InWorld, AActor* InActor)
+	: World(InWorld)
+{
+	if (InActor)
+	{
+		check(InActor->GetWorld() == InWorld);
+		Actor = InActor;
+	}
+	else
+	{
+		Actor = (World ? World->GetWorldSettings() : nullptr);
+	}
+	
+	AudioDevice = (World ? World->GetAudioDevice() : nullptr);
+	CommonInit();
+}
 
-		const float FocusAzimuth = FMath::Clamp(FocusSettings.FocusAzimuthScale * AttenuationSettings.FocusAzimuth, 0.0f, 180.0f);
-		const float NonFocusAzimuth = FMath::Clamp(FocusSettings.NonFocusAzimuthScale * AttenuationSettings.NonFocusAzimuth, 0.0f, 180.0f);
+FAudioDevice::FCreateComponentParams::FCreateComponentParams(AActor* InActor)
+	: Actor(InActor)
+{
+	World = (Actor ? Actor->GetWorld() : nullptr);
+	AudioDevice = (World ? World->GetAudioDevice() : nullptr);
+	CommonInit();
+}
 
-		if (FocusAzimuth != NonFocusAzimuth)
-		{
-			FocusFactor = (FocusAngle - FocusAzimuth) / (NonFocusAzimuth - FocusAzimuth);
-			FocusFactor = FMath::Clamp(FocusFactor, 0.0f, 1.0f);
-		}
-		else if (FocusAngle >= FocusAzimuth)
-		{
-			FocusFactor = 1.0f;
-		}
-	}
+FAudioDevice::FCreateComponentParams::FCreateComponentParams(FAudioDevice* InAudioDevice)
+	: World(nullptr)
+	, Actor(nullptr)
+	, AudioDevice(InAudioDevice)
+{
+	CommonInit();
+}
 
-	return FocusFactor;
+void FAudioDevice::FCreateComponentParams::CommonInit()
+{
+	bPlay = false;
+	bLocationSet = false;
+	AttenuationSettings = nullptr;
+	ConcurrencySettings = nullptr;
+}
+
+void FAudioDevice::FCreateComponentParams::SetLocation(const FVector InLocation)
+{
+	if (World)
+	{
+		bLocationSet = true;
+		Location = InLocation;
+	}
+	else
+	{
+		UE_LOG(LogAudio, Warning, TEXT("AudioComponents created without a World cannot have a location."));
+	}
 }
 
 UAudioComponent* FAudioDevice::CreateComponent(USoundBase* Sound, UWorld* World, AActor* Actor, bool bPlay, bool bStopWhenOwnerDestroyed, const FVector* Location, USoundAttenuation* AttenuationSettings, USoundConcurrency* ConcurrencySettings)
 {
-	check(IsInGameThread());
+	TUniquePtr<FCreateComponentParams> Params;
 
-	UAudioComponent* AudioComponent = nullptr;
+	if (Actor)
+	{
+		Params = MakeUnique<FCreateComponentParams>(Actor);
+	}
+	else if (World)
+	{
+		Params = MakeUnique<FCreateComponentParams>(World);
+	}
+	else
+	{
+		Params = MakeUnique<FCreateComponentParams>(GEngine->GetMainAudioDevice());
+	}
 
-	if (Sound && GEngine && GEngine->UseSound())
+	Params->bPlay = bPlay;
+	Params->bStopWhenOwnerDestroyed = bStopWhenOwnerDestroyed;
+	Params->AttenuationSettings = AttenuationSettings;
+	Params->ConcurrencySettings = ConcurrencySettings;
+	if (Location)
 	{
-		// Get the world's audio device if there is a world, otherwise get the main audio device
-		FAudioDevice* AudioDevice = nullptr;
-		if (World)
-		{
-			AudioDevice = World->GetAudioDevice();
-		}
-		else
-		{
-			AudioDevice = GEngine->GetMainAudioDevice();
-		}
-		
-		if (AudioDevice == nullptr)
-		{
-			return nullptr;
-		}
+		Params->SetLocation(*Location);
+	}
+	return CreateComponent(Sound, *Params);
+}
+
+UAudioComponent* FAudioDevice::CreateComponent(USoundBase* Sound, const FCreateComponentParams& Params)
+{
+	check(IsInGameThread());
+
+	UAudioComponent* AudioComponent = nullptr;
 
+	if (Sound && Params.AudioDevice && GEngine && GEngine->UseSound())
+	{	
 		// Avoid creating component if we're trying to play a sound on an already destroyed actor.
-		if (Actor && Actor->IsPendingKill())
-		{
-			// Don't create component on destroyed actor.
-		}
-		// Either no actor or actor is still alive.
-		else
+		if (Params.Actor == nullptr || !Params.Actor->IsPendingKill())
 		{
 			// Listener position could change before long sounds finish
-			const FAttenuationSettings* AttenuationSettingsToApply = (AttenuationSettings ? &AttenuationSettings->Attenuation : Sound->GetAttenuationSettingsToApply());
+			const FAttenuationSettings* AttenuationSettingsToApply = (Params.AttenuationSettings ? &Params.AttenuationSettings->Attenuation : Sound->GetAttenuationSettingsToApply());
 
 			bool bIsAudible = true;
 			// If a sound is a long duration, the position might change before sound finishes so assume it's audible
-			if (Location && Sound->GetDuration() <= 1.0f)
+			if (Params.bLocationSet && Sound->GetDuration() <= 1.0f)
 			{
 				float MaxDistance = 0.0f;
 				float FocusFactor = 0.0f;
-				AudioDevice->GetMaxDistanceAndFocusFactor(Sound, World, *Location, AttenuationSettingsToApply, MaxDistance, FocusFactor);
-				bIsAudible = AudioDevice->SoundIsAudible(Sound, World, *Location, AttenuationSettingsToApply, MaxDistance, FocusFactor);
+				Params.AudioDevice->GetMaxDistanceAndFocusFactor(Sound, Params.World, Params.Location, AttenuationSettingsToApply, MaxDistance, FocusFactor);
+				bIsAudible = Params.AudioDevice->SoundIsAudible(Sound, Params.World, Params.Location, AttenuationSettingsToApply, MaxDistance, FocusFactor);
 			}
 
 			if (bIsAudible)
 			{
 				// Use actor as outer if we have one.
-				if (Actor)
+				if (Params.Actor)
 				{
-					AudioComponent = NewObject<UAudioComponent>(Actor);
+					AudioComponent = NewObject<UAudioComponent>(Params.Actor);
 				}
 				// Let engine pick the outer (transient package).
 				else
@@ -3552,27 +3745,31 @@ UAudioComponent* FAudioDevice::CreateComponent(USoundBase* Sound, UWorld* World,
 				AudioComponent->Sound = Sound;
 				AudioComponent->bAutoActivate = false;
 				AudioComponent->bIsUISound = false;
-				AudioComponent->bAutoDestroy = bPlay;
-				AudioComponent->bStopWhenOwnerDestroyed = bStopWhenOwnerDestroyed;
+				AudioComponent->bAutoDestroy = Params.bPlay;
+				AudioComponent->bStopWhenOwnerDestroyed = Params.bStopWhenOwnerDestroyed;
 #if WITH_EDITORONLY_DATA
 				AudioComponent->bVisualizeComponent = false;
 #endif
-				AudioComponent->AttenuationSettings = AttenuationSettings;
-				AudioComponent->ConcurrencySettings = ConcurrencySettings;
+				AudioComponent->AttenuationSettings = Params.AttenuationSettings;
+				AudioComponent->ConcurrencySettings = Params.ConcurrencySettings;
 
-				if (Location)
+				if (Params.bLocationSet)
 				{
-					AudioComponent->SetWorldLocation(*Location);
+					AudioComponent->SetWorldLocation(Params.Location);
 				}
 
 				// AudioComponent used in PlayEditorSound sets World to nullptr to avoid situations where the world becomes invalid
 				// and the component is left with invalid pointer.
-				if (World)
+				if (Params.World)
 				{
-					AudioComponent->RegisterComponentWithWorld(World);
+					AudioComponent->RegisterComponentWithWorld(Params.World);
+				}
+				else
+				{
+					AudioComponent->AudioDeviceHandle = Params.AudioDevice->DeviceHandle;
 				}
 
-				if (bPlay)
+				if (Params.bPlay)
 				{
 					AudioComponent->Play();
 				}
@@ -3585,7 +3782,7 @@ UAudioComponent* FAudioDevice::CreateComponent(USoundBase* Sound, UWorld* World,
 		}
 	}
 
-	return(AudioComponent);
+	return AudioComponent;
 }
 
 void FAudioDevice::PlaySoundAtLocation(USoundBase* Sound, UWorld* World, float VolumeMultiplier, float PitchMultiplier, float StartTime, const FVector& Location, const FRotator& Rotation, USoundAttenuation* AttenuationSettings, USoundConcurrency* ConcurrencySettings, const TArray<FAudioComponentParam>* Params)
@@ -3913,6 +4110,83 @@ void FAudioDevice::UnregisterSoundClass(USoundClass* SoundClass)
 	}
 }
 
+void FAudioDevice::InitSoundEffectPresets()
+{
+	// Reset existing submixes if they exist
+	SoundEffectSourcePresetInstances.Reset();
+	SoundEffectSubmixPresetInstances.Reset();
+
+#if WITH_EDITOR
+
+	FAssetToolsModule& AssetToolsModule = FModuleManager::Get().LoadModuleChecked<FAssetToolsModule>("AssetTools");
+
+	for (TObjectIterator<UClass> It; It; ++It)
+	{
+		UClass* ChildClass = *It;
+		if (ChildClass->HasAnyClassFlags(CLASS_Abstract))
+		{
+			continue;
+		}
+
+		// Look for submix or source preset classes
+		if (ChildClass->IsChildOf<USoundEffectPreset>())
+		{		
+			USoundEffectPreset* EffectPreset = ChildClass->GetDefaultObject<USoundEffectPreset>();
+			AssetToolsModule.Get().RegisterAssetTypeActions(MakeShareable(new FAssetTypeActions_SoundEffectPreset(EffectPreset)));
+		}
+	}
+
+#endif
+
+	// Reset the maps of sound instances properties
+	for (TObjectIterator<USoundEffectSourcePreset> It; It; ++It)
+	{
+		RegisterSoundEffectSourcePreset(*It);
+	}
+
+	for (TObjectIterator<USoundEffectSubmixPreset> It; It; ++It)
+	{
+		RegisterSoundEffectSubmixPreset(*It);
+	}
+
+}
+
+void FAudioDevice::RegisterSoundEffectSourcePreset(USoundEffectSourcePreset* SoundEffectPreset)
+{
+	if (SoundEffectPreset && !SoundEffectSourcePresetInstances.Contains(SoundEffectPreset))
+	{
+		USoundEffectSourcePreset* SoundEffectSubmixPresetInstance = NewObject<USoundEffectSourcePreset>((UObject*)GetTransientPackage(), SoundEffectPreset->GetClass());
+		SoundEffectSourcePresetInstances.Add(SoundEffectPreset, SoundEffectSubmixPresetInstance);
+	}
+}
+
+void FAudioDevice::UnregisterSoundEffectSourcePreset(USoundEffectSourcePreset* SoundEffectPreset)
+{
+	if (SoundEffectPreset)
+	{
+		SoundEffectSourcePresetInstances.Remove(SoundEffectPreset);
+	}
+}
+
+void FAudioDevice::RegisterSoundEffectSubmixPreset(USoundEffectSubmixPreset* SoundEffectPreset)
+{
+
+	if (SoundEffectPreset && !SoundEffectSubmixPresetInstances.Contains(SoundEffectPreset))
+	{
+		USoundEffectSubmixPreset* SoundEffectSubmixPresetInstance = NewObject<USoundEffectSubmixPreset>((UObject*)GetTransientPackage(), SoundEffectPreset->GetClass());
+		SoundEffectSubmixPresetInstances.Add(SoundEffectPreset, SoundEffectSubmixPresetInstance);
+	}
+}
+
+void FAudioDevice::UnregisterSoundEffectSubmixPreset(USoundEffectSubmixPreset* SoundEffectPreset)
+{
+	if (SoundEffectPreset)
+	{
+		SoundEffectSubmixPresetInstances.Remove(SoundEffectPreset);
+	}
+
+}
+
 FSoundClassProperties* FAudioDevice::GetSoundClassCurrentProperties(USoundClass* InSoundClass)
 {
 	if (InSoundClass)
diff --git a/Engine/Source/Runtime/Engine/Private/AudioDeviceManager.cpp b/Engine/Source/Runtime/Engine/Private/AudioDeviceManager.cpp
--- a/Engine/Source/Runtime/Engine/Private/AudioDeviceManager.cpp
+++ b/Engine/Source/Runtime/Engine/Private/AudioDeviceManager.cpp
@@ -17,9 +17,6 @@ static const uint32 AUDIO_DEVICE_HANDLE_GENERATION_MASK = (1 << AUDIO_DEVICE_HAN
 
 static const uint16 AUDIO_DEVICE_MINIMUM_FREE_AUDIO_DEVICE_INDICES = 32;
 
-// Invalid handle used for initializing audio device handles
-static const uint32 AUDIO_DEVICE_HANDLE_INVALID			= ~0u;
-
 // The number of multiple audio devices allowed by default
 static const uint32 AUDIO_DEVICE_DEFAULT_ALLOWED_DEVICE_COUNT = 2;
 
@@ -27,7 +24,7 @@ static const uint32 AUDIO_DEVICE_DEFAULT_ALLOWED_DEVICE_COUNT = 2;
 static const uint32 AUDIO_DEVICE_MAX_DEVICE_COUNT = 8;
 
 FAudioDeviceManager::FCreateAudioDeviceResults::FCreateAudioDeviceResults()
-	: Handle(AUDIO_DEVICE_HANDLE_INVALID)
+	: Handle(INDEX_NONE)
 	, bNewDevice(false)
 	, AudioDevice(nullptr)
 {
@@ -107,7 +104,7 @@ bool FAudioDeviceManager::CreateAudioDevice(bool bCreateNewDevice, FCreateAudioD
 		// Now generation a new audio device handle for the device and store the
 		// ptr to the new device in the array of audio devices.
 
-		uint32 AudioDeviceIndex(AUDIO_DEVICE_HANDLE_INVALID);
+		uint32 AudioDeviceIndex(INDEX_NONE);
 
 		// First check to see if we should start recycling audio device indices, if not
 		// then we add a new entry to the Generation array and generate a new index
@@ -161,7 +158,7 @@ bool FAudioDeviceManager::CreateAudioDevice(bool bCreateNewDevice, FCreateAudioD
 
 bool FAudioDeviceManager::IsValidAudioDeviceHandle(uint32 Handle) const
 {
-	if (AudioDeviceModule == nullptr || Handle == AUDIO_DEVICE_HANDLE_INVALID)
+	if (AudioDeviceModule == nullptr || Handle == INDEX_NONE)
 	{
 		return false;
 	}
@@ -254,7 +251,7 @@ bool FAudioDeviceManager::ShutdownAllAudioDevices()
 	return true;
 }
 
-class FAudioDevice* FAudioDeviceManager::GetAudioDevice(uint32 Handle)
+FAudioDevice* FAudioDeviceManager::GetAudioDevice(uint32 Handle)
 {
 	if (!IsValidAudioDeviceHandle(Handle))
 	{
@@ -268,7 +265,7 @@ class FAudioDevice* FAudioDeviceManager::GetAudioDevice(uint32 Handle)
 	return AudioDevice;
 }
 
-class FAudioDevice* FAudioDeviceManager::GetActiveAudioDevice()
+FAudioDevice* FAudioDeviceManager::GetActiveAudioDevice()
 {
 	if (ActiveAudioDeviceHandle != INDEX_NONE)
 	{
@@ -304,7 +301,7 @@ void FAudioDeviceManager::AddReferencedObjects(FReferenceCollector& Collector)
 	}
 }
 
-void FAudioDeviceManager::StopSoundsUsingResource(class USoundWave* InSoundWave, TArray<UAudioComponent*>* StoppedComponents)
+void FAudioDeviceManager::StopSoundsUsingResource(USoundWave* InSoundWave, TArray<UAudioComponent*>* StoppedComponents)
 {
 	for (FAudioDevice* AudioDevice : Devices)
 	{
@@ -348,6 +345,50 @@ void FAudioDeviceManager::InitSoundClasses()
 	}
 }
 
+void FAudioDeviceManager::RegisterSoundSubmix(USoundSubmix* SoundSubmix)
+{
+	for (FAudioDevice* AudioDevice : Devices)
+	{
+		if (AudioDevice)
+		{
+			AudioDevice->RegisterSoundSubmix(SoundSubmix);
+		}
+	}
+}
+
+void FAudioDeviceManager::UnregisterSoundSubmix(USoundSubmix* SoundSubmix)
+{
+	for (FAudioDevice* AudioDevice : Devices)
+	{
+		if (AudioDevice)
+		{
+			AudioDevice->UnregisterSoundSubmix(SoundSubmix);
+		}
+	}
+}
+
+void FAudioDeviceManager::InitSoundSubmixes()
+{
+	for (FAudioDevice* AudioDevice : Devices)
+	{
+		if (AudioDevice)
+		{
+			AudioDevice->InitSoundSubmixes();
+		}
+	}
+}
+
+void FAudioDeviceManager::InitSoundEffectPresets()
+{
+	for (FAudioDevice* AudioDevice : Devices)
+	{
+		if (AudioDevice)
+		{
+			AudioDevice->InitSoundEffectPresets();
+		}
+	}
+}
+
 void FAudioDeviceManager::SetActiveDevice(uint32 InAudioDeviceHandle)
 {
 	// Only change the active device if there are no solo'd audio devices
@@ -609,4 +650,14 @@ const FString& FAudioDeviceManager::GetDebugSoloSoundCue() const
 	return DebugNames.DebugSoloSoundCue;
 }
 
+void FAudioDeviceManager::SetAudioMixerDebugSound(const TCHAR* SoundName)
+{
+	DebugNames.DebugAudioMixerSoundName = SoundName;
+}
+
+const FString& FAudioDeviceManager::GetAudioMixerDebugSoundName() const
+{
+	return DebugNames.DebugAudioMixerSoundName;
+}
+
 
diff --git a/Engine/Source/Runtime/Engine/Private/AudioSettings.cpp b/Engine/Source/Runtime/Engine/Private/AudioSettings.cpp
--- a/Engine/Source/Runtime/Engine/Private/AudioSettings.cpp
+++ b/Engine/Source/Runtime/Engine/Private/AudioSettings.cpp
@@ -6,6 +6,7 @@
 
 #include "EnginePrivate.h"
 #include "Sound/AudioSettings.h"
+#include "Sound/SoundBase.h"
 #include "Sound/SoundNodeQualityLevel.h"
 
 #define LOCTEXT_NAMESPACE "AudioSettings"
@@ -79,6 +80,13 @@ void UAudioSettings::PostEditChangeChainProperty(FPropertyChangedChainEvent& Pro
 		{
 			bReconcileNodes = true;
 		}
+		else if (PropertyChangedEvent.Property->GetFName() == GET_MEMBER_NAME_CHECKED(UAudioSettings, DefaultSoundSubmixName))
+		{
+			if (DefaultSoundSubmixName.IsValid())
+			{
+				USoundBase::DefaultSoundSubmixObject = LoadObject<USoundSubmix>(nullptr, *DefaultSoundSubmixName.ToString());
+			}
+		}
 
 		if (bReconcileNodes)
 		{
diff --git a/Engine/Source/Runtime/Engine/Private/BlueprintGeneratedClass.cpp b/Engine/Source/Runtime/Engine/Private/BlueprintGeneratedClass.cpp
--- a/Engine/Source/Runtime/Engine/Private/BlueprintGeneratedClass.cpp
+++ b/Engine/Source/Runtime/Engine/Private/BlueprintGeneratedClass.cpp
@@ -772,7 +772,7 @@ void UBlueprintGeneratedClass::CreateTimelineComponent(AActor* Actor, const UTim
 		const FTTFloatTrack* FloatTrackTemplate = &TimelineTemplate->FloatTracks[TrackIdx];
 		if (FloatTrackTemplate->CurveFloat != NULL)
 		{
-			NewTimeline->AddInterpFloat(FloatTrackTemplate->CurveFloat, FOnTimelineFloat(), TimelineTemplate->GetTrackPropertyName(FloatTrackTemplate->TrackName));
+			NewTimeline->AddInterpFloat(FloatTrackTemplate->CurveFloat, FOnTimelineFloat(), TimelineTemplate->GetTrackPropertyName(FloatTrackTemplate->TrackName), FloatTrackTemplate->TrackName);
 		}
 	}
 
@@ -782,7 +782,7 @@ void UBlueprintGeneratedClass::CreateTimelineComponent(AActor* Actor, const UTim
 		const FTTVectorTrack* VectorTrackTemplate = &TimelineTemplate->VectorTracks[TrackIdx];
 		if (VectorTrackTemplate->CurveVector != NULL)
 		{
-			NewTimeline->AddInterpVector(VectorTrackTemplate->CurveVector, FOnTimelineVector(), TimelineTemplate->GetTrackPropertyName(VectorTrackTemplate->TrackName));
+			NewTimeline->AddInterpVector(VectorTrackTemplate->CurveVector, FOnTimelineVector(), TimelineTemplate->GetTrackPropertyName(VectorTrackTemplate->TrackName), VectorTrackTemplate->TrackName);
 		}
 	}
 
@@ -792,7 +792,7 @@ void UBlueprintGeneratedClass::CreateTimelineComponent(AActor* Actor, const UTim
 		const FTTLinearColorTrack* LinearColorTrackTemplate = &TimelineTemplate->LinearColorTracks[TrackIdx];
 		if (LinearColorTrackTemplate->CurveLinearColor != NULL)
 		{
-			NewTimeline->AddInterpLinearColor(LinearColorTrackTemplate->CurveLinearColor, FOnTimelineLinearColor(), TimelineTemplate->GetTrackPropertyName(LinearColorTrackTemplate->TrackName));
+			NewTimeline->AddInterpLinearColor(LinearColorTrackTemplate->CurveLinearColor, FOnTimelineLinearColor(), TimelineTemplate->GetTrackPropertyName(LinearColorTrackTemplate->TrackName), LinearColorTrackTemplate->TrackName);
 		}
 	}
 
diff --git a/Engine/Source/Runtime/Engine/Private/Camera/CameraComponent.cpp b/Engine/Source/Runtime/Engine/Private/Camera/CameraComponent.cpp
--- a/Engine/Source/Runtime/Engine/Private/Camera/CameraComponent.cpp
+++ b/Engine/Source/Runtime/Engine/Private/Camera/CameraComponent.cpp
@@ -79,8 +79,7 @@ void UCameraComponent::OnRegister()
 		{
 			ProxyMeshComponent = NewObject<UStaticMeshComponent>(MyOwner, NAME_None, RF_Transactional | RF_TextExportTransient);
 			ProxyMeshComponent->SetupAttachment(this);
-			ProxyMeshComponent->AlwaysLoadOnClient = false;
-			ProxyMeshComponent->AlwaysLoadOnServer = false;
+			ProxyMeshComponent->bIsEditorOnly = true;
 			ProxyMeshComponent->StaticMesh = CameraMesh;
 			ProxyMeshComponent->SetCollisionProfileName(UCollisionProfile::NoCollision_ProfileName);
 			ProxyMeshComponent->bHiddenInGame = true;
@@ -94,8 +93,7 @@ void UCameraComponent::OnRegister()
 		{
 			DrawFrustum = NewObject<UDrawFrustumComponent>(MyOwner, NAME_None, RF_Transactional | RF_TextExportTransient);
 			DrawFrustum->SetupAttachment(this);
-			DrawFrustum->AlwaysLoadOnClient = false;
-			DrawFrustum->AlwaysLoadOnServer = false;
+			DrawFrustum->bIsEditorOnly = true;
 			DrawFrustum->CreationMethod = CreationMethod;
 			DrawFrustum->RegisterComponentWithWorld(GetWorld());
 		}
diff --git a/Engine/Source/Runtime/Engine/Private/CheatManager.cpp b/Engine/Source/Runtime/Engine/Private/CheatManager.cpp
--- a/Engine/Source/Runtime/Engine/Private/CheatManager.cpp
+++ b/Engine/Source/Runtime/Engine/Private/CheatManager.cpp
@@ -22,7 +22,7 @@
 #include "Components/CapsuleComponent.h"
 #include "Engine/DebugCameraController.h"
 #include "GameFramework/PlayerState.h"
-#include "GameFramework/GameMode.h"
+#include "GameFramework/GameModeBase.h"
 #include "GameFramework/PlayerInput.h"
 #include "GameFramework/InputSettings.h"
 #include "Components/BrushComponent.h"
diff --git a/Engine/Source/Runtime/Engine/Private/Components/ActorComponent.cpp b/Engine/Source/Runtime/Engine/Private/Components/ActorComponent.cpp
--- a/Engine/Source/Runtime/Engine/Private/Components/ActorComponent.cpp
+++ b/Engine/Source/Runtime/Engine/Private/Components/ActorComponent.cpp
@@ -106,6 +106,8 @@ FActorComponentCreatePhysicsSignature UActorComponent::CreatePhysicsDelegate;
 // Destroy Physics global delegate
 FActorComponentDestroyPhysicsSignature UActorComponent::DestroyPhysicsDelegate;
 
+const FString UActorComponent::ComponentTemplateNameSuffix(TEXT("_GEN_VARIABLE"));
+
 UActorComponent::UActorComponent(const FObjectInitializer& ObjectInitializer /*= FObjectInitializer::Get()*/)
 	: Super(ObjectInitializer)
 {
@@ -239,7 +241,7 @@ void UActorComponent::PostRename(UObject* OldOuter, const FName OldName)
 	if (OldOuter != GetOuter())
 	{
 		OwnerPrivate = GetTypedOuter<AActor>();
-		AActor* OldOwner = (OldOuter->IsA<AActor>() ? CastChecked<AActor>(OldOuter) : OldOuter->GetTypedOuter<AActor>());
+		AActor* OldOwner = (OldOuter->IsA<AActor>() ? static_cast<AActor*>(OldOuter) : OldOuter->GetTypedOuter<AActor>());
 
 		if (OwnerPrivate != OldOwner)
 		{
@@ -253,21 +255,28 @@ void UActorComponent::PostRename(UObject* OldOuter, const FName OldName)
 			}
 
 			TArray<UObject*> Children;
-			GetObjectsWithOuter(this, Children);
+			GetObjectsWithOuter(this, Children, /*bIncludeNestedObjects=*/false);
 
-			for (UObject* Child : Children)
+			for (int32 Index = 0; Index < Children.Num(); ++Index)
 			{
-				if (UActorComponent* ChildComponent = Cast<UActorComponent>(Child))
+				UObject* Child = Children[Index];
+
+				// Cut off if we have a nested Actor
+				if (!Child->IsA<AActor>())
 				{
-					ChildComponent->OwnerPrivate = OwnerPrivate;
-					if (OldOwner)
+					if (UActorComponent* ChildComponent = Cast<UActorComponent>(Child))
 					{
-						OldOwner->RemoveOwnedComponent(ChildComponent);
-					}
-					if (OwnerPrivate)
-					{
-						OwnerPrivate->AddOwnedComponent(ChildComponent);
+						ChildComponent->OwnerPrivate = OwnerPrivate;
+						if (OldOwner)
+						{
+							OldOwner->RemoveOwnedComponent(ChildComponent);
+						}
+						if (OwnerPrivate)
+						{
+							OwnerPrivate->AddOwnedComponent(ChildComponent);
+						}
 					}
+					GetObjectsWithOuter(Child, Children, /*bIncludeNestedObjects=*/false);
 				}
 			}
 		}
@@ -425,18 +434,16 @@ void UActorComponent::BeginDestroy()
 	Super::BeginDestroy();
 }
 
-
 bool UActorComponent::NeedsLoadForClient() const
 {
 	check(GetOuter());
-	return (GetOuter()->NeedsLoadForClient() && Super::NeedsLoadForClient());
+	return (!IsEditorOnly() && GetOuter()->NeedsLoadForClient() && Super::NeedsLoadForClient());
 }
 
-
 bool UActorComponent::NeedsLoadForServer() const
 {
 	check(GetOuter());
-	return (GetOuter()->NeedsLoadForServer() && Super::NeedsLoadForServer());
+	return (!IsEditorOnly() && GetOuter()->NeedsLoadForServer() && Super::NeedsLoadForServer());
 }
 
 int32 UActorComponent::GetFunctionCallspace( UFunction* Function, void* Parameters, FFrame* Stack )
@@ -554,21 +561,28 @@ void UActorComponent::PostEditUndo()
 		}
 
 		TArray<UObject*> Children;
-		GetObjectsWithOuter(this, Children);
+		GetObjectsWithOuter(this, Children, /*bIncludeNestedObjects=*/false);
 
-		for (UObject* Child : Children)
+		for (int32 Index = 0; Index < Children.Num(); ++Index)
 		{
-			if (UActorComponent* ChildComponent = Cast<UActorComponent>(Child))
+			UObject* Child = Children[Index];
+
+			// Cut off if we have a nested Actor
+			if (!Child->IsA<AActor>())
 			{
-				if (ChildComponent->OwnerPrivate)
-				{
-					ChildComponent->OwnerPrivate->RemoveOwnedComponent(ChildComponent);
-				}
-				ChildComponent->OwnerPrivate = OwnerPrivate;
-				if (OwnerPrivate)
+				if (UActorComponent* ChildComponent = Cast<UActorComponent>(Child))
 				{
-					OwnerPrivate->AddOwnedComponent(ChildComponent);
+					if (ChildComponent->OwnerPrivate)
+					{
+						ChildComponent->OwnerPrivate->RemoveOwnedComponent(ChildComponent);
+					}
+					ChildComponent->OwnerPrivate = OwnerPrivate;
+					if (OwnerPrivate)
+					{
+						OwnerPrivate->AddOwnedComponent(ChildComponent);
+					}
 				}
+				GetObjectsWithOuter(Child, Children, /*bIncludeNestedObjects=*/false);
 			}
 		}
 
@@ -916,7 +930,7 @@ void UActorComponent::RegisterComponentWithWorld(UWorld* InWorld)
 		if (MyOwner->HasActorBegunPlay() || MyOwner->IsActorBeginningPlay())
 		{
 			RegisterAllComponentTickFunctions(true);
-			if (bWantsBeginPlay && !bHasBegunPlay)
+			if (!bHasBegunPlay)
 			{
 				BeginPlay();
 			}
@@ -932,7 +946,7 @@ void UActorComponent::RegisterComponentWithWorld(UWorld* InWorld)
 		for (UObject* Child : Children)
 		{
 			UActorComponent* ChildComponent = Cast<UActorComponent>(Child);
-			if (ChildComponent && !ChildComponent->IsRegistered())
+			if (ChildComponent && !ChildComponent->IsRegistered() && ChildComponent->GetOwner() == MyOwner)
 			{
 				ChildComponent->RegisterComponentWithWorld(InWorld);
 			}
@@ -1698,4 +1712,4 @@ void UActorComponent::Serialize(FArchive& Ar)
 	}
 }
 
-#undef LOCTEXT_NAMESPACE
\ No newline at end of file
+#undef LOCTEXT_NAMESPACE
diff --git a/Engine/Source/Runtime/Engine/Private/Components/AudioComponent.cpp b/Engine/Source/Runtime/Engine/Private/Components/AudioComponent.cpp
--- a/Engine/Source/Runtime/Engine/Private/Components/AudioComponent.cpp
+++ b/Engine/Source/Runtime/Engine/Private/Components/AudioComponent.cpp
@@ -31,6 +31,7 @@ UAudioComponent::UAudioComponent(const FObjectInitializer& ObjectInitializer)
 	VolumeMultiplier = 1.f;
 	bOverridePriority = false;
 	bIsPreviewSound = false;
+	bIsPaused = false;
 	Priority = 1.f;
 	PitchMultiplier = 1.f;
 	VolumeModulationMin = 1.f;
@@ -42,6 +43,7 @@ UAudioComponent::UAudioComponent(const FObjectInitializer& ObjectInitializer)
 	OcclusionCheckInterval = 0.1f;
 	ActiveCount = 0;
 
+	AudioDeviceHandle = INDEX_NONE;
 	AudioComponentID = ++AudioComponentIDCounter;
 
 	// TODO: Consider only putting played/active components in to the map
@@ -248,6 +250,8 @@ void UAudioComponent::PlayInternal(const float StartTime, const float FadeInDura
 			NewActiveSound.bCenterChannelOnly = bCenterChannelOnly;
 			NewActiveSound.bIsPreviewSound = bIsPreviewSound;
 			NewActiveSound.bLocationDefined = !bPreviewComponent;
+			NewActiveSound.bIsPaused = bIsPaused;
+
 			if (NewActiveSound.bLocationDefined)
 			{
 				NewActiveSound.Transform = ComponentToWorld;
@@ -261,6 +265,8 @@ void UAudioComponent::PlayInternal(const float StartTime, const float FadeInDura
 				NewActiveSound.FocusPriorityScale = AttenuationSettingsToApply->GetFocusPriorityScale(AudioDevice->GetGlobalFocusSettings(), FocusFactor);
 			}
 
+			NewActiveSound.bUpdatePlayPercentage = OnAudioPlaybackPercentNative.IsBound() || OnAudioPlaybackPercent.IsBound();
+
 			NewActiveSound.MaxDistance = MaxDistance;
 
 			NewActiveSound.InstanceParameters = InstanceParameters;
@@ -291,8 +297,12 @@ FAudioDevice* UAudioComponent::GetAudioDevice() const
 
 	if (GEngine)
 	{
-		UWorld* World = GetWorld();
-		if (World)
+		if (AudioDeviceHandle != INDEX_NONE)
+		{
+			FAudioDeviceManager* AudioDeviceManager = GEngine->GetAudioDeviceManager();
+			AudioDevice = (AudioDeviceManager ? AudioDeviceManager->GetAudioDevice(AudioDeviceHandle) : nullptr);
+		}
+		else if (UWorld* World = GetWorld())
 		{
 			AudioDevice = World->GetAudioDevice();
 		}
@@ -393,6 +403,30 @@ void UAudioComponent::Stop()
 	}
 }
 
+void UAudioComponent::SetPaused(bool bPause)
+{
+	if (bIsPaused != bPause)
+	{
+		bIsPaused = bPause;
+
+		if (bIsActive)
+		{
+			UE_LOG(LogAudio, Verbose, TEXT("%g: Pausing AudioComponent : '%s' with Sound: '%s'"), GetWorld() ? GetWorld()->GetAudioTimeSeconds() : 0.0f, *GetFullName(), Sound ? *Sound->GetName() : TEXT("nullptr"));
+
+			if (FAudioDevice* AudioDevice = GetAudioDevice())
+			{
+				DECLARE_CYCLE_STAT(TEXT("FAudioThreadTask.PauseActiveSound"), STAT_AudioPauseActiveSound, STATGROUP_AudioThreadCommands);
+
+				const uint64 MyAudioComponentID = AudioComponentID;
+				FAudioThread::RunCommandOnAudioThread([AudioDevice, MyAudioComponentID, bPause]()
+				{
+					AudioDevice->PauseActiveSound(MyAudioComponentID, bPause);
+				}, GET_STATID(STAT_AudioPauseActiveSound));
+			}
+		}
+	}
+}
+
 void UAudioComponent::PlaybackCompleted(uint64 AudioComponentID, bool bFailedToStart)
 {
 	check(IsInAudioThread());
diff --git a/Engine/Source/Runtime/Engine/Private/Components/BrushComponent.cpp b/Engine/Source/Runtime/Engine/Private/Components/BrushComponent.cpp
--- a/Engine/Source/Runtime/Engine/Private/Components/BrushComponent.cpp
+++ b/Engine/Source/Runtime/Engine/Private/Components/BrushComponent.cpp
@@ -780,6 +780,11 @@ void UBrushComponent::BuildSimpleBrushCollision()
 	MarkPackageDirty();
 }
 
+bool UBrushComponent::IsEditorOnly() const
+{
+	// Default to actor component behavior instead of primitive component behavior as brush actors handle it themselves
+	return bIsEditorOnly;
+}
 
 #if WITH_EDITOR
 static FVector GetPolyCenter(const FPoly& Poly)
diff --git a/Engine/Source/Runtime/Engine/Private/Components/CharacterMovementComponent.cpp b/Engine/Source/Runtime/Engine/Private/Components/CharacterMovementComponent.cpp
--- a/Engine/Source/Runtime/Engine/Private/Components/CharacterMovementComponent.cpp
+++ b/Engine/Source/Runtime/Engine/Private/Components/CharacterMovementComponent.cpp
@@ -10,7 +10,7 @@
 #include "GameFramework/GameNetworkManager.h"
 #include "GameFramework/Character.h"
 #include "GameFramework/CharacterMovementComponent.h"
-#include "GameFramework/GameState.h"
+#include "GameFramework/GameStateBase.h"
 #include "Components/PrimitiveComponent.h"
 #include "Animation/AnimMontage.h"
 #include "PhysicsEngine/DestructibleActor.h"
@@ -29,6 +29,7 @@
 #include "HAL/IConsoleManager.h"
 #include "PerfCountersHelpers.h"
 
+
 DEFINE_LOG_CATEGORY_STATIC(LogCharacterMovement, Log, All);
 DEFINE_LOG_CATEGORY_STATIC(LogNavMeshMovement, Log, All);
 DEFINE_LOG_CATEGORY_STATIC(LogCharacterNetSmoothing, Log, All);
@@ -5881,12 +5882,6 @@ bool UCharacterMovementComponent::StepUp(const FVector& GravDir, const FVector&
 		return false;
 	}
 
-	// Don't step up if the impact is below us
-	if (InitialImpactZ <= OldLocation.Z - PawnHalfHeight)
-	{
-		return false;
-	}
-
 	if (GravDir.IsZero())
 	{
 		return false;
@@ -5904,7 +5899,7 @@ bool UCharacterMovementComponent::StepUp(const FVector& GravDir, const FVector&
 	if (IsMovingOnGround() && CurrentFloor.IsWalkableFloor())
 	{
 		// Since we float a variable amount off the floor, we need to enforce max step height off the actual point of impact with the floor.
-		const float FloorDist = FMath::Max(0.f, CurrentFloor.FloorDist);
+		const float FloorDist = FMath::Max(0.f, CurrentFloor.GetDistanceToFloor());
 		PawnInitialFloorBaseZ -= FloorDist;
 		StepTravelUpHeight = FMath::Max(StepTravelUpHeight - FloorDist, 0.f);
 		StepTravelDownHeight = (MaxStepHeight + MAX_FLOOR_DIST*2.f);
@@ -5921,6 +5916,12 @@ bool UCharacterMovementComponent::StepUp(const FVector& GravDir, const FVector&
 		}
 	}
 
+	// Don't step up if the impact is below us, accounting for distance from floor.
+	if (InitialImpactZ <= PawnInitialFloorBaseZ)
+	{
+		return false;
+	}
+
 	// Scope our movement updates, and do not apply them until all intermediate moves are completed.
 	FScopedMovementUpdate ScopedStepUpMovement(UpdatedComponent, EScopedUpdate::DeferredUpdates);
 
@@ -7331,9 +7332,10 @@ void UCharacterMovementComponent::ReplicateMoveToServer(float DeltaTime, const F
 			// Decide whether to hold off on move
 			// send moves more frequently in small games where server isn't likely to be saturated
 			float NetMoveDelta;
-			UPlayer* Player = (PC ? PC->Player : NULL);
+			UPlayer* Player = (PC ? PC->Player : nullptr);
+			AGameStateBase const* const GameState = GetWorld()->GetGameState();
 
-			if (Player && (Player->CurrentNetSpeed > 10000) && (GetWorld()->GameState != NULL) && (GetWorld()->GameState->PlayerArray.Num() <= 10))
+			if (Player && (Player->CurrentNetSpeed > 10000) && (GameState != nullptr) && (GameState->PlayerArray.Num() <= 10))
 			{
 				NetMoveDelta = 0.011f;
 			}
diff --git a/Engine/Source/Runtime/Engine/Private/Components/ChildActorComponent.cpp b/Engine/Source/Runtime/Engine/Private/Components/ChildActorComponent.cpp
--- a/Engine/Source/Runtime/Engine/Private/Components/ChildActorComponent.cpp
+++ b/Engine/Source/Runtime/Engine/Private/Components/ChildActorComponent.cpp
@@ -9,7 +9,6 @@ DEFINE_LOG_CATEGORY_STATIC(LogChildActorComponent, Warning, All);
 UChildActorComponent::UChildActorComponent(const FObjectInitializer& ObjectInitializer)
 	: Super(ObjectInitializer)
 {
-	bWantsBeginPlay = true;
 	bAllowReregistration = false;
 }
 
@@ -48,15 +47,65 @@ void UChildActorComponent::OnRegister()
 	}
 }
 
+void UChildActorComponent::Serialize(FArchive& Ar)
+{
+	Super::Serialize(Ar);
+
+	if (Ar.HasAllPortFlags(PPF_DuplicateForPIE))
+	{
+		// PIE duplication should just work normally
+		Ar << ChildActorTemplate;
+	}
+	else if (Ar.HasAllPortFlags(PPF_Duplicate))
+	{
+		if (GIsEditor && Ar.IsLoading() && !IsTemplate())
+		{
+			// If we're not a template then we do not want the duplicate so serialize manually and destroy the template that was created for us
+			Ar.Serialize(&ChildActorTemplate, sizeof(UObject*));
+
+			if (AActor* UnwantedDuplicate = static_cast<AActor*>(FindObjectWithOuter(this, AActor::StaticClass())))
+			{
+				UnwantedDuplicate->MarkPendingKill();
+			}
+		}
+		else if (!GIsEditor && !Ar.IsLoading() && !GIsDuplicatingClassForReinstancing)
+		{
+			// Avoid the archiver in the duplicate writer case because we want to avoid the duplicate being created
+			Ar.Serialize(&ChildActorTemplate, sizeof(UObject*));
+		}
+		else
+		{
+			// When we're loading outside of the editor we won't have created the duplicate, so its fine to just use the normal path
+			// When we're loading a template then we want the duplicate, so it is fine to use normal archiver
+			// When we're saving in the editor we'll create the duplicate, but on loading decide whether to take it or not
+			Ar << ChildActorTemplate;
+		}
+	}
+
+#if WITH_EDITOR
+	// Since we sometimes serialize properties in instead of using duplication if we are a template
+	// and are not pointing at a component we own we'll need to fix that
+	if (ChildActorTemplate && ChildActorTemplate->GetOuter() != this && IsTemplate())
+	{
+		const FString TemplateName = FString::Printf(TEXT("%s_%s_CAT"), *GetName(), *ChildActorClass->GetName());
+		ChildActorTemplate = CastChecked<AActor>(StaticDuplicateObject(ChildActorTemplate, this, *TemplateName));
+	}
+#endif
+}
+
 #if WITH_EDITOR
 void UChildActorComponent::PostEditChangeProperty(FPropertyChangedEvent& PropertyChangedEvent)
 {
-	static const FName NAME_ChildActorClass = GET_MEMBER_NAME_CHECKED(UChildActorComponent, ChildActorClass);
-
-	if (PropertyChangedEvent.Property && PropertyChangedEvent.Property->GetFName() == NAME_ChildActorClass)
+	if (PropertyChangedEvent.Property && PropertyChangedEvent.Property->GetFName() == GET_MEMBER_NAME_CHECKED(UChildActorComponent, ChildActorClass))
 	{
 		ChildActorName = NAME_None;
 
+		if (!IsTemplate())
+		{
+			UChildActorComponent* Archetype = CastChecked<UChildActorComponent>(GetArchetype());
+			ChildActorTemplate = (Archetype->ChildActorClass == ChildActorClass ? Archetype->ChildActorTemplate : nullptr);
+		}
+
 		// If this was created by construction script, the post edit change super call will destroy it anyways
 		if (!IsCreatedByConstructionScript())
 		{
@@ -68,6 +117,55 @@ void UChildActorComponent::PostEditChangeProperty(FPropertyChangedEvent& Propert
 	Super::PostEditChangeProperty(PropertyChangedEvent);
 }
 
+void UChildActorComponent::PostEditChangeChainProperty(FPropertyChangedChainEvent& PropertyChangedEvent)
+{
+	if (PropertyChangedEvent.Property && PropertyChangedEvent.Property->GetFName() == GET_MEMBER_NAME_CHECKED(UChildActorComponent, ChildActorClass))
+	{
+		if (IsTemplate())
+		{
+			if (ChildActorClass)
+			{
+				if (ChildActorTemplate == nullptr || (ChildActorTemplate->GetClass() != ChildActorClass))
+				{
+					Modify();
+
+					AActor* NewChildActorTemplate = NewObject<AActor>(GetTransientPackage(), ChildActorClass, NAME_None, RF_ArchetypeObject | RF_Transactional | RF_Public);
+
+					if (ChildActorTemplate)
+					{
+						UEngine::CopyPropertiesForUnrelatedObjects(ChildActorTemplate, NewChildActorTemplate);
+
+						ChildActorTemplate->Rename(nullptr, GetTransientPackage(), REN_DontCreateRedirectors);
+					}
+
+					ChildActorTemplate = NewChildActorTemplate;
+
+					// Record initial object state in case we're in a transaction context.
+					ChildActorTemplate->Modify();
+
+					// Now set the actual name and outer to the BPGC.
+					const FString TemplateName = FString::Printf(TEXT("%s_%s_CAT"), *GetName(), *ChildActorClass->GetName());
+
+					ChildActorTemplate->Rename(*TemplateName, this, REN_DoNotDirty | REN_DontCreateRedirectors | REN_ForceNoResetLoaders);
+				}
+			}
+			else if (ChildActorTemplate)
+			{
+				Modify();
+
+				ChildActorTemplate->Rename(nullptr, GetTransientPackage(), REN_DontCreateRedirectors);
+				ChildActorTemplate = nullptr;
+			}
+		}
+		else
+		{
+			ChildActorTemplate = CastChecked<UChildActorComponent>(GetArchetype())->ChildActorTemplate;
+		}
+	}
+
+	Super::PostEditChangeChainProperty(PropertyChangedEvent);
+}
+
 void UChildActorComponent::PostEditUndo()
 {
 	Super::PostEditUndo();
@@ -306,6 +404,13 @@ void UChildActorComponent::PostLoad()
 	{
 		FActorParentComponentSetter::Set(ChildActor, this);
 	}
+
+	// Since the template could have been changed we need to respawn the child actor
+	// Don't do this if there is no linker which implies component was created via duplication
+	if (ChildActorTemplate && GetLinker())
+	{
+		DestroyChildActor();
+	}
 }
 #endif
 
@@ -353,6 +458,7 @@ void UChildActorComponent::CreateChildActor()
 				Params.bAllowDuringConstructionScript = true;
 				Params.OverrideLevel = (MyOwner ? MyOwner->GetLevel() : nullptr);
 				Params.Name = ChildActorName;
+				Params.Template = ChildActorTemplate;
 				if (!HasAllFlags(RF_Transactional))
 				{
 					Params.ObjectFlags &= ~RF_Transactional;
diff --git a/Engine/Source/Runtime/Engine/Private/Components/ModelComponent.cpp b/Engine/Source/Runtime/Engine/Private/Components/ModelComponent.cpp
--- a/Engine/Source/Runtime/Engine/Private/Components/ModelComponent.cpp
+++ b/Engine/Source/Runtime/Engine/Private/Components/ModelComponent.cpp
@@ -542,11 +542,26 @@ bool UModelComponent::GetPhysicsTriMeshData(struct FTriMeshCollisionData* Collis
 	int32 nBadArea = 0;
 	const float AreaThreshold = UPhysicsSettings::Get()->TriangleMeshTriangleMinAreaThreshold;
 
+	// See if we should copy UVs
+	bool bCopyUVs = UPhysicsSettings::Get()->bSupportUVFromHitResults;
+	if (bCopyUVs)
+	{
+		CollisionData->UVs.AddZeroed(1); // only one UV channel
+	}
+
 	const int32 NumVerts = Model->VertexBuffer.Vertices.Num();
 	CollisionData->Vertices.AddUninitialized(NumVerts);
+	if (bCopyUVs)
+	{
+		CollisionData->UVs[0].AddUninitialized(NumVerts);
+	}
 	for(int32 i=0; i<NumVerts; i++)
 	{
 		CollisionData->Vertices[i] = Model->VertexBuffer.Vertices[i].Position;
+		if (bCopyUVs)
+		{
+			CollisionData->UVs[0][i] = Model->VertexBuffer.Vertices[i].TexCoord;
+		}
 	}
 
 	for(int32 ElementIndex = 0;ElementIndex < Elements.Num();ElementIndex++)
diff --git a/Engine/Source/Runtime/Engine/Private/Components/MovementComponent.cpp b/Engine/Source/Runtime/Engine/Private/Components/MovementComponent.cpp
--- a/Engine/Source/Runtime/Engine/Private/Components/MovementComponent.cpp
+++ b/Engine/Source/Runtime/Engine/Private/Components/MovementComponent.cpp
@@ -89,6 +89,19 @@ void UMovementComponent::InitializeComponent()
 	TGuardValue<bool> InInitializeComponentGuard(bInInitializeComponent, true);
 	Super::InitializeComponent();
 
+	// RootComponent is null in OnRegister for blueprint (non-native) root components.
+	if (!UpdatedComponent && bAutoRegisterUpdatedComponent)
+	{
+		// Auto-register owner's root component if found.
+		if (AActor* MyActor = GetOwner())
+		{
+			if (USceneComponent* NewUpdatedComponent = MyActor->GetRootComponent())
+			{
+				SetUpdatedComponent(NewUpdatedComponent);
+			}
+		}
+	}
+
 	if (bSnapToPlaneAtStart)
 	{
 		SnapUpdatedComponentToPlane();
diff --git a/Engine/Source/Runtime/Engine/Private/Components/PrimitiveComponent.cpp b/Engine/Source/Runtime/Engine/Private/Components/PrimitiveComponent.cpp
--- a/Engine/Source/Runtime/Engine/Private/Components/PrimitiveComponent.cpp
+++ b/Engine/Source/Runtime/Engine/Private/Components/PrimitiveComponent.cpp
@@ -236,7 +236,7 @@ void UPrimitiveComponent::InvalidateLightingCacheDetailed(bool bInvalidateBuildE
 
 bool UPrimitiveComponent::IsEditorOnly() const
 {
-	return (AlwaysLoadOnClient == false) && (AlwaysLoadOnServer == false);
+	return ((AlwaysLoadOnClient == false) && (AlwaysLoadOnServer == false)) || Super::IsEditorOnly();
 }
 
 bool UPrimitiveComponent::HasStaticLighting() const
diff --git a/Engine/Source/Runtime/Engine/Private/Components/SceneCaptureComponent.cpp b/Engine/Source/Runtime/Engine/Private/Components/SceneCaptureComponent.cpp
--- a/Engine/Source/Runtime/Engine/Private/Components/SceneCaptureComponent.cpp
+++ b/Engine/Source/Runtime/Engine/Private/Components/SceneCaptureComponent.cpp
@@ -39,8 +39,7 @@ ASceneCapture2D::ASceneCapture2D(const FObjectInitializer& ObjectInitializer)
 	: Super(ObjectInitializer)
 {
 	DrawFrustum = CreateDefaultSubobject<UDrawFrustumComponent>(TEXT("DrawFrust0"));
-	DrawFrustum->AlwaysLoadOnClient = false;
-	DrawFrustum->AlwaysLoadOnServer = false;
+	DrawFrustum->bIsEditorOnly = true;
 	DrawFrustum->SetupAttachment(GetMeshComp());
 
 	CaptureComponent2D = CreateDefaultSubobject<USceneCaptureComponent2D>(TEXT("NewSceneCaptureComponent2D"));
@@ -95,8 +94,7 @@ ASceneCaptureCube::ASceneCaptureCube(const FObjectInitializer& ObjectInitializer
 	: Super(ObjectInitializer)
 {
 	DrawFrustum = CreateDefaultSubobject<UDrawFrustumComponent>(TEXT("DrawFrust0"));
-	DrawFrustum->AlwaysLoadOnClient = false;
-	DrawFrustum->AlwaysLoadOnServer = false;
+	DrawFrustum->bIsEditorOnly = true;
 	DrawFrustum->SetupAttachment(GetMeshComp());
 
 	CaptureComponentCube = CreateDefaultSubobject<USceneCaptureComponentCube>(TEXT("NewSceneCaptureComponentCube"));
diff --git a/Engine/Source/Runtime/Engine/Private/Components/SceneComponent.cpp b/Engine/Source/Runtime/Engine/Private/Components/SceneComponent.cpp
--- a/Engine/Source/Runtime/Engine/Private/Components/SceneComponent.cpp
+++ b/Engine/Source/Runtime/Engine/Private/Components/SceneComponent.cpp
@@ -94,7 +94,7 @@ DECLARE_DELEGATE_RetVal_OneParam(bool, FMobilityQueryDelegate, EComponentMobilit
  * @param  NewMobilityType			The Mobility type you want to switch sub-components over to
  * @param  ShouldOverrideMobility	A delegate used to determine if a sub-component's Mobility should be overridden
  *									(if left unset it will default to the AreMobilitiesDifferent() function)
- * @return The number of decedents that had their mobility altered.
+ * @return The number of descendants that had their mobility altered.
  */
 static int32 SetDescendantMobility(USceneComponent const* SceneComponentObject, EComponentMobility::Type NewMobilityType, FMobilityQueryDelegate ShouldOverrideMobility = FMobilityQueryDelegate())
 {
@@ -125,13 +125,15 @@ static int32 SetDescendantMobility(USceneComponent const* SceneComponentObject,
 	}
 
 	int32 NumDescendantsChanged = 0;
-	// recursively alter the mobility for children and deeper decedents 
+	// recursively alter the mobility for children and deeper descendants 
 	for (USceneComponent* ChildSceneComponent : AttachedChildren)
 	{
 		if (ChildSceneComponent)
 		{
 			if (ShouldOverrideMobility.Execute(ChildSceneComponent->Mobility))
 			{
+				ChildSceneComponent->Modify();
+
 				// USceneComponents shouldn't be set Stationary 
 				if ((NewMobilityType == EComponentMobility::Stationary) && ChildSceneComponent->IsA(UStaticMeshComponent::StaticClass()))
 				{
@@ -182,20 +184,29 @@ static int32 SetAncestorMobility(USceneComponent const* SceneComponentObject, EC
 			switch(NewMobilityType)
 			{
 			case EComponentMobility::Stationary:
+
 				if (UStaticMeshComponent* StaticMeshParent = Cast<UStaticMeshComponent>(AttachedParent))
 				{
+					StaticMeshParent->Modify();
+
 					// make it Static (because it is acceptable for Stationary children to have Static parents)
 					StaticMeshParent->Mobility = EComponentMobility::Static;
 					StaticMeshParent->SetSimulatePhysics(false);
 				}
 				break;
+
 			case EComponentMobility::Static:
+
+				AttachedParent->Modify();
 				if (UPrimitiveComponent* PrimitiveComponentParent = Cast<UPrimitiveComponent>(AttachedParent))
 				{
 					PrimitiveComponentParent->SetSimulatePhysics(false);
 				}
-				//FALLTHROUGH: we still want to set mobility
+				AttachedParent->Mobility = NewMobilityType;
+				break;
+
 			default:
+				AttachedParent->Modify();
 				AttachedParent->Mobility = NewMobilityType;
 			}
 
@@ -208,20 +219,20 @@ static int32 SetAncestorMobility(USceneComponent const* SceneComponentObject, EC
 }
 
 /**
-* When a scene component's Mobility is altered, we need to make sure the scene hierarchy is
-* updated. Parents can't be more mobile than their children. This means that certain
-* mobility hierarchy structures are disallowed, like:
-*
-*   Movable
-*   |-Stationary   <-- NOT allowed
-*   Movable
-*   |-Static       <-- NOT allowed
-*   Stationary
-*   |-Static       <-- NOT allowed
-*
-* This method walks the hierarchy and alters parent/child component's Mobility as a result of
-* this property change.
-*/
+ * When a scene component's Mobility is altered, we need to make sure the scene hierarchy is
+ * updated. Parents can't be more mobile than their children. This means that certain
+ * mobility hierarchy structures are disallowed, like:
+ *
+ *    Movable
+ *   |-Stationary   <-- NOT allowed
+ *   Movable
+ *   |-Static       <-- NOT allowed
+ *   Stationary
+ *   |-Static       <-- NOT allowed
+ *
+ * This method walks the hierarchy and alters parent/child component's Mobility as a result of
+ * this property change.
+ */
 static void UpdateAttachedMobility(USceneComponent* ComponentThatChanged)
 {
 	// Attached parent components can't be more mobile than their children. This means that 
@@ -253,7 +264,7 @@ static void UpdateAttachedMobility(USceneComponent* ComponentThatChanged)
 		// a delegate for checking if components are Movable
 		FMobilityQueryDelegate IsMovableDelegate = FMobilityQueryDelegate::CreateRaw(&EquivalenceFunctor, &FMobilityEqualityFunctor::operator(), EComponentMobility::Movable);
 
-		// if any decedents are static, change them to stationary (or movable for static meshes)
+		// if any descendants are static, change them to stationary (or movable for static meshes)
 		NumMobilityChanges += SetDescendantMobility(ComponentThatChanged, EComponentMobility::Stationary, IsStaticDelegate);
 
 		// if any ancestors are movable, change them to stationary (or static for static meshes)
@@ -301,6 +312,81 @@ static void UpdateAttachedMobility(USceneComponent* ComponentThatChanged)
 	}
 }
 
+/**
+ * A static helper function that recursively alters the bIsEditorOnly property for all
+ * sub-components (descending from the specified USceneComponent)
+ *
+ * @param  SceneComponentObject		The component whose sub-components you want to alter
+ *
+ * @return The number of descendants that had their bIsEditorOnly flag altered.
+ */
+static int32 SetDescendantIsEditorOnly(USceneComponent const* SceneComponentObject)
+{
+	if (!ensure(SceneComponentObject != nullptr))
+	{
+		return 0;
+	}
+
+	TArray<USceneComponent*> AttachedChildren = SceneComponentObject->GetAttachChildren();
+
+	// Do we need the templates too?
+	// gather children for component templates
+	USCS_Node* SCSNode = ComponentUtils::FindCorrespondingSCSNode(SceneComponentObject);
+	if (SCSNode != nullptr)
+	{
+		// gather children from the SCSNode
+		for (USCS_Node* SCSChild : SCSNode->GetChildNodes())
+		{
+			USceneComponent* ChildSceneComponent = Cast<USceneComponent>(SCSChild->ComponentTemplate);
+			if (ChildSceneComponent != nullptr)
+			{
+				AttachedChildren.Add(ChildSceneComponent);
+			}
+		}
+	}
+
+	int32 NumDescendantsChanged = 0;
+	// recursively alter the bIsEditorOnly flag for children and deeper descendants 
+	for (USceneComponent* ChildSceneComponent : AttachedChildren)
+	{
+		if (!ChildSceneComponent->bIsEditorOnly)
+		{
+			ChildSceneComponent->Modify();
+			ChildSceneComponent->bIsEditorOnly = true;
+			++NumDescendantsChanged;
+		}
+		NumDescendantsChanged += SetDescendantIsEditorOnly(ChildSceneComponent);
+	}
+
+	return NumDescendantsChanged;
+}
+
+/**
+ * When a scene component's bIsEditorOnly behavior is altered, we need to make sure the scene hierarchy is
+ * updated. All children of an editor only component must be editor only too.
+ *
+ * This method walks the hierarchy and alters parent/child component's bIsEditorOnly flag as a result of
+ * this property change.
+ */
+static void UpdateAttachedIsEditorOnly(USceneComponent* ComponentThatChanged)
+{
+	const int32 NumComponentsChanged = SetDescendantIsEditorOnly(ComponentThatChanged);
+
+	// if we altered any components (other than the ones selected), then notify the user
+	if (NumComponentsChanged > 0 && !ComponentThatChanged->GetName().EndsWith(UActorComponent::ComponentTemplateNameSuffix))
+	{
+		FText NotificationText = LOCTEXT("IsEditorOnlyAlteredSingularNotification", "Caused 1 component to also change its IsEditorOnly behaviour");
+		if (NumComponentsChanged > 1)
+		{
+			NotificationText = FText::Format(LOCTEXT("IsEditorOnlyAlteredPluralNotification", "Caused {0} other components to also change their IsEditorOnly behaviour"), FText::AsNumber(NumComponentsChanged));
+		}
+		FNotificationInfo Info(NotificationText);
+		Info.bFireAndForget = true;
+		Info.bUseThrobber = true;
+		FSlateNotificationManager::Get().AddNotification(Info);
+	}
+}
+
 void USceneComponent::PostEditChangeProperty(FPropertyChangedEvent& PropertyChangedEvent)
 {
 	const static FName LocationName("RelativeLocation");
@@ -314,6 +400,10 @@ void USceneComponent::PostEditChangeProperty(FPropertyChangedEvent& PropertyChan
 	{
 		UpdateAttachedMobility(this);
 	}
+	if (bIsEditorOnly && PropertyName == GET_MEMBER_NAME_CHECKED(UActorComponent, bIsEditorOnly))
+	{
+		UpdateAttachedIsEditorOnly(this);
+	}
 
 	Super::PostEditChangeProperty(PropertyChangedEvent);
 
@@ -330,6 +420,10 @@ void USceneComponent::PostEditChangeChainProperty(FPropertyChangedChainEvent& Pr
 	{
 		UpdateAttachedMobility(this);
 	}
+	if (bIsEditorOnly && PropertyChangedEvent.Property && PropertyChangedEvent.Property->GetFName() == GET_MEMBER_NAME_CHECKED(UActorComponent, bIsEditorOnly))
+	{
+		UpdateAttachedIsEditorOnly(this);
+	}
 
 	Super::PostEditChangeChainProperty(PropertyChangedEvent);
 }
@@ -461,8 +555,7 @@ void USceneComponent::OnRegister()
 		SpriteComponent->Sprite = LoadObject<UTexture2D>(nullptr, TEXT("/Engine/EditorResources/EmptyActor.EmptyActor"));
 		SpriteComponent->RelativeScale3D = FVector(0.5f, 0.5f, 0.5f);
 		SpriteComponent->Mobility = EComponentMobility::Movable;
-		SpriteComponent->AlwaysLoadOnClient = false;
-		SpriteComponent->AlwaysLoadOnServer = false;
+		SpriteComponent->bIsEditorOnly = true;
 		SpriteComponent->SpriteInfo.Category = TEXT("Misc");
 		SpriteComponent->SpriteInfo.DisplayName = NSLOCTEXT( "SpriteCategory", "Misc", "Misc" );
 		SpriteComponent->CreationMethod = CreationMethod;
@@ -1629,11 +1722,6 @@ bool USceneComponent::AttachToComponent(USceneComponent* Parent, const FAttachme
 
 						return false;
 					}
-					else
-					{
-						//A simulated object needs to be detached at runtime. We are in the construction script so we can't do it here. However, we want to make sure it is done in BeginPlay.
-						bWantsBeginPlay = true;
-					}
 				}
 			}
 		}
@@ -2832,6 +2920,20 @@ bool USceneComponent::CanEditChange( const UProperty* Property ) const
 				bIsEditable = !Owner->bLockLocation;
 			}
 		}
+
+		if (Property->GetFName() == GET_MEMBER_NAME_CHECKED(UActorComponent, bIsEditorOnly))
+		{
+			const USceneComponent* SceneComponentObject = this;
+			while (USceneComponent* AttachedParent = ComponentUtils::GetAttachedParent(SceneComponentObject))
+			{
+				if (AttachedParent->IsEditorOnly())
+				{
+					bIsEditable = false;
+					break;
+				}
+				SceneComponentObject = AttachedParent;
+			}
+		}
 	}
 
 	return bIsEditable;
diff --git a/Engine/Source/Runtime/Engine/Private/Components/SkeletalMeshComponent.cpp b/Engine/Source/Runtime/Engine/Private/Components/SkeletalMeshComponent.cpp
--- a/Engine/Source/Runtime/Engine/Private/Components/SkeletalMeshComponent.cpp
+++ b/Engine/Source/Runtime/Engine/Private/Components/SkeletalMeshComponent.cpp
@@ -36,6 +36,7 @@
 #endif
 
 #include "MessageLog.h"
+#include "Animation/AnimNode_SubInput.h"
 
 
 #define LOCTEXT_NAMESPACE "SkeletalMeshComponent"
@@ -98,6 +99,11 @@ public:
 				FPlatformProcess::Sleep(Stall / 1000.0f);
 			}
 #endif
+			if (CurrentThread != ENamedThreads::GameThread)
+			{
+				GInitRunaway();
+			}
+
 			Comp->ParallelAnimationEvaluation();
 		}
 	}
@@ -201,6 +207,8 @@ USkeletalMeshComponent::USkeletalMeshComponent(const FObjectInitializer& ObjectI
 
 	bTickInEditor = true;
 
+	CachedAnimCurveUidVersion = 0;
+
 	ResetRootBodyIndex();
 }
 
@@ -303,7 +311,25 @@ bool USkeletalMeshComponent::NeedToSpawnAnimScriptInstance() const
 		(SkeletalMesh != nullptr) && (SkeletalMesh->Skeleton->IsCompatible(AnimSkeleton)
 		&& AnimSkeleton->IsCompatibleMesh(SkeletalMesh)))
 	{
-		if ( (AnimScriptInstance == NULL) || (AnimScriptInstance->GetClass() != AnimClass) )
+		if ( (AnimScriptInstance == nullptr) || (AnimScriptInstance->GetClass() != AnimClass) )
+		{
+			return true;
+		}
+	}
+
+	return false;
+}
+
+bool USkeletalMeshComponent::NeedToSpawnPostPhysicsInstance() const
+{
+	if(SkeletalMesh)
+	{
+		const UClass* MainInstanceClass = *AnimClass;
+		const UClass* ClassToUse = *SkeletalMesh->PostProcessAnimBlueprint;
+		const UClass* CurrentClass = PostProcessAnimInstance ? PostProcessAnimInstance->GetClass() : nullptr;
+
+		// We need to have an instance, and we have the wrong class (different or null)
+		if(ClassToUse && ClassToUse != CurrentClass && MainInstanceClass != ClassToUse)
 		{
 			return true;
 		}
@@ -366,6 +392,9 @@ void USkeletalMeshComponent::InitAnim(bool bForceReinit)
 	// I'm moving the check here
 	if ( SkeletalMesh != nullptr && IsRegistered() )
 	{
+		//clear cache UID since we don't know if skeleton changed
+		CachedAnimCurveUidVersion = 0;
+
 		// we still need this in case users doesn't call tick, but sent to renderer
 		MorphTargetWeights.SetNumZeroed(SkeletalMesh->MorphTargets.Num());
 
@@ -376,15 +405,16 @@ void USkeletalMeshComponent::InitAnim(bool bForceReinit)
 		const bool bPerformPostAnimEvaluation = false; // Skip post evaluation, it would be wasted work
 		HandleExistingParallelEvaluationTask(bBlockOnTask, bPerformPostAnimEvaluation);
 
-		bool bBlueprintMismatch = (AnimClass != NULL) &&
-			(AnimScriptInstance != NULL) && (AnimScriptInstance->GetClass() != AnimClass);
+		bool bBlueprintMismatch = (AnimClass != nullptr) &&
+			(AnimScriptInstance != nullptr) && (AnimScriptInstance->GetClass() != AnimClass);
 
 		const USkeleton* AnimSkeleton = (AnimScriptInstance)? AnimScriptInstance->CurrentSkeleton : nullptr;
 
+		bool bClearAnimInstance = AnimScriptInstance && !AnimSkeleton;
 		bool bSkeletonMismatch = AnimSkeleton && (AnimScriptInstance->CurrentSkeleton!=SkeletalMesh->Skeleton);
 		bool bSkeletonNotCompatible = AnimSkeleton && (AnimSkeleton->IsCompatibleMesh(SkeletalMesh) == false);
 
-		if (bBlueprintMismatch || bSkeletonMismatch || bSkeletonNotCompatible)
+		if (bBlueprintMismatch || bSkeletonMismatch || bSkeletonNotCompatible || bClearAnimInstance)
 		{
 			ClearAnimScriptInstance();
 		}
@@ -421,9 +451,13 @@ void USkeletalMeshComponent::InitAnim(bool bForceReinit)
 
 bool USkeletalMeshComponent::InitializeAnimScriptInstance(bool bForceReinit)
 {
-	bool bCalledInitialize = false;
+	bool bInitializedMainInstance = false;
+	bool bInitializedPostInstance = false;
+
 	if (IsRegistered())
 	{
+		check(SkeletalMesh);
+
 		if (NeedToSpawnAnimScriptInstance())
 		{
 			SCOPE_CYCLE_COUNTER(STAT_AnimSpawnTime);
@@ -435,51 +469,85 @@ bool USkeletalMeshComponent::InitializeAnimScriptInstance(bool bForceReinit)
 				SubInstances.Empty();
 
 				AnimScriptInstance->InitializeAnimation();
-				bCalledInitialize = true;
+				bInitializedMainInstance = true;
 			}
 		}
-		else if (AnimationMode == EAnimationMode::AnimationSingleNode)
+		else 
 		{
-			SCOPE_CYCLE_COUNTER(STAT_AnimSpawnTime);
-
-			UAnimSingleNodeInstance* OldInstance = nullptr;
-			if (!bForceReinit)
+			bool bShouldSpawnSingleNodeInstance = SkeletalMesh && SkeletalMesh->Skeleton && AnimationMode == EAnimationMode::AnimationSingleNode;
+			if (bShouldSpawnSingleNodeInstance)
 			{
-				OldInstance = Cast<UAnimSingleNodeInstance>(AnimScriptInstance);
-			}
+				SCOPE_CYCLE_COUNTER(STAT_AnimSpawnTime);
 
-			AnimScriptInstance = NewObject<UAnimSingleNodeInstance>(this);
+				UAnimSingleNodeInstance* OldInstance = nullptr;
+				if (!bForceReinit)
+				{
+					OldInstance = Cast<UAnimSingleNodeInstance>(AnimScriptInstance);
+				}
 
-			if (AnimScriptInstance)
-			{
-				AnimScriptInstance->InitializeAnimation();
-				bCalledInitialize = true;
-			}
+				AnimScriptInstance = NewObject<UAnimSingleNodeInstance>(this);
 
-			if (OldInstance && AnimScriptInstance)
-			{
-				// Copy data from old instance unless we force reinitialized
-				FSingleAnimationPlayData CachedData;
-				CachedData.PopulateFrom(OldInstance);
-				CachedData.Initialize(Cast<UAnimSingleNodeInstance>(AnimScriptInstance));
+				if (AnimScriptInstance)
+				{
+					AnimScriptInstance->InitializeAnimation();
+					bInitializedMainInstance = true;
+				}
+
+				if (OldInstance && AnimScriptInstance)
+				{
+					// Copy data from old instance unless we force reinitialized
+					FSingleAnimationPlayData CachedData;
+					CachedData.PopulateFrom(OldInstance);
+					CachedData.Initialize(Cast<UAnimSingleNodeInstance>(AnimScriptInstance));
+				}
+				else
+				{
+					// otherwise, initialize with AnimationData
+					AnimationData.Initialize(Cast<UAnimSingleNodeInstance>(AnimScriptInstance));
+				}
 			}
-			else
+		}
+
+		// May need to clear out the post physics instance
+		UClass* NewMeshInstanceClass = *SkeletalMesh->PostProcessAnimBlueprint;
+		if(!NewMeshInstanceClass || NewMeshInstanceClass == *AnimClass)
+		{
+			PostProcessAnimInstance = nullptr;
+		}
+
+		if(NeedToSpawnPostPhysicsInstance())
+		{
+			PostProcessAnimInstance = NewObject<UAnimInstance>(this, *SkeletalMesh->PostProcessAnimBlueprint);
+
+			if(PostProcessAnimInstance)
 			{
-				// otherwise, initialize with AnimationData
-				AnimationData.Initialize(Cast<UAnimSingleNodeInstance>(AnimScriptInstance));
+				PostProcessAnimInstance->InitializeAnimation();
+
+				if(FAnimNode_SubInput* InputNode = PostProcessAnimInstance->GetSubInputNode())
+				{
+					InputNode->InputPose.SetBoneContainer(&PostProcessAnimInstance->GetRequiredBones());
+				}
+
+				bInitializedPostInstance = true;
 			}
 		}
 
-		if (AnimScriptInstance && !bCalledInitialize && bForceReinit)
+		if (AnimScriptInstance && !bInitializedMainInstance && bForceReinit)
 		{
 			AnimScriptInstance->InitializeAnimation();
-			bCalledInitialize = true;
-		}		
+			bInitializedMainInstance = true;
+		}
+
+		if(PostProcessAnimInstance && !bInitializedPostInstance && bForceReinit)
+		{
+			PostProcessAnimInstance->InitializeAnimation();
+			bInitializedPostInstance = true;
+		}
 
 		// refresh morph targets - this can happen when re-registration happens
 		RefreshMorphTargets();
 	}
-	return bCalledInitialize;
+	return bInitializedMainInstance || bInitializedPostInstance;
 }
 
 bool USkeletalMeshComponent::IsWindEnabled() const
@@ -494,7 +562,7 @@ bool USkeletalMeshComponent::IsWindEnabled() const
 
 void USkeletalMeshComponent::ClearAnimScriptInstance()
 {
-	AnimScriptInstance = NULL;
+	AnimScriptInstance = nullptr;
 }
 
 void USkeletalMeshComponent::CreateRenderState_Concurrent()
@@ -519,14 +587,14 @@ void USkeletalMeshComponent::PostEditChangeProperty(FPropertyChangedEvent& Prope
 
 	UProperty* PropertyThatChanged = PropertyChangedEvent.Property;
 
-	if ( PropertyThatChanged != NULL )
+	if ( PropertyThatChanged != nullptr )
 	{
 		// if the blueprint has changed, recreate the AnimInstance
 		if ( PropertyThatChanged->GetFName() == GET_MEMBER_NAME_CHECKED( USkeletalMeshComponent, AnimationMode ) )
 		{
 			if (AnimationMode == EAnimationMode::AnimationBlueprint)
 			{
-				if (AnimClass == NULL)
+				if (AnimClass == nullptr)
 				{
 					ClearAnimScriptInstance();
 				}
@@ -551,6 +619,26 @@ void USkeletalMeshComponent::PostEditChangeProperty(FPropertyChangedEvent& Prope
 		{
 			ValidateAnimation();
 
+			// Check the post physics mesh instance, as the mesh has changed
+			if(PostProcessAnimInstance)
+			{
+				UClass* CurrentClass = PostProcessAnimInstance->GetClass();
+				UClass* MeshClass = *SkeletalMesh->PostProcessAnimBlueprint;
+				if(CurrentClass != MeshClass)
+				{
+					if(MeshClass)
+					{
+						PostProcessAnimInstance = NewObject<UAnimInstance>(this, *SkeletalMesh->PostProcessAnimBlueprint);
+						PostProcessAnimInstance->InitializeAnimation();
+					}
+					else
+					{
+						// No instance needed for the new mesh
+						PostProcessAnimInstance = nullptr;
+					}
+				}
+			}
+
 			if(OnSkeletalMeshPropertyChanged.IsBound())
 			{
 				OnSkeletalMeshPropertyChanged.Broadcast();
@@ -567,10 +655,10 @@ void USkeletalMeshComponent::PostEditChangeProperty(FPropertyChangedEvent& Prope
 		if ( PropertyThatChanged->GetFName() == GET_MEMBER_NAME_CHECKED( FSingleAnimationPlayData, AnimToPlay ))
 		{
 			// make sure the animation skeleton matches the current skeletalmesh
-			if (AnimationData.AnimToPlay != NULL && SkeletalMesh && AnimationData.AnimToPlay->GetSkeleton() != SkeletalMesh->Skeleton)
+			if (AnimationData.AnimToPlay != nullptr && SkeletalMesh && AnimationData.AnimToPlay->GetSkeleton() != SkeletalMesh->Skeleton)
 			{
 				UE_LOG(LogAnimation, Warning, TEXT("Invalid animation"));
-				AnimationData.AnimToPlay = NULL;
+				AnimationData.AnimToPlay = nullptr;
 			}
 			else
 			{
@@ -599,10 +687,10 @@ void USkeletalMeshComponent::LoadedFromAnotherClass(const FName& OldClassName)
 			SetAnimationMode(EAnimationMode::Type::AnimationSingleNode);
 
 			// support old compatibility code that changed variable name
-			if (SequenceToPlay_DEPRECATED!=NULL && AnimToPlay_DEPRECATED== NULL)
+			if (SequenceToPlay_DEPRECATED!=nullptr && AnimToPlay_DEPRECATED== nullptr)
 			{
 				AnimToPlay_DEPRECATED = SequenceToPlay_DEPRECATED;
-				SequenceToPlay_DEPRECATED = NULL;
+				SequenceToPlay_DEPRECATED = nullptr;
 			}
 
 			AnimationData.AnimToPlay = AnimToPlay_DEPRECATED;
@@ -622,9 +710,9 @@ void USkeletalMeshComponent::TickAnimation(float DeltaTime, bool bNeedsValidRoot
 
 	SCOPE_CYCLE_COUNTER(STAT_AnimGameThreadTime);
 	SCOPE_CYCLE_COUNTER(STAT_AnimTickTime);
-	if (SkeletalMesh != NULL)
+	if (SkeletalMesh != nullptr)
 	{
-		if (AnimScriptInstance != NULL)
+		if (AnimScriptInstance != nullptr)
 		{
 			// Tick the animation
 			AnimScriptInstance->UpdateAnimation(DeltaTime * GlobalAnimRateScale, bNeedsValidRootMotion);
@@ -634,6 +722,11 @@ void USkeletalMeshComponent::TickAnimation(float DeltaTime, bool bNeedsValidRoot
 		{
 			SubInstance->UpdateAnimation(DeltaTime * GlobalAnimRateScale, false);
 		}
+
+		if(PostProcessAnimInstance)
+		{
+			PostProcessAnimInstance->UpdateAnimation(DeltaTime * GlobalAnimRateScale, false);
+		}
 	}
 }
 
@@ -667,7 +760,7 @@ bool USkeletalMeshComponent::ShouldUpdateTransform(bool bLODHasChanged) const
 	const bool bSkipBecauseOfRefPose = bForceRefpose && bOldForceRefPose && (MorphTargetCurves.Num() == 0) && ((AnimScriptInstance) ? !AnimScriptInstance->HasMorphTargetCurves() : true);
 
 	// LOD changing should always trigger an update.
-	return (bLODHasChanged || (!bNoSkeletonUpdate && !bSkipBecauseOfRefPose && Super::ShouldUpdateTransform(bLODHasChanged)));
+	return (bLODHasChanged || !bRequiredBonesUpToDate || (!bNoSkeletonUpdate && !bSkipBecauseOfRefPose && Super::ShouldUpdateTransform(bLODHasChanged)));
 }
 
 bool USkeletalMeshComponent::ShouldTickPose() const
@@ -675,7 +768,7 @@ bool USkeletalMeshComponent::ShouldTickPose() const
 	// When we stop root motion we go back to ticking after CharacterMovement. Unfortunately that means that we could tick twice that frame.
 	// So only enforce a single tick per frame.
 	const bool bAlreadyTickedThisFrame = PoseTickedThisFrame();
-	return (Super::ShouldTickPose() && IsRegistered() && AnimScriptInstance && !bAutonomousTickPose && !bPauseAnims && GetWorld()->AreActorsInitialized() && !bNoSkeletonUpdate && !bAlreadyTickedThisFrame);
+	return (Super::ShouldTickPose() && IsRegistered() && (AnimScriptInstance || PostProcessAnimInstance) && !bAutonomousTickPose && !bPauseAnims && GetWorld()->AreActorsInitialized() && !bNoSkeletonUpdate && !bAlreadyTickedThisFrame);
 }
 
 static FThreadSafeCounter Ticked;
@@ -930,6 +1023,23 @@ static void MergeInBoneIndexArrays(TArray<FBoneIndexType>& BaseArray, const TArr
 	}
 }
 
+// this is optimized version of updating only curves
+// if you call RecalcRequiredBones, curve should be refreshed
+void USkeletalMeshComponent::RecalcRequiredCurves()
+{
+	if (!SkeletalMesh)
+	{
+		return;
+	}
+
+	// make sure animation requiredcurve to mark as dirty
+	if (AnimScriptInstance)
+	{
+		AnimScriptInstance->RecalcRequiredCurves();
+	}
+
+	MarkRequiredCurveUpToDate();
+}
 
 void USkeletalMeshComponent::RecalcRequiredBones(int32 LODIndex)
 {
@@ -956,7 +1066,7 @@ void USkeletalMeshComponent::RecalcRequiredBones(int32 LODIndex)
 	// The list of bones we want is taken from the predicted LOD level.
 	FStaticLODModel& LODModel = SkelMeshResource->LODModels[LODIndex];
 	RequiredBones = LODModel.RequiredBones;
-	
+
 	const UPhysicsAsset* const PhysicsAsset = GetPhysicsAsset();
 	// If we have a PhysicsAsset, we also need to make sure that all the bones used by it are always updated, as its used
 	// by line checks etc. We might also want to kick in the physics, which means having valid bone transforms.
@@ -1087,6 +1197,14 @@ void USkeletalMeshComponent::RecalcRequiredBones(int32 LODIndex)
 		AnimScriptInstance->RecalcRequiredBones();
 	}
 
+	if(PostProcessAnimInstance)
+	{
+		PostProcessAnimInstance->RecalcRequiredBones();
+	}
+
+	// when recalc requiredbones happend
+	// this should always happen
+	MarkRequiredCurveUpToDate();
 	bRequiredBonesUpToDate = true;
 
 	// Invalidate cached bones.
@@ -1094,6 +1212,19 @@ void USkeletalMeshComponent::RecalcRequiredBones(int32 LODIndex)
 	CachedComponentSpaceTransforms.Empty();
 }
 
+void USkeletalMeshComponent::MarkRequiredCurveUpToDate()
+{
+	if (SkeletalMesh && SkeletalMesh->Skeleton)
+	{
+		CachedAnimCurveUidVersion = SkeletalMesh->Skeleton->GetAnimCurveUidVersion();
+	}
+}
+
+bool USkeletalMeshComponent::AreRequiredCurvesUpToDate() const
+{
+	return (!SkeletalMesh || !SkeletalMesh->Skeleton || CachedAnimCurveUidVersion == SkeletalMesh->Skeleton->GetAnimCurveUidVersion());
+}
+
 void USkeletalMeshComponent::EvaluateAnimation(const USkeletalMesh* InSkeletalMesh, UAnimInstance* InAnimInstance, TArray<FTransform>& OutBoneSpaceTransforms, FVector& OutRootBoneTranslation, FBlendedHeapCurve& OutCurve) const
 {
 	ANIM_MT_SCOPE_CYCLE_COUNTER(SkeletalComponentAnimEvaluate, IsRunningParallelEvaluation());
@@ -1167,15 +1298,42 @@ void USkeletalMeshComponent::PerformAnimationEvaluation(const USkeletalMesh* InS
 	if(AnimEvaluationContext.bDoUpdate)
 	{
 		InAnimInstance->ParallelUpdateAnimation();
+
+		if(PostProcessAnimInstance)
+		{
+			PostProcessAnimInstance->ParallelUpdateAnimation();
+		}
+	}
+	else if(!InAnimInstance && PostProcessAnimInstance)
+	{
+		// If we don't have an anim instance, we may still have a post physics instance
+		PostProcessAnimInstance->ParallelUpdateAnimation();
 	}
 
 	// evaluate pure animations, and fill up BoneSpaceTransforms
 	EvaluateAnimation(InSkeletalMesh, InAnimInstance, OutBoneSpaceTransforms, OutRootBoneTranslation, OutCurve);
+	EvaluatePostProcessMeshInstance(OutBoneSpaceTransforms, OutCurve, InSkeletalMesh, OutRootBoneTranslation);
+
 	// Fill SpaceBases from LocalAtoms
 	FillComponentSpaceTransforms(InSkeletalMesh, OutBoneSpaceTransforms, OutSpaceBases);
 }
 
 
+void USkeletalMeshComponent::EvaluatePostProcessMeshInstance(TArray<FTransform>& OutBoneSpaceTransforms, FBlendedHeapCurve& OutCurve, const USkeletalMesh* InSkeletalMesh, FVector& OutRootBoneTranslation) const
+{
+	if(PostProcessAnimInstance)
+	{
+		// Push the previous pose to any input nodes required
+		if(FAnimNode_SubInput* InputNode = PostProcessAnimInstance->GetSubInputNode())
+		{
+			InputNode->InputPose.CopyBonesFrom(OutBoneSpaceTransforms);
+			InputNode->InputCurve.CopyFrom(OutCurve);
+		}
+
+		EvaluateAnimation(InSkeletalMesh, PostProcessAnimInstance, OutBoneSpaceTransforms, OutRootBoneTranslation, OutCurve);
+	}
+}
+
 int32 GetCurveNumber(USkeleton* Skeleton)
 {
 	// get all curve list
@@ -1243,15 +1401,17 @@ void USkeletalMeshComponent::RefreshBoneTransforms(FActorComponentTickFunction*
 		return;
 	}
 
-	// Cache Animation curve mapping names UIds from Skeleton
-	UpdateCachedAnimCurveMappingNameUids();
-
 	// Recalculate the RequiredBones array, if necessary
 	if (!bRequiredBonesUpToDate)
 	{
 		QUICK_SCOPE_CYCLE_COUNTER(STAT_USkeletalMeshComponent_RefreshBoneTransforms_RecalcRequiredBones);
 		RecalcRequiredBones(PredictedLODLevel);
 	}
+	// if curves have to be refreshed
+	else if (!AreRequiredCurvesUpToDate())
+	{
+		RecalcRequiredCurves();
+	}
 
 	const bool bDoEvaluationRateOptimization = ShouldUseUpdateRateOptimizations() && AnimUpdateRateParams->DoEvaluationRateOptimizations();
 
@@ -1263,8 +1423,11 @@ void USkeletalMeshComponent::RefreshBoneTransforms(FActorComponentTickFunction*
 									 || (GetNumComponentSpaceTransforms() != CachedComponentSpaceTransforms.Num()));
 
 
+	TArray<SmartName::UID_Type> const* CurrentAnimCurveMappingNameUids = (AnimScriptInstance) ? &AnimScriptInstance->GetRequiredBones().GetAnimCurveNameUids() : nullptr;
+
 	const bool bInvalidCachedCurve = bDoEvaluationRateOptimization && 
-									((CachedCurve.Num() != GetCurveNumber(SkeletalMesh->Skeleton)) || (CachedCurve.UIDList != &CachedAnimCurveMappingNameUids));
+									CurrentAnimCurveMappingNameUids != nullptr &&
+									((CachedCurve.Num() != GetCurveNumber(SkeletalMesh->Skeleton)) || (CachedCurve.UIDList != CurrentAnimCurveMappingNameUids));
 
 	const bool bShouldDoEvaluation = !bDoEvaluationRateOptimization || bInvalidCachedBones || bInvalidCachedCurve || !AnimUpdateRateParams->ShouldSkipEvaluation();
 
@@ -1286,16 +1449,23 @@ void USkeletalMeshComponent::RefreshBoneTransforms(FActorComponentTickFunction*
 	AnimEvaluationContext.SkeletalMesh = SkeletalMesh;
 	AnimEvaluationContext.AnimInstance = AnimScriptInstance;
 
-	if ((AnimEvaluationContext.Curve.Num() != GetCurveNumber(SkeletalMesh->Skeleton)) || (AnimEvaluationContext.Curve.UIDList != &CachedAnimCurveMappingNameUids))
+	if (CurrentAnimCurveMappingNameUids && 
+		((AnimEvaluationContext.Curve.Num() != GetCurveNumber(SkeletalMesh->Skeleton)) || (AnimEvaluationContext.Curve.UIDList != CurrentAnimCurveMappingNameUids)))
 	{
-		AnimEvaluationContext.Curve.InitFrom(&CachedAnimCurveMappingNameUids);
+		AnimEvaluationContext.Curve.InitFrom(CurrentAnimCurveMappingNameUids);
 	}
+	else if(CurrentAnimCurveMappingNameUids == nullptr)
+	{
+		AnimEvaluationContext.Curve.Empty();
+		CachedCurve.Empty();
+	}
+
 	AnimEvaluationContext.bDoEvaluation = bShouldDoEvaluation;
 	AnimEvaluationContext.bDoUpdate = AnimScriptInstance && AnimScriptInstance->NeedsUpdate();
 	
-	AnimEvaluationContext.bDoInterpolation = bDoEvaluationRateOptimization && !bInvalidCachedBones && AnimUpdateRateParams->ShouldInterpolateSkippedFrames();
+	AnimEvaluationContext.bDoInterpolation = bDoEvaluationRateOptimization && !bInvalidCachedBones && AnimUpdateRateParams->ShouldInterpolateSkippedFrames() && CurrentAnimCurveMappingNameUids != nullptr;
 	AnimEvaluationContext.bDuplicateToCacheBones = bInvalidCachedBones || (bDoEvaluationRateOptimization && AnimEvaluationContext.bDoEvaluation && !AnimEvaluationContext.bDoInterpolation);
-	AnimEvaluationContext.bDuplicateToCacheCurve = bInvalidCachedCurve || (bDoEvaluationRateOptimization && AnimEvaluationContext.bDoEvaluation && !AnimEvaluationContext.bDoInterpolation);
+	AnimEvaluationContext.bDuplicateToCacheCurve = bInvalidCachedCurve || (bDoEvaluationRateOptimization && AnimEvaluationContext.bDoEvaluation && !AnimEvaluationContext.bDoInterpolation && CurrentAnimCurveMappingNameUids != nullptr);
 	if (!bDoEvaluationRateOptimization)
 	{
 		//If we aren't optimizing clear the cached local atoms
@@ -1314,6 +1484,11 @@ void USkeletalMeshComponent::RefreshBoneTransforms(FActorComponentTickFunction*
 		}
 	}
 
+	if(PostProcessAnimInstance)
+	{
+		PostProcessAnimInstance->PreEvaluateAnimation();
+	}
+
 	if (bDoParallelEvaluation)
 	{
 		QUICK_SCOPE_CYCLE_COUNTER(STAT_USkeletalMeshComponent_RefreshBoneTransforms_SetupParallel); 
@@ -1365,18 +1540,26 @@ void USkeletalMeshComponent::RefreshBoneTransforms(FActorComponentTickFunction*
 				TArray<FTransform>& LocalEditableSpaceBases = GetEditableComponentSpaceTransforms();
 				LocalEditableSpaceBases.Reset();
 				LocalEditableSpaceBases.Append(CachedComponentSpaceTransforms);
-				AnimEvaluationContext.Curve.CopyFrom(CachedCurve);
+				if (CachedCurve.IsValid())
+				{
+					AnimEvaluationContext.Curve.CopyFrom(CachedCurve);
+				}
 			}
 			if(AnimEvaluationContext.bDoUpdate)
 			{
 				AnimScriptInstance->ParallelUpdateAnimation();
+
+				if(PostProcessAnimInstance)
+				{
+					PostProcessAnimInstance->ParallelUpdateAnimation();
+				}
 			}
 		}
 
 		PostAnimEvaluation(AnimEvaluationContext);
 	}
 
-	if (TickFunction == NULL)
+	if (TickFunction == nullptr)
 	{
 		//Since we aren't doing this through the tick system, assume we want the buffer flipped now
 		FinalizeBoneTransform();
@@ -1451,6 +1634,11 @@ void USkeletalMeshComponent::PostAnimEvaluation(FAnimationEvaluationContext& Eva
 		}
 	}
 
+	if(PostProcessAnimInstance)
+	{
+		PostProcessAnimInstance->UpdateCurves(EvaluationContext.Curve);
+	}
+
 	bNeedToFlipSpaceBaseBuffers = true;
 
 	// update physics data from animated data
@@ -1646,9 +1834,12 @@ UAnimInstance* USkeletalMeshComponent::GetAnimInstance() const
 	return AnimScriptInstance;
 }
 
-void USkeletalMeshComponent::NotifySkelControlBeyondLimit( USkelControlLookAt* LookAt ) {}
-
+UAnimInstance* USkeletalMeshComponent::GetPostProcessInstance() const
+{
+	return PostProcessAnimInstance;
+}
 
+void USkeletalMeshComponent::NotifySkelControlBeyondLimit( USkelControlLookAt* LookAt ) {}
 
 void USkeletalMeshComponent::SkelMeshCompOnParticleSystemFinished( UParticleSystemComponent* PSC )
 {
@@ -1724,7 +1915,7 @@ bool USkeletalMeshComponent::IsAnySimulatingPhysics() const
  */
 void USkeletalMeshComponent::DebugDrawBones(UCanvas* Canvas, bool bSimpleBones) const
 {
-	if (GetWorld()->IsGameWorld() && SkeletalMesh && Canvas && MasterPoseComponent == NULL)
+	if (GetWorld()->IsGameWorld() && SkeletalMesh && Canvas && MasterPoseComponent == nullptr)
 	{
 		// draw spacebases, we could cache parent bones, but this is mostly debug feature, I'm not caching it right now
 		for ( int32 Index=0; Index<RequiredBones.Num(); ++Index )
@@ -1871,7 +2062,7 @@ SIZE_T USkeletalMeshComponent::GetResourceSize( EResourceSizeMode::Type Mode )
 
 	for (int32 i=0; i < Bodies.Num(); ++i)
 	{
-		if (Bodies[i] != NULL && Bodies[i]->IsValidBodyInstance())
+		if (Bodies[i] != nullptr && Bodies[i]->IsValidBodyInstance())
 		{
 			ResSize += Bodies[i]->GetBodyInstanceResourceSize(Mode);
 		}
@@ -1890,7 +2081,10 @@ void USkeletalMeshComponent::SetAnimationMode(EAnimationMode::Type InAnimationMo
 
 	// when mode is swapped, make sure to reinitialize
 	// even if it was same mode
-	InitializeAnimScriptInstance(true);
+	if(SkeletalMesh != nullptr)
+	{
+		InitializeAnimScriptInstance(true);
+	}
 }
 
 EAnimationMode::Type USkeletalMeshComponent::GetAnimationMode() const
@@ -1913,7 +2107,7 @@ void USkeletalMeshComponent::SetAnimation(UAnimationAsset* NewAnimToPlay)
 		SingleNodeInstance->SetAnimationAsset(NewAnimToPlay, false);
 		SingleNodeInstance->SetPlaying(false);
 	}
-	else if( AnimScriptInstance != NULL )
+	else if( AnimScriptInstance != nullptr )
 	{
 		UE_LOG(LogAnimation, Warning, TEXT("Currently in Animation Blueprint mode. Please change AnimationMode to Use Animation Asset"));
 	}
@@ -1927,7 +2121,7 @@ void USkeletalMeshComponent::Play(bool bLooping)
 		SingleNodeInstance->SetPlaying(true);
 		SingleNodeInstance->SetLooping(bLooping);
 	}
-	else if( AnimScriptInstance != NULL )
+	else if( AnimScriptInstance != nullptr )
 	{
 		UE_LOG(LogAnimation, Warning, TEXT("Currently in Animation Blueprint mode. Please change AnimationMode to Use Animation Asset"));
 	}
@@ -1940,7 +2134,7 @@ void USkeletalMeshComponent::Stop()
 	{
 		SingleNodeInstance->SetPlaying(false);
 	}
-	else if( AnimScriptInstance != NULL )
+	else if( AnimScriptInstance != nullptr )
 	{
 		UE_LOG(LogAnimation, Warning, TEXT("Currently in Animation Blueprint mode. Please change AnimationMode to Use Animation Asset"));
 	}
@@ -1953,7 +2147,7 @@ bool USkeletalMeshComponent::IsPlaying() const
 	{
 		return SingleNodeInstance->IsPlaying();
 	}
-	else if( AnimScriptInstance != NULL )
+	else if( AnimScriptInstance != nullptr )
 	{
 		UE_LOG(LogAnimation, Warning, TEXT("Currently in Animation Blueprint mode. Please change AnimationMode to Use Animation Asset"));
 	}
@@ -1968,7 +2162,7 @@ void USkeletalMeshComponent::SetPosition(float InPos, bool bFireNotifies)
 	{
 		SingleNodeInstance->SetPosition(InPos, bFireNotifies);
 	}
-	else if( AnimScriptInstance != NULL )
+	else if( AnimScriptInstance != nullptr )
 	{
 		UE_LOG(LogAnimation, Warning, TEXT("Currently in Animation Blueprint mode. Please change AnimationMode to Use Animation Asset"));
 	}
@@ -1981,7 +2175,7 @@ float USkeletalMeshComponent::GetPosition() const
 	{
 		return SingleNodeInstance->GetCurrentTime();
 	}
-	else if( AnimScriptInstance != NULL )
+	else if( AnimScriptInstance != nullptr )
 	{
 		UE_LOG(LogAnimation, Warning, TEXT("Currently in Animation Blueprint mode. Please change AnimationMode to Use Animation Asset"));
 	}
@@ -1996,7 +2190,7 @@ void USkeletalMeshComponent::SetPlayRate(float Rate)
 	{
 		SingleNodeInstance->SetPlayRate(Rate);
 	}
-	else if( AnimScriptInstance != NULL )
+	else if( AnimScriptInstance != nullptr )
 	{
 		UE_LOG(LogAnimation, Warning, TEXT("Currently in Animation Blueprint mode. Please change AnimationMode to Use Animation Asset"));
 	}
@@ -2009,7 +2203,7 @@ float USkeletalMeshComponent::GetPlayRate() const
 	{
 		return SingleNodeInstance->GetPlayRate();
 	}
-	else if( AnimScriptInstance != NULL )
+	else if( AnimScriptInstance != nullptr )
 	{
 		UE_LOG(LogAnimation, Warning, TEXT("Currently in Animation Blueprint mode. Please change AnimationMode to Use Animation Asset"));
 	}
@@ -2288,7 +2482,7 @@ void USkeletalMeshComponent::ValidateAnimation()
 
 bool USkeletalMeshComponent::IsPlayingRootMotion()
 {
-	return (AnimScriptInstance ? (AnimScriptInstance->GetRootMotionMontageInstance() != NULL) : false);
+	return (AnimScriptInstance ? (AnimScriptInstance->GetRootMotionMontageInstance() != nullptr) : false);
 }
 
 bool USkeletalMeshComponent::IsPlayingRootMotionFromEverything()
@@ -2514,19 +2708,6 @@ void USkeletalMeshComponent::FinalizeBoneTransform()
 	}
 }
 
-void USkeletalMeshComponent::UpdateCachedAnimCurveMappingNameUids()
-{
-	if (SkeletalMesh->Skeleton)
-	{
-		CachedAnimCurveMappingNameUids = SkeletalMesh->Skeleton->GetCachedAnimCurveMappingNameUids();
-	}
-}
-
-TArray<FSmartNameMapping::UID> const * USkeletalMeshComponent::GetCachedAnimCurveMappingNameUids()
-{
-	return &CachedAnimCurveMappingNameUids;
-}
-
 FDelegateHandle USkeletalMeshComponent::RegisterOnPhysicsCreatedDelegate(const FOnSkelMeshPhysicsCreated& Delegate)
 {
 	return OnSkelMeshPhysicsCreated.Add(Delegate);
@@ -2537,7 +2718,7 @@ void USkeletalMeshComponent::UnregisterOnPhysicsCreatedDelegate(const FDelegateH
 	OnSkelMeshPhysicsCreated.Remove(DelegateHandle);
 }
 
-bool USkeletalMeshComponent::MoveComponentImpl(const FVector& Delta, const FQuat& NewRotation, bool bSweep, FHitResult* OutHit /*= NULL*/, EMoveComponentFlags MoveFlags /*= MOVECOMP_NoFlags*/, ETeleportType Teleport /*= ETeleportType::None*/)
+bool USkeletalMeshComponent::MoveComponentImpl(const FVector& Delta, const FQuat& NewRotation, bool bSweep, FHitResult* OutHit /*= nullptr*/, EMoveComponentFlags MoveFlags /*= MOVECOMP_NoFlags*/, ETeleportType Teleport /*= ETeleportType::None*/)
 {
 #if WITH_EDITOR
 	UWorld* World = GetWorld();
@@ -2558,4 +2739,11 @@ bool USkeletalMeshComponent::MoveComponentImpl(const FVector& Delta, const FQuat
 	return Super::MoveComponentImpl(Delta, NewRotation, bSweep, OutHit, MoveFlags, Teleport);
 }
 
+void USkeletalMeshComponent::AddSlavePoseComponent(USkinnedMeshComponent* SkinnedMeshComponent)
+{
+	Super::AddSlavePoseComponent(SkinnedMeshComponent);
+
+	bRequiredBonesUpToDate = false;
+}
+
 #undef LOCTEXT_NAMESPACE
