Copying //UE4/Dev-Framework to //UE4/Dev-Main (Source: //UE4/Dev-Framework @ 3136612)

diff --git a/Engine/Source/Runtime/Engine/Private/Components/SkinnedMeshComponent.cpp b/Engine/Source/Runtime/Engine/Private/Components/SkinnedMeshComponent.cpp
--- a/Engine/Source/Runtime/Engine/Private/Components/SkinnedMeshComponent.cpp
+++ b/Engine/Source/Runtime/Engine/Private/Components/SkinnedMeshComponent.cpp
@@ -622,7 +622,7 @@ bool USkinnedMeshComponent::IsMaterialSlotNameValid(FName MaterialSlotName) cons
 
 bool USkinnedMeshComponent::ShouldCPUSkin()
 {
-	return false;
+	return bCPUSkinning;
 }
 
 void USkinnedMeshComponent::GetStreamingTextureInfo(FStreamingTextureLevelContext& LevelContext, TArray<FStreamingTexturePrimitiveInfo>& OutStreamingTextures) const
@@ -1143,7 +1143,7 @@ void USkinnedMeshComponent::SetMasterPoseComponent(class USkinnedMeshComponent*
 
 		if (bAddNew)
 		{
-			MasterPoseComponent->SlavePoseComponents.Add(this);
+			MasterPoseComponent->AddSlavePoseComponent(this);
 		}
 	}
 
@@ -1164,6 +1164,11 @@ void USkinnedMeshComponent::SetMasterPoseComponent(class USkinnedMeshComponent*
 	UpdateMasterBoneMap();
 }
 
+void USkinnedMeshComponent::AddSlavePoseComponent(USkinnedMeshComponent* SkinnedMeshComponent)
+{
+	SlavePoseComponents.Add(SkinnedMeshComponent);
+}
+
 void USkinnedMeshComponent::InvalidateCachedBounds()
 {
 	bCachedLocalBoundsUpToDate = false;
@@ -1997,6 +2002,16 @@ bool USkinnedMeshComponent::UpdateLODStatus()
 	{
 		MaxDistanceFactor = MeshObject->MaxDistanceFactor;
 	}
+	
+	// also update slave component LOD status, as we may need to recalc required bones if this changes
+	// independently of our LOD
+	for (TWeakObjectPtr<USkinnedMeshComponent> SlaveComponents : SlavePoseComponents)
+	{
+		if (SlaveComponents.IsValid())
+		{
+			bLODChanged |= SlaveComponents->UpdateLODStatus();
+		}
+	}
 
 	return bLODChanged;
 }
@@ -2045,6 +2060,40 @@ void USkinnedMeshComponent::UpdateRecomputeTangent(int32 MaterialIndex, int32 Lo
 	}
 }
 
+void USkinnedMeshComponent::GetCPUSkinnedVertices(TArray<FFinalSkinVertex>& OutVertices, int32 InLODIndex)
+{
+	// switch to CPU skinning
+	bool bCachedCPUSkinning = bCPUSkinning;
+	bCPUSkinning = true;
+
+	if (MasterPoseComponent.IsValid())
+	{
+		MasterPoseComponent->ForcedLodModel = InLODIndex + 1;
+		MasterPoseComponent->UpdateLODStatus();
+		MasterPoseComponent->RefreshBoneTransforms(nullptr);
+	}
+	else
+	{
+		ForcedLodModel = InLODIndex + 1;
+		UpdateLODStatus();
+		RefreshBoneTransforms(nullptr);
+	}
+
+	// Recreate render state and flush the renderer
+	RecreateRenderState_Concurrent();
+	FlushRenderingCommands();
+
+	check(MeshObject);
+
+	// Copy our vertices out. We know we are using CPU skinning now, so this cast is safe
+	OutVertices = static_cast<FSkeletalMeshObjectCPUSkin*>(MeshObject)->GetCachedFinalVertices();
+
+	// switch skinning mode, LOD etc. back
+	bCPUSkinning = bCachedCPUSkinning;
+	ForcedLodModel = 0;
+	RecreateRenderState_Concurrent();
+}
+
 void FAnimUpdateRateParameters::SetTrailMode(float DeltaTime, uint8 UpdateRateShift, int32 NewUpdateRate, int32 NewEvaluationRate, bool bNewInterpSkippedFrames)
 {
 	OptimizeMode = TrailMode;
@@ -2141,4 +2190,4 @@ float FAnimUpdateRateParameters::GetRootMotionInterp() const
 		return FMath::Clamp(ThisTickDelta / (TickedPoseOffestTime + ThisTickDelta), 0.f, 1.f);
 	}
 	return 1.f;
-}
\ No newline at end of file
+}
diff --git a/Engine/Source/Runtime/Engine/Private/Controller.cpp b/Engine/Source/Runtime/Engine/Private/Controller.cpp
--- a/Engine/Source/Runtime/Engine/Private/Controller.cpp
+++ b/Engine/Source/Runtime/Engine/Private/Controller.cpp
@@ -10,8 +10,8 @@
 #include "NetworkingDistanceConstants.h"
 #include "VisualLogger/VisualLogger.h"
 #include "GameFramework/Character.h"
-#include "GameFramework/GameMode.h"
-#include "GameFramework/GameState.h"
+#include "GameFramework/GameModeBase.h"
+#include "GameFramework/GameStateBase.h"
 #include "GameFramework/PawnMovementComponent.h"
 #include "MessageLog.h"
 
@@ -448,7 +448,7 @@ void AController::Destroyed()
 	if (Role == ROLE_Authority && PlayerState != NULL)
 	{
 		// if we are a player, log out
-		AGameMode* const GameMode = GetWorld()->GetAuthGameMode();
+		AGameModeBase* const GameMode = GetWorld()->GetAuthGameMode();
 		if (GameMode)
 		{
 			GameMode->Logout(this);
@@ -480,14 +480,14 @@ void AController::InitPlayerState()
 	if ( GetNetMode() != NM_Client )
 	{
 		UWorld* const World = GetWorld();
-		const AGameMode* GameMode = World ? World->GetAuthGameMode() : NULL;
+		const AGameModeBase* GameMode = World ? World->GetAuthGameMode() : NULL;
 
 		// If the GameMode is null, this might be a network client that's trying to
 		// record a replay. Try to use the default game mode in this case so that
 		// we can still spawn a PlayerState.
 		if (GameMode == NULL)
 		{
-			const AGameState* const GameState = World ? World->GetGameState() : NULL;
+			const AGameStateBase* const GameState = World ? World->GetGameState() : NULL;
 			GameMode = GameState ? GameState->GetDefaultGameMode() : NULL;
 		}
 
diff --git a/Engine/Source/Runtime/Engine/Private/DebugCameraController.cpp b/Engine/Source/Runtime/Engine/Private/DebugCameraController.cpp
--- a/Engine/Source/Runtime/Engine/Private/DebugCameraController.cpp
+++ b/Engine/Source/Runtime/Engine/Private/DebugCameraController.cpp
@@ -260,7 +260,7 @@ ASpectatorPawn* ADebugCameraController::SpawnSpectatorPawn()
 	// Only spawned for the local player
 	if (GetSpectatorPawn() == NULL && IsLocalController())
 	{
-		AGameState const* const GameState = GetWorld()->GameState;
+		AGameStateBase const* const GameState = GetWorld()->GetGameState();
 		if (GameState)
 		{
 			FActorSpawnParameters SpawnParams;
diff --git a/Engine/Source/Runtime/Engine/Private/DemoNetDriver.cpp b/Engine/Source/Runtime/Engine/Private/DemoNetDriver.cpp
--- a/Engine/Source/Runtime/Engine/Private/DemoNetDriver.cpp
+++ b/Engine/Source/Runtime/Engine/Private/DemoNetDriver.cpp
@@ -26,8 +26,8 @@
 #include "Net/UnrealNetwork.h"
 #include "Net/NetworkProfiler.h"
 #include "Net/DataReplication.h"
-#include "GameFramework/GameMode.h"
-#include "GameFramework/GameState.h"
+#include "GameFramework/GameModeBase.h"
+#include "GameFramework/GameStateBase.h"
 #include "GameFramework/PlayerState.h"
 #include "NetworkVersion.h"
 
@@ -588,13 +588,14 @@ bool UDemoNetDriver::InitListen( FNetworkNotify* InNotify, FURL& ListenURL, bool
 	const TCHAR* FriendlyNameOption = ListenURL.GetOption( TEXT("DemoFriendlyName="), nullptr );
 
 	TArray< FString > UserNames;
+	AGameStateBase* GameState = GetWorld()->GetGameState();
 
 	// If a client is recording a replay, GameState may not have replicated yet
-	if ( GetWorld()->GameState != nullptr )
+	if (GameState != nullptr )
 	{
-		for ( int32 i = 0; i < GetWorld()->GameState->PlayerArray.Num(); i++ )
+		for ( int32 i = 0; i < GameState->PlayerArray.Num(); i++ )
 		{
-			APlayerState* PlayerState = GetWorld()->GameState->PlayerArray[i];
+			APlayerState* PlayerState = GameState->PlayerArray[i];
 
 			if ( !PlayerState->bIsABot && !PlayerState->bIsSpectator )
 			{
@@ -1091,11 +1092,11 @@ void UDemoNetDriver::TickCheckpoint()
 	// Save the replicated server time so we can restore it after the checkpoint has been serialized.
 	// This preserves the existing behavior and prevents clients from receiving updated server time
 	// more often than the normal update rate.
-	AGameState* const GameState = World != nullptr ? World->GetGameState() : nullptr;
+	AGameStateBase* const GameState = World != nullptr ? World->GetGameState() : nullptr;
 
 	const float SavedReplicatedServerTimeSeconds = GameState ? GameState->ReplicatedWorldTimeSeconds : -1.0f;
 
-	// Normally AGameState::ReplicatedWorldTimeSeconds is only updated periodically,
+	// Normally AGameStateBase::ReplicatedWorldTimeSeconds is only updated periodically,
 	// but we want to make sure it's accurate for the checkpoint.
 	if ( GameState )
 	{
@@ -2034,7 +2035,7 @@ void UDemoNetDriver::FinalizeFastForward( const float StartTime )
 	// This must be set before we CallRepNotifies or they might be skipped again
 	bIsFastForwarding = false;
 
-	AGameState* const GameState = World != nullptr ? World->GetGameState() : nullptr;
+	AGameStateBase* const GameState = World != nullptr ? World->GetGameState() : nullptr;
 
 	// Correct server world time for fast-forwarding after a checkpoint
 	if (GameState != nullptr)
@@ -2114,8 +2115,8 @@ void UDemoNetDriver::SpawnDemoRecSpectator( UNetConnection* Connection, const FU
 
 	// Get the replay spectator controller class from the default game mode object,
 	// since the game mode instance isn't replicated to clients of live games.
-	AGameState* GameState = GetWorld() != nullptr ? GetWorld()->GetGameState() : nullptr;
-	TSubclassOf<AGameMode> DefaultGameModeClass = GameState != nullptr ? GameState->GameModeClass : nullptr;
+	AGameStateBase* GameState = GetWorld() != nullptr ? GetWorld()->GetGameState() : nullptr;
+	TSubclassOf<AGameModeBase> DefaultGameModeClass = GameState != nullptr ? GameState->GameModeClass : nullptr;
 	
 	// If we don't have a game mode class from the world, try to get it from the URL option.
 	// This may be true on clients who are recording a replay before the game mode class was replicated to them.
@@ -2124,12 +2125,12 @@ void UDemoNetDriver::SpawnDemoRecSpectator( UNetConnection* Connection, const FU
 		const TCHAR* URLGameModeClass = ListenURL.GetOption(TEXT("game="), nullptr);
 		if (URLGameModeClass != nullptr)
 		{
-			UClass* GameModeFromURL = StaticLoadClass(AGameMode::StaticClass(), nullptr, URLGameModeClass);
+			UClass* GameModeFromURL = StaticLoadClass(AGameModeBase::StaticClass(), nullptr, URLGameModeClass);
 			DefaultGameModeClass = GameModeFromURL;
 		}
 	}
 
-	AGameMode* DefaultGameMode = Cast<AGameMode>(DefaultGameModeClass.GetDefaultObject());
+	AGameModeBase* DefaultGameMode = Cast<AGameModeBase>(DefaultGameModeClass.GetDefaultObject());
 	UClass* C = DefaultGameMode != nullptr ? DefaultGameMode->ReplaySpectatorPlayerControllerClass : nullptr;
 
 	if ( C == NULL )
@@ -2485,7 +2486,7 @@ void UDemoNetDriver::LoadCheckpoint( FArchive* GotoCheckpointArchive, int64 Goto
 	// Save the replicated server time here
 	if (World != nullptr)
 	{
-		const AGameState* const GameState = World->GetGameState();
+		const AGameStateBase* const GameState = World->GetGameState();
 		if (GameState != nullptr)
 		{
 			SavedReplicatedWorldTimeSeconds = GameState->ReplicatedWorldTimeSeconds;
diff --git a/Engine/Source/Runtime/Engine/Private/EngineAnalytics.cpp b/Engine/Source/Runtime/Engine/Private/EngineAnalytics.cpp
--- a/Engine/Source/Runtime/Engine/Private/EngineAnalytics.cpp
+++ b/Engine/Source/Runtime/Engine/Private/EngineAnalytics.cpp
@@ -107,6 +107,10 @@ void FEngineAnalytics::Initialize()
 		{
 			Config.AppEnvironment = TEXT("datacollector-source");
 		}
+		if (Config.AppVersionET.IsEmpty())
+		{
+			Config.AppVersionET = FEngineVersion::Current().ToString();
+		}
 
 
 		// Connect the engine analytics provider (if there is a configuration delegate installed)
@@ -124,7 +128,8 @@ void FEngineAnalytics::Initialize()
 					FPlatformMisc::SetStoredValue(TEXT("Epic Games"), TEXT("Unreal Engine/Privacy"), TEXT("AnonymousID"), AnonymousId);
 				}
 
-				Analytics->SetUserID(FString::Printf(TEXT("ANON-%s"), *AnonymousId));
+				// Place the anonymous user id into the first field of the UserID set, as it most closely matches the LoginID semantics.
+				Analytics->SetUserID(FString::Printf(TEXT("ANON-%s||"), *AnonymousId));
 			}
 			else
 			{
diff --git a/Engine/Source/Runtime/Engine/Private/GameEngine.cpp b/Engine/Source/Runtime/Engine/Private/GameEngine.cpp
--- a/Engine/Source/Runtime/Engine/Private/GameEngine.cpp
+++ b/Engine/Source/Runtime/Engine/Private/GameEngine.cpp
@@ -37,7 +37,7 @@
 #include "Components/ReflectionCaptureComponent.h"
 #include "Engine/GameEngine.h"
 #include "GameFramework/GameUserSettings.h"
-#include "GameFramework/GameMode.h"
+#include "GameFramework/GameModeBase.h"
 #include "GameDelegates.h"
 #include "Engine/CoreSettings.h"
 #include "EngineAnalytics.h"
diff --git a/Engine/Source/Runtime/Engine/Private/GameInstance.cpp b/Engine/Source/Runtime/Engine/Private/GameInstance.cpp
--- a/Engine/Source/Runtime/Engine/Private/GameInstance.cpp
+++ b/Engine/Source/Runtime/Engine/Private/GameInstance.cpp
@@ -900,3 +900,120 @@ void UGameInstance::NotifyPreClientTravel(const FString& PendingURL, ETravelType
 	OnNotifyPreClientTravel().Broadcast(PendingURL, TravelType, bIsSeamlessTravel);
 }
 
+void UGameInstance::PreloadContentForURL(FURL InURL)
+{
+	// Preload game mode and other content if needed here
+}
+
+AGameModeBase* UGameInstance::CreateGameModeForURL(FURL InURL)
+{
+	UWorld* World = GetWorld();
+	// Init the game info.
+	FString Options(TEXT(""));
+	TCHAR GameParam[256] = TEXT("");
+	FString	Error = TEXT("");
+	AWorldSettings* Settings = World->GetWorldSettings();
+	for (int32 i = 0; i < InURL.Op.Num(); i++)
+	{
+		Options += TEXT("?");
+		Options += InURL.Op[i];
+		FParse::Value(*InURL.Op[i], TEXT("GAME="), GameParam, ARRAY_COUNT(GameParam));
+	}
+
+	UGameEngine* const GameEngine = Cast<UGameEngine>(GEngine);
+
+	// Get the GameMode class. Start by using the default game type specified in the map's worldsettings.  It may be overridden by settings below.
+	TSubclassOf<AGameModeBase> GameClass = Settings->DefaultGameMode;
+
+	// If there is a GameMode parameter in the URL, allow it to override the default game type
+	if (GameParam[0])
+	{
+		FString const GameClassName = UGameMapsSettings::GetGameModeForName(FString(GameParam));
+
+		// If the gamename was specified, we can use it to fully load the pergame PreLoadClass packages
+		if (GameEngine)
+		{
+			GameEngine->LoadPackagesFully(World, FULLYLOAD_Game_PreLoadClass, *GameClassName);
+		}
+
+		// Don't overwrite the map's world settings if we failed to load the value off the command line parameter
+		TSubclassOf<AGameModeBase> GameModeParamClass = LoadClass<AGameModeBase>(nullptr, *GameClassName);
+		if (GameModeParamClass)
+		{
+			GameClass = GameModeParamClass;
+		}
+		else
+		{
+			UE_LOG(LogLoad, Warning, TEXT("Failed to load game mode '%s' specified by URL options."), *GameClassName);
+		}
+	}
+
+	// Next try to parse the map prefix
+	if (!GameClass)
+	{
+		FString MapName = InURL.Map;
+		FString MapNameNoPath = FPaths::GetBaseFilename(MapName);
+		if (MapNameNoPath.StartsWith(PLAYWORLD_PACKAGE_PREFIX))
+		{
+			const int32 PrefixLen = UWorld::BuildPIEPackagePrefix(WorldContext->PIEInstance).Len();
+			MapNameNoPath = MapNameNoPath.Mid(PrefixLen);
+		}
+
+		FString const GameClassName = UGameMapsSettings::GetGameModeForMapName(FString(MapNameNoPath));
+
+		if (!GameClassName.IsEmpty())
+		{
+			if (GameEngine)
+			{
+				GameEngine->LoadPackagesFully(World, FULLYLOAD_Game_PreLoadClass, *GameClassName);
+			}
+
+			TSubclassOf<AGameModeBase> GameModeParamClass = LoadClass<AGameModeBase>(nullptr, *GameClassName);
+			if (GameModeParamClass)
+			{
+				GameClass = GameModeParamClass;
+			}
+			else
+			{
+				UE_LOG(LogLoad, Warning, TEXT("Failed to load game mode '%s' specified by prefixed map name %s."), *GameClassName, *MapNameNoPath);
+			}
+		}
+	}
+
+	// Fall back to game default
+	if (!GameClass)
+	{
+		GameClass = LoadClass<AGameModeBase>(nullptr, *UGameMapsSettings::GetGlobalDefaultGameMode());
+	}
+
+	if (!GameClass)
+	{
+		// Fall back to raw GameMode
+		GameClass = AGameModeBase::StaticClass();
+	}
+	else
+	{
+		// See if game instance wants to override it
+		GameClass = OverrideGameModeClass(GameClass, FPaths::GetBaseFilename(InURL.Map), Options, *InURL.Portal);
+	}
+
+	// no matter how the game was specified, we can use it to load the PostLoadClass packages
+	if (GameEngine)
+	{
+		GameEngine->LoadPackagesFully(World, FULLYLOAD_Game_PostLoadClass, GameClass->GetPathName());
+		GameEngine->LoadPackagesFully(World, FULLYLOAD_Game_PostLoadClass, TEXT("LoadForAllGameModes"));
+	}
+
+	// Spawn the GameMode.
+	UE_LOG(LogLoad, Log, TEXT("Game class is '%s'"), *GameClass->GetName());
+	FActorSpawnParameters SpawnInfo;
+	SpawnInfo.SpawnCollisionHandlingOverride = ESpawnActorCollisionHandlingMethod::AlwaysSpawn;
+	SpawnInfo.ObjectFlags |= RF_Transient;	// We never want to save game modes into a map
+
+	return World->SpawnActor<AGameModeBase>(GameClass, SpawnInfo);
+}
+
+TSubclassOf<AGameModeBase> UGameInstance::OverrideGameModeClass(TSubclassOf<AGameModeBase> GameModeClass, const FString& MapName, const FString& Options, const FString& Portal) const
+{
+	 return GameModeClass;
+}
\ No newline at end of file
diff --git a/Engine/Source/Runtime/Engine/Private/GameMode.cpp b/Engine/Source/Runtime/Engine/Private/GameMode.cpp
--- a/Engine/Source/Runtime/Engine/Private/GameMode.cpp
+++ b/Engine/Source/Runtime/Engine/Private/GameMode.cpp
@@ -5,9 +5,7 @@
 =============================================================================*/
 
 #include "EnginePrivate.h"
-#include "Engine/LevelScriptActor.h"
 #include "GameFramework/GameNetworkManager.h"
-#include "Matinee/MatineeActor.h"
 #include "Net/OnlineEngineInterface.h"
 #include "GameFramework/HUD.h"
 #include "GameFramework/DefaultPawn.h"
@@ -19,16 +17,8 @@
 #include "GameFramework/GameSession.h"
 #include "GameFramework/CheatManager.h"
 #include "GameFramework/GameMode.h"
-#include "Engine/ChildConnection.h"
 #include "Engine/GameInstance.h"
-	
-#if WITH_EDITOR
-	#include "IMovieSceneCapture.h"
-	#include "MovieSceneCaptureModule.h"
-	#include "MovieSceneCaptureSettings.h"
-#endif
-
-DEFINE_LOG_CATEGORY_STATIC(LogGameMode, Log, All);
+#include "GameDelegates.h"
 
 namespace MatchState
 {
@@ -40,166 +30,55 @@ namespace MatchState
 	const FName Aborted = FName(TEXT("Aborted"));
 }
 
-// Statically declared events for plugins to use
-FGameModeEvents::FGameModePostLoginEvent FGameModeEvents::GameModePostLoginEvent;
-FGameModeEvents::FGameModeLogoutEvent FGameModeEvents::GameModeLogoutEvent;
-
 AGameMode::AGameMode(const FObjectInitializer& ObjectInitializer)
-	: Super(ObjectInitializer
-		.DoNotCreateDefaultSubobject(TEXT("Sprite"))
-	)
+	: Super(ObjectInitializer)
 {
-	bStartPlayersAsSpectators = false;
 	bDelayedStart = false;
-	bNetLoadOnClient = false;
 
 	// One-time initialization
 	PrimaryActorTick.bCanEverTick = true;
 	PrimaryActorTick.TickGroup = TG_PrePhysics;
-	HUDClass = AHUD::StaticClass();
 	MatchState = MatchState::EnteringMap;
-	bPauseable = true;
-	DefaultPawnClass = ADefaultPawn::StaticClass();
-	PlayerControllerClass = APlayerController::StaticClass();
-	SpectatorClass = ASpectatorPawn::StaticClass();
-	ReplaySpectatorPlayerControllerClass = APlayerController::StaticClass();
 	EngineMessageClass = UEngineMessage::StaticClass();
 	GameStateClass = AGameState::StaticClass();
-	CurrentID = 1;
-	PlayerStateClass = APlayerState::StaticClass();
 	MinRespawnDelay = 1.0f;
 	InactivePlayerStateLifeSpan = 300.f;
 }
 
-FString AGameMode::GetNetworkNumber()
+FString AGameMode::GetDefaultGameClassPath(const FString& MapName, const FString& Options, const FString& Portal) const
 {
-	UNetDriver* NetDriver = GetWorld()->GetNetDriver();
-	return NetDriver ? NetDriver->LowLevelGetNetworkNumber() : FString(TEXT(""));
+	// This is called on the CDO
+	return GetClass()->GetPathName();
 }
 
-void AGameMode::SwapPlayerControllers(APlayerController* OldPC, APlayerController* NewPC)
+TSubclassOf<AGameMode> AGameMode::GetGameModeClass(const FString& MapName, const FString& Options, const FString& Portal) const
 {
-	if (OldPC != NULL && !OldPC->IsPendingKill() && NewPC != NULL && !NewPC->IsPendingKill() && OldPC->Player != NULL)
-	{
-		// move the Player to the new PC
-		UPlayer* Player = OldPC->Player;
-		NewPC->NetPlayerIndex = OldPC->NetPlayerIndex; //@warning: critical that this is first as SetPlayer() may trigger RPCs
-		NewPC->NetConnection = OldPC->NetConnection;
-		NewPC->SetPlayer(Player);
-		NewPC->CopyRemoteRoleFrom( OldPC );
-
-		K2_OnSwapPlayerControllers(OldPC, NewPC);
-
-		// send destroy event to old PC immediately if it's local
-		if (Cast<ULocalPlayer>(Player))
-		{
-			GetWorld()->DestroyActor(OldPC);
-		}
-		else
-		{
-			OldPC->PendingSwapConnection = Cast<UNetConnection>(Player);
-			//@note: at this point, any remaining RPCs sent by the client on the old PC will be discarded
-			// this is consistent with general owned Actor destruction,
-			// however in this particular case it could easily be changed
-			// by modifying UActorChannel::ReceivedBunch() to account for PendingSwapConnection when it is setting bNetOwner
-		}
-	}
-	else
-	{
-		UE_LOG(LogGameMode, Warning, TEXT("SwapPlayerControllers(): Invalid OldPC, invalid NewPC, or OldPC has no Player!"));
-	}
+	// This is called on the CDO
+	return GetClass();
 }
 
-void AGameMode::ForceClearUnpauseDelegates( AActor* PauseActor )
+FString AGameMode::StaticGetFullGameClassName(FString const& Str)
 {
-	if ( PauseActor != NULL )
-	{
-		bool bUpdatePausedState = false;
-		for ( int32 PauserIdx = Pausers.Num() - 1; PauserIdx >= 0; PauserIdx-- )
-		{
-			FCanUnpause& CanUnpauseDelegate = Pausers[PauserIdx];
-			if ( CanUnpauseDelegate.GetUObject() == PauseActor )
-			{
-				Pausers.RemoveAt(PauserIdx);
-				bUpdatePausedState = true;
-			}
-		}
-
-		// if we removed some CanUnpause delegates, we may be able to unpause the game now
-		if ( bUpdatePausedState )
-		{
-			ClearPause();
-		}
-
-		APlayerController* PC = Cast<APlayerController>(PauseActor);
-		AWorldSettings * WorldSettings = GetWorldSettings();
-		if ( PC != NULL && PC->PlayerState != NULL && WorldSettings != NULL && WorldSettings->Pauser == PC->PlayerState )
-		{
-			// try to find another player to be the worldsettings's Pauser
-			for( FConstPlayerControllerIterator Iterator = GetWorld()->GetPlayerControllerIterator(); Iterator; ++Iterator )
-			{
-				APlayerController* Player = *Iterator;
-				if (Player->PlayerState != NULL
-					&&	Player->PlayerState != PC->PlayerState
-					&&	!Player->IsPendingKillPending() && !Player->PlayerState->IsPendingKillPending())
-				{
-					WorldSettings->Pauser = Player->PlayerState;
-					break;
-				}
-			}
-
-			// if it's still pointing to the original player's PlayerState, clear it completely
-			if ( WorldSettings->Pauser == PC->PlayerState )
-			{
-				WorldSettings->Pauser = NULL;
-			}
-		}
-	}
+	return UGameMapsSettings::GetGameModeForName(Str);
 }
 
-void AGameMode::InitGame(const FString& MapName, const FString& Options, FString& ErrorMessage)
+FString AGameMode::GetNetworkNumber()
 {
-	UWorld* World = GetWorld();
-
-	// save Options for future use
-	OptionsString = Options;
-
-	UClass* const SessionClass = GetGameSessionClass();
-	FActorSpawnParameters SpawnInfo;
-	SpawnInfo.Instigator = Instigator;
-	SpawnInfo.ObjectFlags |= RF_Transient;	// We never want to save game sessions into a map
-	GameSession = World->SpawnActor<AGameSession>(SessionClass, SpawnInfo);
-	GameSession->InitOptions(Options);
-
-	if (GetNetMode() != NM_Standalone)
-	{
-		// Attempt to login, returning true means an async login is in flight
-		if (!UOnlineEngineInterface::Get()->DoesSessionExist(World, GameSession->SessionName) &&
-			!GameSession->ProcessAutoLogin())
-		{
-			GameSession->RegisterServer();
-		}
-	}
-
-	SetMatchState(MatchState::EnteringMap);
+	UNetDriver* NetDriver = GetWorld()->GetNetDriver();
+	return NetDriver ? NetDriver->LowLevelGetNetworkNumber() : FString(TEXT(""));
 }
 
-bool AGameMode::SetPause(APlayerController* PC, FCanUnpause CanUnpauseDelegate)
+void AGameMode::InitGame(const FString& MapName, const FString& Options, FString& ErrorMessage)
 {
-	if ( AllowPausing(PC) )
-	{
-		// Add it for querying
-		Pausers.Add(CanUnpauseDelegate);
+	Super::InitGame(MapName, Options, ErrorMessage);
+	SetMatchState(MatchState::EnteringMap);
 
-		// Let the first one in "own" the pause state
-		AWorldSettings * WorldSettings = GetWorldSettings();
-		if (WorldSettings->Pauser == NULL)
-		{
-			WorldSettings->Pauser = PC->PlayerState;
-		}
-		return true;
-	}
-	return false;
+	// Bind to delegates
+	FGameDelegates::Get().GetMatineeCancelledDelegate().AddUObject(this, &AGameMode::MatineeCancelled);
+	FGameDelegates::Get().GetPendingConnectionLostDelegate().AddUObject(this, &AGameMode::NotifyPendingConnectionLost);
+	FGameDelegates::Get().GetPreCommitMapChangeDelegate().AddUObject(this, &AGameMode::PreCommitMapChange);
+	FGameDelegates::Get().GetPostCommitMapChangeDelegate().AddUObject(this, &AGameMode::PostCommitMapChange);
+	FGameDelegates::Get().GetHandleDisconnectDelegate().AddUObject(this, &AGameMode::HandleDisconnect);
 }
 
 void AGameMode::RestartGame()
@@ -215,20 +94,12 @@ void AGameMode::RestartGame()
 	}
 }
 
-void AGameMode::ReturnToMainMenuHost()
-{
-	if (GameSession)
-	{
-		GameSession->ReturnToMainMenuHost();
-	}
-}
-
 void AGameMode::PostLogin( APlayerController* NewPlayer )
 {
 	UWorld* World = GetWorld();
 
 	// update player count
-	if (NewPlayer->PlayerState->bOnlySpectator)
+	if (MustSpectate(NewPlayer))
 	{
 		NumSpectators++;
 	}
@@ -249,300 +120,25 @@ void AGameMode::PostLogin( APlayerController* NewPlayer )
 	// check if this player is reconnecting and already has PlayerState
 	FindInactivePlayer(NewPlayer);
 
-	StartNewPlayer(NewPlayer);
-	NewPlayer->ClientCapBandwidth(NewPlayer->Player->CurrentNetSpeed);
-	if(World->NetworkManager)
-	{
-		World->NetworkManager->UpdateNetSpeeds(true);	 // @TODO ONLINE - Revisit LAN netspeeds
-	}
-
-	GenericPlayerInitialization(NewPlayer);
-
-	if (NewPlayer->PlayerState->bOnlySpectator)
-	{
-		NewPlayer->ClientGotoState(NAME_Spectating);
-	}
-	else
-	{
-		// If NewPlayer is not only a spectator and has a valid ID, add him as a user to the replay.
-		if (NewPlayer->PlayerState->UniqueId.IsValid())
-		{
-			GetGameInstance()->AddUserToReplay(NewPlayer->PlayerState->UniqueId.ToString());
-		}
-	}
-
-	if (GameSession)
-	{
-		GameSession->PostLogin(NewPlayer);
-	}
-
-	// add the player to any matinees running so that it gets in on any cinematics already running, etc
-	TArray<AMatineeActor*> AllMatineeActors;
-	World->GetMatineeActors(AllMatineeActors);
-	for (int32 i = 0; i < AllMatineeActors.Num(); i++)
-	{
-		AllMatineeActors[i]->AddPlayerToDirectorTracks(NewPlayer);
-	}
-
-	bool HidePlayer=false, HideHUD=false, DisableMovement=false, DisableTurning=false;
-
-	//Check to see if we should start in cinematic mode (matinee movie capture)
-	if(ShouldStartInCinematicMode(NewPlayer, HidePlayer, HideHUD, DisableMovement, DisableTurning))
-	{
-		NewPlayer->SetCinematicMode(true, HidePlayer, HideHUD, DisableMovement, DisableTurning);
-	}
-
-	// Tell the player to enable voice by default or use the push to talk method
-	NewPlayer->ClientEnableNetworkVoice(!GameSession->RequiresPushToTalk());
-
-	// Notify Blueprints that a new player has logged in.  Calling it here, because this is the first time that the PlayerController can take RPCs
-	K2_PostLogin(NewPlayer);
-	FGameModeEvents::GameModePostLoginEvent.Broadcast(this, NewPlayer);
-}
-
-bool AGameMode::ShouldStartInCinematicMode(APlayerController* Player, bool& OutHidePlayer,bool& OutHideHUD,bool& OutDisableMovement,bool& OutDisableTurning)
-{
-	ULocalPlayer* LocPlayer = Player->GetLocalPlayer();
-	if (!LocPlayer)
-	{
-		return false;
-	}
-
-#if WITH_EDITOR
-	// If we have an active movie scene capture, we can take the settings from that
-	if(LocPlayer->ViewportClient && LocPlayer->ViewportClient->Viewport)
-	{
-		if(auto* MovieSceneCapture = IMovieSceneCaptureModule::Get().GetFirstActiveMovieSceneCapture())
-		{
-			const FMovieSceneCaptureSettings& Settings = MovieSceneCapture->GetSettings();
-			if (Settings.bCinematicMode)
-			{
-				OutDisableMovement = !Settings.bAllowMovement;
-				OutDisableTurning = !Settings.bAllowTurning;
-				OutHidePlayer = !Settings.bShowPlayer;
-				OutHideHUD = !Settings.bShowHUD;
-				return true;
-			}
-		}
-	}
-#endif
-
-	return false;
-}
-
-
-void AGameMode::SetPlayerDefaults(APawn* PlayerPawn)
-{
-	PlayerPawn->SetPlayerDefaults();
+	Super::PostLogin(NewPlayer);
 }
 
 void AGameMode::Logout( AController* Exiting )
 {
 	APlayerController* PC = Cast<APlayerController>(Exiting);
-	if ( PC != NULL )
+	if ( PC != nullptr )
 	{
-		FGameModeEvents::GameModeLogoutEvent.Broadcast(this, Exiting);
-		K2_OnLogout(Exiting);
-
 		RemovePlayerControllerFromPlayerCount(PC);
-
-		if (GameSession)
-		{
-			GameSession->NotifyLogout(PC);
-		}
-
-		// @TODO ONLINE - Revisit LAN netspeeds
-		UWorld* World = GetWorld();
-		if (World && World->NetworkManager)
-		{
-			World->NetworkManager->UpdateNetSpeeds(true);
-		}
-	}
-}
-
-void AGameMode::InitGameState()
-{
-	GameState->GameModeClass = GetClass();
-	GameState->ReceivedGameModeClass();
-	GameState->SpectatorClass = SpectatorClass;
-	GameState->ReceivedSpectatorClass();
-}
-
-
-AActor* AGameMode::FindPlayerStart_Implementation( AController* Player, const FString& IncomingName )
-{
-	UWorld* World = GetWorld();
-
-	// if incoming start is specified, then just use it
-	if( !IncomingName.IsEmpty() )
-	{
-		const FName IncomingPlayerStartTag = FName(*IncomingName);
-		for (TActorIterator<APlayerStart> It(World); It; ++It)
-		{
-			APlayerStart* Start = *It;
-			if (Start && Start->PlayerStartTag == IncomingPlayerStartTag)
-			{
-				return Start;
-			}
-		}
-	}
-
-	// always pick StartSpot at start of match
-	if ( ShouldSpawnAtStartSpot(Player) )
-	{
-		if (AActor* PlayerStartSpot = Player->StartSpot.Get())
-		{
-			return PlayerStartSpot;
-		}
-		else
-		{
-			UE_LOG(LogGameMode, Error, TEXT("Error - ShouldSpawnAtStartSpot returned true but the Player StartSpot was null."));
-		}
-	}
-
-	AActor* BestStart = ChoosePlayerStart(Player);
-	if (BestStart == nullptr)
-	{
-		// no player start found
-		UE_LOG(LogGameMode, Log, TEXT("Warning - PATHS NOT DEFINED or NO PLAYERSTART with positive rating"));
-
-		// This is a bit odd, but there was a complex chunk of code that in the end always resulted in this, so we may as well just 
-		// short cut it down to this.  Basically we are saying spawn at 0,0,0 if we didn't find a proper player start
-		BestStart = World->GetWorldSettings();
-	}
-
-	return BestStart;
-}
-
-AActor* AGameMode::K2_FindPlayerStart( AController* Player )
-{
-	return FindPlayerStart(Player);
-}
-
-void AGameMode::PreInitializeComponents()
-{
-	Super::PreInitializeComponents();
-
-	FActorSpawnParameters SpawnInfo;
-	SpawnInfo.Instigator = Instigator;
-	SpawnInfo.ObjectFlags |= RF_Transient;	// We never want to save game states or network managers into a map
-
-	// Fallback to default GameState if none was specified.
-	if (GameStateClass == NULL)
-	{
-		UE_LOG(LogGameMode, Warning, TEXT("No GameStateClass was specified in %s (%s)"), *GetName(), *GetClass()->GetName());
-		GameStateClass = AGameState::StaticClass();
-	}
-
-	GameState = GetWorld()->SpawnActor<AGameState>(GameStateClass, SpawnInfo);
-	GetWorld()->GameState = GameState;
-	if (GameState)
-	{
-		GameState->AuthorityGameMode = this;
-	}
-
-	// Only need NetworkManager for servers in net games
-	GetWorld()->NetworkManager = GetWorldSettings()->GameNetworkManagerClass ? GetWorld()->SpawnActor<AGameNetworkManager>(GetWorldSettings()->GameNetworkManagerClass, SpawnInfo ) : NULL;
-
-	InitGameState();
-}
-
-void AGameMode::RestartPlayer(AController* NewPlayer)
-{
-	if ( NewPlayer == NULL || NewPlayer->IsPendingKillPending() )
-	{
-		return;
-	}
-
-	UE_LOG(LogGameMode, Verbose, TEXT("RestartPlayer %s"), (NewPlayer && NewPlayer->PlayerState) ? *NewPlayer->PlayerState->PlayerName : TEXT("Unknown"));
-
-	if (NewPlayer->PlayerState && NewPlayer->PlayerState->bOnlySpectator)
-	{
-		UE_LOG(LogGameMode, Verbose, TEXT("RestartPlayer tried to restart a spectator-only player!"));
-		return;
-	}
-
-	AActor* StartSpot = FindPlayerStart(NewPlayer);
-
-	// if a start spot wasn't found,
-	if (StartSpot == NULL)
-	{
-		// check for a previously assigned spot
-		if (NewPlayer->StartSpot != NULL)
-		{
-			StartSpot = NewPlayer->StartSpot.Get();
-			UE_LOG(LogGameMode, Warning, TEXT("Player start not found, using last start spot"));
-		}
-		else
-		{
-			// otherwise abort
-			UE_LOG(LogGameMode, Warning, TEXT("Player start not found, failed to restart player"));
-			return;
-		}
-	}
-	// try to create a pawn to use of the default class for this player
-	if (NewPlayer->GetPawn() == NULL && GetDefaultPawnClassForController(NewPlayer) != NULL)
-	{
-		NewPlayer->SetPawn(SpawnDefaultPawnFor(NewPlayer, StartSpot));
-	}
-
-	if (NewPlayer->GetPawn() == NULL)
-	{
-		NewPlayer->FailedToSpawnPawn();
+		AddInactivePlayer(PC->PlayerState, PC);
 	}
-	else
-	{
-		// initialize and start it up
-		InitStartSpot(StartSpot, NewPlayer);
-
-		// @todo: this was related to speedhack code, which is disabled.
-		/*
-		if ( NewPlayer->GetAPlayerController() )
-		{
-			NewPlayer->GetAPlayerController()->TimeMargin = -0.1f;
-		}
-		*/
-		NewPlayer->Possess(NewPlayer->GetPawn());
-
-		// If the Pawn is destroyed as part of possession we have to abort
-		if (NewPlayer->GetPawn() == nullptr)
-		{
-			NewPlayer->FailedToSpawnPawn();
-		}
-		else
-		{
-			// set initial control rotation to player start's rotation
-			NewPlayer->ClientSetRotation(NewPlayer->GetPawn()->GetActorRotation(), true);
-
-			FRotator NewControllerRot = StartSpot->GetActorRotation();
-			NewControllerRot.Roll = 0.f;
-			NewPlayer->SetControlRotation( NewControllerRot );
-
-			SetPlayerDefaults(NewPlayer->GetPawn());
-
-			K2_OnRestartPlayer(NewPlayer);
-		}
-	}
-
-#if !UE_WITH_PHYSICS
-	if (NewPlayer->GetPawn() != NULL)
-	{
-		UCharacterMovementComponent* CharacterMovement = Cast<UCharacterMovementComponent>(NewPlayer->GetPawn()->GetMovementComponent());
-		if (CharacterMovement)
-		{
-			CharacterMovement->bCheatFlying = true;
-			CharacterMovement->SetMovementMode(MOVE_Flying);
-		}
-	}
-#endif	//!UE_WITH_PHYSICS
-}
 
-void AGameMode::InitStartSpot_Implementation(AActor* StartSpot, AController* NewPlayer)
-{
+	Super::Logout(Exiting);
 }
 
 void AGameMode::StartPlay()
 {
+	// Don't call super, this class handles begin play/match start itself
+
 	if (MatchState == MatchState::EnteringMap)
 	{
 		SetMatchState(MatchState::WaitingToStart);
@@ -557,7 +153,7 @@ void AGameMode::StartPlay()
 
 void AGameMode::HandleMatchIsWaitingToStart()
 {
-	if (GameSession != NULL)
+	if (GameSession != nullptr)
 	{
 		GameSession->HandleMatchIsWaitingToStart();
 	}
@@ -614,7 +210,7 @@ void AGameMode::HandleMatchHasStarted()
 	for( FConstPlayerControllerIterator Iterator = GetWorld()->GetPlayerControllerIterator(); Iterator; ++Iterator )
 	{
 		APlayerController* PlayerController = *Iterator;
-		if ((PlayerController->GetPawn() == NULL) && PlayerCanRestart(PlayerController))
+		if ((PlayerController->GetPawn() == nullptr) && PlayerCanRestart(PlayerController))
 		{
 			RestartPlayer(PlayerController);
 		}
@@ -637,7 +233,7 @@ void AGameMode::HandleMatchHasStarted()
 		for( FConstPlayerControllerIterator Iterator = GetWorld()->GetPlayerControllerIterator(); Iterator; ++Iterator )
 		{
 			APlayerController* PlayerController = *Iterator;
-			if( PlayerController->CheatManager != NULL )
+			if( PlayerController->CheatManager != nullptr )
 			{
 				PlayerController->CheatManager->BugItGoString( BugLocString, BugRotString );
 			}
@@ -739,9 +335,10 @@ void AGameMode::SetMatchState(FName NewState)
 
 	OnMatchStateSet();
 
-	if (GameState)
+	AGameState* FullGameState = GetGameState<AGameState>();
+	if (FullGameState)
 	{
-		GameState->SetMatchState(NewState);
+		FullGameState->SetMatchState(NewState);
 	}
 
 	K2_OnSetMatchState(NewState);
@@ -796,56 +393,18 @@ void AGameMode::Tick(float DeltaSeconds)
 	}
 }
 
-void AGameMode::ResetLevel()
-{
-	UE_LOG(LogGameMode, Log, TEXT("Reset %s"), *GetName());
-
-	// Reset ALL controllers first
-	for( FConstControllerIterator Iterator = GetWorld()->GetControllerIterator(); Iterator; ++Iterator )
-	{
-		AController* Controller = *Iterator;
-		APlayerController* PlayerController = Cast<APlayerController>(Controller);
-		if ( PlayerController )
-		{
-			PlayerController->ClientReset();
-		}
-		Controller->Reset();
-	}
-
-	// Reset all actors (except controllers, the GameMode, and any other actors specified by ShouldReset())
-	for (FActorIterator It(GetWorld()); It; ++It)
-	{
-		AActor* A = *It;
-		if(	A && !A->IsPendingKill() && A != this && !A->IsA<AController>() && ShouldReset(A))
-		{
-			A->Reset();
-		}
-	}
-
-	// reset the GameMode
-	Reset();
-
-	// Notify the level script that the level has been reset
-	ALevelScriptActor* LevelScript = GetWorld()->GetLevelScriptActor();
-	if( LevelScript )
-	{
-		LevelScript->LevelReset();
-	}
-}
-
-
 void AGameMode::HandleSeamlessTravelPlayer(AController*& C)
 {
 	UE_LOG(LogGameMode, Log, TEXT(">> GameMode::HandleSeamlessTravelPlayer: %s "), *C->GetName());
 
 	APlayerController* PC = Cast<APlayerController>(C);
-	if (PC != NULL && PC->GetClass() != PlayerControllerClass)
+	if (PC != nullptr && PC->GetClass() != PlayerControllerClass)
 	{
-		if (PC->Player != NULL)
+		if (PC->Player != nullptr)
 		{
 			// we need to spawn a new PlayerController to replace the old one
 			APlayerController* NewPC = SpawnPlayerController(PC->IsLocalPlayerController() ? ROLE_SimulatedProxy : ROLE_AutonomousProxy, PC->GetFocalLocation(), PC->GetControlRotation());
-			if (NewPC == NULL)
+			if (NewPC == nullptr)
 			{
 				UE_LOG(LogGameMode, Warning, TEXT("Failed to spawn new PlayerController for %s (old class %s)"), *PC->GetHumanReadableName(), *PC->GetClass()->GetName());
 				PC->Destroy();
@@ -853,7 +412,6 @@ void AGameMode::HandleSeamlessTravelPlayer(AController*& C)
 			}
 			else
 			{
-				PC->CleanUpAudioComponents();
 				PC->SeamlessTravelTo(NewPC);
 				NewPC->SeamlessTravelFrom(PC);
 				SwapPlayerControllers(PC, NewPC);
@@ -879,267 +437,47 @@ void AGameMode::HandleSeamlessTravelPlayer(AController*& C)
 		OldPlayerState->Destroy();
 	}
 
-	// Find a start spot
-	AActor* StartSpot = FindPlayerStart(C);
+	InitSeamlessTravelPlayer(C);
 
-	if (StartSpot == NULL)
-	{
-		UE_LOG(LogGameMode, Warning, TEXT("Could not find a starting spot"));
-	}
-	else
+	// Initialize hud and other player details, shared with PostLogin
+	GenericPlayerInitialization(C);
+
+	if (PC)
 	{
-		FRotator StartRotation(0, StartSpot->GetActorRotation().Yaw, 0);
-		C->SetInitialLocationAndRotation(StartSpot->GetActorLocation(), StartRotation);
+		// This may spawn the player pawn if the game is in progress
+		HandleStartingNewPlayer(PC);
 	}
 
-	C->StartSpot = StartSpot;
-
-	if (PC != NULL)
-	{
-		// Track the last completed seamless travel for the player
-		PC->LastCompletedSeamlessTravelCount = PC->SeamlessTravelCount;
+	UE_LOG(LogGameMode, Log, TEXT("<< GameMode::HandleSeamlessTravelPlayer: %s"), *C->GetName());
+}
 
-		PC->CleanUpAudioComponents();
+void AGameMode::InitSeamlessTravelPlayer(AController* NewController)
+{
+	Super::InitSeamlessTravelPlayer(NewController);
 
-		if (PC->PlayerCameraManager == NULL)
-		{
-			PC->SpawnPlayerCameraManager();
-		}
+	APlayerController* NewPC = Cast<APlayerController>(NewController);
 
-		SetSeamlessTravelViewTarget(PC);
-		if (PC->PlayerState->bOnlySpectator)
-		{
-			PC->StartSpectatingOnly();
+	if (NewPC != nullptr)
+	{
+		SetSeamlessTravelViewTarget(NewPC);
 
-			// NumSpectators is incremented in PostSeamlessTravel
-		}
-		else
+		if (!MustSpectate(NewPC))
 		{
 			NumPlayers++;
 			NumTravellingPlayers--;
-			PC->bPlayerIsWaiting = true;
-			PC->ChangeState(NAME_Spectating);
-			PC->ClientGotoState(NAME_Spectating);
 		}
 	}
 	else
 	{
 		NumBots++;
 	}
-
-	if (PC)
-	{
-		// This handles setting hud, and may spawn the player pawn if the game is in progress
-		StartNewPlayer(PC);
-	}
-
-	GenericPlayerInitialization(C);
-
-	UE_LOG(LogGameMode, Log, TEXT("<< GameMode::HandleSeamlessTravelPlayer: %s"), *C->GetName());
-}
+}
 
 void AGameMode::SetSeamlessTravelViewTarget(APlayerController* PC)
 {
 	PC->SetViewTarget(PC);
 }
 
-
-void AGameMode::ProcessServerTravel(const FString& URL, bool bAbsolute)
-{
-#if WITH_SERVER_CODE
-	StartToLeaveMap();
-
-	// force an old style load screen if the server has been up for a long time so that TimeSeconds doesn't overflow and break everything
-	bool bSeamless = (bUseSeamlessTravel && GetWorld()->TimeSeconds < 172800.0f); // 172800 seconds == 48 hours
-
-	FString NextMap;
-	if (URL.ToUpper().Contains(TEXT("?RESTART")))
-	{
-		NextMap = UWorld::RemovePIEPrefix(GetOutermost()->GetName());
-	}
-	else
-	{
-		int32 OptionStart = URL.Find(TEXT("?"));
-		if (OptionStart == INDEX_NONE)
-		{
-			NextMap = URL;
-		}
-		else
-		{
-			NextMap = URL.Left(OptionStart);
-		}
-	}
-	
-	// Handle short package names (convenience code so that short map names
-	// can still be specified in the console).
-	if (FPackageName::IsShortPackageName(NextMap))
-	{
-		UE_LOG(LogGameMode, Fatal, TEXT("ProcessServerTravel: %s: Short package names are not supported."), *URL);
-		return;
-	}
-	
-	FGuid NextMapGuid = UEngine::GetPackageGuid(FName(*NextMap), GetWorld()->IsPlayInEditor());
-
-	// Notify clients we're switching level and give them time to receive.
-	FString URLMod = URL;
-	APlayerController* LocalPlayer = ProcessClientTravel(URLMod, NextMapGuid, bSeamless, bAbsolute);
-
-	UE_LOG(LogGameMode, Log, TEXT("ProcessServerTravel: %s"), *URL);
-	UWorld* World = GetWorld(); 
-	check(World);
-	World->NextURL = URL;
-	ENetMode NetMode = GetNetMode();
-
-	if (bSeamless)
-	{
-		World->SeamlessTravel(World->NextURL, bAbsolute);
-		World->NextURL = TEXT("");
-	}
-	// Switch immediately if not networking.
-	else if (NetMode != NM_DedicatedServer && NetMode != NM_ListenServer)
-	{
-		World->NextSwitchCountdown = 0.0f;
-	}
-#endif // WITH_SERVER_CODE
-}
-
-
-void AGameMode::GetSeamlessTravelActorList(bool bToTransition, TArray<AActor*>& ActorList)
-{
-	UWorld* World = GetWorld();
-
-	// Get allocations for the elements we're going to add handled in one go
-	const int32 ActorsToAddCount = World->GameState->PlayerArray.Num() + (bToTransition ?  3 : 0);
-	ActorList.Reserve(ActorsToAddCount);
-
-	// always keep PlayerStates, so that after we restart we can keep players on the same team, etc
-	ActorList.Append(World->GameState->PlayerArray);
-
-	if (bToTransition)
-	{
-		// keep ourselves until we transition to the final destination
-		ActorList.Add(this);
-		// keep general game state until we transition to the final destination
-		ActorList.Add(World->GameState);
-		// keep the game session state until we transition to the final destination
-		ActorList.Add(GameSession);
-
-		// If adding in this section best to increase the literal above for the ActorsToAddCount
-	}
-}
-
-void AGameMode::SetBandwidthLimit(float AsyncIOBandwidthLimit)
-{
-	GAsyncIOBandwidthLimit = AsyncIOBandwidthLimit;
-}
-
-FString AGameMode::InitNewPlayer(APlayerController* NewPlayerController, const TSharedPtr<const FUniqueNetId>& UniqueId, const FString& Options, const FString& Portal)
-{
-	FUniqueNetIdRepl UniqueIdRepl(UniqueId);
-	return InitNewPlayer(NewPlayerController, UniqueIdRepl, Options, Portal);
-}
-
-FString AGameMode::InitNewPlayer(APlayerController* NewPlayerController, const FUniqueNetIdRepl& UniqueId, const FString& Options, const FString& Portal)
-{
-	check(NewPlayerController);
-
-	FString ErrorMessage;
-
-	// Register the player with the session
-	GameSession->RegisterPlayer(NewPlayerController, UniqueId.GetUniqueNetId(), UGameplayStatics::HasOption(Options, TEXT("bIsFromInvite")));
-
-	// Init player's name
-	FString InName = UGameplayStatics::ParseOption(Options, TEXT("Name")).Left(20);
-	if (InName.IsEmpty())
-	{
-		InName = FString::Printf(TEXT("%s%i"), *DefaultPlayerName.ToString(), NewPlayerController->PlayerState->PlayerId);
-	}
-
-	ChangeName(NewPlayerController, InName, false);
-
-	// Find a starting spot
-	AActor* const StartSpot = FindPlayerStart(NewPlayerController, Portal);
-	if (StartSpot != NULL)
-	{
-		// Set the player controller / camera in this new location
-		FRotator InitialControllerRot = StartSpot->GetActorRotation();
-		InitialControllerRot.Roll = 0.f;
-		NewPlayerController->SetInitialLocationAndRotation(StartSpot->GetActorLocation(), InitialControllerRot);
-		NewPlayerController->StartSpot = StartSpot;
-	}
-	else
-	{
-		ErrorMessage = FString::Printf(TEXT("Failed to find PlayerStart"));
-	}
-
-	return ErrorMessage;
-}
-
-bool AGameMode::MustSpectate_Implementation(APlayerController* NewPlayerController) const
-{
-	return NewPlayerController->PlayerState->bOnlySpectator;
-}
-
-APlayerController* AGameMode::Login(UPlayer* NewPlayer, ENetRole InRemoteRole, const FString& Portal, const FString& Options, const TSharedPtr<const FUniqueNetId>& UniqueId, FString& ErrorMessage)
-{
-	FUniqueNetIdRepl UniqueIdRepl(UniqueId);
-	return Login(NewPlayer, InRemoteRole, Portal, Options, UniqueIdRepl, ErrorMessage);
-}
-
-APlayerController* AGameMode::Login(UPlayer* NewPlayer, ENetRole InRemoteRole, const FString& Portal, const FString& Options, const FUniqueNetIdRepl& UniqueId, FString& ErrorMessage)
-{
-	ErrorMessage = GameSession->ApproveLogin(Options);
-	if (!ErrorMessage.IsEmpty())
-	{
-		return nullptr;
-	}
-
-	APlayerController* NewPlayerController = SpawnPlayerController(InRemoteRole, FVector::ZeroVector, FRotator::ZeroRotator);
-
-	// Handle spawn failure.
-	if (NewPlayerController == nullptr)
-	{
-		UE_LOG(LogGameMode, Log, TEXT("Couldn't spawn player controller of class %s"), PlayerControllerClass ? *PlayerControllerClass->GetName() : TEXT("NULL"));
-		ErrorMessage = FString::Printf(TEXT("Failed to spawn player controller"));
-		return nullptr;
-	}
-
-	// Customize incoming player based on URL options
-	ErrorMessage = InitNewPlayer(NewPlayerController, UniqueId, Options, Portal);
-	if (!ErrorMessage.IsEmpty())
-	{
-		return nullptr;
-	}
-
-	// Set up spectating
-	bool bSpectator = FCString::Stricmp(*UGameplayStatics::ParseOption(Options, TEXT("SpectatorOnly")), TEXT("1")) == 0;
-	if (bSpectator || MustSpectate(NewPlayerController))
-	{
-		NewPlayerController->StartSpectatingOnly();
-		return NewPlayerController;
-	}
-
-	return NewPlayerController;
-}
-
-
-void AGameMode::DisplayDebug(UCanvas* Canvas, const FDebugDisplayInfo& DebugDisplay, float& YL, float& YPos)
-{
-}
-
-
-void AGameMode::Reset()
-{
-	Super::Reset();
-	InitGameState();
-}
-
-
-bool AGameMode::ShouldReset_Implementation(AActor* ActorToReset)
-{
-	return true;
-}
-
 void AGameMode::PlayerSwitchedToSpectatorOnly(APlayerController* PC)
 {
 	RemovePlayerControllerFromPlayerCount(PC);
@@ -1150,7 +488,7 @@ void AGameMode::RemovePlayerControllerFromPlayerCount(APlayerController* PC)
 {
 	if ( PC )
 	{
-		if ( PC->PlayerState->bOnlySpectator )
+		if ( MustSpectate(PC) )
 		{
 			NumSpectators--;
 		}
@@ -1173,259 +511,25 @@ int32 AGameMode::GetNumPlayers()
 	return NumPlayers + NumTravellingPlayers;
 }
 
-
-
-void AGameMode::ClearPause()
-{
-	if ( !AllowPausing() && Pausers.Num() > 0 )
-	{
-		UE_LOG(LogGameMode, Log, TEXT("Clearing list of UnPause delegates for %s because game type is not pauseable"), *GetFName().ToString());
-		Pausers.Empty();
-	}
-
-	for (int32 Index = 0; Index < Pausers.Num(); Index++)
-	{
-		FCanUnpause CanUnpauseCriteriaMet = Pausers[Index];
-		if( CanUnpauseCriteriaMet.IsBound() )
-		{
-			bool Result = CanUnpauseCriteriaMet.Execute();
-			if (Result)
-			{
-				Pausers.RemoveAt(Index--,1);
-			}
-		}
-		else
-		{
-			Pausers.RemoveAt(Index--,1);
-		}
-	}
-
-	// Clear the pause state if the list is empty
-	if ( Pausers.Num() == 0 )
-	{
-		GetWorldSettings()->Pauser = NULL;
-	}
-}
-
-bool AGameMode::GrabOption( FString& Options, FString& Result )
-{
-	return UGameplayStatics::GrabOption(Options, Result);
-}
-
-void AGameMode::GetKeyValue( const FString& Pair, FString& Key, FString& Value )
-{
-	return UGameplayStatics::GetKeyValue(Pair, Key, Value);
-}
-
-FString AGameMode::ParseOption( FString Options, const FString& InKey )
-{
-	return UGameplayStatics::ParseOption(Options, InKey);
-}
-
-bool AGameMode::HasOption( FString Options, const FString& InKey )
-{
-	return UGameplayStatics::HasOption(Options, InKey);
-}
-
-int32 AGameMode::GetIntOption( const FString& Options, const FString& ParseString, int32 CurrentValue)
-{
-	return UGameplayStatics::GetIntOption(Options, ParseString, CurrentValue);
-}
-
-FString AGameMode::GetDefaultGameClassPath(const FString& MapName, const FString& Options, const FString& Portal) const
-{
-	return GetClass()->GetPathName();
-}
-
-TSubclassOf<AGameMode> AGameMode::GetGameModeClass(const FString& MapName, const FString& Options, const FString& Portal) const
-{
-	return GetClass();
-}
-
-TSubclassOf<AGameMode> AGameMode::SetGameMode(const FString& MapName, const FString& Options, const FString& Portal)
-{
-	return GetGameModeClass(MapName, Options, Portal);
-}
-
-TSubclassOf<AGameSession> AGameMode::GetGameSessionClass() const
-{
-	return AGameSession::StaticClass();
-}
-
-void AGameMode::NotifyPendingConnectionLost() {}
-
-APlayerController* AGameMode::ProcessClientTravel( FString& FURL, FGuid NextMapGuid, bool bSeamless, bool bAbsolute )
-{
-	// We call PreClientTravel directly on any local PlayerPawns (ie listen server)
-	APlayerController* LocalPlayerController = NULL;
-	for( FConstPlayerControllerIterator Iterator = GetWorld()->GetPlayerControllerIterator(); Iterator; ++Iterator )
-	{
-		APlayerController* PlayerController = *Iterator;
-		if ( Cast<UNetConnection>(PlayerController->Player) != NULL )
-		{
-			// remote player
-			PlayerController->ClientTravel(FURL, TRAVEL_Relative, bSeamless, NextMapGuid);
-		}
-		else
-		{
-			// local player
-			LocalPlayerController = PlayerController;
-			PlayerController->PreClientTravel(FURL, bAbsolute ? TRAVEL_Absolute : TRAVEL_Relative, bSeamless);
-		}
-	}
-
-	return LocalPlayerController;
-}
-
-void AGameMode::PreLogin(const FString& Options, const FString& Address, const TSharedPtr<const FUniqueNetId>& UniqueId, FString& ErrorMessage)
+int32 AGameMode::GetNumSpectators()
 {
-	FUniqueNetIdRepl UniqueIdRepl(UniqueId);
-	PreLogin(Options, Address, UniqueIdRepl, ErrorMessage);
+	return NumSpectators;
 }
 
-void AGameMode::PreLogin(const FString& Options, const FString& Address, const FUniqueNetIdRepl& UniqueId, FString& ErrorMessage)
-{
-	ErrorMessage = GameSession->ApproveLogin(Options);
-}
-
-APlayerController* AGameMode::SpawnPlayerController(ENetRole InRemoteRole, FVector const& SpawnLocation, FRotator const& SpawnRotation)
-{
-	FActorSpawnParameters SpawnInfo;
-	SpawnInfo.Instigator = Instigator;	
-	SpawnInfo.ObjectFlags |= RF_Transient;	// We never want to save player controllers into a map
-	SpawnInfo.bDeferConstruction = true;
-	APlayerController* NewPC = GetWorld()->SpawnActor<APlayerController>(PlayerControllerClass, SpawnLocation, SpawnRotation, SpawnInfo);
-	if (NewPC)
-	{
-		if (InRemoteRole == ROLE_SimulatedProxy)
-		{
-			// This is a local player because it has no authority/autonomous remote role
-			NewPC->SetAsLocalPlayerController();
-		}
-		
-		UGameplayStatics::FinishSpawningActor(NewPC, FTransform(SpawnRotation, SpawnLocation));
-	}
-
-	return NewPC;
-}
-
-UClass* AGameMode::GetDefaultPawnClassForController_Implementation(AController* InController)
-{
-	return DefaultPawnClass;
-}
-
-APawn* AGameMode::SpawnDefaultPawnFor_Implementation(AController* NewPlayer, AActor* StartSpot)
-{
-	// don't allow pawn to be spawned with any pitch or roll
-	FRotator StartRotation(ForceInit);
-	StartRotation.Yaw = StartSpot->GetActorRotation().Yaw;
-	FVector StartLocation = StartSpot->GetActorLocation();
-
-	FActorSpawnParameters SpawnInfo;
-	SpawnInfo.Instigator = Instigator;	
-	SpawnInfo.ObjectFlags |= RF_Transient;	// We never want to save default player pawns into a map
-	UClass* PawnClass = GetDefaultPawnClassForController(NewPlayer);
-	APawn* ResultPawn = GetWorld()->SpawnActor<APawn>(PawnClass, StartLocation, StartRotation, SpawnInfo );
-	if ( ResultPawn == NULL )
-	{
-		UE_LOG(LogGameMode, Warning, TEXT("Couldn't spawn Pawn of type %s at %s"), *GetNameSafe(PawnClass), *StartSpot->GetName());
-	}
-	return ResultPawn;
-}
-
-void AGameMode::ReplicateStreamingStatus(APlayerController* PC)
-{
-	UWorld* MyWorld = GetWorld();
-
-	if (MyWorld->GetWorldSettings()->bUseClientSideLevelStreamingVolumes)
-	{
-		// Client will itself decide what to stream
-		return;
-	}
-
-	// don't do this for local players or players after the first on a splitscreen client
-	if (Cast<ULocalPlayer>(PC->Player) == nullptr && Cast<UChildConnection>(PC->Player) == nullptr)
-	{
-		// if we've loaded levels via CommitMapChange() that aren't normally in the StreamingLevels array, tell the client about that
-		if (MyWorld->CommittedPersistentLevelName != NAME_None)
-		{
-			PC->ClientPrepareMapChange(MyWorld->CommittedPersistentLevelName, true, true);
-			// tell the client to commit the level immediately
-			PC->ClientCommitMapChange();
-		}
-
-		if (MyWorld->StreamingLevels.Num() > 0)
-		{
-			// Tell the player controller the current streaming level status
-			for (int32 LevelIndex = 0; LevelIndex < MyWorld->StreamingLevels.Num(); LevelIndex++)
-			{
-				// streamingServer
-				ULevelStreaming* TheLevel = MyWorld->StreamingLevels[LevelIndex];
-
-				if( TheLevel != NULL )
-				{
-					const ULevel* LoadedLevel = TheLevel->GetLoadedLevel();
-					
-					UE_LOG(LogGameMode, Log, TEXT("levelStatus: %s %i %i %i %s %i"),
-						*TheLevel->GetWorldAssetPackageName(),
-						TheLevel->bShouldBeVisible,
-						LoadedLevel && LoadedLevel->bIsVisible,
-						TheLevel->bShouldBeLoaded,
-						*GetNameSafe(LoadedLevel),
-						TheLevel->bHasLoadRequestPending);
-
-					PC->ClientUpdateLevelStreamingStatus(
-						TheLevel->GetWorldAssetPackageFName(),
-						TheLevel->bShouldBeLoaded,
-						TheLevel->bShouldBeVisible,
-						TheLevel->bShouldBlockOnLoad,
-						TheLevel->LevelLODIndex);
-				}
-			}
-			PC->ClientFlushLevelStreaming();
-		}
-
-		// if we're preparing to load different levels using PrepareMapChange() inform the client about that now
-		if (MyWorld->PreparingLevelNames.Num() > 0)
-		{
-			for (int32 LevelIndex = 0; LevelIndex < MyWorld->PreparingLevelNames.Num(); LevelIndex++)
-			{
-				PC->ClientPrepareMapChange(MyWorld->PreparingLevelNames[LevelIndex], LevelIndex == 0, LevelIndex == MyWorld->PreparingLevelNames.Num() - 1);
-			}
-			// DO NOT commit these changes yet - we'll send that when we're done preparing them
-		}
-	}
-}
-
-void AGameMode::GenericPlayerInitialization(AController* C)
+void AGameMode::StartNewPlayer(APlayerController* NewPlayer)
 {
-	APlayerController* PC = Cast<APlayerController>(C);
-	if (PC != NULL)
-	{
-		// Notify the game that we can now be muted and mute others
-		UpdateGameplayMuteList(PC);
 
-		ReplicateStreamingStatus(PC);
-	}
 }
 
-void AGameMode::StartNewPlayer(APlayerController* NewPlayer)
+void AGameMode::HandleStartingNewPlayer_Implementation(APlayerController* NewPlayer)
 {
-	// tell client what hud class to use
-	NewPlayer->ClientSetHUD(HUDClass);
-
 	// If players should start as spectators, leave them in the spectator state
-	if (!bStartPlayersAsSpectators && !NewPlayer->PlayerState->bOnlySpectator)
+	if (!bStartPlayersAsSpectators && !MustSpectate(NewPlayer))
 	{
 		// If match is in progress, start the player
-		if (IsMatchInProgress())
+		if (IsMatchInProgress() && PlayerCanRestart(NewPlayer))
 		{
 			RestartPlayer(NewPlayer);
-
-			if (NewPlayer->GetPawn() != NULL)
-			{
-				NewPlayer->GetPawn()->ClientSetRotation(NewPlayer->GetPawn()->GetActorRotation());
-			}
 		}
 		// Check to see if we should start right away, avoids a one frame lag in single player games
 		else if (GetMatchState() == MatchState::WaitingToStart)
@@ -1439,17 +543,14 @@ void AGameMode::StartNewPlayer(APlayerController* NewPlayer)
 	}
 }
 
-bool AGameMode::CanSpectate_Implementation( APlayerController* Viewer, APlayerState* ViewTarget )
-{
-	return true;
-}
-
-void AGameMode::ChangeName( AController* Other, const FString& S, bool bNameChange )
+bool AGameMode::PlayerCanRestart_Implementation(APlayerController* Player)
 {
-	if( Other && !S.IsEmpty() )
+	if (!IsMatchInProgress())
 	{
-		Other->PlayerState->SetPlayerName(S);
+		return false;
 	}
+
+	return Super::PlayerCanRestart_Implementation(Player);
 }
 
 void AGameMode::SendPlayer( APlayerController* aPlayer, const FString& FURL )
@@ -1462,14 +563,19 @@ bool AGameMode::GetTravelType()
 	return false;
 }
 
+void AGameMode::Say(const FString& Msg)
+{
+	Broadcast(nullptr, Msg, NAME_None);
+}
+
 void AGameMode::Broadcast( AActor* Sender, const FString& Msg, FName Type )
 {
-	APlayerState* SenderPlayerState = NULL;
-	if ( Cast<APawn>(Sender) != NULL )
+	APlayerState* SenderPlayerState = nullptr;
+	if ( Cast<APawn>(Sender) != nullptr)
 	{
 		SenderPlayerState = Cast<APawn>(Sender)->PlayerState;
 	}
-	else if ( Cast<AController>(Sender) != NULL )
+	else if ( Cast<AController>(Sender) != nullptr)
 	{
 		SenderPlayerState = Cast<AController>(Sender)->PlayerState;
 	}
@@ -1489,103 +595,18 @@ void AGameMode::BroadcastLocalized( AActor* Sender, TSubclassOf<ULocalMessage> M
 	}
 }
 
-bool AGameMode::ShouldSpawnAtStartSpot(AController* Player)
-{
-	return ( Player != NULL && Player->StartSpot != NULL );
-}
-
-AActor* AGameMode::ChoosePlayerStart_Implementation( AController* Player )
-{
-	// Choose a player start
-	APlayerStart* FoundPlayerStart = NULL;
-	UClass* PawnClass = GetDefaultPawnClassForController(Player);
-	APawn* PawnToFit = PawnClass ? PawnClass->GetDefaultObject<APawn>() : nullptr;
-	TArray<APlayerStart*> UnOccupiedStartPoints;
-	TArray<APlayerStart*> OccupiedStartPoints;
-	for (TActorIterator<APlayerStart> It(GetWorld()); It; ++It)
-	{
-		APlayerStart* PlayerStart = *It;
-
-		if (PlayerStart->IsA<APlayerStartPIE>())
-		{
-			// Always prefer the first "Play from Here" PlayerStart, if we find one while in PIE mode
-			FoundPlayerStart = PlayerStart;
-			break;
-		}
-		else
-		{
-			FVector ActorLocation = PlayerStart->GetActorLocation();
-			const FRotator ActorRotation = PlayerStart->GetActorRotation();
-			if (!GetWorld()->EncroachingBlockingGeometry(PawnToFit, ActorLocation, ActorRotation))
-			{
-				UnOccupiedStartPoints.Add(PlayerStart);
-			}
-			else if (GetWorld()->FindTeleportSpot(PawnToFit, ActorLocation, ActorRotation))
-			{
-				OccupiedStartPoints.Add(PlayerStart);
-			}
-		}
-	}
-	if (FoundPlayerStart == NULL)
-	{
-		if (UnOccupiedStartPoints.Num() > 0)
-		{
-			FoundPlayerStart = UnOccupiedStartPoints[FMath::RandRange(0, UnOccupiedStartPoints.Num() - 1)];
-		}
-		else if (OccupiedStartPoints.Num() > 0)
-		{
-			FoundPlayerStart = OccupiedStartPoints[FMath::RandRange(0, OccupiedStartPoints.Num() - 1)];
-		}
-	}
-	return FoundPlayerStart;
-}
-
-bool AGameMode::PlayerCanRestart_Implementation( APlayerController* Player )
-{
-	if (!IsMatchInProgress() || Player == NULL || Player->IsPendingKillPending())
-	{
-		return false;
-	}
-
-	// Ask the player controller if it's ready to restart as well
-	return Player->CanRestartPlayer();
-}
-
-void AGameMode::UpdateGameplayMuteList( APlayerController* aPlayer )
-{
-	if (aPlayer)
-	{
-		aPlayer->MuteList.bHasVoiceHandshakeCompleted = true;
-		aPlayer->ClientVoiceHandshakeComplete();
-	}
-}
-
-bool AGameMode::AllowCheats(APlayerController* P)
-{
-	return ( GetNetMode() == NM_Standalone || GIsEditor ); // Always allow cheats in editor (PIE now supports networking)
-}
-
-bool AGameMode::AllowPausing( APlayerController* PC )
-{
-	return bPauseable || GetNetMode() == NM_Standalone;
-}
-
-void AGameMode::PreCommitMapChange(const FString& PreviousMapName, const FString& NextMapName) {}
-
-void AGameMode::PostCommitMapChange() {}
-
 void AGameMode::AddInactivePlayer(APlayerState* PlayerState, APlayerController* PC)
 {
 	check(PlayerState)
 	UWorld* LocalWorld = GetWorld();
 	// don't store if it's an old PlayerState from the previous level or if it's a spectator... or if we are shutting down
-	if (!PlayerState->bFromPreviousLevel && !PlayerState->bOnlySpectator && !LocalWorld->bIsTearingDown)
+	if (!PlayerState->bFromPreviousLevel && !MustSpectate(PC) && !LocalWorld->bIsTearingDown)
 	{
 		APlayerState* const NewPlayerState = PlayerState->Duplicate();
 		if (NewPlayerState)
 		{
 			// Side effect of Duplicate() adding PlayerState to PlayerArray (see APlayerState::PostInitializeComponents)
-			LocalWorld->GameState->RemovePlayerState(NewPlayerState);
+			GameState->RemovePlayerState(NewPlayerState);
 
 			// make PlayerState inactive
 			NewPlayerState->SetReplicates(false);
@@ -1646,12 +667,11 @@ void AGameMode::AddInactivePlayer(APlayerState* PlayerState, APlayerController*
 	PlayerState->OnDeactivated();
 }
 
-
 bool AGameMode::FindInactivePlayer(APlayerController* PC)
 {
 	check(PC && PC->PlayerState);
 	// don't bother for spectators
-	if (PC->PlayerState->bOnlySpectator)
+	if (MustSpectate(PC))
 	{
 		return false;
 	}
@@ -1669,7 +689,7 @@ bool AGameMode::FindInactivePlayer(APlayerController* PC)
 	for (int32 i=0; i < InactivePlayerArray.Num(); i++)
 	{
 		APlayerState* CurrentPlayerState = InactivePlayerArray[i];
-		if ( (CurrentPlayerState == NULL) || CurrentPlayerState->IsPendingKill() )
+		if ( (CurrentPlayerState == nullptr) || CurrentPlayerState->IsPendingKill() )
 		{
 			InactivePlayerArray.RemoveAt(i,1);
 			i--;
@@ -1684,12 +704,12 @@ bool AGameMode::FindInactivePlayer(APlayerController* PC)
 			PC->PlayerState->SetReplicates(true);
 			PC->PlayerState->SetLifeSpan(0.0f);
 			OverridePlayerState(PC, OldPlayerState);
-			GetWorld()->GameState->AddPlayerState(PC->PlayerState);
+			GameState->AddPlayerState(PC->PlayerState);
 			InactivePlayerArray.RemoveAt(i, 1);
 			OldPlayerState->bIsInactive = true;
-			// Set the uniqueId to NULL so it will not kill the player's registration 
+			// Set the uniqueId to nullptr so it will not kill the player's registration 
 			// in UnregisterPlayerWithSession()
-			OldPlayerState->SetUniqueId(NULL);
+			OldPlayerState->SetUniqueId(nullptr);
 			OldPlayerState->Destroy();
 			PC->PlayerState->OnReactivated();
 			return true;
@@ -1704,9 +724,27 @@ void AGameMode::OverridePlayerState(APlayerController* PC, APlayerState* OldPlay
 	PC->PlayerState->DispatchOverrideWith(OldPlayerState);
 }
 
+
+bool AGameMode::CanServerTravel(const FString& URL, bool bAbsolute)
+{
+	if (!Super::CanServerTravel(URL, bAbsolute))
+	{
+		return false;
+	}
+
+	// Check for an error in the server's connection
+	if (GetMatchState() == MatchState::Aborted)
+	{
+		UE_LOG(LogGameMode, Log, TEXT("Not traveling because of network error"));
+		return false;
+	}
+
+	return true;
+}
+
 void AGameMode::PostSeamlessTravel()
 {
-	if ( GameSession != NULL )
+	if (GameSession)
 	{
 		GameSession->PostSeamlessTravel();
 	}
@@ -1726,13 +764,13 @@ void AGameMode::PostSeamlessTravel()
 		if (Controller->PlayerState)
 		{
 			APlayerController* PlayerController = Cast<APlayerController>(Controller);
-			if (PlayerController == NULL)
+			if (PlayerController == nullptr)
 			{
 				HandleSeamlessTravelPlayer(Controller);
 			}
 			else
 			{
-				if (Controller->PlayerState->bOnlySpectator)
+				if (MustSpectate(PlayerController))
 				{
 					// The spectator count must be incremented here, instead of in HandleSeamlessTravelPlayer,
 					// as otherwise spectators can 'hide' from player counters, by making HasClientLoadedCurrentWorld return false
@@ -1751,49 +789,45 @@ void AGameMode::PostSeamlessTravel()
 	}
 }
 
-void AGameMode::MatineeCancelled()
-{
-}
-
-FString AGameMode::StaticGetFullGameClassName(FString const& Str)
+bool AGameMode::IsHandlingReplays()
 {
-	// look to see if this should be remapped from a shortname to full class name
-	const AGameMode* const DefaultGameMode = GetDefault<AGameMode>();
-	if (DefaultGameMode)
+	// If we're running in PIE, don't record demos
+	if (GetWorld() != nullptr && GetWorld()->IsPlayInEditor())
 	{
-		int32 const NumAliases = DefaultGameMode->GameModeClassAliases.Num();
-		for (int32 Idx=0; Idx<NumAliases; ++Idx)
-		{
-			const FGameClassShortName& Alias = DefaultGameMode->GameModeClassAliases[Idx];
-			if ( Str == Alias.ShortName )
-			{
-				// switch GameClassName to the full name
-				return Alias.GameClassName;
-			}
-		}
+		return false;
 	}
 
-	return Str;
+	return bHandleDedicatedServerReplays && GetNetMode() == ENetMode::NM_DedicatedServer;
 }
 
-FString AGameMode::GetRedirectURL(const FString& MapName) const
+void AGameMode::SetBandwidthLimit(float AsyncIOBandwidthLimit)
 {
-	return FString();
+	GAsyncIOBandwidthLimit = AsyncIOBandwidthLimit;
 }
-void AGameMode::GameWelcomePlayer(UNetConnection* Connection, FString& RedirectURL)
+
+void AGameMode::MatineeCancelled() {}
+
+void AGameMode::PreCommitMapChange(const FString& PreviousMapName, const FString& NextMapName) {}
+
+void AGameMode::PostCommitMapChange() {}
+
+void AGameMode::NotifyPendingConnectionLost() {}
+
+void AGameMode::HandleDisconnect(UWorld* InWorld, UNetDriver* NetDriver)
 {
-	PRAGMA_DISABLE_DEPRECATION_WARNINGS
-	RedirectURL = GetRedirectURL(Connection->ClientWorldPackageName.ToString());
-	PRAGMA_ENABLE_DEPRECATION_WARNINGS
+	AbortMatch();
 }
 
-bool AGameMode::IsHandlingReplays()
+bool AGameMode::SpawnPlayerFromSimulate(const FVector& NewLocation, const FRotator& NewRotation)
 {
-	// If we're running in PIE, don't record demos
-	if (GetWorld() != nullptr && GetWorld()->IsPlayInEditor())
+#if WITH_EDITOR
+	APlayerController* PC = GetGameInstance()->GetFirstLocalPlayerController();
+	if (PC != nullptr)
 	{
-		return false;
+		RemovePlayerControllerFromPlayerCount(PC);
+		NumPlayers++;
 	}
+#endif
 
-	return bHandleDedicatedServerReplays && GetNetMode() == ENetMode::NM_DedicatedServer;
-}
+	return Super::SpawnPlayerFromSimulate(NewLocation, NewRotation);
+}
\ No newline at end of file
