Copying //UE4/Dev-Framework to //UE4/Dev-Main (Source: //UE4/Dev-Framework @ 3136612)

diff --git a/Engine/Source/Runtime/Engine/Private/GameModeBase.cpp b/Engine/Source/Runtime/Engine/Private/GameModeBase.cpp
new file mode 100644
--- /dev/null
+++ b/Engine/Source/Runtime/Engine/Private/GameModeBase.cpp
@@ -0,0 +1,1312 @@
+// Copyright 1998-2016 Epic Games, Inc. All Rights Reserved.
+
+#include "EnginePrivate.h"
+#include "GameFramework/GameNetworkManager.h"
+#include "Matinee/MatineeActor.h"
+#include "Engine/LevelScriptActor.h"
+#include "Net/OnlineEngineInterface.h"
+#include "GameFramework/GameStateBase.h"
+#include "GameFramework/DefaultPawn.h"
+#include "GameFramework/SpectatorPawn.h"
+#include "GameFramework/HUD.h"
+#include "GameFramework/PlayerState.h"
+#include "GameFramework/GameSession.h"
+#include "GameFramework/GameModeBase.h"
+#include "Engine/ChildConnection.h"
+#include "Engine/PlayerStartPIE.h"
+
+#if WITH_EDITOR
+	#include "IMovieSceneCapture.h"
+	#include "MovieSceneCaptureModule.h"
+	#include "MovieSceneCaptureSettings.h"
+#endif
+
+DEFINE_LOG_CATEGORY(LogGameMode);
+
+// Statically declared events for plugins to use
+FGameModeEvents::FGameModePostLoginEvent FGameModeEvents::GameModePostLoginEvent;
+FGameModeEvents::FGameModeLogoutEvent FGameModeEvents::GameModeLogoutEvent;
+
+AGameModeBase::AGameModeBase(const FObjectInitializer& ObjectInitializer)
+	: Super(ObjectInitializer.DoNotCreateDefaultSubobject(TEXT("Sprite")))
+{
+	bNetLoadOnClient = false;
+	bPauseable = true;
+	bStartPlayersAsSpectators = false;
+
+	DefaultPawnClass = ADefaultPawn::StaticClass();
+	PlayerControllerClass = APlayerController::StaticClass();
+	PlayerStateClass = APlayerState::StaticClass();
+	GameStateClass = AGameStateBase::StaticClass();
+	HUDClass = AHUD::StaticClass();
+	GameSessionClass = AGameSession::StaticClass();
+	SpectatorClass = ASpectatorPawn::StaticClass();
+	ReplaySpectatorPlayerControllerClass = APlayerController::StaticClass();
+}
+
+void AGameModeBase::InitGame(const FString& MapName, const FString& Options, FString& ErrorMessage)
+{
+	UWorld* World = GetWorld();
+
+	// Save Options for future use
+	OptionsString = Options;
+
+	FActorSpawnParameters SpawnInfo;
+	SpawnInfo.Instigator = Instigator;
+	SpawnInfo.ObjectFlags |= RF_Transient;	// We never want to save game sessions into a map
+	GameSession = World->SpawnActor<AGameSession>(GetGameSessionClass(), SpawnInfo);
+	GameSession->InitOptions(Options);
+
+	if (GetNetMode() != NM_Standalone)
+	{
+		// Attempt to login, returning true means an async login is in flight
+		if (!UOnlineEngineInterface::Get()->DoesSessionExist(World, GameSession->SessionName) && 
+			!GameSession->ProcessAutoLogin())
+		{
+			GameSession->RegisterServer();
+		}
+	}
+}
+
+void AGameModeBase::InitGameState()
+{
+	GameState->GameModeClass = GetClass();
+	GameState->ReceivedGameModeClass();
+
+	GameState->SpectatorClass = SpectatorClass;
+	GameState->ReceivedSpectatorClass();
+}
+
+void AGameModeBase::PreInitializeComponents()
+{
+	Super::PreInitializeComponents();
+
+	FActorSpawnParameters SpawnInfo;
+	SpawnInfo.Instigator = Instigator;
+	SpawnInfo.ObjectFlags |= RF_Transient;	// We never want to save game states or network managers into a map									
+											
+	// Fallback to default GameState if none was specified.
+	if (GameStateClass == nullptr)
+	{
+		UE_LOG(LogGameMode, Warning, TEXT("No GameStateClass was specified in %s (%s)"), *GetName(), *GetClass()->GetName());
+		GameStateClass = AGameStateBase::StaticClass();
+	}
+
+	GameState = GetWorld()->SpawnActor<AGameStateBase>(GameStateClass, SpawnInfo);
+	GetWorld()->SetGameState(GameState);
+	if (GameState)
+	{
+		GameState->AuthorityGameMode = this;
+	}
+
+	// Only need NetworkManager for servers in net games
+	GetWorld()->NetworkManager = GetWorldSettings()->GameNetworkManagerClass ? GetWorld()->SpawnActor<AGameNetworkManager>(GetWorldSettings()->GameNetworkManagerClass, SpawnInfo) : nullptr;
+
+	InitGameState();
+}
+
+TSubclassOf<AGameSession> AGameModeBase::GetGameSessionClass() const
+{
+	return AGameSession::StaticClass();
+}
+
+UClass* AGameModeBase::GetDefaultPawnClassForController_Implementation(AController* InController)
+{
+	return DefaultPawnClass;
+}
+
+int32 AGameModeBase::GetNumPlayers()
+{
+	int32 PlayerCount = 0;
+	for (FConstPlayerControllerIterator Iterator = GetWorld()->GetPlayerControllerIterator(); Iterator; ++Iterator)
+	{
+		APlayerController* PlayerActor = *Iterator;
+		if (PlayerActor->PlayerState && !MustSpectate(PlayerActor))
+		{
+			PlayerCount++;
+		}
+	}
+	return PlayerCount;
+}
+
+int32 AGameModeBase::GetNumSpectators()
+{
+	int32 PlayerCount = 0;
+	for (FConstPlayerControllerIterator Iterator = GetWorld()->GetPlayerControllerIterator(); Iterator; ++Iterator)
+	{
+		APlayerController* PlayerActor = *Iterator;
+		if (PlayerActor->PlayerState && MustSpectate(PlayerActor))
+		{
+			PlayerCount++;
+		}
+	}
+	return PlayerCount;
+}
+
+void AGameModeBase::StartPlay()
+{
+	GameState->HandleBeginPlay();
+}
+
+bool AGameModeBase::HasMatchStarted() const
+{
+	return GameState->HasMatchStarted();
+}
+
+bool AGameModeBase::SetPause(APlayerController* PC, FCanUnpause CanUnpauseDelegate /*= FCanUnpause()*/)
+{
+	if (AllowPausing(PC))
+	{
+		// Add it for querying
+		Pausers.Add(CanUnpauseDelegate);
+
+		// Let the first one in "own" the pause state
+		AWorldSettings * WorldSettings = GetWorldSettings();
+		if (WorldSettings->Pauser == nullptr)
+		{
+			WorldSettings->Pauser = PC->PlayerState;
+		}
+		return true;
+	}
+	return false;
+}
+
+bool AGameModeBase::ClearPause()
+{
+	bool bPauseCleared = false;
+
+	if (!AllowPausing() && Pausers.Num() > 0)
+	{
+		UE_LOG(LogGameMode, Log, TEXT("Clearing list of UnPause delegates for %s because game type is not pauseable"), *GetFName().ToString());
+		Pausers.Reset();
+		bPauseCleared = true;
+	}
+
+	for (int32 Index = Pausers.Num() - 1; Index >= 0; --Index)
+	{
+		FCanUnpause CanUnpauseCriteriaMet = Pausers[Index];
+		if (CanUnpauseCriteriaMet.IsBound())
+		{
+			const bool bResult = CanUnpauseCriteriaMet.Execute();
+			if (bResult)
+			{
+				Pausers.RemoveAtSwap(Index, 1, false);
+				bPauseCleared = true;
+			}
+		}
+		else
+		{
+			Pausers.RemoveAtSwap(Index, 1, false);
+			bPauseCleared = true;
+		}
+	}
+
+	// Clear the pause state if the list is empty
+	if (Pausers.Num() == 0)
+	{
+		GetWorldSettings()->Pauser = nullptr;
+	}
+
+	return bPauseCleared;
+}
+
+void AGameModeBase::ForceClearUnpauseDelegates(AActor* PauseActor)
+{
+	if (PauseActor != nullptr)
+	{
+		bool bUpdatePausedState = false;
+		for (int32 PauserIdx = Pausers.Num() - 1; PauserIdx >= 0; PauserIdx--)
+		{
+			FCanUnpause& CanUnpauseDelegate = Pausers[PauserIdx];
+			if (CanUnpauseDelegate.GetUObject() == PauseActor)
+			{
+				Pausers.RemoveAt(PauserIdx);
+				bUpdatePausedState = true;
+			}
+		}
+
+		// If we removed some CanUnpause delegates, we may be able to unpause the game now
+		if (bUpdatePausedState)
+		{
+			ClearPause();
+		}
+
+		APlayerController* PC = Cast<APlayerController>(PauseActor);
+		AWorldSettings * WorldSettings = GetWorldSettings();
+		if (PC != nullptr && PC->PlayerState != nullptr && WorldSettings != nullptr && WorldSettings->Pauser == PC->PlayerState)
+		{
+			// Try to find another player to be the worldsettings's Pauser
+			for (FConstPlayerControllerIterator Iterator = GetWorld()->GetPlayerControllerIterator(); Iterator; ++Iterator)
+			{
+				APlayerController* Player = *Iterator;
+				if (Player->PlayerState != nullptr
+					&&	Player->PlayerState != PC->PlayerState
+					&& !Player->IsPendingKillPending() && !Player->PlayerState->IsPendingKillPending())
+				{
+					WorldSettings->Pauser = Player->PlayerState;
+					break;
+				}
+			}
+
+			// If it's still pointing to the original player's PlayerState, clear it completely
+			if (WorldSettings->Pauser == PC->PlayerState)
+			{
+				WorldSettings->Pauser = nullptr;
+			}
+		}
+	}
+}
+
+bool AGameModeBase::AllowPausing(APlayerController* PC /*= nullptr*/)
+{
+	return bPauseable || GetNetMode() == NM_Standalone;
+}
+
+void AGameModeBase::Reset()
+{
+	Super::Reset();
+	InitGameState();
+}
+
+bool AGameModeBase::ShouldReset_Implementation(AActor* ActorToReset)
+{
+	return true;
+}
+
+void AGameModeBase::ResetLevel()
+{
+	UE_LOG(LogGameMode, Log, TEXT("Reset %s"), *GetName());
+
+	// Reset ALL controllers first
+	for (FConstControllerIterator Iterator = GetWorld()->GetControllerIterator(); Iterator; ++Iterator)
+	{
+		AController* Controller = *Iterator;
+		APlayerController* PlayerController = Cast<APlayerController>(Controller);
+		if (PlayerController)
+		{
+			PlayerController->ClientReset();
+		}
+		Controller->Reset();
+	}
+
+	// Reset all actors (except controllers, the GameMode, and any other actors specified by ShouldReset())
+	for (FActorIterator It(GetWorld()); It; ++It)
+	{
+		AActor* A = *It;
+		if (A && !A->IsPendingKill() && A != this && !A->IsA<AController>() && ShouldReset(A))
+		{
+			A->Reset();
+		}
+	}
+
+	// Reset the GameMode
+	Reset();
+
+	// Notify the level script that the level has been reset
+	ALevelScriptActor* LevelScript = GetWorld()->GetLevelScriptActor();
+	if (LevelScript)
+	{
+		LevelScript->LevelReset();
+	}
+}
+
+void AGameModeBase::ReturnToMainMenuHost()
+{
+	if (GameSession)
+	{
+		GameSession->ReturnToMainMenuHost();
+	}
+}
+
+APlayerController* AGameModeBase::ProcessClientTravel(FString& FURL, FGuid NextMapGuid, bool bSeamless, bool bAbsolute)
+{
+	// We call PreClientTravel directly on any local PlayerPawns (ie listen server)
+	APlayerController* LocalPlayerController = nullptr;
+	for (FConstPlayerControllerIterator Iterator = GetWorld()->GetPlayerControllerIterator(); Iterator; ++Iterator)
+	{
+		APlayerController* PlayerController = *Iterator;
+		if (Cast<UNetConnection>(PlayerController->Player) != nullptr)
+		{
+			// Remote player
+			PlayerController->ClientTravel(FURL, TRAVEL_Relative, bSeamless, NextMapGuid);
+		}
+		else
+		{
+			// Local player
+			LocalPlayerController = PlayerController;
+			PlayerController->PreClientTravel(FURL, bAbsolute ? TRAVEL_Absolute : TRAVEL_Relative, bSeamless);
+		}
+	}
+
+	return LocalPlayerController;
+}
+
+bool AGameModeBase::CanServerTravel(const FString& FURL, bool bAbsolute)
+{
+	UWorld* World = GetWorld();
+
+	check(World);
+
+	// NOTE - This is a temp check while we work on a long term fix
+	// There are a few issues with seamless travel using single process PIE, so we're disabling that for now while working on a fix
+	if (World->WorldType == EWorldType::PIE && bUseSeamlessTravel && !FParse::Param(FCommandLine::Get(), TEXT("MultiprocessOSS")))
+	{
+		UE_LOG(LogGameMode, Warning, TEXT("AGameModeBase::CanServerTravel: Seamless travel currently NOT supported in single process PIE."));
+		return false;
+	}
+
+	if (FURL.Contains(TEXT("%")))
+	{
+		UE_LOG(LogGameMode, Error, TEXT("FURL %s Contains illegal character '%%'."), *FURL);
+		return false;
+	}
+
+	if (FURL.Contains(TEXT(":")) || FURL.Contains(TEXT("\\")))
+	{
+		UE_LOG(LogGameMode, Error, TEXT("FURL %s blocked, contains : or \\"), *FURL);
+		return false;
+	}
+
+	FString MapName;
+	int32 OptionStart = FURL.Find(TEXT("?"));
+	if (OptionStart == INDEX_NONE)
+	{
+		MapName = FURL;
+	}
+	else
+	{
+		MapName = FURL.Left(OptionStart);
+	}
+
+	// Check for invalid package names.
+	FText InvalidPackageError;
+	if (MapName.StartsWith(TEXT("/")) && !FPackageName::IsValidLongPackageName(MapName, true, &InvalidPackageError))
+	{
+		UE_LOG(LogGameMode, Log, TEXT("FURL %s blocked (%s)"), *FURL, *InvalidPackageError.ToString());
+		return false;
+	}
+	
+	return true;
+}
+
+void AGameModeBase::ProcessServerTravel(const FString& URL, bool bAbsolute)
+{
+#if WITH_SERVER_CODE
+	StartToLeaveMap();
+
+	// Force an old style load screen if the server has been up for a long time so that TimeSeconds doesn't overflow and break everything
+	bool bSeamless = (bUseSeamlessTravel && GetWorld()->TimeSeconds < 172800.0f); // 172800 seconds == 48 hours
+
+	FString NextMap;
+	if (URL.ToUpper().Contains(TEXT("?RESTART")))
+	{
+		NextMap = UWorld::RemovePIEPrefix(GetOutermost()->GetName());
+	}
+	else
+	{
+		int32 OptionStart = URL.Find(TEXT("?"));
+		if (OptionStart == INDEX_NONE)
+		{
+			NextMap = URL;
+		}
+		else
+		{
+			NextMap = URL.Left(OptionStart);
+		}
+	}
+
+	FGuid NextMapGuid = UEngine::GetPackageGuid(FName(*NextMap), GetWorld()->IsPlayInEditor());
+
+	// Notify clients we're switching level and give them time to receive.
+	FString URLMod = URL;
+	APlayerController* LocalPlayer = ProcessClientTravel(URLMod, NextMapGuid, bSeamless, bAbsolute);
+
+	UE_LOG(LogGameMode, Log, TEXT("ProcessServerTravel: %s"), *URL);
+	UWorld* World = GetWorld();
+	check(World);
+	World->NextURL = URL;
+	ENetMode NetMode = GetNetMode();
+
+	if (bSeamless)
+	{
+		World->SeamlessTravel(World->NextURL, bAbsolute);
+		World->NextURL = TEXT("");
+	}
+	// Switch immediately if not networking.
+	else if (NetMode != NM_DedicatedServer && NetMode != NM_ListenServer)
+	{
+		World->NextSwitchCountdown = 0.0f;
+	}
+#endif // WITH_SERVER_CODE
+}
+
+void AGameModeBase::GetSeamlessTravelActorList(bool bToTransition, TArray<AActor*>& ActorList)
+{
+	// Get allocations for the elements we're going to add handled in one go
+	const int32 ActorsToAddCount = GameState->PlayerArray.Num() + (bToTransition ? 3 : 0);
+	ActorList.Reserve(ActorsToAddCount);
+
+	// Always keep PlayerStates, so that after we restart we can keep players on the same team, etc
+	ActorList.Append(GameState->PlayerArray);
+
+	if (bToTransition)
+	{
+		// Keep ourselves until we transition to the final destination
+		ActorList.Add(this);
+		// Keep general game state until we transition to the final destination
+		ActorList.Add(GameState);
+		// Keep the game session state until we transition to the final destination
+		ActorList.Add(GameSession);
+
+		// If adding in this section best to increase the literal above for the ActorsToAddCount
+	}
+}
+
+void AGameModeBase::SwapPlayerControllers(APlayerController* OldPC, APlayerController* NewPC)
+{
+	if (OldPC != nullptr && !OldPC->IsPendingKill() && NewPC != nullptr && !NewPC->IsPendingKill() && OldPC->Player != nullptr)
+	{
+		// move the Player to the new PC
+		UPlayer* Player = OldPC->Player;
+		NewPC->NetPlayerIndex = OldPC->NetPlayerIndex; //@warning: critical that this is first as SetPlayer() may trigger RPCs
+		NewPC->NetConnection = OldPC->NetConnection;
+		NewPC->SetPlayer(Player);
+		NewPC->CopyRemoteRoleFrom(OldPC);
+
+		K2_OnSwapPlayerControllers(OldPC, NewPC);
+
+		// send destroy event to old PC immediately if it's local
+		if (Cast<ULocalPlayer>(Player))
+		{
+			GetWorld()->DestroyActor(OldPC);
+		}
+		else
+		{
+			OldPC->PendingSwapConnection = Cast<UNetConnection>(Player);
+			//@note: at this point, any remaining RPCs sent by the client on the old PC will be discarded
+			// this is consistent with general owned Actor destruction,
+			// however in this particular case it could easily be changed
+			// by modifying UActorChannel::ReceivedBunch() to account for PendingSwapConnection when it is setting bNetOwner
+		}
+	}
+	else
+	{
+		UE_LOG(LogGameMode, Warning, TEXT("SwapPlayerControllers(): Invalid OldPC, invalid NewPC, or OldPC has no Player!"));
+	}
+}
+
+void AGameModeBase::HandleSeamlessTravelPlayer(AController*& C)
+{
+	// Default behavior is to spawn new controllers and copy data
+	APlayerController* PC = Cast<APlayerController>(C);
+	if (PC && PC->Player)
+	{
+		// We need to spawn a new PlayerController to replace the old one
+		APlayerController* NewPC = SpawnPlayerController(PC->IsLocalPlayerController() ? ROLE_SimulatedProxy : ROLE_AutonomousProxy, PC->GetFocalLocation(), PC->GetControlRotation());
+
+		if (NewPC)
+		{
+			PC->SeamlessTravelTo(NewPC);
+			NewPC->SeamlessTravelFrom(PC);
+			SwapPlayerControllers(PC, NewPC);
+			PC = NewPC;
+			C = NewPC;
+		}
+		else
+		{
+			UE_LOG(LogGameMode, Warning, TEXT("Failed to spawn new PlayerController for %s (old class %s)"), *PC->GetHumanReadableName(), *PC->GetClass()->GetName());
+			PC->Destroy();
+			return;
+		}
+	}
+
+	InitSeamlessTravelPlayer(C);
+
+	// Initialize hud and other player details, shared with PostLogin
+	GenericPlayerInitialization(C);
+
+	if (PC)
+	{
+		// This may spawn the player pawn if the game is in progress
+		HandleStartingNewPlayer(PC);
+	}
+}
+
+void AGameModeBase::PostSeamlessTravel()
+{
+	if (GameSession)
+	{
+		GameSession->PostSeamlessTravel();
+	}
+
+	// We have to make a copy of the controller list, since the code after this will destroy
+	// and create new controllers in the world's list
+	TArray<TAutoWeakObjectPtr<AController> >	OldControllerList;
+	for (auto It = GetWorld()->GetControllerIterator(); It; ++It)
+	{
+		OldControllerList.Add(*It);
+	}
+
+	// Handle players that are already loaded
+	for (FConstControllerIterator Iterator = OldControllerList.CreateConstIterator(); Iterator; ++Iterator)
+	{
+		AController* Controller = *Iterator;
+		if (Controller->PlayerState)
+		{
+			APlayerController* PlayerController = Cast<APlayerController>(Controller);
+			if (!PlayerController || PlayerController->HasClientLoadedCurrentWorld())
+			{
+				// Don't handle if player is still loading world, that gets called in ServerNotifyLoadedWorld
+				HandleSeamlessTravelPlayer(Controller);
+			}
+		}
+	}
+}
+
+void AGameModeBase::StartToLeaveMap()
+{
+
+}
+
+void AGameModeBase::GameWelcomePlayer(UNetConnection* Connection, FString& RedirectURL) 
+{
+
+}
+
+void AGameModeBase::PreLogin(const FString& Options, const FString& Address, const TSharedPtr<const FUniqueNetId>& UniqueId, FString& ErrorMessage)
+{
+
+}
+
+void AGameModeBase::PreLogin(const FString& Options, const FString& Address, const FUniqueNetIdRepl& UniqueId, FString& ErrorMessage)
+{
+PRAGMA_DISABLE_DEPRECATION_WARNINGS
+	// Try calling deprecated version first
+	PreLogin(Options, Address, UniqueId.GetUniqueNetId(), ErrorMessage);
+	if (!ErrorMessage.IsEmpty())
+	{
+		return;
+	}
+PRAGMA_ENABLE_DEPRECATION_WARNINGS
+
+	ErrorMessage = GameSession->ApproveLogin(Options);
+}
+
+APlayerController* AGameModeBase::Login(UPlayer* NewPlayer, ENetRole InRemoteRole, const FString& Portal, const FString& Options, const TSharedPtr<const FUniqueNetId>& UniqueId, FString& ErrorMessage)
+{
+	return nullptr;
+}
+
+APlayerController* AGameModeBase::Login(UPlayer* NewPlayer, ENetRole InRemoteRole, const FString& Portal, const FString& Options, const FUniqueNetIdRepl& UniqueId, FString& ErrorMessage)
+{
+PRAGMA_DISABLE_DEPRECATION_WARNINGS
+	// Try calling deprecated version first
+	APlayerController* DeprecatedController = Login(NewPlayer, InRemoteRole, Portal, Options, UniqueId.GetUniqueNetId(), ErrorMessage);
+	if (DeprecatedController)
+	{
+		return DeprecatedController;
+	}
+PRAGMA_ENABLE_DEPRECATION_WARNINGS
+
+	ErrorMessage = GameSession->ApproveLogin(Options);
+	if (!ErrorMessage.IsEmpty())
+	{
+		return nullptr;
+	}
+
+	APlayerController* NewPlayerController = SpawnPlayerController(InRemoteRole, FVector::ZeroVector, FRotator::ZeroRotator);
+
+	// Handle spawn failure.
+	if (NewPlayerController == nullptr)
+	{
+		UE_LOG(LogGameMode, Log, TEXT("Couldn't spawn player controller of class %s"), PlayerControllerClass ? *PlayerControllerClass->GetName() : TEXT("NULL"));
+		ErrorMessage = FString::Printf(TEXT("Failed to spawn player controller"));
+		return nullptr;
+	}
+
+	// Customize incoming player based on URL options
+	ErrorMessage = InitNewPlayer(NewPlayerController, UniqueId, Options, Portal);
+	if (!ErrorMessage.IsEmpty())
+	{
+		return nullptr;
+	}
+
+	return NewPlayerController;
+}
+
+APlayerController* AGameModeBase::SpawnPlayerController(ENetRole InRemoteRole, FVector const& SpawnLocation, FRotator const& SpawnRotation)
+{
+	FActorSpawnParameters SpawnInfo;
+	SpawnInfo.Instigator = Instigator;
+	SpawnInfo.ObjectFlags |= RF_Transient;	// We never want to save player controllers into a map
+	SpawnInfo.bDeferConstruction = true;
+	APlayerController* NewPC = GetWorld()->SpawnActor<APlayerController>(PlayerControllerClass, SpawnLocation, SpawnRotation, SpawnInfo);
+	if (NewPC)
+	{
+		if (InRemoteRole == ROLE_SimulatedProxy)
+		{
+			// This is a local player because it has no authority/autonomous remote role
+			NewPC->SetAsLocalPlayerController();
+		}
+
+		UGameplayStatics::FinishSpawningActor(NewPC, FTransform(SpawnRotation, SpawnLocation));
+	}
+
+	return NewPC;
+}
+
+FString AGameModeBase::InitNewPlayer(APlayerController* NewPlayerController, const TSharedPtr<const FUniqueNetId>& UniqueId, const FString& Options, const FString& Portal)
+{
+	return TEXT("DEPRECATED");
+}
+
+FString AGameModeBase::InitNewPlayer(APlayerController* NewPlayerController, const FUniqueNetIdRepl& UniqueId, const FString& Options, const FString& Portal)
+{
+PRAGMA_DISABLE_DEPRECATION_WARNINGS
+	// Try calling deprecated version first
+	FString DeprecatedError = InitNewPlayer(NewPlayerController, UniqueId.GetUniqueNetId(), Options, Portal);
+	if (DeprecatedError != TEXT("DEPRECATED"))
+	{
+		// This means it was implemented in subclass
+		return DeprecatedError;
+	}
+PRAGMA_ENABLE_DEPRECATION_WARNINGS
+
+	check(NewPlayerController);
+
+	FString ErrorMessage;
+
+	// Register the player with the session
+	GameSession->RegisterPlayer(NewPlayerController, UniqueId.GetUniqueNetId(), UGameplayStatics::HasOption(Options, TEXT("bIsFromInvite")));
+
+	// Find a starting spot
+	AActor* const StartSpot = FindPlayerStart(NewPlayerController, Portal);
+	if (StartSpot != nullptr)
+	{
+		// Set the player controller / camera in this new location
+		FRotator InitialControllerRot = StartSpot->GetActorRotation();
+		InitialControllerRot.Roll = 0.f;
+		NewPlayerController->SetInitialLocationAndRotation(StartSpot->GetActorLocation(), InitialControllerRot);
+		NewPlayerController->StartSpot = StartSpot;
+	}
+	else
+	{
+		ErrorMessage = FString::Printf(TEXT("Failed to find PlayerStart"));
+	}
+
+	// Set up spectating
+	bool bSpectator = FCString::Stricmp(*UGameplayStatics::ParseOption(Options, TEXT("SpectatorOnly")), TEXT("1")) == 0;
+	if (bSpectator || MustSpectate(NewPlayerController))
+	{
+		NewPlayerController->StartSpectatingOnly();
+	}
+
+	// Init player's name
+	FString InName = UGameplayStatics::ParseOption(Options, TEXT("Name")).Left(20);
+	if (InName.IsEmpty())
+	{
+		InName = FString::Printf(TEXT("%s%i"), *DefaultPlayerName.ToString(), NewPlayerController->PlayerState->PlayerId);
+	}
+
+	ChangeName(NewPlayerController, InName, false);
+
+	return ErrorMessage;
+}
+
+void AGameModeBase::InitSeamlessTravelPlayer(AController* NewController)
+{
+	APlayerController* NewPC = Cast<APlayerController>(NewController);
+	// Find a start spot
+	AActor* StartSpot = FindPlayerStart(NewController);
+
+	if (StartSpot == nullptr)
+	{
+		UE_LOG(LogGameMode, Warning, TEXT("Could not find a starting spot"));
+	}
+	else
+	{
+		FRotator StartRotation(0, StartSpot->GetActorRotation().Yaw, 0);
+		NewController->SetInitialLocationAndRotation(StartSpot->GetActorLocation(), StartRotation);
+	}
+
+	NewController->StartSpot = StartSpot;
+
+	if (NewPC != nullptr)
+	{
+		NewPC->PostSeamlessTravel();
+
+		if (MustSpectate(NewPC))
+		{
+			NewPC->StartSpectatingOnly();
+		}
+		else
+		{
+			NewPC->bPlayerIsWaiting = true;
+			NewPC->ChangeState(NAME_Spectating);
+			NewPC->ClientGotoState(NAME_Spectating);
+		}
+	}
+}
+
+bool AGameModeBase::ShouldStartInCinematicMode(APlayerController* Player, bool& OutHidePlayer, bool& OutHideHUD, bool& OutDisableMovement, bool& OutDisableTurning)
+{
+	ULocalPlayer* LocPlayer = Player->GetLocalPlayer();
+	if (!LocPlayer)
+	{
+		return false;
+	}
+
+#if WITH_EDITOR
+	// If we have an active movie scene capture, we can take the settings from that
+	if (LocPlayer->ViewportClient && LocPlayer->ViewportClient->Viewport)
+	{
+		if (auto* MovieSceneCapture = IMovieSceneCaptureModule::Get().GetFirstActiveMovieSceneCapture())
+		{
+			const FMovieSceneCaptureSettings& Settings = MovieSceneCapture->GetSettings();
+			if (Settings.bCinematicMode)
+			{
+				OutDisableMovement = !Settings.bAllowMovement;
+				OutDisableTurning = !Settings.bAllowTurning;
+				OutHidePlayer = !Settings.bShowPlayer;
+				OutHideHUD = !Settings.bShowHUD;
+				return true;
+			}
+		}
+	}
+#endif
+
+	return false;
+}
+
+void AGameModeBase::InitializeHUDForPlayer_Implementation(APlayerController* NewPlayer)
+{
+	// Tell client what HUD class to use
+	NewPlayer->ClientSetHUD(HUDClass);
+}
+
+void AGameModeBase::UpdateGameplayMuteList(APlayerController* aPlayer)
+{
+	if (aPlayer)
+	{
+		aPlayer->MuteList.bHasVoiceHandshakeCompleted = true;
+		aPlayer->ClientVoiceHandshakeComplete();
+	}
+}
+
+void AGameModeBase::ReplicateStreamingStatus(APlayerController* PC)
+{
+	UWorld* MyWorld = GetWorld();
+
+	if (MyWorld->GetWorldSettings()->bUseClientSideLevelStreamingVolumes)
+	{
+		// Client will itself decide what to stream
+		return;
+	}
+
+	// Don't do this for local players or players after the first on a splitscreen client
+	if (Cast<ULocalPlayer>(PC->Player) == nullptr && Cast<UChildConnection>(PC->Player) == nullptr)
+	{
+		// If we've loaded levels via CommitMapChange() that aren't normally in the StreamingLevels array, tell the client about that
+		if (MyWorld->CommittedPersistentLevelName != NAME_None)
+		{
+			PC->ClientPrepareMapChange(MyWorld->CommittedPersistentLevelName, true, true);
+			// Tell the client to commit the level immediately
+			PC->ClientCommitMapChange();
+		}
+
+		if (MyWorld->StreamingLevels.Num() > 0)
+		{
+			// Tell the player controller the current streaming level status
+			for (int32 LevelIndex = 0; LevelIndex < MyWorld->StreamingLevels.Num(); LevelIndex++)
+			{
+				ULevelStreaming* TheLevel = MyWorld->StreamingLevels[LevelIndex];
+
+				if (TheLevel != nullptr)
+				{
+					const ULevel* LoadedLevel = TheLevel->GetLoadedLevel();
+
+					UE_LOG(LogGameMode, Log, TEXT("levelStatus: %s %i %i %i %s %i"),
+						*TheLevel->GetWorldAssetPackageName(),
+						TheLevel->bShouldBeVisible,
+						LoadedLevel && LoadedLevel->bIsVisible,
+						TheLevel->bShouldBeLoaded,
+						*GetNameSafe(LoadedLevel),
+						TheLevel->bHasLoadRequestPending);
+
+					PC->ClientUpdateLevelStreamingStatus(
+						TheLevel->GetWorldAssetPackageFName(),
+						TheLevel->bShouldBeLoaded,
+						TheLevel->bShouldBeVisible,
+						TheLevel->bShouldBlockOnLoad,
+						TheLevel->LevelLODIndex);
+				}
+			}
+			PC->ClientFlushLevelStreaming();
+		}
+
+		// If we're preparing to load different levels using PrepareMapChange() inform the client about that now
+		if (MyWorld->PreparingLevelNames.Num() > 0)
+		{
+			for (int32 LevelIndex = 0; LevelIndex < MyWorld->PreparingLevelNames.Num(); LevelIndex++)
+			{
+				PC->ClientPrepareMapChange(MyWorld->PreparingLevelNames[LevelIndex], LevelIndex == 0, LevelIndex == MyWorld->PreparingLevelNames.Num() - 1);
+			}
+			// DO NOT commit these changes yet - we'll send that when we're done preparing them
+		}
+	}
+}
+
+void AGameModeBase::GenericPlayerInitialization(AController* C)
+{
+	APlayerController* PC = Cast<APlayerController>(C);
+	if (PC != nullptr)
+	{
+		InitializeHUDForPlayer(PC);
+
+		// Notify the game that we can now be muted and mute others
+		UpdateGameplayMuteList(PC);
+
+		// Tell the player to enable voice by default or use the push to talk method
+		PC->ClientEnableNetworkVoice(!GameSession->RequiresPushToTalk());
+
+		ReplicateStreamingStatus(PC);
+
+		bool HidePlayer = false, HideHUD = false, DisableMovement = false, DisableTurning = false;
+
+		// Check to see if we should start in cinematic mode (matinee movie capture)
+		if (ShouldStartInCinematicMode(PC, HidePlayer, HideHUD, DisableMovement, DisableTurning))
+		{
+			PC->SetCinematicMode(true, HidePlayer, HideHUD, DisableMovement, DisableTurning);
+		}
+
+		// Add the player to any matinees running so that it gets in on any cinematics already running, etc
+		TArray<AMatineeActor*> AllMatineeActors;
+		GetWorld()->GetMatineeActors(AllMatineeActors);
+		for (int32 i = 0; i < AllMatineeActors.Num(); i++)
+		{
+			AllMatineeActors[i]->AddPlayerToDirectorTracks(PC);
+		}
+	}
+}
+
+void AGameModeBase::PostLogin(APlayerController* NewPlayer)
+{
+	// Runs shared initialization that can happen during seamless travel as well
+
+	GenericPlayerInitialization(NewPlayer);
+
+	// Perform initialization that only happens on initially joining a server
+
+	UWorld* World = GetWorld();
+
+	NewPlayer->ClientCapBandwidth(NewPlayer->Player->CurrentNetSpeed);
+
+	if (MustSpectate(NewPlayer))
+	{
+		NewPlayer->ClientGotoState(NAME_Spectating);
+	}
+	else
+	{
+		// If NewPlayer is not only a spectator and has a valid ID, add him as a user to the replay.
+		if (NewPlayer->PlayerState->UniqueId.IsValid())
+		{
+			GetGameInstance()->AddUserToReplay(NewPlayer->PlayerState->UniqueId.ToString());
+		}
+	}
+
+	if (GameSession)
+	{
+		GameSession->PostLogin(NewPlayer);
+	}
+
+	// Notify Blueprints that a new player has logged in.  Calling it here, because this is the first time that the PlayerController can take RPCs
+	K2_PostLogin(NewPlayer);
+	FGameModeEvents::GameModePostLoginEvent.Broadcast(this, NewPlayer);
+
+	// Now that initialization is done, try to spawn the player's pawn and start match
+	HandleStartingNewPlayer(NewPlayer);
+}
+
+void AGameModeBase::Logout(AController* Exiting)
+{
+	APlayerController* PC = Cast<APlayerController>(Exiting);
+	if (PC != nullptr)
+	{
+		FGameModeEvents::GameModeLogoutEvent.Broadcast(this, Exiting);
+		K2_OnLogout(Exiting);
+
+		if (GameSession)
+		{
+			GameSession->NotifyLogout(PC);
+		}
+	}
+}
+
+void AGameModeBase::HandleStartingNewPlayer_Implementation(APlayerController* NewPlayer)
+{
+	// If players should start as spectators, leave them in the spectator state
+	if (!bStartPlayersAsSpectators && !MustSpectate(NewPlayer) && PlayerCanRestart(NewPlayer))
+	{
+		// Otherwise spawn their pawn immediately
+		RestartPlayer(NewPlayer);
+	}
+}
+
+bool AGameModeBase::MustSpectate_Implementation(APlayerController* NewPlayerController) const
+{
+	if (!NewPlayerController || !NewPlayerController->PlayerState)
+	{
+		return false;
+	}
+
+	return NewPlayerController->PlayerState->bOnlySpectator;
+}
+
+bool AGameModeBase::CanSpectate_Implementation(APlayerController* Viewer, APlayerState* ViewTarget)
+{
+	return true;
+}
+
+AActor* AGameModeBase::ChoosePlayerStart_Implementation(AController* Player)
+{
+	// Choose a player start
+	APlayerStart* FoundPlayerStart = nullptr;
+	UClass* PawnClass = GetDefaultPawnClassForController(Player);
+	APawn* PawnToFit = PawnClass ? PawnClass->GetDefaultObject<APawn>() : nullptr;
+	TArray<APlayerStart*> UnOccupiedStartPoints;
+	TArray<APlayerStart*> OccupiedStartPoints;
+	for (TActorIterator<APlayerStart> It(GetWorld()); It; ++It)
+	{
+		APlayerStart* PlayerStart = *It;
+
+		if (PlayerStart->IsA<APlayerStartPIE>())
+		{
+			// Always prefer the first "Play from Here" PlayerStart, if we find one while in PIE mode
+			FoundPlayerStart = PlayerStart;
+			break;
+		}
+		else
+		{
+			FVector ActorLocation = PlayerStart->GetActorLocation();
+			const FRotator ActorRotation = PlayerStart->GetActorRotation();
+			if (!GetWorld()->EncroachingBlockingGeometry(PawnToFit, ActorLocation, ActorRotation))
+			{
+				UnOccupiedStartPoints.Add(PlayerStart);
+			}
+			else if (GetWorld()->FindTeleportSpot(PawnToFit, ActorLocation, ActorRotation))
+			{
+				OccupiedStartPoints.Add(PlayerStart);
+			}
+		}
+	}
+	if (FoundPlayerStart == nullptr)
+	{
+		if (UnOccupiedStartPoints.Num() > 0)
+		{
+			FoundPlayerStart = UnOccupiedStartPoints[FMath::RandRange(0, UnOccupiedStartPoints.Num() - 1)];
+		}
+		else if (OccupiedStartPoints.Num() > 0)
+		{
+			FoundPlayerStart = OccupiedStartPoints[FMath::RandRange(0, OccupiedStartPoints.Num() - 1)];
+		}
+	}
+	return FoundPlayerStart;
+}
+
+bool AGameModeBase::ShouldSpawnAtStartSpot(AController* Player)
+{
+	return (Player != nullptr && Player->StartSpot != nullptr);
+}
+
+AActor* AGameModeBase::FindPlayerStart_Implementation(AController* Player, const FString& IncomingName)
+{
+	UWorld* World = GetWorld();
+
+	// If incoming start is specified, then just use it
+	if (!IncomingName.IsEmpty())
+	{
+		const FName IncomingPlayerStartTag = FName(*IncomingName);
+		for (TActorIterator<APlayerStart> It(World); It; ++It)
+		{
+			APlayerStart* Start = *It;
+			if (Start && Start->PlayerStartTag == IncomingPlayerStartTag)
+			{
+				return Start;
+			}
+		}
+	}
+
+	// Always pick StartSpot at start of match
+	if (ShouldSpawnAtStartSpot(Player))
+	{
+		if (AActor* PlayerStartSpot = Player->StartSpot.Get())
+		{
+			return PlayerStartSpot;
+		}
+		else
+		{
+			UE_LOG(LogGameMode, Error, TEXT("Error - ShouldSpawnAtStartSpot returned true but the Player StartSpot was null."));
+		}
+	}
+
+	AActor* BestStart = ChoosePlayerStart(Player);
+	if (BestStart == nullptr)
+	{
+		// No player start found
+		UE_LOG(LogGameMode, Log, TEXT("Warning - PATHS NOT DEFINED or NO PLAYERSTART with positive rating"));
+
+		// This is a bit odd, but there was a complex chunk of code that in the end always resulted in this, so we may as well just 
+		// short cut it down to this.  Basically we are saying spawn at 0,0,0 if we didn't find a proper player start
+		BestStart = World->GetWorldSettings();
+	}
+
+	return BestStart;
+}
+
+AActor* AGameModeBase::K2_FindPlayerStart(AController* Player, const FString& IncomingName)
+{
+	return FindPlayerStart(Player, IncomingName);
+}
+
+bool AGameModeBase::PlayerCanRestart_Implementation(APlayerController* Player)
+{
+	if (Player == nullptr || Player->IsPendingKillPending())
+	{
+		return false;
+	}
+
+	// Ask the player controller if it's ready to restart as well
+	return Player->CanRestartPlayer();
+}
+
+APawn* AGameModeBase::SpawnDefaultPawnFor_Implementation(AController* NewPlayer, AActor* StartSpot)
+{
+	// Don't allow pawn to be spawned with any pitch or roll
+	FRotator StartRotation(ForceInit);
+	StartRotation.Yaw = StartSpot->GetActorRotation().Yaw;
+	FVector StartLocation = StartSpot->GetActorLocation();
+
+	FTransform Transform = FTransform(StartRotation, StartLocation);
+	return SpawnDefaultPawnAtTransform(NewPlayer, Transform);
+}
+
+APawn* AGameModeBase::SpawnDefaultPawnAtTransform_Implementation(AController* NewPlayer, const FTransform& SpawnTransform)
+{
+	FActorSpawnParameters SpawnInfo;
+	SpawnInfo.Instigator = Instigator;
+	SpawnInfo.ObjectFlags |= RF_Transient;	// We never want to save default player pawns into a map
+	UClass* PawnClass = GetDefaultPawnClassForController(NewPlayer);
+	APawn* ResultPawn = GetWorld()->SpawnActor<APawn>(PawnClass, SpawnTransform, SpawnInfo);
+	if (!ResultPawn)
+	{
+		UE_LOG(LogGameMode, Warning, TEXT("Couldn't spawn Pawn of type %s at %s"), *GetNameSafe(PawnClass), *SpawnTransform.ToHumanReadableString());
+	}
+	return ResultPawn;
+}
+
+void AGameModeBase::RestartPlayer(AController* NewPlayer)
+{
+	if (NewPlayer == nullptr || NewPlayer->IsPendingKillPending())
+	{
+		return;
+	}
+
+	AActor* StartSpot = FindPlayerStart(NewPlayer);
+
+	// If a start spot wasn't found,
+	if (StartSpot == nullptr)
+	{
+		// Check for a previously assigned spot
+		if (NewPlayer->StartSpot != nullptr)
+		{
+			StartSpot = NewPlayer->StartSpot.Get();
+			UE_LOG(LogGameMode, Warning, TEXT("Player start not found, using last start spot"));
+		}	
+	}
+
+	RestartPlayerAtPlayerStart(NewPlayer, StartSpot);
+}
+
+void AGameModeBase::RestartPlayerAtPlayerStart(AController* NewPlayer, AActor* StartSpot)
+{
+	if (NewPlayer == nullptr || NewPlayer->IsPendingKillPending())
+	{
+		return;
+	}
+
+	if (!StartSpot)
+	{
+		UE_LOG(LogGameMode, Warning, TEXT("Player start not found, failed to RestartPlayerAtPlayerStart"));
+		return;
+	}
+
+	UE_LOG(LogGameMode, Verbose, TEXT("RestartPlayerAtPlayerStart %s"), (NewPlayer && NewPlayer->PlayerState) ? *NewPlayer->PlayerState->PlayerName : TEXT("Unknown"));
+
+	if (MustSpectate(Cast<APlayerController>(NewPlayer)))
+	{
+		UE_LOG(LogGameMode, Verbose, TEXT("RestartPlayerAtTransform tried to restart a spectator-only player!"));
+		return;
+	}
+
+	// Try to create a pawn to use of the default class for this player
+	if (NewPlayer->GetPawn() == nullptr && GetDefaultPawnClassForController(NewPlayer) != nullptr)
+	{
+		NewPlayer->SetPawn(SpawnDefaultPawnFor(NewPlayer, StartSpot));
+	}
+
+	if (NewPlayer->GetPawn() == nullptr)
+	{
+		NewPlayer->FailedToSpawnPawn();
+	}
+	else
+	{
+		// Tell the start spot it was used
+		InitStartSpot(StartSpot, NewPlayer);
+
+		FinishRestartPlayer(NewPlayer, StartSpot->GetActorRotation());
+	}
+}
+
+void AGameModeBase::RestartPlayerAtTransform(AController* NewPlayer, const FTransform& SpawnTransform)
+{
+	if (NewPlayer == nullptr || NewPlayer->IsPendingKillPending())
+	{
+		return;
+	}
+
+	UE_LOG(LogGameMode, Verbose, TEXT("RestartPlayerAtTransform %s"), (NewPlayer && NewPlayer->PlayerState) ? *NewPlayer->PlayerState->PlayerName : TEXT("Unknown"));
+
+	if (MustSpectate(Cast<APlayerController>(NewPlayer)))
+	{
+		UE_LOG(LogGameMode, Verbose, TEXT("RestartPlayerAtTransform tried to restart a spectator-only player!"));
+		return;
+	}
+
+	// Try to create a pawn to use of the default class for this player
+	if (NewPlayer->GetPawn() == nullptr && GetDefaultPawnClassForController(NewPlayer) != nullptr)
+	{
+		NewPlayer->SetPawn(SpawnDefaultPawnAtTransform(NewPlayer, SpawnTransform));
+	}
+
+	if (NewPlayer->GetPawn() == nullptr)
+	{
+		NewPlayer->FailedToSpawnPawn();
+	}
+	else
+	{
+		FinishRestartPlayer(NewPlayer, SpawnTransform.GetRotation().Rotator());
+	}
+}
+
+void AGameModeBase::FinishRestartPlayer(AController* NewPlayer, const FRotator& StartRotation)
+{
+	NewPlayer->Possess(NewPlayer->GetPawn());
+
+	// If the Pawn is destroyed as part of possession we have to abort
+	if (NewPlayer->GetPawn() == nullptr)
+	{
+		NewPlayer->FailedToSpawnPawn();
+	}
+	else
+	{
+		// Set initial control rotation to starting rotation rotation
+		NewPlayer->ClientSetRotation(NewPlayer->GetPawn()->GetActorRotation(), true);
+
+		FRotator NewControllerRot = StartRotation;
+		NewControllerRot.Roll = 0.f;
+		NewPlayer->SetControlRotation(NewControllerRot);
+
+		SetPlayerDefaults(NewPlayer->GetPawn());
+
+		K2_OnRestartPlayer(NewPlayer);
+	}
+}
+
+void AGameModeBase::InitStartSpot_Implementation(AActor* StartSpot, AController* NewPlayer)
+{
+
+}
+
+void AGameModeBase::SetPlayerDefaults(APawn* PlayerPawn)
+{
+	PlayerPawn->SetPlayerDefaults();
+
+#if !UE_WITH_PHYSICS
+	// If there is no physics, set to flying by default
+	UCharacterMovementComponent* CharacterMovement = Cast<UCharacterMovementComponent>(NewPlayer->GetPawn()->GetMovementComponent());
+	if (CharacterMovement)
+	{
+		CharacterMovement->bCheatFlying = true;
+		CharacterMovement->SetMovementMode(MOVE_Flying);
+	}
+#endif	//!UE_WITH_PHYSICS
+
+}
+
+void AGameModeBase::ChangeName(AController* Other, const FString& S, bool bNameChange)
+{
+	if (Other && !S.IsEmpty())
+	{
+		Other->PlayerState->SetPlayerName(S);
+
+		K2_OnChangeName(Other, S, bNameChange);
+	}
+}
+
+bool AGameModeBase::AllowCheats(APlayerController* P)
+{
+	return (GetNetMode() == NM_Standalone || GIsEditor); // Always allow cheats in editor (PIE now supports networking)
+}
+
+bool AGameModeBase::IsHandlingReplays()
+{
+	return false;
+}
+
+bool AGameModeBase::SpawnPlayerFromSimulate(const FVector& NewLocation, const FRotator& NewRotation)
+{
+#if WITH_EDITOR
+	APlayerController* PC = GetGameInstance()->GetFirstLocalPlayerController();
+	if (PC != nullptr)
+	{
+		PC->PlayerState->bOnlySpectator = false;
+
+		bool bNeedsRestart = true;
+		if (PC->GetPawn() == nullptr)
+		{
+			// Use the "auto-possess" pawn in the world, if there is one.
+			for (FConstPawnIterator Iterator = GetWorld()->GetPawnIterator(); Iterator; ++Iterator)
+			{
+				APawn* Pawn = *Iterator;
+				if (Pawn && Pawn->AutoPossessPlayer == EAutoReceiveInput::Player0)
+				{
+					if (Pawn->Controller == nullptr)
+					{
+						PC->Possess(Pawn);
+						bNeedsRestart = false;
+					}
+					break;
+				}
+			}
+		}
+
+		if (bNeedsRestart)
+		{
+			RestartPlayer(PC);
+
+			if (PC->GetPawn())
+			{
+				// If there was no player start, then try to place the pawn where the camera was.						
+				if (PC->StartSpot == nullptr || Cast<AWorldSettings>(PC->StartSpot.Get()))
+				{
+					const FVector Location = NewLocation;
+					const FRotator Rotation = NewRotation;
+					PC->SetControlRotation(Rotation);
+					PC->GetPawn()->TeleportTo(Location, Rotation);
+				}
+			}
+		}
+	}
+#endif
+	return true;
+}
\ No newline at end of file
diff --git a/Engine/Source/Runtime/Engine/Private/GameNetworkManager.cpp b/Engine/Source/Runtime/Engine/Private/GameNetworkManager.cpp
--- a/Engine/Source/Runtime/Engine/Private/GameNetworkManager.cpp
+++ b/Engine/Source/Runtime/Engine/Private/GameNetworkManager.cpp
@@ -6,7 +6,7 @@
 
 #include "EnginePrivate.h"
 #include "GameFramework/GameNetworkManager.h"
-#include "GameFramework/GameMode.h"
+#include "GameFramework/GameModeBase.h"
 
 DEFINE_LOG_CATEGORY_STATIC(LogGameNetworkManager, Log, All);
 
@@ -71,6 +71,11 @@ void AGameNetworkManager::UpdateNetSpeedsTimer()
 	UpdateNetSpeeds(false);
 }
 
+bool AGameNetworkManager::IsInLowBandwidthMode()
+{
+	return false;
+}
+
 void AGameNetworkManager::UpdateNetSpeeds(bool bIsLanMatch)
 {
 	// Don't adjust net speeds for LAN matches or dedicated servers
@@ -103,7 +108,15 @@ void AGameNetworkManager::UpdateNetSpeeds(bool bIsLanMatch)
 
 int32 AGameNetworkManager::CalculatedNetSpeed()
 {
-	return FMath::Clamp(TotalNetBandwidth/FMath::Max(GetWorld()->GetAuthGameMode()->NumPlayers,1), MinDynamicBandwidth, MaxDynamicBandwidth);
+	int32 NumPlayers = 1;
+	AGameModeBase* GameMode = GetWorld()->GetAuthGameMode();
+
+	if (GameMode)
+	{
+		NumPlayers = FMath::Max(GameMode->GetNumPlayers(), 1);
+	}
+
+	return FMath::Clamp(TotalNetBandwidth/NumPlayers, MinDynamicBandwidth, MaxDynamicBandwidth);
 }
 
 void AGameNetworkManager::StandbyCheatDetected(EStandbyType StandbyType) {}
diff --git a/Engine/Source/Runtime/Engine/Private/GameSession.cpp b/Engine/Source/Runtime/Engine/Private/GameSession.cpp
--- a/Engine/Source/Runtime/Engine/Private/GameSession.cpp
+++ b/Engine/Source/Runtime/Engine/Private/GameSession.cpp
@@ -10,7 +10,7 @@
 #include "Kismet/GameplayStatics.h"
 #include "GameFramework/PlayerState.h"
 #include "GameFramework/GameSession.h"
-#include "GameFramework/GameMode.h"
+#include "GameFramework/GameModeBase.h"
 
 DEFINE_LOG_CATEGORY_STATIC(LogGameSession, Log, All);
 
@@ -129,7 +129,7 @@ void AGameSession::InitOptions( const FString& Options )
 {
 	UWorld* const World = GetWorld();
 	check(World);
-	AGameMode* const GameMode = World ? World->GetAuthGameMode() : nullptr;
+	AGameModeBase* const GameMode = World ? World->GetAuthGameMode() : nullptr;
 
 	MaxPlayers = UGameplayStatics::GetIntOption( Options, TEXT("MaxPlayers"), MaxPlayers );
 	MaxSpectators = UGameplayStatics::GetIntOption( Options, TEXT("MaxSpectators"), MaxSpectators );
@@ -191,7 +191,7 @@ FString AGameSession::ApproveLogin(const FString& Options)
 	UWorld* const World = GetWorld();
 	check(World);
 
-	AGameMode* const GameMode = World->GetAuthGameMode();
+	AGameModeBase* const GameMode = World->GetAuthGameMode();
 	check(GameMode);
 
 	int32 SpectatorOnly = 0;
@@ -268,16 +268,18 @@ bool AGameSession::AtCapacity(bool bSpectator)
 		return false;
 	}
 
+	AGameModeBase* GameMode = GetWorld()->GetAuthGameMode();
+
 	if ( bSpectator )
 	{
-		return ( (GetWorld()->GetAuthGameMode()->NumSpectators >= MaxSpectators)
-		&& ((GetNetMode() != NM_ListenServer) || (GetWorld()->GetAuthGameMode()->NumPlayers > 0)) );
+		return ( (GameMode->GetNumSpectators() >= MaxSpectators)
+		&& ((GetNetMode() != NM_ListenServer) || (GameMode->GetNumPlayers() > 0)) );
 	}
 	else
 	{
 		const int32 MaxPlayersToUse = CVarMaxPlayersOverride.GetValueOnGameThread() > 0 ? CVarMaxPlayersOverride.GetValueOnGameThread() : MaxPlayers;
 
-		return ( (MaxPlayersToUse>0) && (GetWorld()->GetAuthGameMode()->GetNumPlayers() >= MaxPlayersToUse) );
+		return ( (MaxPlayersToUse>0) && (GameMode->GetNumPlayers() >= MaxPlayersToUse) );
 	}
 }
 
diff --git a/Engine/Source/Runtime/Engine/Private/GameState.cpp b/Engine/Source/Runtime/Engine/Private/GameState.cpp
--- a/Engine/Source/Runtime/Engine/Private/GameState.cpp
+++ b/Engine/Source/Runtime/Engine/Private/GameState.cpp
@@ -8,40 +8,14 @@
 #include "Net/UnrealNetwork.h"
 #include "GameFramework/GameState.h"
 #include "GameFramework/PlayerState.h"
+#include "GameFramework/GameModeBase.h"
 #include "GameFramework/GameMode.h"
 
-DEFINE_LOG_CATEGORY_STATIC(LogGameState, Log, All);
-
 AGameState::AGameState(const FObjectInitializer& ObjectInitializer)
-	: Super(ObjectInitializer
-		.DoNotCreateDefaultSubobject(TEXT("Sprite")) )
+	: Super(ObjectInitializer)
 {
-	SetRemoteRoleForBackwardsCompat(ROLE_SimulatedProxy);
-	bReplicates = true;
-	bAlwaysRelevant = true;
-	bReplicateMovement = false;
-
-	// Note: this is very important to set to false. Though all replication infos are spawned at run time, during seamless travel
-	// they are held on to and brought over into the new world. In ULevel::InitializeNetworkActors, these PlayerStates may be treated as map/startup actors
-	// and given static NetGUIDs. This also causes their deletions to be recorded and sent to new clients, which if unlucky due to name conflicts,
-	// may end up deleting the new PlayerStates they had just spaned.
-	bNetLoadOnClient = false;
 	MatchState = MatchState::EnteringMap;
 	PreviousMatchState = MatchState::EnteringMap;
-
-	// Default to every few seconds.
-	ServerWorldTimeSecondsUpdateFrequency = 5.f;
-}
-
-/** Helper to return the default object of the GameMode class corresponding to this GameState. */
-AGameMode* AGameState::GetDefaultGameMode() const
-{
-	if ( GameModeClass )
-	{
-		AGameMode* const DefaultGameActor = GameModeClass->GetDefaultObject<AGameMode>();
-		return DefaultGameActor;
-	}
-	return NULL;
 }
 
 void AGameState::DefaultTimer()
@@ -58,98 +32,17 @@ void AGameState::DefaultTimer()
 	GetWorldTimerManager().SetTimer(TimerHandle_DefaultTimer, this, &AGameState::DefaultTimer, GetWorldSettings()->GetEffectiveTimeDilation() / GetWorldSettings()->DemoPlayTimeDilation, true);
 }
 
+bool AGameState::ShouldShowGore() const
+{
+	return true;
+}
+
 void AGameState::PostInitializeComponents()
 {
 	Super::PostInitializeComponents();
 
 	FTimerManager& TimerManager = GetWorldTimerManager();
 	TimerManager.SetTimer(TimerHandle_DefaultTimer, this, &AGameState::DefaultTimer, GetWorldSettings()->GetEffectiveTimeDilation() / GetWorldSettings()->DemoPlayTimeDilation, true);
-
-	UWorld* World = GetWorld();
-	World->GameState = this;
-
-	if (World->IsGameWorld() && Role == ROLE_Authority)
-	{
-		UpdateServerTimeSeconds();
-		if (ServerWorldTimeSecondsUpdateFrequency > 0.f)
-		{
-			TimerManager.SetTimer(TimerHandle_UpdateServerTimeSeconds, this, &AGameState::UpdateServerTimeSeconds, ServerWorldTimeSecondsUpdateFrequency, true);
-		}
-	}
-
-	for (TActorIterator<APlayerState> It(World); It; ++It)
-	{
-		AddPlayerState(*It);
-	}
-}
-
-void AGameState::OnRep_GameModeClass()
-{
-	ReceivedGameModeClass();
-}
-
-void AGameState::OnRep_SpectatorClass()
-{
-	ReceivedSpectatorClass();
-}
-
-void AGameState::ReceivedGameModeClass()
-{
-	// Tell each PlayerController that the Game class is here
-	for( FConstPlayerControllerIterator Iterator = GetWorld()->GetPlayerControllerIterator(); Iterator; ++Iterator )
-	{
-		APlayerController* const PlayerController = *Iterator;
-		if (PlayerController)
-		{
-			PlayerController->ReceivedGameModeClass(GameModeClass);
-		}
-	}
-}
-
-void AGameState::ReceivedSpectatorClass()
-{
-	// Tell each PlayerController that the Spectator class is here
-	for (FConstPlayerControllerIterator Iterator = GetWorld()->GetPlayerControllerIterator(); Iterator; ++Iterator)
-	{
-		APlayerController* const PlayerController = *Iterator;
-		if (PlayerController && PlayerController->IsLocalController())
-		{
-			PlayerController->ReceivedSpectatorClass(SpectatorClass);
-		}
-	}
-}
-
-void AGameState::SeamlessTravelTransitionCheckpoint(bool bToTransitionMap)
-{
-	UWorld* World = GetWorld();
-
-	// mark all existing player states as from previous level for various bookkeeping
-	for (int32 i = 0; i < World->GameState->PlayerArray.Num(); i++)
-	{
-		World->GameState->PlayerArray[i]->bFromPreviousLevel = true;
-	}
-}
-
-void AGameState::AddPlayerState(APlayerState* PlayerState)
-{
-	// Determine whether it should go in the active or inactive list
-	if (!PlayerState->bIsInactive)
-	{
-		// make sure no duplicates
-		PlayerArray.AddUnique(PlayerState);
-	}
-}
-
-void AGameState::RemovePlayerState(APlayerState* PlayerState)
-{
-	for (int32 i=0; i<PlayerArray.Num(); i++)
-	{
-		if (PlayerArray[i] == PlayerState)
-		{
-			PlayerArray.RemoveAt(i,1);
-			return;
-		}
-	}
 }
 
 void AGameState::HandleMatchIsWaitingToStart()
@@ -168,6 +61,11 @@ void AGameState::HandleMatchHasStarted()
 		// Server handles this in AGameMode::HandleMatchHasStarted
 		GetWorldSettings()->NotifyMatchStarted();
 	}
+	else
+	{
+		// Now that match has started, act like the base class and set replicated flag
+		bReplicatedHasBegunPlay = true;
+	}
 }
 
 void AGameState::HandleMatchHasEnded()
@@ -190,6 +88,11 @@ bool AGameState::HasMatchStarted() const
 	return true;
 }
 
+void AGameState::HandleBeginPlay()
+{
+	// Overridden to not do anything, the state machine tells world when to start
+}
+
 bool AGameState::IsMatchInProgress() const
 {
 	if (GetMatchState() == MatchState::InProgress)
@@ -249,41 +152,24 @@ void AGameState::OnRep_MatchState()
 
 void AGameState::OnRep_ElapsedTime()
 {
-	//Blank on purpose
+	// Blank on purpose
 }
 
-bool AGameState::ShouldShowGore() const
+float AGameState::GetPlayerStartTime(class AController* Controller) const
 {
-	return true;
+	return ElapsedTime;
 }
 
-float AGameState::GetServerWorldTimeSeconds() const
+float AGameState::GetPlayerRespawnDelay(class AController* Controller) const
 {
-	UWorld* World = GetWorld();
-	if (World)
-	{
-		return World->GetTimeSeconds() + ServerWorldTimeSecondsDelta;
-	}
+	const AGameMode* GameMode = GetDefaultGameMode<AGameMode>();
 
-	return 0.f;
-}
-
-void AGameState::UpdateServerTimeSeconds()
-{
-	UWorld* World = GetWorld();
-	if (World)
+	if (GameMode)
 	{
-		ReplicatedWorldTimeSeconds = World->GetTimeSeconds();
+		return GameMode->MinRespawnDelay;
 	}
-}
 
-void AGameState::OnRep_ReplicatedWorldTimeSeconds()
-{
-	UWorld* World = GetWorld();
-	if (World)
-	{
-		ServerWorldTimeSecondsDelta = ReplicatedWorldTimeSeconds - World->GetTimeSeconds();
-	}
+	return Super::GetPlayerRespawnDelay(Controller);
 }
 
 void AGameState::GetLifetimeReplicatedProps( TArray< FLifetimeProperty > & OutLifetimeProps ) const
@@ -291,10 +177,5 @@ void AGameState::GetLifetimeReplicatedProps( TArray< FLifetimeProperty > & OutLi
 	Super::GetLifetimeReplicatedProps( OutLifetimeProps );
 
 	DOREPLIFETIME( AGameState, MatchState );
-	DOREPLIFETIME( AGameState, SpectatorClass );
-
-	DOREPLIFETIME_CONDITION( AGameState, GameModeClass,	COND_InitialOnly );
 	DOREPLIFETIME_CONDITION( AGameState, ElapsedTime,	COND_InitialOnly );
-
-	DOREPLIFETIME( AGameState, ReplicatedWorldTimeSeconds );
 }
\ No newline at end of file
diff --git a/Engine/Source/Runtime/Engine/Private/GameViewportClient.cpp b/Engine/Source/Runtime/Engine/Private/GameViewportClient.cpp
--- a/Engine/Source/Runtime/Engine/Private/GameViewportClient.cpp
+++ b/Engine/Source/Runtime/Engine/Private/GameViewportClient.cpp
@@ -1456,6 +1456,11 @@ void UGameViewportClient::Deactivated(FViewport* InViewport, const FWindowActiva
 	LostFocus(InViewport);
 }
 
+bool UGameViewportClient::WindowCloseRequested()
+{
+	return !WindowCloseRequestedDelegate.IsBound() || WindowCloseRequestedDelegate.Execute();
+}
+
 void UGameViewportClient::CloseRequested(FViewport* InViewport)
 {
 	check(InViewport == Viewport);
diff --git a/Engine/Source/Runtime/Engine/Private/GameplayStatics.cpp b/Engine/Source/Runtime/Engine/Private/GameplayStatics.cpp
--- a/Engine/Source/Runtime/Engine/Private/GameplayStatics.cpp
+++ b/Engine/Source/Runtime/Engine/Private/GameplayStatics.cpp
@@ -13,7 +13,8 @@
 #include "Sound/SoundBase.h"
 #include "Sound/SoundCue.h"
 #include "Engine/GameInstance.h"
-#include "GameFramework/GameMode.h"
+#include "GameFramework/GameModeBase.h"
+#include "GameFramework/GameStateBase.h"
 #include "GameFramework/SaveGame.h"
 #include "GameFramework/DamageType.h"
 #include "PhysicalMaterials/PhysicalMaterial.h"
@@ -52,13 +53,13 @@ UGameplayStatics::UGameplayStatics(const FObjectInitializer& ObjectInitializer)
 {
 }
 
-UGameInstance* UGameplayStatics::GetGameInstance(UObject* WorldContextObject)
+class UGameInstance* UGameplayStatics::GetGameInstance(const UObject* WorldContextObject)
 {
 	UWorld* World = GEngine->GetWorldFromContextObject(WorldContextObject);
 	return World ? World->GetGameInstance() : nullptr;
 }
 
-APlayerController* UGameplayStatics::GetPlayerController(UObject* WorldContextObject, int32 PlayerIndex ) 
+class APlayerController* UGameplayStatics::GetPlayerController(const UObject* WorldContextObject, int32 PlayerIndex ) 
 {
 	if (UWorld* World = GEngine->GetWorldFromContextObject(WorldContextObject))
 	{
@@ -76,25 +77,25 @@ APlayerController* UGameplayStatics::GetPlayerController(UObject* WorldContextOb
 	return nullptr;
 }
 
-ACharacter* UGameplayStatics::GetPlayerCharacter(UObject* WorldContextObject, int32 PlayerIndex)
+ACharacter* UGameplayStatics::GetPlayerCharacter(const UObject* WorldContextObject, int32 PlayerIndex)
 {
 	APlayerController* PC = GetPlayerController(WorldContextObject, PlayerIndex);
 	return PC ? Cast<ACharacter>(PC->GetPawn()) : nullptr;
 }
 
-APawn* UGameplayStatics::GetPlayerPawn(UObject* WorldContextObject, int32 PlayerIndex)
+APawn* UGameplayStatics::GetPlayerPawn(const UObject* WorldContextObject, int32 PlayerIndex)
 {
 	APlayerController* PC = GetPlayerController(WorldContextObject, PlayerIndex);
 	return PC ? PC->GetPawnOrSpectator() : nullptr;
 }
 
-APlayerCameraManager* UGameplayStatics::GetPlayerCameraManager(UObject* WorldContextObject, int32 PlayerIndex)
+APlayerCameraManager* UGameplayStatics::GetPlayerCameraManager(const UObject* WorldContextObject, int32 PlayerIndex)
 {
 	APlayerController* const PC = GetPlayerController(WorldContextObject, PlayerIndex);
 	return PC ? PC->PlayerCameraManager : nullptr;
 }
 
-APlayerController* UGameplayStatics::CreatePlayer(UObject* WorldContextObject, int32 ControllerId, bool bSpawnPawn)
+APlayerController* UGameplayStatics::CreatePlayer(const UObject* WorldContextObject, int32 ControllerId, bool bSpawnPawn)
 {
 	UWorld* World = GEngine->GetWorldFromContextObject(WorldContextObject);
 	FString Error;
@@ -151,16 +152,16 @@ void UGameplayStatics::SetPlayerControllerID(APlayerController* PlayerController
 	}
 }
 
-AGameMode* UGameplayStatics::GetGameMode(UObject* WorldContextObject)
+AGameModeBase* UGameplayStatics::GetGameMode(const UObject* WorldContextObject)
 {
 	UWorld* const World = GEngine->GetWorldFromContextObject(WorldContextObject);
 	return World ? World->GetAuthGameMode() : NULL;
 }
 
-AGameState* UGameplayStatics::GetGameState(UObject* WorldContextObject)
+AGameStateBase* UGameplayStatics::GetGameState(const UObject* WorldContextObject)
 {
 	UWorld* const World = GEngine->GetWorldFromContextObject(WorldContextObject);
-	return World ? World->GameState : nullptr;
+	return World ? World->GetGameState() : nullptr;
 }
 
 class UClass* UGameplayStatics::GetObjectClass(const UObject* Object)
@@ -168,13 +169,13 @@ class UClass* UGameplayStatics::GetObjectClass(const UObject* Object)
 	return Object ? Object->GetClass() : nullptr;
 }
 
-float UGameplayStatics::GetGlobalTimeDilation(UObject* WorldContextObject)
+float UGameplayStatics::GetGlobalTimeDilation(const UObject* WorldContextObject)
 {
 	UWorld* const World = GEngine->GetWorldFromContextObject( WorldContextObject );
 	return World ? World->GetWorldSettings()->TimeDilation : 1.f;
 }
 
-void UGameplayStatics::SetGlobalTimeDilation(UObject* WorldContextObject, float TimeDilation)
+void UGameplayStatics::SetGlobalTimeDilation(const UObject* WorldContextObject, float TimeDilation)
 {
 	UWorld* const World = GEngine->GetWorldFromContextObject( WorldContextObject );
 	if (World != nullptr)
@@ -191,14 +192,14 @@ void UGameplayStatics::SetGlobalTimeDilation(UObject* WorldContextObject, float
 	}
 }
 
-bool UGameplayStatics::SetGamePaused(UObject* WorldContextObject, bool bPaused)
+bool UGameplayStatics::SetGamePaused(const UObject* WorldContextObject, bool bPaused)
 {
 	UGameInstance* const GameInstance = GetGameInstance( WorldContextObject );
 	APlayerController* const PC = GameInstance ? GameInstance->GetFirstLocalPlayerController() : nullptr;
 	return PC ? PC->SetPause(bPaused) : false;
 }
 
-bool UGameplayStatics::IsGamePaused(UObject* WorldContextObject)
+bool UGameplayStatics::IsGamePaused(const UObject* WorldContextObject)
 {
 	UWorld* const World = GEngine->GetWorldFromContextObject( WorldContextObject );
 	return World ? World->IsPaused() : false;
@@ -249,13 +250,13 @@ static bool ComponentIsDamageableFrom(UPrimitiveComponent* VictimComp, FVector c
 	return true;
 }
 
-bool UGameplayStatics::ApplyRadialDamage(UObject* WorldContextObject, float BaseDamage, const FVector& Origin, float DamageRadius, TSubclassOf<UDamageType> DamageTypeClass, const TArray<AActor*>& IgnoreActors, AActor* DamageCauser, AController* InstigatedByController, bool bDoFullDamage, ECollisionChannel DamagePreventionChannel )
+bool UGameplayStatics::ApplyRadialDamage(const UObject* WorldContextObject, float BaseDamage, const FVector& Origin, float DamageRadius, TSubclassOf<UDamageType> DamageTypeClass, const TArray<AActor*>& IgnoreActors, AActor* DamageCauser, AController* InstigatedByController, bool bDoFullDamage, ECollisionChannel DamagePreventionChannel )
 {
 	float DamageFalloff = bDoFullDamage ? 0.f : 1.f;
 	return ApplyRadialDamageWithFalloff(WorldContextObject, BaseDamage, 0.f, Origin, 0.f, DamageRadius, DamageFalloff, DamageTypeClass, IgnoreActors, DamageCauser, InstigatedByController, DamagePreventionChannel);
 }
 
-bool UGameplayStatics::ApplyRadialDamageWithFalloff(UObject* WorldContextObject, float BaseDamage, float MinimumDamage, const FVector& Origin, float DamageInnerRadius, float DamageOuterRadius, float DamageFalloff, TSubclassOf<class UDamageType> DamageTypeClass, const TArray<AActor*>& IgnoreActors, AActor* DamageCauser, AController* InstigatedByController, ECollisionChannel DamagePreventionChannel)
+bool UGameplayStatics::ApplyRadialDamageWithFalloff(const UObject* WorldContextObject, float BaseDamage, float MinimumDamage, const FVector& Origin, float DamageInnerRadius, float DamageOuterRadius, float DamageFalloff, TSubclassOf<class UDamageType> DamageTypeClass, const TArray<AActor*>& IgnoreActors, AActor* DamageCauser, AController* InstigatedByController, ECollisionChannel DamagePreventionChannel)
 {
 	static FName NAME_ApplyRadialDamage = FName(TEXT("ApplyRadialDamage"));
 	FCollisionQueryParams SphereParams(NAME_ApplyRadialDamage, false, DamageCauser);
@@ -413,7 +414,7 @@ UObject* UGameplayStatics::SpawnObject(TSubclassOf<UObject> ObjectClass, UObject
 	return NewObject<UObject>(Outer, ObjectClass, NAME_None, RF_StrongRefOnFrame);
 }
 
-AActor* UGameplayStatics::BeginSpawningActorFromBlueprint(UObject* WorldContextObject, UBlueprint const* Blueprint, const FTransform& SpawnTransform, bool bNoCollisionFail)
+class AActor* UGameplayStatics::BeginSpawningActorFromBlueprint(UObject* WorldContextObject, const class UBlueprint* Blueprint, const FTransform& SpawnTransform, bool bNoCollisionFail)
 {
 	if (Blueprint && Blueprint->GeneratedClass)
 	{
@@ -431,13 +432,13 @@ AActor* UGameplayStatics::BeginSpawningActorFromBlueprint(UObject* WorldContextO
 }
 
 // deprecated
-AActor* UGameplayStatics::BeginSpawningActorFromClass(UObject* WorldContextObject, TSubclassOf<AActor> ActorClass, const FTransform& SpawnTransform, bool bNoCollisionFail, AActor* Owner)
+class AActor* UGameplayStatics::BeginSpawningActorFromClass(UObject* WorldContextObject, TSubclassOf<AActor> ActorClass, const FTransform& SpawnTransform, bool bNoCollisionFail /*= false*/, AActor* Owner /*= nullptr*/)
 {
 	ESpawnActorCollisionHandlingMethod const CollisionHandlingOverride = bNoCollisionFail ? ESpawnActorCollisionHandlingMethod::AdjustIfPossibleButDontSpawnIfColliding : ESpawnActorCollisionHandlingMethod::AlwaysSpawn;
 	return BeginDeferredActorSpawnFromClass(WorldContextObject, ActorClass, SpawnTransform, CollisionHandlingOverride, Owner);
 }
 
-AActor* UGameplayStatics::BeginDeferredActorSpawnFromClass(UObject* WorldContextObject, TSubclassOf<AActor> ActorClass, const FTransform& SpawnTransform, ESpawnActorCollisionHandlingMethod CollisionHandlingMethod, AActor* Owner)
+class AActor* UGameplayStatics::BeginDeferredActorSpawnFromClass(UObject* WorldContextObject, TSubclassOf<AActor> ActorClass, const FTransform& SpawnTransform, ESpawnActorCollisionHandlingMethod CollisionHandlingOverride /*= ESpawnActorCollisionHandlingMethod::Undefined*/, AActor* Owner /*= nullptr*/)
 {
 	if (UClass* Class = *ActorClass)
 	{
@@ -455,7 +456,7 @@ AActor* UGameplayStatics::BeginDeferredActorSpawnFromClass(UObject* WorldContext
 
 		if (UWorld* World = GEngine->GetWorldFromContextObject(WorldContextObject))
 		{
-			return World->SpawnActorDeferred<AActor>(Class, SpawnTransform, Owner, AutoInstigator, CollisionHandlingMethod);
+			return World->SpawnActorDeferred<AActor>(Class, SpawnTransform, Owner, AutoInstigator, CollisionHandlingOverride);
 		}
 		else
 		{
@@ -479,7 +480,7 @@ AActor* UGameplayStatics::FinishSpawningActor(AActor* Actor, const FTransform& S
 	return Actor;
 }
 
-void UGameplayStatics::LoadStreamLevel(UObject* WorldContextObject, FName LevelName,bool bMakeVisibleAfterLoad,bool bShouldBlockOnLoad,FLatentActionInfo LatentInfo)
+void UGameplayStatics::LoadStreamLevel(const UObject* WorldContextObject, FName LevelName,bool bMakeVisibleAfterLoad,bool bShouldBlockOnLoad,FLatentActionInfo LatentInfo)
 {
 	if (UWorld* World = GEngine->GetWorldFromContextObject(WorldContextObject))
 	{
@@ -492,7 +493,7 @@ void UGameplayStatics::LoadStreamLevel(UObject* WorldContextObject, FName LevelN
 	}
 }
 
-void UGameplayStatics::UnloadStreamLevel(UObject* WorldContextObject, FName LevelName,FLatentActionInfo LatentInfo)
+void UGameplayStatics::UnloadStreamLevel(const UObject* WorldContextObject, FName LevelName,FLatentActionInfo LatentInfo)
 {
 	if (UWorld* World = GEngine->GetWorldFromContextObject(WorldContextObject))
 	{
@@ -505,7 +506,7 @@ void UGameplayStatics::UnloadStreamLevel(UObject* WorldContextObject, FName Leve
 	}
 }
 
-ULevelStreaming* UGameplayStatics::GetStreamingLevel(UObject* WorldContextObject, FName InPackageName)
+ULevelStreaming* UGameplayStatics::GetStreamingLevel(const UObject* WorldContextObject, FName InPackageName)
 {
 	if (InPackageName != NAME_None)
 	{
@@ -533,7 +534,7 @@ ULevelStreaming* UGameplayStatics::GetStreamingLevel(UObject* WorldContextObject
 	return NULL;
 }
 
-void UGameplayStatics::FlushLevelStreaming(UObject* WorldContextObject)
+void UGameplayStatics::FlushLevelStreaming(const UObject* WorldContextObject)
 {
 	if (UWorld* World = GEngine->GetWorldFromContextObject(WorldContextObject))
 	{
@@ -546,7 +547,7 @@ void UGameplayStatics::CancelAsyncLoading()
 	::CancelAsyncLoading();
 }
 
-void UGameplayStatics::OpenLevel(UObject* WorldContextObject, FName LevelName, bool bAbsolute, FString Options)
+void UGameplayStatics::OpenLevel(const UObject* WorldContextObject, FName LevelName, bool bAbsolute, FString Options)
 {
 	UWorld* World = GEngine->GetWorldFromContextObject(WorldContextObject);
 	const ETravelType TravelType = (bAbsolute ? TRAVEL_Absolute : TRAVEL_Relative);
@@ -569,7 +570,7 @@ void UGameplayStatics::OpenLevel(UObject* WorldContextObject, FName LevelName, b
 	GEngine->SetClientTravel( World, *Cmd, TravelType );
 }
 
-FString UGameplayStatics::GetCurrentLevelName(UObject* WorldContextObject, bool bRemovePrefixString)
+FString UGameplayStatics::GetCurrentLevelName(const UObject* WorldContextObject, bool bRemovePrefixString)
 {
 	if (UWorld* World = GEngine->GetWorldFromContextObject(WorldContextObject))
 	{
@@ -631,7 +632,7 @@ void UGameplayStatics::GetActorArrayBounds(const TArray<AActor*>& Actors, bool b
 	}
 }
 
-void UGameplayStatics::GetAllActorsOfClass(UObject* WorldContextObject, TSubclassOf<AActor> ActorClass, TArray<AActor*>& OutActors)
+void UGameplayStatics::GetAllActorsOfClass(const UObject* WorldContextObject, TSubclassOf<AActor> ActorClass, TArray<AActor*>& OutActors)
 {
 	OutActors.Empty();
 
@@ -651,7 +652,7 @@ void UGameplayStatics::GetAllActorsOfClass(UObject* WorldContextObject, TSubclas
 	}
 }
 
-void UGameplayStatics::GetAllActorsWithInterface(UObject* WorldContextObject, TSubclassOf<UInterface> Interface, TArray<AActor*>& OutActors)
+void UGameplayStatics::GetAllActorsWithInterface(const UObject* WorldContextObject, TSubclassOf<UInterface> Interface, TArray<AActor*>& OutActors)
 {
 	OutActors.Empty();
 
@@ -670,7 +671,7 @@ void UGameplayStatics::GetAllActorsWithInterface(UObject* WorldContextObject, TS
 	}
 }
 
-void UGameplayStatics::GetAllActorsWithTag(UObject* WorldContextObject, FName Tag, TArray<AActor*>& OutActors)
+void UGameplayStatics::GetAllActorsWithTag(const UObject* WorldContextObject, FName Tag, TArray<AActor*>& OutActors)
 {
 	OutActors.Empty();
 
@@ -690,7 +691,7 @@ void UGameplayStatics::GetAllActorsWithTag(UObject* WorldContextObject, FName Ta
 	}
 }
 
-void UGameplayStatics::PlayWorldCameraShake(UObject* WorldContextObject, TSubclassOf<class UCameraShake> Shake, FVector Epicenter, float InnerRadius, float OuterRadius, float Falloff, bool bOrientShakeTowardsEpicenter)
+void UGameplayStatics::PlayWorldCameraShake(const UObject* WorldContextObject, TSubclassOf<class UCameraShake> Shake, FVector Epicenter, float InnerRadius, float OuterRadius, float Falloff, bool bOrientShakeTowardsEpicenter)
 {
 	UWorld* World = GEngine->GetWorldFromContextObject(WorldContextObject);
 	if(World != nullptr)
@@ -712,7 +713,7 @@ UParticleSystemComponent* CreateParticleSystem(UParticleSystem* EmitterTemplate,
 	return PSC;
 }
 
-UParticleSystemComponent* UGameplayStatics::SpawnEmitterAtLocation(UObject* WorldContextObject, UParticleSystem* EmitterTemplate, FVector SpawnLocation, FRotator SpawnRotation, bool bAutoDestroy)
+UParticleSystemComponent* UGameplayStatics::SpawnEmitterAtLocation(const UObject* WorldContextObject, UParticleSystem* EmitterTemplate, FVector SpawnLocation, FRotator SpawnRotation, bool bAutoDestroy)
 {
 	if (EmitterTemplate)
 	{
@@ -803,13 +804,25 @@ UParticleSystemComponent* UGameplayStatics::SpawnEmitterAttached(UParticleSystem
 					const FTransform RelativeTM = ComponentToWorld.GetRelativeTransform(ParentToWorld);
 					PSC->RelativeLocation = RelativeTM.GetLocation();
 					PSC->RelativeRotation = RelativeTM.GetRotation().Rotator();
+					PSC->RelativeScale3D = FVector(1.f);
 				}
 				else
 				{
 					PSC->RelativeLocation = Location;
 					PSC->RelativeRotation = Rotation;
+					
+					if (LocationType == EAttachLocation::SnapToTarget)
+					{
+						// SnapToTarget indicates we "keep world scale", and since we pass in no scale, we assume we want to keep a scale of 1,
+						// which indicates we want the inverse of the parent-to-world scale to maintain world scale of 1.
+						const FTransform ParentToWorld = AttachToComponent->GetSocketTransform(AttachPointName);
+						PSC->RelativeScale3D = ParentToWorld.GetSafeScaleReciprocal(ParentToWorld.GetScale3D());
+					}
+					else
+					{
+						PSC->RelativeScale3D = FVector(1.f);
+					}
 				}
-				PSC->RelativeScale3D = FVector(1.f);
 
 				PSC->RegisterComponentWithWorld(World);
 				PSC->ActivateSystem(true);
@@ -900,7 +913,7 @@ bool UGameplayStatics::FindCollisionUV(const struct FHitResult& Hit, int32 UVCha
 	return bSuccess;
 }
 
-bool UGameplayStatics::AreAnyListenersWithinRange(UObject* WorldContextObject, FVector Location, float MaximumRange)
+bool UGameplayStatics::AreAnyListenersWithinRange(const UObject* WorldContextObject, FVector Location, float MaximumRange)
 {
 	if (!GEngine || !GEngine->UseSound())
 	{
@@ -922,7 +935,7 @@ bool UGameplayStatics::AreAnyListenersWithinRange(UObject* WorldContextObject, F
 	return false;
 }
 
-void UGameplayStatics::SetGlobalPitchModulation(UObject* WorldContextObject, float PitchModulation, float TimeSec)
+void UGameplayStatics::SetGlobalPitchModulation(const UObject* WorldContextObject, float PitchModulation, float TimeSec)
 {
 	if (!GEngine || !GEngine->UseSound())
 	{
@@ -941,7 +954,7 @@ void UGameplayStatics::SetGlobalPitchModulation(UObject* WorldContextObject, flo
 	}
 }
 
-void UGameplayStatics::SetGlobalListenerFocusParameters(UObject* WorldContextObject, float FocusAzimuthScale, float NonFocusAzimuthScale, float FocusDistanceScale, float NonFocusDistanceScale, float FocusVolumeScale, float NonFocusVolumeScale, float FocusPriorityScale, float NonFocusPriorityScale)
+void UGameplayStatics::SetGlobalListenerFocusParameters(const UObject* WorldContextObject, float FocusAzimuthScale, float NonFocusAzimuthScale, float FocusDistanceScale, float NonFocusDistanceScale, float FocusVolumeScale, float NonFocusVolumeScale, float FocusPriorityScale, float NonFocusPriorityScale)
 {
 	if (!GEngine || !GEngine->UseSound())
 	{
@@ -970,7 +983,7 @@ void UGameplayStatics::SetGlobalListenerFocusParameters(UObject* WorldContextObj
 	}
 }
 
-void UGameplayStatics::PlaySound2D(UObject* WorldContextObject, class USoundBase* Sound, float VolumeMultiplier, float PitchMultiplier, float StartTime, class USoundConcurrency* ConcurrencySettings)
+void UGameplayStatics::PlaySound2D(const UObject* WorldContextObject, class USoundBase* Sound, float VolumeMultiplier, float PitchMultiplier, float StartTime, class USoundConcurrency* ConcurrencySettings)
 {
 	if (!Sound || !GEngine || !GEngine->UseSound())
 	{
@@ -1002,7 +1015,7 @@ void UGameplayStatics::PlaySound2D(UObject* WorldContextObject, class USoundBase
 	}
 }
 
-UAudioComponent* UGameplayStatics::CreateSound2D(UObject* WorldContextObject, class USoundBase* Sound, float VolumeMultiplier, float PitchMultiplier, float StartTime, class USoundConcurrency* ConcurrencySettings)
+UAudioComponent* UGameplayStatics::CreateSound2D(const UObject* WorldContextObject, class USoundBase* Sound, float VolumeMultiplier, float PitchMultiplier, float StartTime, class USoundConcurrency* ConcurrencySettings, bool bPersistAcrossLevelTransition)
 {
 	if (!Sound || !GEngine || !GEngine->UseSound())
 	{
@@ -1015,23 +1028,39 @@ UAudioComponent* UGameplayStatics::CreateSound2D(UObject* WorldContextObject, cl
 		return nullptr;
 	}
 
-	UAudioComponent* AudioComponent = FAudioDevice::CreateComponent(Sound, ThisWorld, ThisWorld->GetWorldSettings(), false, false, nullptr, nullptr, ConcurrencySettings);
+	UAudioComponent* AudioComponent;
+
+	if (bPersistAcrossLevelTransition)
+	{
+		FAudioDevice::FCreateComponentParams Params(ThisWorld->GetAudioDevice());
+		Params.ConcurrencySettings = ConcurrencySettings;
+		
+		AudioComponent = FAudioDevice::CreateComponent(Sound, Params);
+	}
+	else
+	{
+		FAudioDevice::FCreateComponentParams Params(ThisWorld);
+		Params.ConcurrencySettings = ConcurrencySettings;
+
+		AudioComponent = FAudioDevice::CreateComponent(Sound, Params);
+	}
+	
 	if (AudioComponent)
 	{
-		const bool bIsInGameWorld = AudioComponent->GetWorld()->IsGameWorld();
 		AudioComponent->SetVolumeMultiplier(VolumeMultiplier);
 		AudioComponent->SetPitchMultiplier(PitchMultiplier);
 		AudioComponent->bAllowSpatialization = false;
 		AudioComponent->bIsUISound = true;
 		AudioComponent->bAutoDestroy = true;
+		AudioComponent->bIgnoreForFlushing = bPersistAcrossLevelTransition;
 		AudioComponent->SubtitlePriority = Sound->GetSubtitlePriority();
 	}
 	return AudioComponent;
 }
 
-UAudioComponent* UGameplayStatics::SpawnSound2D(UObject* WorldContextObject, class USoundBase* Sound, float VolumeMultiplier, float PitchMultiplier, float StartTime, class USoundConcurrency* ConcurrencySettings)
+UAudioComponent* UGameplayStatics::SpawnSound2D(const UObject* WorldContextObject, class USoundBase* Sound, float VolumeMultiplier, float PitchMultiplier, float StartTime, class USoundConcurrency* ConcurrencySettings, bool bPersistAcrossLevelTransition)
 {
-	UAudioComponent* AudioComponent = CreateSound2D(WorldContextObject, Sound, VolumeMultiplier, PitchMultiplier, StartTime, ConcurrencySettings);
+	UAudioComponent* AudioComponent = CreateSound2D(WorldContextObject, Sound, VolumeMultiplier, PitchMultiplier, StartTime, ConcurrencySettings, bPersistAcrossLevelTransition);
 	if (AudioComponent)
 	{
 		AudioComponent->Play(StartTime);
@@ -1039,7 +1068,7 @@ UAudioComponent* UGameplayStatics::SpawnSound2D(UObject* WorldContextObject, cla
 	return AudioComponent;
 }
 
-void UGameplayStatics::PlaySoundAtLocation(UObject* WorldContextObject, class USoundBase* Sound, FVector Location, FRotator Rotation, float VolumeMultiplier, float PitchMultiplier, float StartTime, class USoundAttenuation* AttenuationSettings, class USoundConcurrency* ConcurrencySettings)
+void UGameplayStatics::PlaySoundAtLocation(const UObject* WorldContextObject, class USoundBase* Sound, FVector Location, FRotator Rotation, float VolumeMultiplier, float PitchMultiplier, float StartTime, class USoundAttenuation* AttenuationSettings, class USoundConcurrency* ConcurrencySettings)
 {
 	if (!Sound || !GEngine || !GEngine->UseSound())
 	{
@@ -1058,7 +1087,7 @@ void UGameplayStatics::PlaySoundAtLocation(UObject* WorldContextObject, class US
 	}
 }
 
-UAudioComponent* UGameplayStatics::SpawnSoundAtLocation(UObject* WorldContextObject, class USoundBase* Sound, FVector Location, FRotator Rotation, float VolumeMultiplier, float PitchMultiplier, float StartTime, class USoundAttenuation* AttenuationSettings, class USoundConcurrency* ConcurrencySettings)
+UAudioComponent* UGameplayStatics::SpawnSoundAtLocation(const UObject* WorldContextObject, class USoundBase* Sound, FVector Location, FRotator Rotation, float VolumeMultiplier, float PitchMultiplier, float StartTime, class USoundAttenuation* AttenuationSettings, class USoundConcurrency* ConcurrencySettings)
 {
 	if (!Sound || !GEngine || !GEngine->UseSound())
 	{
@@ -1073,7 +1102,12 @@ UAudioComponent* UGameplayStatics::SpawnSoundAtLocation(UObject* WorldContextObj
 
 	const bool bIsInGameWorld = ThisWorld->IsGameWorld();
 
-	UAudioComponent* AudioComponent = FAudioDevice::CreateComponent(Sound, ThisWorld, ThisWorld->GetWorldSettings(), false, false, &Location, AttenuationSettings, ConcurrencySettings);
+	FAudioDevice::FCreateComponentParams Params(ThisWorld);
+	Params.SetLocation(Location);
+	Params.AttenuationSettings = AttenuationSettings;
+	Params.ConcurrencySettings = ConcurrencySettings;
+
+	UAudioComponent* AudioComponent = FAudioDevice::CreateComponent(Sound, Params);
 
 	if (AudioComponent)
 	{
@@ -1117,33 +1151,42 @@ class UAudioComponent* UGameplayStatics::SpawnSoundAttached(class USoundBase* So
 		}
 	}
 
-	UAudioComponent* AudioComponent = FAudioDevice::CreateComponent(Sound, AttachToComponent->GetWorld(), AttachToComponent->GetOwner(), false, bStopWhenAttachedToDestroyed, &TestLocation, AttenuationSettings, ConcurrencySettings);
-	if (AudioComponent && AudioComponent->GetWorld())
-	{
-		const bool bIsInGameWorld = AudioComponent->GetWorld()->IsGameWorld();
+	FAudioDevice::FCreateComponentParams Params(AttachToComponent->GetOwner());
+	Params.SetLocation(TestLocation);
+	Params.bStopWhenOwnerDestroyed = bStopWhenAttachedToDestroyed;
+	Params.AttenuationSettings = AttenuationSettings;
+	Params.ConcurrencySettings = ConcurrencySettings;
 
-		AudioComponent->AttachToComponent(AttachToComponent, FAttachmentTransformRules::KeepRelativeTransform, AttachPointName);
-		if (LocationType == EAttachLocation::KeepWorldPosition)
-		{
-			AudioComponent->SetWorldLocationAndRotation(Location, Rotation);
-		}
-		else
+	UAudioComponent* AudioComponent = FAudioDevice::CreateComponent(Sound, Params);
+	if (AudioComponent)
+	{
+		if (UWorld* ComponentWorld = AudioComponent->GetWorld())
 		{
-			AudioComponent->SetRelativeLocationAndRotation(Location, Rotation);
+			const bool bIsInGameWorld = ComponentWorld->IsGameWorld();
+
+			AudioComponent->AttachToComponent(AttachToComponent, FAttachmentTransformRules::KeepRelativeTransform, AttachPointName);
+			if (LocationType == EAttachLocation::KeepWorldPosition)
+			{
+				AudioComponent->SetWorldLocationAndRotation(Location, Rotation);
+			}
+			else
+			{
+				AudioComponent->SetRelativeLocationAndRotation(Location, Rotation);
+			}
+			AudioComponent->SetVolumeMultiplier(VolumeMultiplier);
+			AudioComponent->SetPitchMultiplier(PitchMultiplier);
+			AudioComponent->bAllowSpatialization = bIsInGameWorld;
+			AudioComponent->bIsUISound = !bIsInGameWorld;
+			AudioComponent->bAutoDestroy = true;
+			AudioComponent->SubtitlePriority = DEFAULT_SUBTITLE_PRIORITY; // Fixme: pass in? Do we want that exposed to blueprints though?
+			AudioComponent->Play(StartTime);
 		}
-		AudioComponent->SetVolumeMultiplier(VolumeMultiplier);
-		AudioComponent->SetPitchMultiplier(PitchMultiplier);
-		AudioComponent->bAllowSpatialization	= bIsInGameWorld;
-		AudioComponent->bIsUISound				= !bIsInGameWorld;
-		AudioComponent->bAutoDestroy			= true;
-		AudioComponent->SubtitlePriority		= DEFAULT_SUBTITLE_PRIORITY; // Fixme: pass in? Do we want that exposed to blueprints though?
-		AudioComponent->Play(StartTime);
 	}
 
 	return AudioComponent;
 }
 
-void UGameplayStatics::PlayDialogue2D(UObject* WorldContextObject, class UDialogueWave* Dialogue, const FDialogueContext& Context, float VolumeMultiplier, float PitchMultiplier, float StartTime)
+void UGameplayStatics::PlayDialogue2D(const UObject* WorldContextObject, class UDialogueWave* Dialogue, const FDialogueContext& Context, float VolumeMultiplier, float PitchMultiplier, float StartTime)
 {
 	if (Dialogue)
 	{
@@ -1151,7 +1194,7 @@ void UGameplayStatics::PlayDialogue2D(UObject* WorldContextObject, class UDialog
 	}
 }
 
-UAudioComponent* UGameplayStatics::SpawnDialogue2D(UObject* WorldContextObject, class UDialogueWave* Dialogue, const FDialogueContext& Context, float VolumeMultiplier, float PitchMultiplier, float StartTime)
+UAudioComponent* UGameplayStatics::SpawnDialogue2D(const UObject* WorldContextObject, class UDialogueWave* Dialogue, const FDialogueContext& Context, float VolumeMultiplier, float PitchMultiplier, float StartTime)
 {
 	if (Dialogue)
 	{
@@ -1160,7 +1203,7 @@ UAudioComponent* UGameplayStatics::SpawnDialogue2D(UObject* WorldContextObject,
 	return nullptr;
 }
 
-void UGameplayStatics::PlayDialogueAtLocation(UObject* WorldContextObject, class UDialogueWave* Dialogue, const FDialogueContext& Context, FVector Location, FRotator Rotation, float VolumeMultiplier, float PitchMultiplier, float StartTime, class USoundAttenuation* AttenuationSettings)
+void UGameplayStatics::PlayDialogueAtLocation(const UObject* WorldContextObject, class UDialogueWave* Dialogue, const FDialogueContext& Context, FVector Location, FRotator Rotation, float VolumeMultiplier, float PitchMultiplier, float StartTime, class USoundAttenuation* AttenuationSettings)
 {
 	if (Dialogue)
 	{
@@ -1168,7 +1211,7 @@ void UGameplayStatics::PlayDialogueAtLocation(UObject* WorldContextObject, class
 	}
 }
 
-UAudioComponent* UGameplayStatics::SpawnDialogueAtLocation(UObject* WorldContextObject, class UDialogueWave* Dialogue, const FDialogueContext& Context, FVector Location, FRotator Rotation, float VolumeMultiplier, float PitchMultiplier, float StartTime, class USoundAttenuation* AttenuationSettings)
+UAudioComponent* UGameplayStatics::SpawnDialogueAtLocation(const UObject* WorldContextObject, class UDialogueWave* Dialogue, const FDialogueContext& Context, FVector Location, FRotator Rotation, float VolumeMultiplier, float PitchMultiplier, float StartTime, class USoundAttenuation* AttenuationSettings)
 {
 	if (Dialogue)
 	{
@@ -1203,7 +1246,7 @@ bool UGameplayStatics::AreSubtitlesEnabled()
 	return 0;
 }
 
-void UGameplayStatics::SetBaseSoundMix(UObject* WorldContextObject, USoundMix* InSoundMix)
+void UGameplayStatics::SetBaseSoundMix(const UObject* WorldContextObject, USoundMix* InSoundMix)
 {
 	if (!InSoundMix || !GEngine || !GEngine->UseSound())
 	{
@@ -1222,7 +1265,7 @@ void UGameplayStatics::SetBaseSoundMix(UObject* WorldContextObject, USoundMix* I
 	}
 }
 
-void UGameplayStatics::PushSoundMixModifier(UObject* WorldContextObject, USoundMix* InSoundMixModifier)
+void UGameplayStatics::PushSoundMixModifier(const UObject* WorldContextObject, USoundMix* InSoundMixModifier)
 {
 	if (!InSoundMixModifier || !GEngine || !GEngine->UseSound())
 	{
@@ -1241,7 +1284,7 @@ void UGameplayStatics::PushSoundMixModifier(UObject* WorldContextObject, USoundM
 	}
 }
 
-void UGameplayStatics::SetSoundMixClassOverride(UObject* WorldContextObject, class USoundMix* InSoundMixModifier, class USoundClass* InSoundClass, float Volume, float Pitch, float FadeInTime, bool bApplyToChildren)
+void UGameplayStatics::SetSoundMixClassOverride(const UObject* WorldContextObject, class USoundMix* InSoundMixModifier, class USoundClass* InSoundClass, float Volume, float Pitch, float FadeInTime, bool bApplyToChildren)
 {
 	if (!InSoundMixModifier || !GEngine || !GEngine->UseSound())
 	{
@@ -1260,7 +1303,7 @@ void UGameplayStatics::SetSoundMixClassOverride(UObject* WorldContextObject, cla
 	}
 }
 
-void UGameplayStatics::ClearSoundMixClassOverride(UObject* WorldContextObject, class USoundMix* InSoundMixModifier, class USoundClass* InSoundClass, float FadeOutTime)
+void UGameplayStatics::ClearSoundMixClassOverride(const UObject* WorldContextObject, class USoundMix* InSoundMixModifier, class USoundClass* InSoundClass, float FadeOutTime)
 {
 	if (!InSoundMixModifier || !GEngine || !GEngine->UseSound())
 	{
@@ -1279,7 +1322,7 @@ void UGameplayStatics::ClearSoundMixClassOverride(UObject* WorldContextObject, c
 	}
 }
 
-void UGameplayStatics::PopSoundMixModifier(UObject* WorldContextObject, USoundMix* InSoundMixModifier)
+void UGameplayStatics::PopSoundMixModifier(const UObject* WorldContextObject, USoundMix* InSoundMixModifier)
 {
 	if (InSoundMixModifier == nullptr || GEngine == nullptr || !GEngine->UseSound())
 	{
@@ -1298,7 +1341,7 @@ void UGameplayStatics::PopSoundMixModifier(UObject* WorldContextObject, USoundMi
 	}
 }
 
-void UGameplayStatics::ClearSoundMixModifiers(UObject* WorldContextObject)
+void UGameplayStatics::ClearSoundMixModifiers(const UObject* WorldContextObject)
 {
 	if (!GEngine || !GEngine->UseSound())
 	{
@@ -1317,7 +1360,7 @@ void UGameplayStatics::ClearSoundMixModifiers(UObject* WorldContextObject)
 	}
 }
 
-void UGameplayStatics::ActivateReverbEffect(UObject* WorldContextObject, class UReverbEffect* ReverbEffect, FName TagName, float Priority, float Volume, float FadeTime)
+void UGameplayStatics::ActivateReverbEffect(const UObject* WorldContextObject, class UReverbEffect* ReverbEffect, FName TagName, float Priority, float Volume, float FadeTime)
 {
 	if (ReverbEffect == nullptr || !GEngine || !GEngine->UseSound())
 	{
@@ -1336,7 +1379,7 @@ void UGameplayStatics::ActivateReverbEffect(UObject* WorldContextObject, class U
 	}
 }
 
-void UGameplayStatics::DeactivateReverbEffect(UObject* WorldContextObject, FName TagName)
+void UGameplayStatics::DeactivateReverbEffect(const UObject* WorldContextObject, FName TagName)
 {
 	if (GEngine == nullptr || !GEngine->UseSound())
 	{
@@ -1355,7 +1398,7 @@ void UGameplayStatics::DeactivateReverbEffect(UObject* WorldContextObject, FName
 	}
 }
 
-class UReverbEffect* UGameplayStatics::GetCurrentReverbEffect(UObject* WorldContextObject)
+class UReverbEffect* UGameplayStatics::GetCurrentReverbEffect(const UObject* WorldContextObject)
 {
 	if (GEngine == nullptr || !GEngine->UseSound())
 	{
@@ -1392,7 +1435,7 @@ UDecalComponent* CreateDecalComponent(class UMaterialInterface* DecalMaterial, F
 	return DecalComp;
 }
 
-UDecalComponent* UGameplayStatics::SpawnDecalAtLocation(UObject* WorldContextObject, class UMaterialInterface* DecalMaterial, FVector DecalSize, FVector Location, FRotator Rotation, float LifeSpan)
+UDecalComponent* UGameplayStatics::SpawnDecalAtLocation(const UObject* WorldContextObject, class UMaterialInterface* DecalMaterial, FVector DecalSize, FVector Location, FRotator Rotation, float LifeSpan)
 {
 	if (DecalMaterial)
 	{
@@ -1610,25 +1653,37 @@ USaveGame* UGameplayStatics::LoadGameFromSlot(const FString& SlotName, const int
 	return OutSaveGameObject;
 }
 
-float UGameplayStatics::GetWorldDeltaSeconds(UObject* WorldContextObject)
+float UGameplayStatics::GetWorldDeltaSeconds(const UObject* WorldContextObject)
 {
 	UWorld* World = GEngine->GetWorldFromContextObject(WorldContextObject);
 	return World ? World->GetDeltaSeconds() : 0.f;
 }
 
-float UGameplayStatics::GetRealTimeSeconds(UObject* WorldContextObject)
+float UGameplayStatics::GetTimeSeconds(const UObject* WorldContextObject)
+{
+	UWorld* World = GEngine->GetWorldFromContextObject(WorldContextObject);
+	return World ? World->GetTimeSeconds() : 0.f;
+}
+
+float UGameplayStatics::GetUnpausedTimeSeconds(const UObject* WorldContextObject)
+{
+	UWorld* World = GEngine->GetWorldFromContextObject(WorldContextObject);
+	return World ? World->GetUnpausedTimeSeconds() : 0.f;
+}
+
+float UGameplayStatics::GetRealTimeSeconds(const UObject* WorldContextObject)
 {
 	UWorld* World = GEngine->GetWorldFromContextObject(WorldContextObject);
 	return World ? World->GetRealTimeSeconds() : 0.f;
 }
 
-float UGameplayStatics::GetAudioTimeSeconds(UObject* WorldContextObject)
+float UGameplayStatics::GetAudioTimeSeconds(const UObject* WorldContextObject)
 {
 	UWorld* World = GEngine->GetWorldFromContextObject(WorldContextObject);
 	return World ? World->GetAudioTimeSeconds() : 0.f;
 }
 
-void UGameplayStatics::GetAccurateRealTime(UObject* WorldContextObject, int32& Seconds, float& PartialSeconds)
+void UGameplayStatics::GetAccurateRealTime(const UObject* WorldContextObject, int32& Seconds, float& PartialSeconds)
 {
 	double TimeSeconds = FPlatformTime::Seconds();
 	Seconds = floor(TimeSeconds);
@@ -1649,7 +1704,7 @@ FString UGameplayStatics::GetPlatformName()
 	return FPlatformProperties::IniPlatformName();
 }
 
-bool UGameplayStatics::BlueprintSuggestProjectileVelocity(UObject* WorldContextObject, FVector& OutTossVelocity, FVector StartLocation, FVector EndLocation, float LaunchSpeed, float OverrideGravityZ, ESuggestProjVelocityTraceOption::Type TraceOption, float CollisionRadius, bool bFavorHighArc, bool bDrawDebug)
+bool UGameplayStatics::BlueprintSuggestProjectileVelocity(const UObject* WorldContextObject, FVector& OutTossVelocity, FVector StartLocation, FVector EndLocation, float LaunchSpeed, float OverrideGravityZ, ESuggestProjVelocityTraceOption::Type TraceOption, float CollisionRadius, bool bFavorHighArc, bool bDrawDebug)
 {
 	// simple pass-through to the C++ interface
 	return UGameplayStatics::SuggestProjectileVelocity(WorldContextObject, OutTossVelocity, StartLocation, EndLocation, LaunchSpeed, bFavorHighArc, CollisionRadius, OverrideGravityZ, TraceOption, FCollisionResponseParams::DefaultResponseParam, TArray<AActor*>(), bDrawDebug);
@@ -1659,7 +1714,7 @@ bool UGameplayStatics::BlueprintSuggestProjectileVelocity(UObject* WorldContextO
 static const FName NAME_SuggestProjVelTrace = FName(TEXT("SuggestProjVelTrace"));
 
 // Based on analytic solution to ballistic angle of launch http://en.wikipedia.org/wiki/Trajectory_of_a_projectile#Angle_required_to_hit_coordinate_.28x.2Cy.29
-bool UGameplayStatics::SuggestProjectileVelocity(UObject* WorldContextObject, FVector& OutTossVelocity, FVector Start, FVector End, float TossSpeed, bool bFavorHighArc, float CollisionRadius, float OverrideGravityZ, ESuggestProjVelocityTraceOption::Type TraceOption, const FCollisionResponseParams& ResponseParam, const TArray<AActor*>& ActorsToIgnore, bool bDrawDebug)
+bool UGameplayStatics::SuggestProjectileVelocity(const UObject* WorldContextObject, FVector& OutTossVelocity, FVector Start, FVector End, float TossSpeed, bool bFavorHighArc, float CollisionRadius, float OverrideGravityZ, ESuggestProjVelocityTraceOption::Type TraceOption, const FCollisionResponseParams& ResponseParam, const TArray<AActor*>& ActorsToIgnore, bool bDrawDebug)
 {
 	const FVector FlightDelta = End - Start;
 	const FVector DirXY = FlightDelta.GetSafeNormal2D();
@@ -1829,7 +1884,7 @@ bool UGameplayStatics::SuggestProjectileVelocity(UObject* WorldContextObject, FV
 static const FName NAME_PredictProjectilePath = FName(TEXT("PredictProjectilePath"));
 
 // note: this will automatically fall back to line test if radius is small enough
-bool UGameplayStatics::PredictProjectilePath(UObject* WorldContextObject, FHitResult& OutHit, TArray<FVector>& OutPathPositions, FVector& OutLastTraceDestination, FVector StartPos, FVector LaunchVelocity, bool bTracePath, float ProjectileRadius, const TArray<TEnumAsByte<EObjectTypeQuery> >& ObjectTypes, bool bTraceComplex, const TArray<AActor*>& ActorsToIgnore, EDrawDebugTrace::Type DrawDebugType, float DrawDebugTime, float SimFrequency /*= 30.f*/, float MaxSimTime /*= 2.f*/, float OverrideGravityZ /*= 0*/)
+bool UGameplayStatics::PredictProjectilePath(const UObject* WorldContextObject, FHitResult& OutHit, TArray<FVector>& OutPathPositions, FVector& OutLastTraceDestination, FVector StartPos, FVector LaunchVelocity, bool bTracePath, float ProjectileRadius, const TArray<TEnumAsByte<EObjectTypeQuery> >& ObjectTypes, bool bTraceComplex, const TArray<AActor*>& ActorsToIgnore, EDrawDebugTrace::Type DrawDebugType, float DrawDebugTime, float SimFrequency /*= 30.f*/, float MaxSimTime /*= 2.f*/, float OverrideGravityZ /*= 0*/)
 {
 	OutPathPositions.Empty();
 	bool bBlockingHit = false;
@@ -1910,7 +1965,7 @@ bool UGameplayStatics::PredictProjectilePath(UObject* WorldContextObject, FHitRe
 	return bBlockingHit;
 }
 
-bool UGameplayStatics::SuggestProjectileVelocity_CustomArc(UObject* WorldContextObject, FVector& OutLaunchVelocity, FVector StartPos, FVector EndPos, float OverrideGravityZ /*= 0*/, float ArcParam /*= 0.5f */)
+bool UGameplayStatics::SuggestProjectileVelocity_CustomArc(const UObject* WorldContextObject, FVector& OutLaunchVelocity, FVector StartPos, FVector EndPos, float OverrideGravityZ /*= 0*/, float ArcParam /*= 0.5f */)
 {
 	/* Make sure the start and end aren't the same location */
 	FVector const StartToEnd = EndPos - StartPos;
@@ -1948,13 +2003,13 @@ bool UGameplayStatics::SuggestProjectileVelocity_CustomArc(UObject* WorldContext
 	return false;
 }
 
-FIntVector UGameplayStatics::GetWorldOriginLocation(UObject* WorldContextObject)
+FIntVector UGameplayStatics::GetWorldOriginLocation(const UObject* WorldContextObject)
 {
 	UWorld* World = GEngine->GetWorldFromContextObject(WorldContextObject);
 	return World ? World->OriginLocation : FIntVector::ZeroValue;
 }
 
-void UGameplayStatics::SetWorldOriginLocation(UObject* WorldContextObject, FIntVector NewLocation)
+void UGameplayStatics::SetWorldOriginLocation(const UObject* WorldContextObject, FIntVector NewLocation)
 {
 	UWorld* World = GEngine->GetWorldFromContextObject(WorldContextObject);
 	if ( World )
@@ -1963,7 +2018,7 @@ void UGameplayStatics::SetWorldOriginLocation(UObject* WorldContextObject, FIntV
 	}
 }
 
-int32 UGameplayStatics::GrassOverlappingSphereCount(UObject* WorldContextObject, const UStaticMesh* Mesh, FVector CenterPosition, float Radius)
+int32 UGameplayStatics::GrassOverlappingSphereCount(const UObject* WorldContextObject, const UStaticMesh* Mesh, FVector CenterPosition, float Radius)
 {
 	int32 Count = 0;
 
@@ -2124,4 +2179,4 @@ int32 UGameplayStatics::GetIntOption( const FString& Options, const FString& Key
 	return DefaultValue;
 }
 
-#undef LOCTEXT_NAMESPACE
+#undef LOCTEXT_NAMESPACE
\ No newline at end of file
diff --git a/Engine/Source/Runtime/Engine/Private/HUD.cpp b/Engine/Source/Runtime/Engine/Private/HUD.cpp
--- a/Engine/Source/Runtime/Engine/Private/HUD.cpp
+++ b/Engine/Source/Runtime/Engine/Private/HUD.cpp
@@ -7,7 +7,7 @@
 #include "EnginePrivate.h"
 #include "Components/LineBatchComponent.h"
 #include "GameFramework/HUD.h"
-#include "GameFramework/GameMode.h"
+#include "GameFramework/GameModeBase.h"
 #include "Net/UnrealNetwork.h"
 #include "MessageLog.h"
 #include "UObjectToken.h"
@@ -429,7 +429,7 @@ void AHUD::ShowDebugInfo(float& YL, float& YPos)
 
 		if (ShouldDisplayDebug(NAME_Game))
 		{
-			AGameMode* AuthGameMode = GetWorld()->GetAuthGameMode();
+			AGameModeBase* AuthGameMode = GetWorld()->GetAuthGameMode();
 			if (AuthGameMode)
 			{
 				AuthGameMode->DisplayDebug(DebugCanvas, DisplayInfo, YL, YPos);
diff --git a/Engine/Source/Runtime/Engine/Private/Interpolation.cpp b/Engine/Source/Runtime/Engine/Private/Interpolation.cpp
--- a/Engine/Source/Runtime/Engine/Private/Interpolation.cpp
+++ b/Engine/Source/Runtime/Engine/Private/Interpolation.cpp
@@ -861,14 +861,17 @@ void AMatineeActor::InitInterp()
 
 		// Cache whether or not we want to enable extreme content within this sequence
 		bShouldShowGore = true;
+
+PRAGMA_DISABLE_DEPRECATION_WARNINGS
 		if( GetWorld() != NULL && GetWorld()->GetWorldSettings() != NULL )
 		{
-			AGameState const* const GameState = GetWorld()->GameState;
+			AGameState const* const GameState = GetWorld()->GetGameState<AGameState>();
 			if( GameState != NULL )
 			{
 				bShouldShowGore = GameState->ShouldShowGore();
 			}
 		}
+PRAGMA_ENABLE_DEPRECATION_WARNINGS
 
 		for(int32 GroupIndex=0; GroupIndex<MatineeData->InterpGroups.Num(); GroupIndex++)
 		{
@@ -8016,7 +8019,8 @@ void UInterpTrackSound::UpdateTrack(float NewPosition, UInterpTrackInst* TrInst,
 				else
 				{
 					// If there is no AudioComponent - create one now.
-					SoundInst->PlayAudioComp = FAudioDevice::CreateComponent(NewSound, SoundInst->GetWorld(), Actor, false, false);
+					FAudioDevice::FCreateComponentParams Params(SoundInst->GetWorld(), Actor);
+					SoundInst->PlayAudioComp = FAudioDevice::CreateComponent(NewSound, Params);
 					if(SoundInst->PlayAudioComp)
 					{
 						// If we have no actor to attach sound to - its location is meaningless, so we turn off spatialization.
@@ -8133,7 +8137,7 @@ void UInterpTrackSound::PreviewUpdateTrack(float NewPosition, UInterpTrackInst*
 				: false;
 
 		USoundCue* TempPlaybackAudioCue = NewObject<USoundCue>();
-		UAudioComponent* Component = FAudioDevice::CreateComponent(TempPlaybackAudioCue, NULL, NULL, false, false);
+		UAudioComponent* Component = FAudioDevice::CreateComponent(TempPlaybackAudioCue);
 
 		if ( bIsInRange && Component )
 		{
diff --git a/Engine/Source/Runtime/Engine/Private/LevelActor.cpp b/Engine/Source/Runtime/Engine/Private/LevelActor.cpp
--- a/Engine/Source/Runtime/Engine/Private/LevelActor.cpp
+++ b/Engine/Source/Runtime/Engine/Private/LevelActor.cpp
@@ -18,7 +18,7 @@
 #include "Components/CapsuleComponent.h"
 #include "Components/BoxComponent.h"
 #include "GameFramework/MovementComponent.h"
-#include "GameFramework/GameMode.h"
+#include "GameFramework/GameModeBase.h"
 
 #define LOCTEXT_NAMESPACE "LevelActor"
 
@@ -684,7 +684,7 @@ APlayerController* UWorld::SpawnPlayActor(UPlayer* NewPlayer, ENetRole RemoteRol
 		Options += InURL.Op[i];
 	}
 
-	AGameMode* GameMode = GetAuthGameMode();
+	AGameModeBase* GameMode = GetAuthGameMode();
 
 	// Give the GameMode a chance to accept the login
 	APlayerController* const NewPlayerController = GameMode->Login(NewPlayer, RemoteRole, *InURL.Portal, Options, UniqueId, Error);
diff --git a/Engine/Source/Runtime/Engine/Private/LevelTick.cpp b/Engine/Source/Runtime/Engine/Private/LevelTick.cpp
--- a/Engine/Source/Runtime/Engine/Private/LevelTick.cpp
+++ b/Engine/Source/Runtime/Engine/Private/LevelTick.cpp
@@ -1219,6 +1219,8 @@ void UWorld::Tick( ELevelTick TickType, float DeltaSeconds )
 	DeltaSeconds = GameDeltaSeconds;
 	DeltaTimeSeconds = DeltaSeconds;
 
+	UnpausedTimeSeconds += DeltaSeconds;
+
 	if ( !bIsPaused )
 	{
 		TimeSeconds += DeltaSeconds;
diff --git a/Engine/Source/Runtime/Engine/Private/LocalPlayer.cpp b/Engine/Source/Runtime/Engine/Private/LocalPlayer.cpp
--- a/Engine/Source/Runtime/Engine/Private/LocalPlayer.cpp
+++ b/Engine/Source/Runtime/Engine/Private/LocalPlayer.cpp
@@ -20,7 +20,8 @@
 #include "DataChannel.h"
 #include "GameFramework/OnlineSession.h"
 #include "GameFramework/PlayerInput.h"
-#include "GameFramework/GameMode.h"
+#include "GameFramework/GameModeBase.h"
+#include "GameFramework/GameStateBase.h"
 #include "GameFramework/PlayerState.h"
 
 #include "GameDelegates.h"
@@ -109,17 +110,17 @@ APlayerController* FLocalPlayerContext::GetPlayerController() const
 	return WorldPtr != nullptr ? LocalPlayer->GetPlayerController(WorldPtr) : LocalPlayer->PlayerController;
 }
 
-AGameState* FLocalPlayerContext::GetGameState() const
+class AGameStateBase* FLocalPlayerContext::GetGameState() const
 {
 	UWorld* WorldPtr = World.Get();
 	if (WorldPtr != nullptr)
 	{
-		return WorldPtr->GameState;
+		return WorldPtr->GetGameState();
 	}
-	
+
 	check(LocalPlayer.IsValid());
 	UWorld* LocalPlayerWorld = LocalPlayer->GetWorld();
-	return LocalPlayerWorld ? LocalPlayerWorld->GameState : NULL;
+	return LocalPlayerWorld ? LocalPlayerWorld->GetGameState() : NULL;
 }
 
 APlayerState* FLocalPlayerContext::GetPlayerState() const
@@ -1205,7 +1206,7 @@ bool ULocalPlayer::HandleCancelMatineeCommand( const TCHAR* Cmd, FOutputDevice&
 	// is the player in cinematic mode?
 	if (PlayerController->bCinematicMode)
 	{
-		TArray<UWorld*> MatineeActorWorldsThatSkipped;
+		bool bFoundMatinee = false;
 		// if so, look for all active matinees that has this Player in a director group
 		for (TActorIterator<AMatineeActor> It(GetWorld()); It; ++It)
 		{
@@ -1226,23 +1227,16 @@ bool ULocalPlayer::HandleCancelMatineeCommand( const TCHAR* Cmd, FOutputDevice&
 						{
 							// skip to end
 							MatineeActor->SetPosition(MatineeActor->MatineeData->InterpLength - RightBeforeEndTime, true);
-							MatineeActorWorldsThatSkipped.AddUnique( MatineeActor->GetWorld() );
+							bFoundMatinee = true;
 						}
 					}
 				}
 			}
 		}
 
-		if (MatineeActorWorldsThatSkipped.Num() != 0 )
+		if (bFoundMatinee)
 		{
-			for (int iActor = 0; iActor < MatineeActorWorldsThatSkipped.Num() ; iActor++)
-			{
-				AGameMode* const GameMode = MatineeActorWorldsThatSkipped[ iActor ]->GetAuthGameMode();
-				if (GameMode)
-				{
-					GameMode->MatineeCancelled();
-				}
-			}
+			FGameDelegates::Get().GetMatineeCancelledDelegate().Broadcast();
 		}
 	}
 	return true;
diff --git a/Engine/Source/Runtime/Engine/Private/NetConnection.cpp b/Engine/Source/Runtime/Engine/Private/NetConnection.cpp
--- a/Engine/Source/Runtime/Engine/Private/NetConnection.cpp
+++ b/Engine/Source/Runtime/Engine/Private/NetConnection.cpp
@@ -12,10 +12,11 @@
 #include "Engine/NetworkObjectList.h"
 #include "DataChannel.h"
 #include "Engine/PackageMapClient.h"
-#include "GameFramework/GameMode.h"
+#include "GameFramework/GameModeBase.h"
 #include "Runtime/PacketHandlers/PacketHandler/Public/PacketHandler.h"
 
 #include "PerfCountersHelpers.h"
+#include "GameDelegates.h"
 #if WITH_EDITOR
 #include "UnrealEd.h"
 #endif
@@ -387,15 +388,7 @@ void UNetConnection::CleanUp()
 		}
 		else 
 		{
-			UWorld* World = Driver ? Driver->GetWorld() : NULL;
-			if (World)
-			{
-				AGameMode* const GameMode = World->GetAuthGameMode();
-				if (GameMode)
-				{
-					GameMode->NotifyPendingConnectionLost();
-				}
-			}
+			FGameDelegates::Get().GetPendingConnectionLostDelegate().Broadcast();
 		}
 	}
 
