Copying //UE4/Dev-Framework to //UE4/Dev-Main (Source: //UE4/Dev-Framework @ 3136612)

diff --git a/Engine/Source/Runtime/Engine/Private/NetworkDriver.cpp b/Engine/Source/Runtime/Engine/Private/NetworkDriver.cpp
--- a/Engine/Source/Runtime/Engine/Private/NetworkDriver.cpp
+++ b/Engine/Source/Runtime/Engine/Private/NetworkDriver.cpp
@@ -18,7 +18,7 @@
 #include "DataChannel.h"
 #include "Engine/PackageMapClient.h"
 #include "GameFramework/PlayerState.h"
-#include "GameFramework/GameMode.h"
+#include "GameFramework/GameModeBase.h"
 #include "PerfCountersHelpers.h"
 
 
@@ -1499,7 +1499,6 @@ void UNetDriver::UpdateStandbyCheatStatus(void)
 			
 			if (FoundWorld)
 			{
-				AGameMode* const GameMode = FoundWorld->GetAuthGameMode();
 				AGameNetworkManager* const NetworkManager = FoundWorld->NetworkManager;
 				if (NetworkManager)
 				{
@@ -1507,21 +1506,17 @@ void UNetDriver::UpdateStandbyCheatStatus(void)
 					if (float(CountBadRx) / float(ClientConnections.Num()) > PercentMissingForRxStandby)
 					{
 						bHasStandbyCheatTriggered = true;
-						// Send to the GameMode for processing
 						NetworkManager->StandbyCheatDetected(STDBY_Rx);
 					}
 					else if (float(CountBadPing) / float(ClientConnections.Num()) > PercentForBadPing)
 					{
 						bHasStandbyCheatTriggered = true;
-						// Send to the GameMode for processing
 						NetworkManager->StandbyCheatDetected(STDBY_BadPing);
 					}
-					// Check for the host not sending to the clients, but only during a match
-					else if ( GameMode && GameMode->IsMatchInProgress() &&
-						float(CountBadTx) / float(ClientConnections.Num()) > PercentMissingForTxStandby)
+					// Check for the host not sending to the clients
+					else if (float(CountBadTx) / float(ClientConnections.Num()) > PercentMissingForTxStandby)
 					{
 						bHasStandbyCheatTriggered = true;
-						// Send to the GameMode for processing
 						NetworkManager->StandbyCheatDetected(STDBY_Tx);
 					}
 				}
@@ -2621,11 +2616,10 @@ int32 UNetDriver::ServerReplicateActors_PrioritizeActors( UNetConnection* Connec
 		OutPriorityList = new ( FMemStack::Get(), MaxSortedActors ) FActorPriority;
 		OutPriorityActors = new ( FMemStack::Get(), MaxSortedActors ) FActorPriority*;
 
-		// determine whether we should priority sort the list of relevant actors based on the saturation/bandwidth of the current connection
-		//@note - if the server is currently CPU saturated then do not sort until framerate improves
 		check( World == Connection->ViewTarget->GetWorld() );
-		AGameMode const* const GameMode = World->GetAuthGameMode();
-		const bool bLowNetBandwidth = !bCPUSaturated && ( Connection->CurrentNetSpeed / float( GameMode->NumPlayers + GameMode->NumBots ) < 500.f );
+
+		AGameNetworkManager* const NetworkManager = World->NetworkManager;
+		const bool bLowNetBandwidth = NetworkManager ? NetworkManager->IsInLowBandwidthMode() : false;
 
 		for ( FNetworkObjectInfo* ActorInfo : ConsiderList )
 		{
diff --git a/Engine/Source/Runtime/Engine/Private/Particles/Emitter.cpp b/Engine/Source/Runtime/Engine/Private/Particles/Emitter.cpp
--- a/Engine/Source/Runtime/Engine/Private/Particles/Emitter.cpp
+++ b/Engine/Source/Runtime/Engine/Private/Particles/Emitter.cpp
@@ -59,8 +59,6 @@ AEmitter::AEmitter(const FObjectInitializer& ObjectInitializer)
 			ArrowComponent->ArrowColor = FColor(0, 255, 128);
 
 			ArrowComponent->ArrowSize = 1.5f;
-			ArrowComponent->AlwaysLoadOnClient = false;
-			ArrowComponent->AlwaysLoadOnServer = false;
 			ArrowComponent->bTreatAsASprite = true;
 			ArrowComponent->bIsScreenSizeScaled = true;
 			ArrowComponent->SpriteInfo.Category = ConstructorStatics.ID_Effects;
diff --git a/Engine/Source/Runtime/Engine/Private/Particles/ParticleComponents.cpp b/Engine/Source/Runtime/Engine/Private/Particles/ParticleComponents.cpp
--- a/Engine/Source/Runtime/Engine/Private/Particles/ParticleComponents.cpp
+++ b/Engine/Source/Runtime/Engine/Private/Particles/ParticleComponents.cpp
@@ -7141,15 +7141,15 @@ void AEmitterCameraLensEffectBase::ActivateLensEffect()
 	check(World);
 	if( !IsNetMode(NM_DedicatedServer) )
 	{
-		UParticleSystem* PSToActuallySpawn;
-		if( World->GameState && World->GameState->ShouldShowGore() )
-		{
-			PSToActuallySpawn = PS_CameraEffect;
-		}
-		else
+		UParticleSystem* PSToActuallySpawn = PS_CameraEffect;
+
+PRAGMA_DISABLE_DEPRECATION_WARNINGS
+		AGameState* GameState = World->GetGameState<AGameState>();
+		if(GameState && !GameState->ShouldShowGore() )
 		{
-			PSToActuallySpawn = PS_CameraEffectNonExtremeContent;
+			PSToActuallySpawn = PS_CameraEffectNonExtremeContent_DEPRECATED;
 		}
+PRAGMA_ENABLE_DEPRECATION_WARNINGS
 
 		if( PSToActuallySpawn != NULL )
 		{
@@ -7174,7 +7174,7 @@ bool AEmitterCameraLensEffectBase::IsLooping() const
 		return true;
 	}
 
-	if ((PS_CameraEffectNonExtremeContent != nullptr) && PS_CameraEffectNonExtremeContent->IsLooping())
+	if ((PS_CameraEffectNonExtremeContent_DEPRECATED != nullptr) && PS_CameraEffectNonExtremeContent_DEPRECATED->IsLooping())
 	{
 		return true;
 	}
diff --git a/Engine/Source/Runtime/Engine/Private/PhysicsEngine/BodyInstance.cpp b/Engine/Source/Runtime/Engine/Private/PhysicsEngine/BodyInstance.cpp
--- a/Engine/Source/Runtime/Engine/Private/PhysicsEngine/BodyInstance.cpp
+++ b/Engine/Source/Runtime/Engine/Private/PhysicsEngine/BodyInstance.cpp
@@ -3565,7 +3565,9 @@ void FBodyInstance::WakeInstance()
 #if WITH_PHYSX
 	ExecuteOnPxRigidDynamicReadWrite(this, [&](PxRigidDynamic* PRigidDynamic)
 	{
-		if (!IsRigidBodyKinematic_AssumesLocked(PRigidDynamic))
+		// Only do the wake if we are not kinematic and we have a scene. If a body is not in a
+		// scene yet (could have been defered) then it will wake when it is added to the scene
+		if (PRigidDynamic->getScene() && !IsRigidBodyKinematic_AssumesLocked(PRigidDynamic))
 		{
 			PRigidDynamic->wakeUp();
 		}
diff --git a/Engine/Source/Runtime/Engine/Private/PhysicsEngine/ConstraintDrives.cpp b/Engine/Source/Runtime/Engine/Private/PhysicsEngine/ConstraintDrives.cpp
--- a/Engine/Source/Runtime/Engine/Private/PhysicsEngine/ConstraintDrives.cpp
+++ b/Engine/Source/Runtime/Engine/Private/PhysicsEngine/ConstraintDrives.cpp
@@ -103,14 +103,14 @@ void FLinearDriveConstraint::SetLinearVelocityDrive(bool bEnableXDrive, bool bEn
 
 void FAngularDriveConstraint::SetAngularPositionDrive(bool InEnableSwingDrive, bool InEnableTwistDrive)
 {
-	SwingDrive.bEnablePositionDrive &= InEnableSwingDrive;
-	TwistDrive.bEnablePositionDrive &= InEnableSwingDrive;
+	SwingDrive.bEnablePositionDrive = InEnableSwingDrive;
+	TwistDrive.bEnablePositionDrive = InEnableTwistDrive;
 }
 
 void FAngularDriveConstraint::SetAngularVelocityDrive(bool InEnableSwingDrive, bool InEnableTwistDrive)
 {
-	SwingDrive.bEnableVelocityDrive &= InEnableSwingDrive;
-	TwistDrive.bEnableVelocityDrive &= InEnableSwingDrive;
+	SwingDrive.bEnableVelocityDrive = InEnableSwingDrive;
+	TwistDrive.bEnableVelocityDrive = InEnableTwistDrive;
 }
 
 void FConstraintDrive::SetDriveParams(float InStiffness, float InDamping, float InForceLimit)
diff --git a/Engine/Source/Runtime/Engine/Private/PhysicsEngine/PhysScene.cpp b/Engine/Source/Runtime/Engine/Private/PhysicsEngine/PhysScene.cpp
--- a/Engine/Source/Runtime/Engine/Private/PhysicsEngine/PhysScene.cpp
+++ b/Engine/Source/Runtime/Engine/Private/PhysicsEngine/PhysScene.cpp
@@ -35,6 +35,10 @@ DECLARE_CYCLE_STAT(TEXT("Fetch Results Time (cloth)"), STAT_PhysicsFetchDynamics
 DECLARE_CYCLE_STAT(TEXT("Start Physics Time (async)"), STAT_PhysicsKickOffDynamicsTime_Async, STATGROUP_Physics);
 DECLARE_CYCLE_STAT(TEXT("Fetch Results Time (async)"), STAT_PhysicsFetchDynamicsTime_Async, STATGROUP_Physics);
 
+DECLARE_CYCLE_STAT(TEXT("Update Kinematics On Deferred SkelMeshes"), STAT_UpdateKinematicsOnDeferredSkelMeshes, STATGROUP_Physics);
+DECLARE_CYCLE_STAT(TEXT("Update Vehicles"), STAT_UpdateVehicles, STATGROUP_Physics);
+DECLARE_CYCLE_STAT(TEXT("Pretick Vehicles"), STAT_PretickVehicles, STATGROUP_Physics);
+
 
 DECLARE_CYCLE_STAT(TEXT("Phys Events Time"), STAT_PhysicsEventTime, STATGROUP_Physics);
 DECLARE_CYCLE_STAT(TEXT("SyncComponentsToBodies (sync)"), STAT_SyncComponentsToBodies, STATGROUP_Physics);
@@ -827,6 +831,8 @@ void FPhysScene::ClearPreSimKinematicUpdate(USkeletalMeshComponent* InSkelComp)
 
 void FPhysScene::UpdateKinematicsOnDeferredSkelMeshes()
 {
+	SCOPE_CYCLE_COUNTER(STAT_UpdateKinematicsOnDeferredSkelMeshes);
+
 	for (TMap< USkeletalMeshComponent*, FDeferredKinematicUpdateInfo >::TIterator It(DeferredKinematicUpdateSkelMeshes); It; ++It)
 	{
 		USkeletalMeshComponent* SkelComp = (*It).Key;
@@ -917,10 +923,14 @@ void FPhysScene::TickPhysScene(uint32 SceneType, FGraphEventRef& InOutCompletion
 		{
 			TickTime = UseSyncTime(SceneType) ? SyncDeltaSeconds : DeltaSeconds;
 		}
-		VehicleManager->PreTick(TickTime);
+		{
+			SCOPE_CYCLE_COUNTER(STAT_PretickVehicles);
+			VehicleManager->PreTick(TickTime);
+		}
 
 		if (IsSubstepping(SceneType) == false)
 		{
+			SCOPE_CYCLE_COUNTER(STAT_UpdateVehicles);
 			VehicleManager->Update(AveragedFrameTime[SceneType]);
 		}
 	}
diff --git a/Engine/Source/Runtime/Engine/Private/PhysicsEngine/PhysicsHandleComponent.cpp b/Engine/Source/Runtime/Engine/Private/PhysicsEngine/PhysicsHandleComponent.cpp
--- a/Engine/Source/Runtime/Engine/Private/PhysicsEngine/PhysicsHandleComponent.cpp
+++ b/Engine/Source/Runtime/Engine/Private/PhysicsEngine/PhysicsHandleComponent.cpp
@@ -241,6 +241,8 @@ void UPhysicsHandleComponent::ReleaseComponent()
 
 		bRotationConstrained = false;
 
+		GrabbedComponent->WakeRigidBody(GrabbedBoneName);
+
 		GrabbedComponent = NULL;
 		GrabbedBoneName = NAME_None;
 	}
diff --git a/Engine/Source/Runtime/Engine/Private/PhysicsEngine/RadialForceComponent.cpp b/Engine/Source/Runtime/Engine/Private/PhysicsEngine/RadialForceComponent.cpp
--- a/Engine/Source/Runtime/Engine/Private/PhysicsEngine/RadialForceComponent.cpp
+++ b/Engine/Source/Runtime/Engine/Private/PhysicsEngine/RadialForceComponent.cpp
@@ -89,6 +89,13 @@ void URadialForceComponent::TickComponent(float DeltaTime, enum ELevelTick TickT
 	}
 }
 
+void URadialForceComponent::BeginPlay()
+{
+	Super::BeginPlay();
+
+	UpdateCollisionObjectQueryParams();
+}
+
 void URadialForceComponent::PostLoad()
 {
 	Super::PostLoad();
diff --git a/Engine/Source/Runtime/Engine/Private/Player.cpp b/Engine/Source/Runtime/Engine/Private/Player.cpp
--- a/Engine/Source/Runtime/Engine/Private/Player.cpp
+++ b/Engine/Source/Runtime/Engine/Private/Player.cpp
@@ -17,8 +17,8 @@
 #include "SlateBasics.h"
 #include "GameFramework/CheatManager.h"
 #include "GameFramework/Pawn.h"
-#include "GameFramework/GameState.h"
-#include "GameFramework/GameMode.h"
+#include "GameFramework/GameStateBase.h"
+#include "GameFramework/GameModeBase.h"
 
 //////////////////////////////////////////////////////////////////////////
 // UPlayer
@@ -129,7 +129,7 @@ bool UPlayer::Exec( UWorld* InWorld, const TCHAR* Cmd,FOutputDevice& Ar)
 		{
 			return true;
 		}
-		else if (World->GameState && World->GameState->ProcessConsoleExec(Cmd, Ar, PCPawn))
+		else if (World->GetGameState() && World->GetGameState()->ProcessConsoleExec(Cmd, Ar, PCPawn))
 		{
 			return true;
 		}
diff --git a/Engine/Source/Runtime/Engine/Private/PlayerController.cpp b/Engine/Source/Runtime/Engine/Private/PlayerController.cpp
--- a/Engine/Source/Runtime/Engine/Private/PlayerController.cpp
+++ b/Engine/Source/Runtime/Engine/Private/PlayerController.cpp
@@ -547,7 +547,7 @@ void APlayerController::ServerNotifyLoadedWorld_Implementation(FName WorldPackag
 
 		// if both the server and this client have completed the transition, handle it
 		FSeamlessTravelHandler& SeamlessTravelHandler = GEngine->SeamlessTravelHandlerForWorld(CurWorld);
-		AGameMode* CurGameMode = CurWorld->GetAuthGameMode();
+		AGameModeBase* CurGameMode = CurWorld->GetAuthGameMode();
 
 		if (!SeamlessTravelHandler.IsInTransition() && WorldPackageName == CurWorld->GetOutermost()->GetFName() && CurGameMode != NULL)
 		{
@@ -977,7 +977,7 @@ void APlayerController::ServerShortTimeout_Implementation()
 		}
 		else 
 		{
-			float NetUpdateTimeOffset = (World->GetAuthGameMode()->NumPlayers < 8) ? 0.2f : 0.5f;
+			float NetUpdateTimeOffset = (World->GetAuthGameMode()->GetNumPlayers() < 8) ? 0.2f : 0.5f;
 			for (FActorIterator It(World); It; ++It)
 			{
 				AActor* A = *It;
@@ -1191,12 +1191,6 @@ void APlayerController::ClientSetHUD_Implementation(TSubclassOf<AHUD> NewHUDClas
 
 void APlayerController::CleanupPlayerState()
 {
-	AGameMode* const GameMode = GetWorld()->GetAuthGameMode();
-	if (GameMode)
-	{
-		GameMode->AddInactivePlayer(PlayerState, this);
-	}
-
 	PlayerState = NULL;
 }
 
@@ -1266,7 +1260,7 @@ void APlayerController::OnNetCleanup(UNetConnection* Connection)
 void APlayerController::ClientReceiveLocalizedMessage_Implementation( TSubclassOf<ULocalMessage> Message, int32 Switch, APlayerState* RelatedPlayerState_1, APlayerState* RelatedPlayerState_2, UObject* OptionalObject )
 {
 	// Wait for player to be up to date with replication when joining a server, before stacking up messages
-	if (GetNetMode() == NM_DedicatedServer || GetWorld()->GameState == NULL || Message == NULL)
+	if (GetNetMode() == NM_DedicatedServer || GetWorld()->GetGameState() == nullptr || Message == nullptr)
 	{
 		return;
 	}
@@ -1294,7 +1288,7 @@ void APlayerController::ClientPlaySoundAtLocation_Implementation(USoundBase* Sou
 
 void APlayerController::ClientMessage_Implementation( const FString& S, FName Type, float MsgLifeTime )
 {
-	if ( GetNetMode() == NM_DedicatedServer || GetWorld()->GameState == NULL )
+	if ( GetNetMode() == NM_DedicatedServer || GetWorld()->GetGameState() == nullptr )
 	{
 		return;
 	}
@@ -1428,7 +1422,7 @@ void APlayerController::Destroyed()
 
 	// Tells the game info to forcibly remove this player's CanUnpause delegates from its list of Pausers.
 	// Prevents the game from being stuck in a paused state when a PC that paused the game is destroyed before the game is unpaused.
-	AGameMode* const GameMode = GetWorld()->GetAuthGameMode();
+	AGameModeBase* const GameMode = GetWorld()->GetAuthGameMode();
 	if (GameMode)
 	{
 		GameMode->ForceClearUnpauseDelegates(this);
@@ -1652,7 +1646,7 @@ bool APlayerController::SetPause( bool bPause, FCanUnpause CanUnpauseDelegate)
 	bool bResult = false;
 	if (GetNetMode() != NM_Client)
 	{
-		AGameMode* const GameMode = GetWorld()->GetAuthGameMode();
+		AGameModeBase* const GameMode = GetWorld()->GetAuthGameMode();
 		if (GameMode != nullptr)
 		{
 			if (bPause)
@@ -1662,7 +1656,7 @@ bool APlayerController::SetPause( bool bPause, FCanUnpause CanUnpauseDelegate)
 			}
 			else
 			{
-				GameMode->ClearPause();
+				bResult = GameMode->ClearPause();
 			}
 		}
 	}
@@ -1705,9 +1699,10 @@ void APlayerController::SetName(const FString& S)
 
 void APlayerController::ServerChangeName_Implementation( const FString& S )
 {
-	if (!S.IsEmpty())
+	AGameModeBase* GameMode = GetWorld()->GetAuthGameMode();
+	if (!S.IsEmpty() && GameMode)
 	{
-		GetWorld()->GetAuthGameMode()->ChangeName( this, S, true );
+		GameMode->ChangeName( this, S, true );
 	}
 }
 
@@ -2647,13 +2642,16 @@ void APlayerController::ServerViewPrevPlayer_Implementation()
 APlayerState* APlayerController::GetNextViewablePlayer(int32 dir)
 {
 	int32 CurrentIndex = -1;
-	UWorld* MyWorld = GetWorld();
+
+	AGameModeBase* GameMode = GetWorld()->GetAuthGameMode();
+	AGameStateBase* GameState = GetWorld()->GetGameState();
+
 	if (PlayerCameraManager->ViewTarget.PlayerState )
 	{
 		// Find index of current viewtarget's PlayerState
-		for ( int32 i=0; i<MyWorld->GameState->PlayerArray.Num(); i++ )
+		for ( int32 i=0; i<GameState->PlayerArray.Num(); i++ )
 		{
-			if (PlayerCameraManager->ViewTarget.PlayerState == MyWorld->GameState->PlayerArray[i])
+			if (PlayerCameraManager->ViewTarget.PlayerState == GameState->PlayerArray[i])
 			{
 				CurrentIndex = i;
 				break;
@@ -2663,23 +2661,23 @@ APlayerState* APlayerController::GetNextViewablePlayer(int32 dir)
 
 	// Find next valid viewtarget in appropriate direction
 	int32 NewIndex;
-	for ( NewIndex=CurrentIndex+dir; (NewIndex>=0)&&(NewIndex<MyWorld->GameState->PlayerArray.Num()); NewIndex=NewIndex+dir )
+	for ( NewIndex=CurrentIndex+dir; (NewIndex>=0)&&(NewIndex<GameState->PlayerArray.Num()); NewIndex=NewIndex+dir )
 	{
-		APlayerState* const NextPlayerState = MyWorld->GameState->PlayerArray[NewIndex];
+		APlayerState* const NextPlayerState = GameState->PlayerArray[NewIndex];
 		AController* NextController = (NextPlayerState ? Cast<AController>(NextPlayerState->GetOwner()) : nullptr);
-		if ( NextController && NextController->GetPawn() != nullptr && MyWorld->GetAuthGameMode()->CanSpectate(this, PlayerState) )
+		if ( NextController && NextController->GetPawn() != nullptr && GameMode->CanSpectate(this, PlayerState) )
 		{
 			return PlayerState;
 		}
 	}
 
 	// wrap around
-	CurrentIndex = (NewIndex < 0) ? MyWorld->GameState->PlayerArray.Num() : -1;
-	for ( NewIndex=CurrentIndex+dir; (NewIndex>=0)&&(NewIndex<GetWorld()->GameState->PlayerArray.Num()); NewIndex=NewIndex+dir )
+	CurrentIndex = (NewIndex < 0) ? GameState->PlayerArray.Num() : -1;
+	for ( NewIndex=CurrentIndex+dir; (NewIndex>=0)&&(NewIndex<GameState->PlayerArray.Num()); NewIndex=NewIndex+dir )
 	{
-		APlayerState* const NextPlayerState = MyWorld->GameState->PlayerArray[NewIndex];
+		APlayerState* const NextPlayerState = GameState->PlayerArray[NewIndex];
 		AController* NextController = (NextPlayerState ? Cast<AController>(NextPlayerState->GetOwner()) : nullptr);
-		if ( NextController && NextController->GetPawn() != nullptr && MyWorld->GetAuthGameMode()->CanSpectate(this, PlayerState) )
+		if ( NextController && NextController->GetPawn() != nullptr && GameMode->CanSpectate(this, PlayerState) )
 		{
 			return PlayerState;
 		}
@@ -2762,8 +2760,8 @@ void APlayerController::ServerRestartPlayer_Implementation()
 
 	if ( IsInState(NAME_Inactive) || (IsInState(NAME_Spectating) && bPlayerIsWaiting) )
 	{
-		AGameMode* const GameMode = GetWorld()->GetAuthGameMode();
-		if ( !GetWorld()->GetAuthGameMode()->PlayerCanRestart(this) )
+		AGameModeBase* const GameMode = GetWorld()->GetAuthGameMode();
+		if ( !GameMode->PlayerCanRestart(this) )
 		{
 			return;
 		}
@@ -2774,7 +2772,7 @@ void APlayerController::ServerRestartPlayer_Implementation()
 			UnPossess();
 		}
 
-		GameMode->RestartPlayer( this );
+		GameMode->RestartPlayer(this);
 	}
 	else if ( GetPawn() != NULL )
 	{
@@ -3039,10 +3037,9 @@ void APlayerController::GetSeamlessTravelActorList(bool bToEntry, TArray<AActor*
 
 void APlayerController::SeamlessTravelTo(APlayerController* NewPC)
 {
-
+	CleanUpAudioComponents();
 }
 
-
 void APlayerController::SeamlessTravelFrom(APlayerController* OldPC)
 {
 	// copy PlayerState data
@@ -3057,6 +3054,20 @@ void APlayerController::SeamlessTravelFrom(APlayerController* OldPC)
 	}
 }
 
+void APlayerController::PostSeamlessTravel()
+{
+	// Track the last completed seamless travel for the player
+	LastCompletedSeamlessTravelCount = SeamlessTravelCount;
+
+	CleanUpAudioComponents();
+
+	if (PlayerCameraManager == nullptr)
+	{
+		SpawnPlayerCameraManager();
+	}
+
+}
+
 void APlayerController::ClientEnableNetworkVoice_Implementation(bool bEnable)
 {
 	ToggleSpeaking(bEnable);
@@ -3768,7 +3779,7 @@ void APlayerController::ClientClearCameraLensEffects_Implementation()
 	}
 }
 
-void APlayerController::ReceivedGameModeClass(TSubclassOf<AGameMode> GameModeClass)
+void APlayerController::ReceivedGameModeClass(TSubclassOf<AGameModeBase> GameModeClass)
 {
 }
 
@@ -4183,7 +4194,7 @@ ASpectatorPawn* APlayerController::SpawnSpectatorPawn()
 	if ((GetSpectatorPawn() == nullptr) && IsLocalController())
 	{
 		UWorld* World = GetWorld();
-		if (AGameState const* const GameState = World->GameState)
+		if (AGameStateBase const* const GameState = World->GetGameState())
 		{
 			if (UClass* SpectatorClass = GameState->SpectatorClass)
 			{
@@ -4329,8 +4340,13 @@ void APlayerController::BeginInactiveState()
 
 float APlayerController::GetMinRespawnDelay()
 {
-	AGameState const* const GameState = GetWorld()->GameState;
-	return ((GameState != NULL) && (GameState->GameModeClass != NULL)) ? GetDefault<AGameMode>(GameState->GameModeClass)->MinRespawnDelay : 1.0f;
+	AGameStateBase const* const GameState = GetWorld()->GetGameState();
+	
+	if (GameState)
+	{
+		return GameState->GetPlayerRespawnDelay(this);
+	}
+	return 1.0f;
 }
 
 void APlayerController::EndInactiveState()
diff --git a/Engine/Source/Runtime/Engine/Private/PlayerState.cpp b/Engine/Source/Runtime/Engine/Private/PlayerState.cpp
--- a/Engine/Source/Runtime/Engine/Private/PlayerState.cpp
+++ b/Engine/Source/Runtime/Engine/Private/PlayerState.cpp
@@ -123,10 +123,12 @@ void APlayerState::PostInitializeComponents()
 	Super::PostInitializeComponents();
 
 	UWorld* World = GetWorld();
-	// register this PlayerState with the game's ReplicationInfo
-	if ( World->GameState != NULL )
+	AGameStateBase* GameStateBase = World->GetGameState();
+
+	// register this PlayerState with the game state
+	if (GameStateBase != nullptr )
 	{
-		World->GameState->AddPlayerState(this);
+		GameStateBase->AddPlayerState(this);
 	}
 
 	if ( Role < ROLE_Authority )
@@ -137,12 +139,12 @@ void APlayerState::PostInitializeComponents()
 	AController* OwningController = Cast<AController>(GetOwner());
 	if (OwningController != NULL)
 	{
-		bIsABot = (Cast<APlayerController>(OwningController) == NULL);
+		bIsABot = (Cast<APlayerController>(OwningController) == nullptr);
 	}
 
-	if (World->GameState)
+	if (GameStateBase)
 	{
-		StartTime = World->GameState->ElapsedTime;
+		StartTime = GameStateBase->GetPlayerStartTime(OwningController);
 	}
 }
 
@@ -203,10 +205,10 @@ void APlayerState::OnRep_bIsInactive()
 {
 	// remove and re-add from the GameState so it's in the right list  
 	UWorld* World = GetWorld();
-	if (ensure(World && World->GameState))
+	if (ensure(World && World->GetGameState()))
 	{
-		World->GameState->RemovePlayerState(this);
-		World->GameState->AddPlayerState(this);
+		World->GetGameState()->RemovePlayerState(this);
+		World->GetGameState()->AddPlayerState(this);
 	}
 }
 
@@ -218,9 +220,9 @@ bool APlayerState::ShouldBroadCastWelcomeMessage(bool bExiting)
 void APlayerState::Destroyed()
 {
 	UWorld* World = GetWorld();
-	if (World->GameState != NULL)
+	if (World->GetGameState() != nullptr)
 	{
-		World->GameState->RemovePlayerState(this);
+		World->GetGameState()->RemovePlayerState(this);
 	}
 
 	if( ShouldBroadCastWelcomeMessage(true) )
@@ -340,7 +342,9 @@ void APlayerState::GetLifetimeReplicatedProps(TArray< FLifetimeProperty > & OutL
 	DOREPLIFETIME( APlayerState, bIsSpectator );
 	DOREPLIFETIME( APlayerState, bOnlySpectator );
 	DOREPLIFETIME( APlayerState, bFromPreviousLevel );
+PRAGMA_DISABLE_DEPRECATION_WARNINGS
 	DOREPLIFETIME( APlayerState, StartTime );
+PRAGMA_ENABLE_DEPRECATION_WARNINGS
 
 	DOREPLIFETIME_CONDITION( APlayerState, Ping,		COND_SkipOwner );
 
diff --git a/Engine/Source/Runtime/Engine/Private/PrimitiveComponentPhysics.cpp b/Engine/Source/Runtime/Engine/Private/PrimitiveComponentPhysics.cpp
--- a/Engine/Source/Runtime/Engine/Private/PrimitiveComponentPhysics.cpp
+++ b/Engine/Source/Runtime/Engine/Private/PrimitiveComponentPhysics.cpp
@@ -390,6 +390,10 @@ void UPrimitiveComponent::SetAllPhysicsRotation(FRotator NewRot)
 	SetWorldRotation(NewRot, NAME_None);
 }
 
+void UPrimitiveComponent::SetAllPhysicsRotation(const FQuat& NewRot)
+{
+	SetWorldRotation(NewRot);
+}
 
 void UPrimitiveComponent::WakeRigidBody(FName BoneName)
 {
diff --git a/Engine/Source/Runtime/Engine/Private/SimpleConstructionScript.cpp b/Engine/Source/Runtime/Engine/Private/SimpleConstructionScript.cpp
--- a/Engine/Source/Runtime/Engine/Private/SimpleConstructionScript.cpp
+++ b/Engine/Source/Runtime/Engine/Private/SimpleConstructionScript.cpp
@@ -16,8 +16,6 @@
 
 namespace
 {
-	static const FString ComponentTemplateNameSuffix(TEXT("_GEN_VARIABLE"));
-
 	// Helper method to register instanced components post-construction
 	void RegisterInstancedComponent(UActorComponent* InstancedComponent)
 	{
@@ -173,10 +171,10 @@ void USimpleConstructionScript::PostLoad()
 			// Fix up any component template objects whose name doesn't match the current variable name; this ensures that there is always one unique template per node.
 			FString VariableName = Node->GetVariableName().ToString();
 			FString ComponentTemplateName = Node->ComponentTemplate->GetName();
-			if(ComponentTemplateName.EndsWith(ComponentTemplateNameSuffix) && !ComponentTemplateName.StartsWith(VariableName) && !GIsDuplicatingClassForReinstancing)
+			if(ComponentTemplateName.EndsWith(UActorComponent::ComponentTemplateNameSuffix) && !ComponentTemplateName.StartsWith(VariableName) && !GIsDuplicatingClassForReinstancing)
 			{
 				Node->ComponentTemplate->ConditionalPostLoad();
-				Node->ComponentTemplate = static_cast<UActorComponent*>(StaticDuplicateObject(Node->ComponentTemplate, Node->ComponentTemplate->GetOuter(), *(VariableName + ComponentTemplateNameSuffix)));
+				Node->ComponentTemplate = static_cast<UActorComponent*>(StaticDuplicateObject(Node->ComponentTemplate, Node->ComponentTemplate->GetOuter(), *(VariableName + UActorComponent::ComponentTemplateNameSuffix)));
 			}
 
 			// Check to see if switched from scene to a non-scene component type
@@ -1170,7 +1168,7 @@ USCS_Node* USimpleConstructionScript::CreateNode(UClass* NewComponentClass, FNam
 	NewComponentTemplate->Modify();
 
 	// Now set the actual name and outer to the BPGC.
-	NewComponentTemplate->Rename(*(NewComponentVariableName.ToString() + ComponentTemplateNameSuffix), Blueprint->GeneratedClass, REN_DoNotDirty|REN_DontCreateRedirectors|REN_ForceNoResetLoaders);
+	NewComponentTemplate->Rename(*(NewComponentVariableName.ToString() + UActorComponent::ComponentTemplateNameSuffix), Blueprint->GeneratedClass, REN_DoNotDirty|REN_DontCreateRedirectors|REN_ForceNoResetLoaders);
 
 	return CreateNodeImpl(NewComponentTemplate, NewComponentVariableName);
 }
@@ -1183,7 +1181,7 @@ USCS_Node* USimpleConstructionScript::CreateNodeAndRenameComponent(UActorCompone
 	FName NewComponentVariableName = GenerateNewComponentName(NewComponentTemplate->GetClass());
 
 	// Relocate the instance from the transient package to the BPGC and assign it a unique object name
-	NewComponentTemplate->Rename(*(NewComponentVariableName.ToString() + ComponentTemplateNameSuffix), GetBlueprint()->GeneratedClass, REN_DontCreateRedirectors | REN_DoNotDirty);
+	NewComponentTemplate->Rename(*(NewComponentVariableName.ToString() + UActorComponent::ComponentTemplateNameSuffix), GetBlueprint()->GeneratedClass, REN_DontCreateRedirectors | REN_DoNotDirty);
 
 	return CreateNodeImpl(NewComponentTemplate, NewComponentVariableName);
 }
diff --git a/Engine/Source/Runtime/Engine/Private/SkeletalMesh.cpp b/Engine/Source/Runtime/Engine/Private/SkeletalMesh.cpp
--- a/Engine/Source/Runtime/Engine/Private/SkeletalMesh.cpp
+++ b/Engine/Source/Runtime/Engine/Private/SkeletalMesh.cpp
@@ -2832,6 +2832,20 @@ void USkeletalMesh::PostEditChangeProperty(FPropertyChangedEvent& PropertyChange
 		}
 	}
 
+	if(PropertyThatChanged && PropertyThatChanged->GetFName() == GET_MEMBER_NAME_CHECKED(USkeletalMesh, PostProcessAnimBlueprint))
+	{
+		TArray<UActorComponent*> ComponentsToReregister;
+		for(TObjectIterator<USkeletalMeshComponent> It; It; ++It)
+		{
+			USkeletalMeshComponent* MeshComponent = *It;
+			if(MeshComponent && !MeshComponent->IsTemplate() && MeshComponent->SkeletalMesh == this)
+			{
+				ComponentsToReregister.Add(*It);
+			}
+		}
+		FMultiComponentReregisterContext ReregisterContext(ComponentsToReregister);
+	}
+
 	Super::PostEditChangeProperty(PropertyChangedEvent);
 }
 
@@ -3582,6 +3596,18 @@ void USkeletalMesh::InitMorphTargets()
 		if (MorphTargetIndexMap.Find(ShapeName) == nullptr)
 		{
 			MorphTargetIndexMap.Add(ShapeName, Index);
+
+			// register as morphtarget curves
+			if (Skeleton)
+			{
+				FSmartName CurveName;
+				CurveName.DisplayName = ShapeName;
+				
+				// verify will make sure it adds to the curve if not found
+				// the reason of using this is to make sure it works in editor/non-editor
+				Skeleton->VerifySmartName(USkeleton::AnimCurveMappingName, CurveName);
+				Skeleton->AccumulateCurveMetaData(ShapeName, false, true);
+			}
 		}
 	}
 }
@@ -3635,15 +3661,12 @@ USkeletalMeshSocket* USkeletalMesh::FindSocketAndIndex(FName InSocketName, int32
 	// If the socket isn't on the mesh, try to find it on the skeleton
 	if (Skeleton)
 	{
-		for (int32 i = 0; i < Skeleton->Sockets.Num(); ++i)
+		USkeletalMeshSocket* SkeletonSocket = Skeleton->FindSocketAndIndex(InSocketName, OutIndex);
+		if (SkeletonSocket != nullptr)
 		{
-			USkeletalMeshSocket* Socket = Skeleton->Sockets[i];
-			if (Socket && Socket->SocketName == InSocketName)
-			{
-				OutIndex = Sockets.Num() + i;
-				return Socket;
-			}
+			OutIndex += Sockets.Num();
 		}
+		return SkeletonSocket;
 	}
 
 	return NULL;
@@ -3733,6 +3756,11 @@ TArray<USkeletalMeshSocket*>& USkeletalMesh::GetMeshOnlySocketList()
 	return Sockets;
 }
 
+const TArray<USkeletalMeshSocket*>& USkeletalMesh::GetMeshOnlySocketList() const
+{
+	return Sockets;
+}
+
 void USkeletalMesh::MoveDeprecatedShadowFlagToMaterials()
 {
 	// First, the easy case where there's no LOD info (in which case, default to true!)
diff --git a/Engine/Source/Runtime/Engine/Private/SkeletalMeshComponentPhysics.cpp b/Engine/Source/Runtime/Engine/Private/SkeletalMeshComponentPhysics.cpp
--- a/Engine/Source/Runtime/Engine/Private/SkeletalMeshComponentPhysics.cpp
+++ b/Engine/Source/Runtime/Engine/Private/SkeletalMeshComponentPhysics.cpp
@@ -647,6 +647,11 @@ void USkeletalMeshComponent::SetAllPhysicsPosition(FVector NewPos)
 }
 
 void USkeletalMeshComponent::SetAllPhysicsRotation(FRotator NewRot)
+{
+	SetAllPhysicsRotation(NewRot.Quaternion());
+}
+
+void USkeletalMeshComponent::SetAllPhysicsRotation(const FQuat& NewRot)
 {
 #if ENABLE_NAN_DIAGNOSTIC
 	if (NewRot.ContainsNaN())
@@ -662,10 +667,9 @@ void USkeletalMeshComponent::SetAllPhysicsRotation(FRotator NewRot)
 		if(RootBI->IsValidBodyInstance())
 		{
 			// move the root body
-			FQuat NewRotQuat = NewRot.Quaternion();
 			FTransform RootBodyTM = RootBI->GetUnrealWorldTransform();
-			FQuat DeltaQuat = RootBodyTM.GetRotation().Inverse() * NewRotQuat;
-			RootBodyTM.SetRotation(NewRotQuat);
+			FQuat DeltaQuat = RootBodyTM.GetRotation().Inverse() * NewRot;
+			RootBodyTM.SetRotation(NewRot);
 			RootBI->SetBodyTransform(RootBodyTM, ETeleportType::TeleportPhysics);
 
 			// apply the delta to all the other bodies
diff --git a/Engine/Source/Runtime/Engine/Private/SkeletalRenderCPUSkin.cpp b/Engine/Source/Runtime/Engine/Private/SkeletalRenderCPUSkin.cpp
--- a/Engine/Source/Runtime/Engine/Private/SkeletalRenderCPUSkin.cpp
+++ b/Engine/Source/Runtime/Engine/Private/SkeletalRenderCPUSkin.cpp
@@ -28,7 +28,7 @@
 struct FMorphTargetDelta;
 
 template<typename BaseVertexType, typename VertexType>
-static void SkinVertices(FFinalSkinVertex* DestVertex, FMatrix* ReferenceToLocal, int32 LODIndex, FStaticLODModel& LOD, TArray<FActiveMorphTarget>& ActiveMorphTargets, TArray<float>& MorphTargetWeights);
+static void SkinVertices(FFinalSkinVertex* DestVertex, FMatrix* ReferenceToLocal, int32 LODIndex, FStaticLODModel& LOD, TArray<FActiveMorphTarget>& ActiveMorphTargets, TArray<float>& MorphTargetWeights, const TMap<int32, FClothSimulData>& ClothSimulUpdateData, float ClothBlendWeight, const FMatrix& WorldToLocal);
 
 #define INFLUENCE_0		0
 #define INFLUENCE_1		1
@@ -294,11 +294,11 @@ void FSkeletalMeshObjectCPUSkin::CacheVertices(int32 LODIndex, bool bForce) cons
 				// do actual skinning
 				if (LOD.DoesVertexBufferHaveExtraBoneInfluences())
 				{
-					SkinVertices< TGPUSkinVertexBase<true>, TGPUSkinVertexFloat32Uvs<1, true> >( DestVertex, ReferenceToLocal, DynamicData->LODIndex, LOD, DynamicData->ActiveMorphTargets, DynamicData->MorphTargetWeights);
+					SkinVertices< TGPUSkinVertexBase<true>, TGPUSkinVertexFloat32Uvs<1, true> >( DestVertex, ReferenceToLocal, DynamicData->LODIndex, LOD, DynamicData->ActiveMorphTargets, DynamicData->MorphTargetWeights, DynamicData->ClothSimulUpdateData, DynamicData->ClothBlendWeight, DynamicData->WorldToLocal);
 				}
 				else
 				{
-					SkinVertices< TGPUSkinVertexBase<false>, TGPUSkinVertexFloat32Uvs<1, false> >( DestVertex, ReferenceToLocal, DynamicData->LODIndex, LOD, DynamicData->ActiveMorphTargets, DynamicData->MorphTargetWeights);
+					SkinVertices< TGPUSkinVertexBase<false>, TGPUSkinVertexFloat32Uvs<1, false> >( DestVertex, ReferenceToLocal, DynamicData->LODIndex, LOD, DynamicData->ActiveMorphTargets, DynamicData->MorphTargetWeights, DynamicData->ClothSimulUpdateData, DynamicData->ClothBlendWeight, DynamicData->WorldToLocal);
 				}
 			}
 			else
@@ -306,11 +306,11 @@ void FSkeletalMeshObjectCPUSkin::CacheVertices(int32 LODIndex, bool bForce) cons
 				// do actual skinning
 				if (LOD.DoesVertexBufferHaveExtraBoneInfluences())
 				{
-					SkinVertices< TGPUSkinVertexBase<true>, TGPUSkinVertexFloat16Uvs<1, true> >( DestVertex, ReferenceToLocal, DynamicData->LODIndex, LOD, DynamicData->ActiveMorphTargets, DynamicData->MorphTargetWeights);
+					SkinVertices< TGPUSkinVertexBase<true>, TGPUSkinVertexFloat16Uvs<1, true> >( DestVertex, ReferenceToLocal, DynamicData->LODIndex, LOD, DynamicData->ActiveMorphTargets, DynamicData->MorphTargetWeights, DynamicData->ClothSimulUpdateData, DynamicData->ClothBlendWeight, DynamicData->WorldToLocal);
 				}
 				else
 				{
-					SkinVertices< TGPUSkinVertexBase<false>, TGPUSkinVertexFloat16Uvs<1, false> >( DestVertex, ReferenceToLocal, DynamicData->LODIndex, LOD, DynamicData->ActiveMorphTargets, DynamicData->MorphTargetWeights);
+					SkinVertices< TGPUSkinVertexBase<false>, TGPUSkinVertexFloat16Uvs<1, false> >( DestVertex, ReferenceToLocal, DynamicData->LODIndex, LOD, DynamicData->ActiveMorphTargets, DynamicData->MorphTargetWeights, DynamicData->ClothSimulUpdateData, DynamicData->ClothBlendWeight, DynamicData->WorldToLocal);
 				}
 			}
 
@@ -487,11 +487,15 @@ FDynamicSkelMeshObjectDataCPUSkin::FDynamicSkelMeshObjectDataCPUSkin(
 :	LODIndex(InLODIndex)
 ,	ActiveMorphTargets(InActiveMorphTargets)
 ,	MorphTargetWeights(InMorphTargetWeights)
+,	ClothBlendWeight(0.0f)
 {
 	UpdateRefToLocalMatrices( ReferenceToLocal, InMeshComponent, InSkeletalMeshResource, LODIndex );
 
 	UpdateCustomLeftRightVectors( CustomLeftRightVectors, InMeshComponent, InSkeletalMeshResource, LODIndex );
 
+	// Update the clothing simulation mesh positions and normals
+	UpdateClothSimulationData(InMeshComponent);
+
 #if !(UE_BUILD_SHIPPING || UE_BUILD_TEST)
 	MeshComponentSpaceTransforms = InMeshComponent->GetComponentSpaceTransforms();
 #endif
@@ -646,8 +650,10 @@ const VectorRegister		VECTOR4_UNPACK_MINUS_1	= DECLARE_VECTOR_REGISTER(-1.f, -1.
 
 const VectorRegister		VECTOR_0001				= DECLARE_VECTOR_REGISTER(0.f, 0.f, 0.f, 1.f);
 
+#define FIXED_VERTEX_INDEX 0xFFFF
+
 template<bool bExtraBoneInfluences, int32 MaxSectionBoneInfluences, typename BaseVertexType, typename VertexType>
-static void SkinVertexSection( FFinalSkinVertex*& DestVertex, TArray<FMorphTargetInfo>& MorphEvalInfos, const TArray<float>& MorphWeights, const FSkelMeshSection& Section, const FStaticLODModel &LOD, int32 VertexBufferBaseIndex, uint32 NumValidMorphs, int32 &CurBaseVertIdx, int32 LODIndex, int32 RigidInfluenceIndex, const FMatrix* RESTRICT ReferenceToLocal )
+static void SkinVertexSection( FFinalSkinVertex*& DestVertex, TArray<FMorphTargetInfo>& MorphEvalInfos, const TArray<float>& MorphWeights, const FSkelMeshSection& Section, const FStaticLODModel &LOD, int32 VertexBufferBaseIndex, uint32 NumValidMorphs, int32 &CurBaseVertIdx, int32 LODIndex, int32 RigidInfluenceIndex, const FMatrix* RESTRICT ReferenceToLocal, const FClothSimulData* ClothSimData, float ClothBlendWeight, const FMatrix& WorldToLocal)
 {
 	// VertexCopy for morph. Need to allocate right struct
 	// To avoid re-allocation, create 2 statics, and assign right struct
@@ -662,6 +668,9 @@ static void SkinVertexSection( FFinalSkinVertex*& DestVertex, TArray<FMorphTarge
 	
 
 	VertexType* SrcSoftVertex = NULL;
+	const FVector MeshExtension = LOD.VertexBufferGPUSkin.GetMeshExtension();
+	const FVector MeshOrigin = LOD.VertexBufferGPUSkin.GetMeshOrigin();
+	const bool bLODUsesAPEXCloth = LOD.HasApexClothData() && ClothSimData != nullptr && ClothBlendWeight > 0.0f;
 	const int32 NumSoftVertices = Section.GetNumVertices();
 	if (NumSoftVertices > 0)
 	{
@@ -669,6 +678,10 @@ static void SkinVertexSection( FFinalSkinVertex*& DestVertex, TArray<FMorphTarge
 
 		// Prefetch first vertex
 		FPlatformMisc::Prefetch( LOD.VertexBufferGPUSkin.GetVertexPtr<(bExtraBoneInfluences)>(Section.GetVertexBufferIndex()) );
+		if (bLODUsesAPEXCloth)
+		{
+			FPlatformMisc::Prefetch(&LOD.APEXClothVertexBuffer.MappingData(Section.GetVertexBufferIndex()));
+		}
 
 		for(int32 VertexIndex = VertexBufferBaseIndex;VertexIndex < NumSoftVertices;VertexIndex++,DestVertex++)
 		{
@@ -682,6 +695,13 @@ static void SkinVertexSection( FFinalSkinVertex*& DestVertex, TArray<FMorphTarge
 				UpdateMorphedVertex<VertexType>( *MorphedVertex, *SrcSoftVertex, CurBaseVertIdx, LODIndex, MorphEvalInfos, MorphWeights);
 			}
 
+			const FApexClothPhysToRenderVertData* APEXVertData = nullptr;
+			if (bLODUsesAPEXCloth)
+			{
+				APEXVertData = &LOD.APEXClothVertexBuffer.MappingData(VertexBufferIndex);
+				FPlatformMisc::Prefetch(APEXVertData, PLATFORM_CACHE_LINE_SIZE);	// Prefetch next cloth vertex
+			}
+
 			const uint8* RESTRICT BoneIndices = MorphedVertex->InfluenceBones;
 			const uint8* RESTRICT BoneWeights = MorphedVertex->InfluenceWeights;
 
@@ -803,6 +823,76 @@ static void SkinVertexSection( FFinalSkinVertex*& DestVertex, TArray<FMorphTarge
 			Pack4( DstNormals[2], &DestVertex->TangentZ.Vector.Packed );
 			VectorResetFloatRegisters(); // Need to call this to be able to use regular floating point registers again after Pack().
 
+			// Apply cloth. This code has been adapted from GpuSkinVertexFactory.usf
+			if (APEXVertData != nullptr && APEXVertData->SimulMeshVertIndices[3] < FIXED_VERTEX_INDEX)
+			{
+				struct APEXClothCPU
+				{
+					FORCEINLINE static FVector GetClothSimulPosition(const FClothSimulData& InClothSimData, int32 InIndex)
+					{
+						return FVector(InClothSimData.ClothSimulPositions[InIndex]);
+					}
+
+					FORCEINLINE static FVector GetClothSimulNormal(const FClothSimulData& InClothSimData, int32 InIndex)
+					{
+						return FVector(InClothSimData.ClothSimulNormals[InIndex]);
+					}
+
+					FORCEINLINE static FVector ClothingPosition(const FApexClothPhysToRenderVertData& InAPEXVertData, const FClothSimulData& InClothSimData)
+					{
+						return    InAPEXVertData.PositionBaryCoordsAndDist.X * (GetClothSimulPosition(InClothSimData, InAPEXVertData.SimulMeshVertIndices[0]) + GetClothSimulNormal(InClothSimData, InAPEXVertData.SimulMeshVertIndices[0]) * InAPEXVertData.PositionBaryCoordsAndDist.W)
+								+ InAPEXVertData.PositionBaryCoordsAndDist.Y * (GetClothSimulPosition(InClothSimData, InAPEXVertData.SimulMeshVertIndices[1]) + GetClothSimulNormal(InClothSimData, InAPEXVertData.SimulMeshVertIndices[1]) * InAPEXVertData.PositionBaryCoordsAndDist.W)
+								+ InAPEXVertData.PositionBaryCoordsAndDist.Z * (GetClothSimulPosition(InClothSimData, InAPEXVertData.SimulMeshVertIndices[2]) + GetClothSimulNormal(InClothSimData, InAPEXVertData.SimulMeshVertIndices[2]) * InAPEXVertData.PositionBaryCoordsAndDist.W);
+					}
+
+					FORCEINLINE static void ClothingTangents(const FApexClothPhysToRenderVertData& InAPEXVertData, const FClothSimulData& InClothSimData, const FVector& InSimulatedPosition, const FMatrix& InWorldToLocal, const FVector& InMeshExtension, const FVector& InMeshOrign, FVector& OutTangentX, FVector& OutTangentZ)
+					{
+						FVector A = GetClothSimulPosition(InClothSimData, InAPEXVertData.SimulMeshVertIndices[0]);
+						FVector B = GetClothSimulPosition(InClothSimData, InAPEXVertData.SimulMeshVertIndices[1]);
+						FVector C = GetClothSimulPosition(InClothSimData, InAPEXVertData.SimulMeshVertIndices[2]);
+
+						FVector NA = GetClothSimulNormal(InClothSimData, InAPEXVertData.SimulMeshVertIndices[0]);
+						FVector NB = GetClothSimulNormal(InClothSimData, InAPEXVertData.SimulMeshVertIndices[1]);
+						FVector NC = GetClothSimulNormal(InClothSimData, InAPEXVertData.SimulMeshVertIndices[2]);
+
+						FVector NormalPosition = InAPEXVertData.NormalBaryCoordsAndDist.X*(A + NA*InAPEXVertData.NormalBaryCoordsAndDist.W)
+												+ InAPEXVertData.NormalBaryCoordsAndDist.Y*(B + NB*InAPEXVertData.NormalBaryCoordsAndDist.W)
+												+ InAPEXVertData.NormalBaryCoordsAndDist.Z*(C + NC*InAPEXVertData.NormalBaryCoordsAndDist.W);
+
+						FVector TangentPosition = InAPEXVertData.TangentBaryCoordsAndDist.X*(A + NA*InAPEXVertData.TangentBaryCoordsAndDist.W)
+												+ InAPEXVertData.TangentBaryCoordsAndDist.Y*(B + NB*InAPEXVertData.TangentBaryCoordsAndDist.W)
+												+ InAPEXVertData.TangentBaryCoordsAndDist.Z*(C + NC*InAPEXVertData.TangentBaryCoordsAndDist.W);
+
+						OutTangentX = (TangentPosition*InMeshExtension + InMeshOrign - InSimulatedPosition).GetUnsafeNormal();
+						OutTangentZ = (NormalPosition*InMeshExtension + InMeshOrign - InSimulatedPosition).GetUnsafeNormal();
+
+						// cloth data are all in world space so need to change into local space
+						OutTangentX = InWorldToLocal.TransformVector(OutTangentX);
+						OutTangentZ = InWorldToLocal.TransformVector(OutTangentZ);
+					}
+				};
+
+				// build sim position (in world space)
+				FVector SimulatedPositionWorld = APEXClothCPU::ClothingPosition(*APEXVertData, *ClothSimData) * MeshExtension + MeshOrigin;
+
+				// transform back to local space
+				FVector SimulatedPosition = WorldToLocal.TransformPosition(SimulatedPositionWorld);
+
+				// Lerp between skinned and simulated position
+				DestVertex->Position = FMath::Lerp(DestVertex->Position, SimulatedPosition, ClothBlendWeight);
+
+				// recompute tangent & normal
+				FVector TangentX;
+				FVector TangentZ;
+				APEXClothCPU::ClothingTangents(*APEXVertData, *ClothSimData, SimulatedPositionWorld, WorldToLocal, MeshExtension, MeshOrigin, TangentX, TangentZ);
+
+				// Lerp between skinned and simulated tangents
+				FVector SkinnedTangentX = DestVertex->TangentX;
+				FVector4 SkinnedTangentZ = DestVertex->TangentZ;
+				DestVertex->TangentX = (TangentX * ClothBlendWeight) + (SkinnedTangentX * (1.0f - ClothBlendWeight));
+				DestVertex->TangentZ = FVector4((TangentZ * ClothBlendWeight) + (SkinnedTangentZ * (1.0f - ClothBlendWeight)), SkinnedTangentZ.W);
+			}
+
 			// Copy UVs.
 			FVector2D UVs = LOD.VertexBufferGPUSkin.GetVertexUVFast<bExtraBoneInfluences>(Section.GetVertexBufferIndex() + VertexIndex, 0);
 			DestVertex->U = UVs.X;
@@ -814,24 +904,24 @@ static void SkinVertexSection( FFinalSkinVertex*& DestVertex, TArray<FMorphTarge
 }
 
 template< bool bExtraBoneInfluences, typename BaseVertexType, typename VertexType>
-static void SkinVertexSection( FFinalSkinVertex*& DestVertex, TArray<FMorphTargetInfo>& MorphEvalInfos, const TArray<float>& MorphWeights, const FSkelMeshSection& Section, const FStaticLODModel &LOD, int32 VertexBufferBaseIndex, uint32 NumValidMorphs, int32 &CurBaseVertIdx, int32 LODIndex, int32 RigidInfluenceIndex, const FMatrix* RESTRICT ReferenceToLocal)
+static void SkinVertexSection( FFinalSkinVertex*& DestVertex, TArray<FMorphTargetInfo>& MorphEvalInfos, const TArray<float>& MorphWeights, const FSkelMeshSection& Section, const FStaticLODModel &LOD, int32 VertexBufferBaseIndex, uint32 NumValidMorphs, int32 &CurBaseVertIdx, int32 LODIndex, int32 RigidInfluenceIndex, const FMatrix* RESTRICT ReferenceToLocal, const FClothSimulData* ClothSimData, float ClothBlendWeight, const FMatrix& WorldToLocal)
 {
 	switch (Section.MaxBoneInfluences)
 	{
-		case 1: SkinVertexSection<bExtraBoneInfluences, 1, BaseVertexType, VertexType>(DestVertex, MorphEvalInfos, MorphWeights, Section, LOD, VertexBufferBaseIndex, NumValidMorphs, CurBaseVertIdx, LODIndex, RigidInfluenceIndex, ReferenceToLocal); break;
-		case 2: SkinVertexSection<bExtraBoneInfluences, 2, BaseVertexType, VertexType>(DestVertex, MorphEvalInfos, MorphWeights, Section, LOD, VertexBufferBaseIndex, NumValidMorphs, CurBaseVertIdx, LODIndex, RigidInfluenceIndex, ReferenceToLocal); break;
-		case 3: SkinVertexSection<bExtraBoneInfluences, 3, BaseVertexType, VertexType>(DestVertex, MorphEvalInfos, MorphWeights, Section, LOD, VertexBufferBaseIndex, NumValidMorphs, CurBaseVertIdx, LODIndex, RigidInfluenceIndex, ReferenceToLocal); break;
-		case 4: SkinVertexSection<bExtraBoneInfluences, 4, BaseVertexType, VertexType>(DestVertex, MorphEvalInfos, MorphWeights, Section, LOD, VertexBufferBaseIndex, NumValidMorphs, CurBaseVertIdx, LODIndex, RigidInfluenceIndex, ReferenceToLocal); break;
-		case 5: SkinVertexSection<bExtraBoneInfluences, 5, BaseVertexType, VertexType>(DestVertex, MorphEvalInfos, MorphWeights, Section, LOD, VertexBufferBaseIndex, NumValidMorphs, CurBaseVertIdx, LODIndex, RigidInfluenceIndex, ReferenceToLocal); break;
-		case 6: SkinVertexSection<bExtraBoneInfluences, 6, BaseVertexType, VertexType>(DestVertex, MorphEvalInfos, MorphWeights, Section, LOD, VertexBufferBaseIndex, NumValidMorphs, CurBaseVertIdx, LODIndex, RigidInfluenceIndex, ReferenceToLocal); break;
-		case 7: SkinVertexSection<bExtraBoneInfluences, 7, BaseVertexType, VertexType>(DestVertex, MorphEvalInfos, MorphWeights, Section, LOD, VertexBufferBaseIndex, NumValidMorphs, CurBaseVertIdx, LODIndex, RigidInfluenceIndex, ReferenceToLocal); break;
-		case 8: SkinVertexSection<bExtraBoneInfluences, 8, BaseVertexType, VertexType>(DestVertex, MorphEvalInfos, MorphWeights, Section, LOD, VertexBufferBaseIndex, NumValidMorphs, CurBaseVertIdx, LODIndex, RigidInfluenceIndex, ReferenceToLocal); break;
+		case 1: SkinVertexSection<bExtraBoneInfluences, 1, BaseVertexType, VertexType>(DestVertex, MorphEvalInfos, MorphWeights, Section, LOD, VertexBufferBaseIndex, NumValidMorphs, CurBaseVertIdx, LODIndex, RigidInfluenceIndex, ReferenceToLocal, ClothSimData, ClothBlendWeight, WorldToLocal); break;
+		case 2: SkinVertexSection<bExtraBoneInfluences, 2, BaseVertexType, VertexType>(DestVertex, MorphEvalInfos, MorphWeights, Section, LOD, VertexBufferBaseIndex, NumValidMorphs, CurBaseVertIdx, LODIndex, RigidInfluenceIndex, ReferenceToLocal, ClothSimData, ClothBlendWeight, WorldToLocal); break;
+		case 3: SkinVertexSection<bExtraBoneInfluences, 3, BaseVertexType, VertexType>(DestVertex, MorphEvalInfos, MorphWeights, Section, LOD, VertexBufferBaseIndex, NumValidMorphs, CurBaseVertIdx, LODIndex, RigidInfluenceIndex, ReferenceToLocal, ClothSimData, ClothBlendWeight, WorldToLocal); break;
+		case 4: SkinVertexSection<bExtraBoneInfluences, 4, BaseVertexType, VertexType>(DestVertex, MorphEvalInfos, MorphWeights, Section, LOD, VertexBufferBaseIndex, NumValidMorphs, CurBaseVertIdx, LODIndex, RigidInfluenceIndex, ReferenceToLocal, ClothSimData, ClothBlendWeight, WorldToLocal); break;
+		case 5: SkinVertexSection<bExtraBoneInfluences, 5, BaseVertexType, VertexType>(DestVertex, MorphEvalInfos, MorphWeights, Section, LOD, VertexBufferBaseIndex, NumValidMorphs, CurBaseVertIdx, LODIndex, RigidInfluenceIndex, ReferenceToLocal, ClothSimData, ClothBlendWeight, WorldToLocal); break;
+		case 6: SkinVertexSection<bExtraBoneInfluences, 6, BaseVertexType, VertexType>(DestVertex, MorphEvalInfos, MorphWeights, Section, LOD, VertexBufferBaseIndex, NumValidMorphs, CurBaseVertIdx, LODIndex, RigidInfluenceIndex, ReferenceToLocal, ClothSimData, ClothBlendWeight, WorldToLocal); break;
+		case 7: SkinVertexSection<bExtraBoneInfluences, 7, BaseVertexType, VertexType>(DestVertex, MorphEvalInfos, MorphWeights, Section, LOD, VertexBufferBaseIndex, NumValidMorphs, CurBaseVertIdx, LODIndex, RigidInfluenceIndex, ReferenceToLocal, ClothSimData, ClothBlendWeight, WorldToLocal); break;
+		case 8: SkinVertexSection<bExtraBoneInfluences, 8, BaseVertexType, VertexType>(DestVertex, MorphEvalInfos, MorphWeights, Section, LOD, VertexBufferBaseIndex, NumValidMorphs, CurBaseVertIdx, LODIndex, RigidInfluenceIndex, ReferenceToLocal, ClothSimData, ClothBlendWeight, WorldToLocal); break;
 		default: check(0);
 	}
 }
 
 template<typename BaseVertexType, typename VertexType>
-static void SkinVertices(FFinalSkinVertex* DestVertex, FMatrix* ReferenceToLocal, int32 LODIndex, FStaticLODModel& LOD, TArray<FActiveMorphTarget>& ActiveMorphTargets, TArray<float>& MorphTargetWeights)
+static void SkinVertices(FFinalSkinVertex* DestVertex, FMatrix* ReferenceToLocal, int32 LODIndex, FStaticLODModel& LOD, TArray<FActiveMorphTarget>& ActiveMorphTargets, TArray<float>& MorphTargetWeights, const TMap<int32, FClothSimulData>& ClothSimulUpdateData, float ClothBlendWeight, const FMatrix& WorldToLocal)
 {
 	uint32 StatusRegister = VectorGetControlRegister();
 	VectorSetControlRegister( StatusRegister | VECTOR_ROUND_TOWARD_ZERO );
@@ -858,13 +948,15 @@ static void SkinVertices(FFinalSkinVertex* DestVertex, FMatrix* ReferenceToLocal
 	{
 		FSkelMeshSection& Section = LOD.Sections[SectionIndex];
 
+		const FClothSimulData* ClothSimData = ClothSimulUpdateData.Find(Section.CorrespondClothAssetIndex);
+
 		if (LOD.DoSectionsNeedExtraBoneInfluences())
 		{
-			SkinVertexSection<true, BaseVertexType, VertexType>(DestVertex, MorphEvalInfos, MorphTargetWeights, Section, LOD, VertexBufferBaseIndex, NumValidMorphs, CurBaseVertIdx, LODIndex, RigidInfluenceIndex, ReferenceToLocal);
+			SkinVertexSection<true, BaseVertexType, VertexType>(DestVertex, MorphEvalInfos, MorphTargetWeights, Section, LOD, VertexBufferBaseIndex, NumValidMorphs, CurBaseVertIdx, LODIndex, RigidInfluenceIndex, ReferenceToLocal, ClothSimData, ClothBlendWeight, WorldToLocal);
 		}
 		else
 		{
-			SkinVertexSection<false, BaseVertexType, VertexType>(DestVertex, MorphEvalInfos, MorphTargetWeights, Section, LOD, VertexBufferBaseIndex, NumValidMorphs, CurBaseVertIdx, LODIndex, RigidInfluenceIndex, ReferenceToLocal);
+			SkinVertexSection<false, BaseVertexType, VertexType>(DestVertex, MorphEvalInfos, MorphTargetWeights, Section, LOD, VertexBufferBaseIndex, NumValidMorphs, CurBaseVertIdx, LODIndex, RigidInfluenceIndex, ReferenceToLocal, ClothSimData, ClothBlendWeight, WorldToLocal);
 		}
 	}
 
@@ -979,4 +1071,41 @@ static void CalculateMorphTargetWeights(FFinalSkinVertex* DestVertex, FStaticLOD
 	}
 }
 
+bool FDynamicSkelMeshObjectDataCPUSkin::UpdateClothSimulationData(USkinnedMeshComponent* InMeshComponent)
+{
+	USkeletalMeshComponent* SimMeshComponent = Cast<USkeletalMeshComponent>(InMeshComponent);
+
+#if WITH_APEX_CLOTHING
+	if (InMeshComponent->MasterPoseComponent.IsValid() && (SimMeshComponent && SimMeshComponent->IsClothBoundToMasterComponent()))
+	{
+		USkeletalMeshComponent* SrcComponent = SimMeshComponent;
+
+		// if I have master, override sim component
+		SimMeshComponent = Cast<USkeletalMeshComponent>(InMeshComponent->MasterPoseComponent.Get());
+
+		// IF we don't have sim component that is skeletalmeshcomponent, just ignore
+		if (!SimMeshComponent)
+		{
+			return false;
+		}
+
+		WorldToLocal = SrcComponent->GetRenderMatrix().InverseFast();
+		ClothBlendWeight = SrcComponent->ClothBlendWeight;
+		SimMeshComponent->GetUpdateClothSimulationData(ClothSimulUpdateData, SrcComponent);
+
+		return true;
+	}
+#endif
+
+	if (SimMeshComponent)
+	{
+		WorldToLocal = SimMeshComponent->GetRenderMatrix().InverseFast();
+		ClothBlendWeight = SimMeshComponent->ClothBlendWeight;
+		SimMeshComponent->GetUpdateClothSimulationData(ClothSimulUpdateData);
+		return true;
+	}
+	return false;
+}
+
+
 MSVC_PRAGMA(warning(pop))
diff --git a/Engine/Source/Runtime/Engine/Private/SkeletalRenderCPUSkin.h b/Engine/Source/Runtime/Engine/Private/SkeletalRenderCPUSkin.h
--- a/Engine/Source/Runtime/Engine/Private/SkeletalRenderCPUSkin.h
+++ b/Engine/Source/Runtime/Engine/Private/SkeletalRenderCPUSkin.h
@@ -10,16 +10,6 @@
 #include "SkeletalRenderPublic.h"
 #include "LocalVertexFactory.h"
 
-/** data for a single skinned skeletal mesh vertex */
-struct FFinalSkinVertex
-{
-	FVector			Position;
-	FPackedNormal	TangentX;
-	FPackedNormal	TangentZ;
-	float			U;
-	float			V;
-};
-
 /**
  * Skeletal mesh vertices which have been skinned to their final positions 
  */
@@ -104,6 +94,9 @@ public:
 	{
 	}
 
+	/** Local to world transform, used for cloth as sim data is in world space */
+	FMatrix WorldToLocal;
+
 	/** ref pose to local space transforms */
 	TArray<FMatrix> ReferenceToLocal;
 
@@ -121,6 +114,12 @@ public:
 	/** Morph Weights to blend when skinning verts */
 	TArray<float> MorphTargetWeights;
 
+	/** data for updating cloth section */
+	TMap<int32, FClothSimulData> ClothSimulUpdateData;
+
+	/** a weight factor to blend between simulated positions and skinned positions */
+	float ClothBlendWeight;
+
 	/**
 	* Returns the size of memory allocated by render data
 	*/
@@ -133,6 +132,9 @@ public:
  		
 		return ResourceSize;
  	}
+
+	/** Update Simulated Positions & Normals from APEX Clothing actor */
+	bool UpdateClothSimulationData(USkinnedMeshComponent* InMeshComponent);
 };
 
 /**
@@ -202,6 +204,9 @@ public:
 	virtual void DrawVertexElements(FPrimitiveDrawInterface* PDI, const FTransform& ToWorldSpace, bool bDrawNormals, bool bDrawTangents, bool bDrawBinormals) const override;
 	//~ End FSkeletalMeshObject Interface
 
+	/** Access cached final vertices */
+	const TArray<FFinalSkinVertex>& GetCachedFinalVertices() const { return CachedFinalVertices; }
+
 private:
 	/** vertex data for rendering a single LOD */
 	struct FSkeletalMeshObjectLOD
@@ -250,7 +255,7 @@ private:
  	/** Index of LOD level's vertices that are currently stored in CachedFinalVertices */
  	mutable int32	CachedVertexLOD;
 
- 	/** Cached skinned vertices. Only updated/accessed by the rendering thread */
+ 	/** Cached skinned vertices. Only updated/accessed by the rendering thread and exporters */
  	mutable TArray<FFinalSkinVertex> CachedFinalVertices;
 
 	/** Array of bone's to render bone weights for */
diff --git a/Engine/Source/Runtime/Engine/Private/Slate/SceneViewport.cpp b/Engine/Source/Runtime/Engine/Private/Slate/SceneViewport.cpp
--- a/Engine/Source/Runtime/Engine/Private/Slate/SceneViewport.cpp
+++ b/Engine/Source/Runtime/Engine/Private/Slate/SceneViewport.cpp
@@ -1004,6 +1004,11 @@ void FSceneViewport::OnViewportClosed()
 	}
 }
 
+FReply FSceneViewport::OnRequestWindowClose()
+{
+	return (ViewportClient && !ViewportClient->WindowCloseRequested()) ? FReply::Handled() : FReply::Unhandled();
+}
+
 TWeakPtr<SWidget> FSceneViewport::GetWidget()
 {
 	return GetViewportWidget();
diff --git a/Engine/Source/Runtime/Engine/Private/SoundBase.cpp b/Engine/Source/Runtime/Engine/Private/SoundBase.cpp
--- a/Engine/Source/Runtime/Engine/Private/SoundBase.cpp
+++ b/Engine/Source/Runtime/Engine/Private/SoundBase.cpp
@@ -7,6 +7,7 @@
 
 USoundClass* USoundBase::DefaultSoundClassObject = nullptr;
 USoundConcurrency* USoundBase::DefaultSoundConcurrencyObject = nullptr;
+USoundSubmix* USoundBase::DefaultSoundSubmixObject = nullptr;
 
 USoundBase::USoundBase(const FObjectInitializer& ObjectInitializer)
 	: Super(ObjectInitializer)
@@ -30,6 +31,16 @@ void USoundBase::PostInitProperties()
 	}
 	SoundClassObject = USoundBase::DefaultSoundClassObject;
 
+	if (USoundBase::DefaultSoundSubmixObject == nullptr)
+	{
+		const FStringAssetReference DefaultSoundSubmixName = GetDefault<UAudioSettings>()->DefaultSoundSubmixName;
+		if (DefaultSoundSubmixName.IsValid())
+		{
+			USoundBase::DefaultSoundSubmixObject = LoadObject<USoundSubmix>(nullptr, *DefaultSoundSubmixName.ToString());
+		}
+	}
+	SoundSubmixObject = USoundBase::DefaultSoundSubmixObject;
+
 	if (USoundBase::DefaultSoundConcurrencyObject == nullptr)
 	{
 		const FStringAssetReference DefaultSoundConcurrencyName = GetDefault<UAudioSettings>()->DefaultSoundConcurrencyName;
@@ -91,6 +102,11 @@ USoundClass* USoundBase::GetSoundClass() const
 	return SoundClassObject;
 }
 
+USoundSubmix* USoundBase::GetSoundSubmix() const
+{
+	return SoundSubmixObject;
+}
+
 const FSoundConcurrencySettings* USoundBase::GetSoundConcurrencySettingsToApply()
 {
 	if (bOverrideConcurrency)
diff --git a/Engine/Source/Runtime/Engine/Private/SoundClass.cpp b/Engine/Source/Runtime/Engine/Private/SoundClass.cpp
--- a/Engine/Source/Runtime/Engine/Private/SoundClass.cpp
+++ b/Engine/Source/Runtime/Engine/Private/SoundClass.cpp
@@ -3,15 +3,14 @@
 
 #include "EnginePrivate.h"
 #include "SoundDefinitions.h"
+#include "Sound/SoundMix.h"
+
 #if WITH_EDITOR
 #include "UnrealEd.h"
-#endif
-#if WITH_EDITOR
 #include "SlateBasics.h"
 #include "SNotificationList.h"
 #include "NotificationManager.h"
 #endif
-#include "Sound/SoundMix.h"
 
 /*-----------------------------------------------------------------------------
 	USoundClass implementation.
diff --git a/Engine/Source/Runtime/Engine/Private/SoundConcurrency.cpp b/Engine/Source/Runtime/Engine/Private/SoundConcurrency.cpp
--- a/Engine/Source/Runtime/Engine/Private/SoundConcurrency.cpp
+++ b/Engine/Source/Runtime/Engine/Private/SoundConcurrency.cpp
@@ -124,13 +124,20 @@ void FConcurrencyGroup::StopQuietSoundsDueToMaxConcurrency()
 
 	// Need to make a new list when stopping the sounds since the process of stopping an active sound
 	// will remove the sound from this concurrency group's ActiveSounds array.
-	for (int32 i = 0; i < NumSoundsToStop; ++i)
+	int32 i = 0;
+	for (; i < NumSoundsToStop; ++i)
 	{
 		// Flag this active sound as needing to be stopped due to volume-based max concurrency.
 		// This will actually be stopped in the audio device update function.
 		ActiveSounds[i]->bShouldStopDueToMaxConcurrency = true;
 	}
 
+	const int32 NumActiveSounds = ActiveSounds.Num();
+	for (; i < NumActiveSounds; ++i)
+	{
+		ActiveSounds[i]->bShouldStopDueToMaxConcurrency = false;
+	}
+
 }
 
 
diff --git a/Engine/Source/Runtime/Engine/Private/SoundEffectBase.cpp b/Engine/Source/Runtime/Engine/Private/SoundEffectBase.cpp
new file mode 100644
--- /dev/null
+++ b/Engine/Source/Runtime/Engine/Private/SoundEffectBase.cpp
@@ -0,0 +1,56 @@
+// Copyright 1998-2016 Epic Games, Inc. All Rights Reserved.
+
+
+#include "EnginePrivate.h"
+#include "SoundDefinitions.h"
+
+#if WITH_EDITOR
+#include "UnrealEd.h"
+#include "SlateBasics.h"
+#include "SNotificationList.h"
+#include "NotificationManager.h"
+#include "Toolkits/SimpleAssetEditor.h"
+#include "Toolkits/AssetEditorManager.h"
+#endif
+
+#include "Sound/SoundEffectBase.h"
+
+/*-----------------------------------------------------------------------------
+	USoundEffectBase Implementation
+-----------------------------------------------------------------------------*/
+
+USoundEffectBase::USoundEffectBase(const FObjectInitializer& ObjectInitializer)
+	: Super(ObjectInitializer)
+	, SoundEffectPreset(nullptr)
+	, PresetSettingsSize(0)
+	, bIsRunning(false)
+	, bIsActive(false)
+{}
+
+USoundEffectBase::~USoundEffectBase()
+{
+}
+
+void USoundEffectBase::SetPreset(USoundEffectPreset* InPreset)
+{
+	check(InPreset->GetEffectClass() == this->GetEffectClass());
+
+	const FPresetSettings& PresetSettings = InPreset->GetPresetSettings();
+
+	// Reset the preset data scratch buffer to the size of the settings
+	RawPresetDataScratchInputBuffer.Init(0, PresetSettings.Size);
+
+	// Copy the data to the raw data array
+	UProperty* Property = PresetSettings.Struct->PropertyLink;
+	while (Property)
+	{
+		uint32 PropertyOffset = Property->GetOffset_ForInternal();
+		Property->CopyCompleteValue((uint8*)RawPresetDataScratchInputBuffer.GetData() + PropertyOffset, (const uint8*)PresetSettings.Data + PropertyOffset);
+		Property = Property->PropertyLinkNext;
+	}
+
+	// Queuing the data will copy the data to the queue
+	EffectSettingsQueue.Enqueue(RawPresetDataScratchInputBuffer);
+}
+
+
diff --git a/Engine/Source/Runtime/Engine/Private/SoundEffectPreset.cpp b/Engine/Source/Runtime/Engine/Private/SoundEffectPreset.cpp
new file mode 100644
--- /dev/null
+++ b/Engine/Source/Runtime/Engine/Private/SoundEffectPreset.cpp
@@ -0,0 +1,147 @@
+// Copyright 1998-2016 Epic Games, Inc. All Rights Reserved.
+
+
+#include "EnginePrivate.h"
+#include "SoundDefinitions.h"
+
+#if WITH_EDITOR
+#include "UnrealEd.h"
+#include "SlateBasics.h"
+#include "SNotificationList.h"
+#include "NotificationManager.h"
+#include "Toolkits/SimpleAssetEditor.h"
+#include "Toolkits/AssetEditorManager.h"
+#endif
+
+#include "Sound/SoundEffectPreset.h"
+
+/*-----------------------------------------------------------------------------
+	FAssetTypeActions_SoundEffectPreset Implementation
+-----------------------------------------------------------------------------*/
+
+#if WITH_EDITOR
+FAssetTypeActions_SoundEffectPreset::FAssetTypeActions_SoundEffectPreset(USoundEffectPreset* InEffectPreset)
+	: EffectPreset(InEffectPreset)
+{
+}
+
+bool FAssetTypeActions_SoundEffectPreset::HasActions(const TArray<UObject*>& InObjects) const
+{
+	return false;
+}
+
+void FAssetTypeActions_SoundEffectPreset::GetActions(const TArray<UObject*>& InObjects, FMenuBuilder& MenuBuilder)
+{
+}
+
+void FAssetTypeActions_SoundEffectPreset::OpenAssetEditor(const TArray<UObject*>& InObjects, TSharedPtr<class IToolkitHost> EditWithinLevelEditor)
+{
+	FSimpleAssetEditor::CreateEditor(EToolkitMode::Standalone, EditWithinLevelEditor, InObjects);
+}
+
+void FAssetTypeActions_SoundEffectPreset::AssetsActivated(const TArray<UObject*>& InObjects, EAssetTypeActivationMethod::Type ActivationType)
+{
+	if (ActivationType == EAssetTypeActivationMethod::DoubleClicked || ActivationType == EAssetTypeActivationMethod::Opened)
+	{
+		if (InObjects.Num() == 1)
+		{
+			FAssetEditorManager::Get().OpenEditorForAsset(InObjects[0]);
+		}
+		else if (InObjects.Num() > 1)
+		{
+			FAssetEditorManager::Get().OpenEditorForAssets(InObjects);
+		}
+	}
+}
+
+bool FAssetTypeActions_SoundEffectPreset::CanFilter()
+{
+	return true;
+}
+
+bool FAssetTypeActions_SoundEffectPreset::CanLocalize() const
+{
+	return false;
+}
+
+bool FAssetTypeActions_SoundEffectPreset::CanMerge() const
+{
+	return false;
+}
+
+void FAssetTypeActions_SoundEffectPreset::Merge(UObject* InObject)
+{
+}
+
+void FAssetTypeActions_SoundEffectPreset::Merge(UObject* BaseAsset, UObject* RemoteAsset, UObject* LocalAsset, const FOnMergeResolved& ResolutionCallback)
+{
+}
+
+bool FAssetTypeActions_SoundEffectPreset::ShouldForceWorldCentric()
+{
+	return false;
+}
+
+void FAssetTypeActions_SoundEffectPreset::PerformAssetDiff(UObject* OldAsset, UObject* NewAsset, const struct FRevisionInfo& OldRevision, const struct FRevisionInfo& NewRevision) const
+{
+	check(OldAsset != nullptr);
+	check(NewAsset != nullptr);
+
+	FAssetToolsModule& AssetToolsModule = FModuleManager::Get().LoadModuleChecked<FAssetToolsModule>("AssetTools");
+
+	// Dump assets to temp text files
+	FString OldTextFilename = AssetToolsModule.Get().DumpAssetToTempFile(OldAsset);
+	FString NewTextFilename = AssetToolsModule.Get().DumpAssetToTempFile(NewAsset);
+	FString DiffCommand = GetDefault<UEditorLoadingSavingSettings>()->TextDiffToolPath.FilePath;
+
+	AssetToolsModule.Get().CreateDiffProcess(DiffCommand, OldTextFilename, NewTextFilename);
+}
+
+class UThumbnailInfo* FAssetTypeActions_SoundEffectPreset::GetThumbnailInfo(UObject* Asset) const
+{
+	return nullptr;
+}
+
+TSharedPtr<class SWidget> FAssetTypeActions_SoundEffectPreset::GetThumbnailOverlay(const FAssetData& AssetData) const
+{
+	return nullptr;
+}
+
+FText FAssetTypeActions_SoundEffectPreset::GetAssetDescription(const class FAssetData& AssetData) const
+{
+	return FText::GetEmpty();
+}
+
+bool FAssetTypeActions_SoundEffectPreset::IsImportedAsset() const
+{
+	return false;
+}
+
+void FAssetTypeActions_SoundEffectPreset::GetResolvedSourceFilePaths(const TArray<UObject*>& TypeAssets, TArray<FString>& OutSourceFilePaths) const
+{
+}
+
+FText FAssetTypeActions_SoundEffectPreset::GetName() const
+{ 
+	return EffectPreset->GetAssetActionName(); 
+}
+
+UClass* FAssetTypeActions_SoundEffectPreset::GetSupportedClass() const
+{ 
+	return EffectPreset->GetClass();
+}
+
+#endif
+
+/*-----------------------------------------------------------------------------
+	USoundEffectPresetBase Implementation
+-----------------------------------------------------------------------------*/
+
+USoundEffectPreset::USoundEffectPreset(const FObjectInitializer& ObjectInitializer)
+	: Super(ObjectInitializer)
+{
+	PresetSettings.Data = GetSettings();
+	PresetSettings.Size = GetSettingsSize();
+	PresetSettings.Struct = GetSettingsStruct();
+}
+
diff --git a/Engine/Source/Runtime/Engine/Private/SoundEffectSource.cpp b/Engine/Source/Runtime/Engine/Private/SoundEffectSource.cpp
new file mode 100644
--- /dev/null
+++ b/Engine/Source/Runtime/Engine/Private/SoundEffectSource.cpp
@@ -0,0 +1,39 @@
+// Copyright 1998-2016 Epic Games, Inc. All Rights Reserved.
+
+
+#include "EnginePrivate.h"
+#include "Sound/SoundEffectSource.h"
+
+USoundEffectSourcePreset::USoundEffectSourcePreset(const FObjectInitializer& ObjectInitializer)
+	: Super(ObjectInitializer)
+{
+}
+
+USoundEffectSource::USoundEffectSource(const FObjectInitializer& ObjectInitializer)
+	: Super(ObjectInitializer)
+{
+}
+
+void USoundEffectSource::ProcessAudio(FSoundEffectSourceInputData& InData, FSoundEffectSourceOutputData& OutData)
+{
+	bIsRunning = true;
+
+	// Update the latest version of the effect settings on the audio thread. If there are new settings,
+	// then RawPresetDataScratchOutputBuffer will have the last data.
+	EffectSettingsQueue.Dequeue(RawPresetDataScratchOutputBuffer);
+
+	InData.PresetData = RawPresetDataScratchOutputBuffer.GetData();
+
+	// Only process the effect if the effect is active
+	if (bIsActive)
+	{
+		OnProcessAudio(InData, OutData);
+	}
+	else
+	{
+		// otherwise, bypass the effect and move inputs to outputs
+		OutData.AudioBuffer = MoveTemp(InData.AudioBuffer);
+		OutData.bIsSpatialized = InData.bIsSpatialized;
+		OutData.OutputChannels = InData.NumSourceChannels;
+	}
+}
\ No newline at end of file
diff --git a/Engine/Source/Runtime/Engine/Private/SoundEffectSubmix.cpp b/Engine/Source/Runtime/Engine/Private/SoundEffectSubmix.cpp
new file mode 100644
--- /dev/null
+++ b/Engine/Source/Runtime/Engine/Private/SoundEffectSubmix.cpp
@@ -0,0 +1,38 @@
+// Copyright 1998-2016 Epic Games, Inc. All Rights Reserved.
+
+
+#include "EnginePrivate.h"
+#include "Sound/SoundEffectSubmix.h"
+
+USoundEffectSubmixPreset::USoundEffectSubmixPreset(const FObjectInitializer& ObjectInitializer)
+	: Super(ObjectInitializer)
+{
+}
+
+
+USoundEffectSubmix::USoundEffectSubmix(const FObjectInitializer& ObjectInitializer)
+	: Super(ObjectInitializer)
+{
+}
+
+void USoundEffectSubmix::ProcessAudio(FSoundEffectSubmixInputData& InData, FSoundEffectSubmixOutputData& OutData)
+{
+	bIsRunning = true;
+
+	// Update the latest version of the effect settings on the audio thread. If there are new settings,
+	// then RawPresetDataScratchOutputBuffer will have the last data.
+	EffectSettingsQueue.Dequeue(RawPresetDataScratchOutputBuffer);
+
+	InData.PresetData = RawPresetDataScratchOutputBuffer.GetData();
+
+	// Only process the effect if the effect is active
+	if (bIsActive)
+	{
+		OnProcessAudio(InData, OutData);
+	}
+	else
+	{
+		// otherwise, bypass the effect and move inputs to outputs
+		OutData.AudioBuffer = MoveTemp(InData.AudioBuffer);
+	}
+}
\ No newline at end of file
diff --git a/Engine/Source/Runtime/Engine/Private/SoundNode.cpp b/Engine/Source/Runtime/Engine/Private/SoundNode.cpp
--- a/Engine/Source/Runtime/Engine/Private/SoundNode.cpp
+++ b/Engine/Source/Runtime/Engine/Private/SoundNode.cpp
@@ -176,6 +176,13 @@ void USoundNode::PostEditChangeProperty(FPropertyChangedEvent& PropertyChangedEv
 
 	MarkPackageDirty();
 }
+
+void USoundNode::PostLoad()
+{
+	Super::PostLoad();
+	// Make sure sound nodes are transactional (so they work with undo system)
+	SetFlags(RF_Transactional);
+}
 #endif // WITH_EDITOR
 
 #if WITH_EDITOR
diff --git a/Engine/Source/Runtime/Engine/Private/SoundWave.cpp b/Engine/Source/Runtime/Engine/Private/SoundWave.cpp
--- a/Engine/Source/Runtime/Engine/Private/SoundWave.cpp
+++ b/Engine/Source/Runtime/Engine/Private/SoundWave.cpp
@@ -387,13 +387,13 @@ void USoundWave::PostLoad()
 	if( !GIsEditor && !IsTemplate( RF_ClassDefaultObject ) && GEngine )
 	{
 		FAudioDevice* AudioDevice = GEngine->GetMainAudioDevice();
-		if( AudioDevice && AudioDevice->AreStartupSoundsPreCached())
+		if (AudioDevice && AudioDevice->AreStartupSoundsPreCached())
 		{
 			// Upload the data to the hardware, but only if we've precached startup sounds already
-			AudioDevice->Precache( this );
+			AudioDevice->Precache(this);
 		}
 		// remove bulk data if no AudioDevice is used and no sounds were initialized
-		else if( IsRunningGame() )
+		else if (IsRunningGame())
 		{
 			RawData.RemoveBulkData();
 		}
@@ -659,6 +659,9 @@ void USoundWave::Parse( FAudioDevice* AudioDevice, const UPTRINT NodeWaveInstanc
 		WaveInstance->UserIndex = ActiveSound.UserIndex;
 		WaveInstance->OmniRadius = ParseParams.OmniRadius;
 		WaveInstance->StereoSpread = ParseParams.StereoSpread;
+		WaveInstance->AttenuationDistance = ParseParams.AttenuationDistance;
+		WaveInstance->AbsoluteAzimuth = ParseParams.AbsoluteAzimuth;
+
 		bool bAlwaysPlay = false;
 
 		// Properties from the sound class
@@ -727,6 +730,7 @@ void USoundWave::Parse( FAudioDevice* AudioDevice, const UPTRINT NodeWaveInstanc
 		WaveInstance->WaveData = this;
 		WaveInstance->NotifyBufferFinishedHooks = ParseParams.NotifyBufferFinishedHooks;
 		WaveInstance->LoopingMode = ((bLooping || ParseParams.bLooping) ? LOOP_Forever : LOOP_Never);
+		WaveInstance->bIsPaused = ParseParams.bIsPaused;
 
 		if (AudioDevice->IsHRTFEnabledForAll() && ParseParams.SpatializationAlgorithm == SPATIALIZATION_Default)
 		{
diff --git a/Engine/Source/Runtime/Engine/Private/Tests/AutomationCommon.cpp b/Engine/Source/Runtime/Engine/Private/Tests/AutomationCommon.cpp
--- a/Engine/Source/Runtime/Engine/Private/Tests/AutomationCommon.cpp
+++ b/Engine/Source/Runtime/Engine/Private/Tests/AutomationCommon.cpp
@@ -5,7 +5,7 @@
 #include "AutomationCommon.h"
 #include "ImageUtils.h"
 #include "ShaderCompiler.h"		// GShaderCompilingManager
-#include "GameFramework/GameMode.h"
+#include "GameFramework/GameStateBase.h"
 #include "Scalability.h"
 
 #if WITH_EDITOR
@@ -272,8 +272,8 @@ bool FWaitForMapToLoadCommand::Update()
 
 	if ( TestWorld && TestWorld->AreActorsInitialized() )
 	{
-		AGameMode* GameMode = TestWorld->GetAuthGameMode();
-		if ( GameMode && GameMode->HasMatchStarted() )
+		AGameStateBase* GameState = TestWorld->GetGameState();
+		if (GameState && GameState->HasMatchStarted() )
 		{
 			return true;
 		}
diff --git a/Engine/Source/Runtime/Engine/Private/TickTaskManager.cpp b/Engine/Source/Runtime/Engine/Private/TickTaskManager.cpp
--- a/Engine/Source/Runtime/Engine/Private/TickTaskManager.cpp
+++ b/Engine/Source/Runtime/Engine/Private/TickTaskManager.cpp
@@ -167,12 +167,15 @@ struct FTickContext
 	ETickingGroup			TickGroup;
 	/** Current or desired thread **/
 	ENamedThreads::Type		Thread;
+	/** The world in which the object being ticked is contained. **/
+	UWorld*					World;
 
 	FTickContext(float InDeltaSeconds = 0.0f, ELevelTick InTickType = LEVELTICK_All, ETickingGroup InTickGroup = TG_PrePhysics, ENamedThreads::Type InThread = ENamedThreads::GameThread)
 		: DeltaSeconds(InDeltaSeconds)
 		, TickType(InTickType)
 		, TickGroup(InTickGroup)
 		, Thread(InThread)
+		, World(nullptr)
 	{
 	}
 
@@ -181,6 +184,7 @@ struct FTickContext
 		, TickType(In.TickType)
 		, TickGroup(In.TickGroup)
 		, Thread(In.Thread)
+		, World(In.World)
 	{
 	}
 	void operator=(const FTickContext& In)
@@ -189,6 +193,7 @@ struct FTickContext
 		TickType = In.TickType;
 		TickGroup = In.TickGroup;
 		Thread = In.Thread;
+		World = In.World;
 	}
 };
 
@@ -241,13 +246,13 @@ public:
 	{
 		if (bLogTick)
 		{
-				UE_LOG(LogTick, Log, TEXT("tick %s [%1d, %1d] %6d %2d %s"), Target->bHighPriority ? TEXT("*") : TEXT(" "), (int32)Target->GetActualTickGroup(), (int32)Target->GetActualEndTickGroup(), GFrameCounter, (int32)CurrentThread, *Target->DiagnosticMessage());
+			UE_LOG(LogTick, Log, TEXT("tick %s [%1d, %1d] %6d %2d %s"), Target->bHighPriority ? TEXT("*") : TEXT(" "), (int32)Target->GetActualTickGroup(), (int32)Target->GetActualEndTickGroup(), GFrameCounter, (int32)CurrentThread, *Target->DiagnosticMessage());
 			if (bLogTicksShowPrerequistes)
 			{
 				Target->ShowPrerequistes();
 			}
 		}
-		Target->ExecuteTick(Context.DeltaSeconds, Context.TickType, CurrentThread, MyCompletionGraphEvent);
+		Target->ExecuteTick(Target->CalculateDeltaTime(Context), Context.TickType, CurrentThread, MyCompletionGraphEvent);
 		Target->TaskPointer = nullptr;  // This is stale and a good time to clear it for safety
 	}
 };
@@ -697,6 +702,7 @@ public:
 		Context.DeltaSeconds = InContext.DeltaSeconds;
 		Context.TickType = InContext.TickType;
 		Context.Thread = ENamedThreads::GameThread;
+		Context.World = InContext.World;
 		bTickNewlySpawned = true;
 
 		int32 CooldownTicksEnabled = 0;
@@ -734,6 +740,7 @@ public:
 		Context.DeltaSeconds = InContext.DeltaSeconds;
 		Context.TickType = InContext.TickType;
 		Context.Thread = ENamedThreads::GameThread;
+		Context.World = InContext.World;
 		bTickNewlySpawned = true;
 
 		{
@@ -896,10 +903,10 @@ public:
 			TickFunction->QueueTickFunction(TTS, Context);
 
 			if (TickFunction->TickInterval > 0.f)
-		{
+			{
 				It.RemoveCurrent();
 				TickFunctionsToReschedule.Add(FTickScheduleDetails(TickFunction, TickFunction->TickInterval));
-		}
+			}
 		}
 		int32 EnabledCooldownTicks = 0;
 		float CumulativeCooldown = 0.f;
@@ -985,10 +992,10 @@ public:
 			if (TickFunction->bTickEvenWhenPaused)
 			{
 				if (CumulativeCooldown < InContext.DeltaSeconds)
-		{
+				{
 					TickFunction->TickVisitedGFrameCounter = GFrameCounter;
 					TickFunction->TickQueuedGFrameCounter = GFrameCounter;
-					TickFunction->ExecuteTick(InContext.DeltaSeconds, InContext.TickType, ENamedThreads::GameThread, FGraphEventRef());
+					TickFunction->ExecuteTick(TickFunction->CalculateDeltaTime(InContext), InContext.TickType, ENamedThreads::GameThread, FGraphEventRef());
 					TickFunction->TaskPointer = nullptr; // this is stale, clear it out now
 
 					TickFunctionsToReschedule.Add(FTickScheduleDetails(TickFunction, TickFunction->TickInterval - (InContext.DeltaSeconds - CumulativeCooldown))); // Give credit for any overrun
@@ -1012,7 +1019,10 @@ public:
 					CumulativeCooldown -= TickFunction->RelativeTickCooldown; // Since the next object in the list will have this cooldown included take it back out of the cumulative
 				}
 			}
-			PrevTickFunction = TickFunction;
+			else
+			{
+				PrevTickFunction = TickFunction;
+			}
 			TickFunction = TickFunction->Next;
 		}
 
@@ -1024,7 +1034,7 @@ public:
 			{
 				TickFunction->TickVisitedGFrameCounter = GFrameCounter;
 				TickFunction->TickQueuedGFrameCounter = GFrameCounter;
-				TickFunction->ExecuteTick(InContext.DeltaSeconds, InContext.TickType, ENamedThreads::GameThread, FGraphEventRef());
+				TickFunction->ExecuteTick(TickFunction->CalculateDeltaTime(InContext), InContext.TickType, ENamedThreads::GameThread, FGraphEventRef());
 
 				if (TickFunction->TickInterval > 0.f)
 				{
@@ -1326,11 +1336,11 @@ public:
 			FPlatformProcess::Sleep(CVarStallStartFrame.GetValueOnGameThread() / 1000.0f);
 		}
 #endif
-		World = InWorld;
 		Context.TickGroup = ETickingGroup(0); // reset this to the start tick group
 		Context.DeltaSeconds = InDeltaSeconds;
 		Context.TickType = InTickType;
 		Context.Thread = ENamedThreads::GameThread;
+		Context.World = InWorld;
 
 		bTickNewlySpawned = true;
 		TickTaskSequencer.StartFrame();
@@ -1403,13 +1413,13 @@ public:
 		Context.DeltaSeconds = InDeltaSeconds;
 		Context.TickType = InTickType;
 		Context.Thread = ENamedThreads::GameThread;
-		World = InWorld;
+		Context.World = InWorld;
 		FillLevelList();
 		for( int32 LevelIndex = 0; LevelIndex < LevelList.Num(); LevelIndex++ )
 		{
 			LevelList[LevelIndex]->RunPauseFrame(Context);
 		}
-		World = NULL;
+		Context.World = nullptr;
 		bTickNewlySpawned = false;
 		LevelList.Reset();
 	}
@@ -1467,7 +1477,7 @@ public:
 		{
 			LevelList[LevelIndex]->EndFrame();
 		}
-		World = NULL;
+		Context.World = nullptr;
 		LevelList.Reset();
 	}
 
@@ -1499,7 +1509,6 @@ private:
 	/** Default constructor **/
 	FTickTaskManager()
 		: TickTaskSequencer(FTickTaskSequencer::Get())
-		, World(NULL)
 		, bTickNewlySpawned(false)
 	{
 		IConsoleManager::Get().RegisterConsoleCommand(TEXT("dumpticks"), TEXT("Dumps all tick functions registered with FTickTaskManager to log."));
@@ -1509,11 +1518,11 @@ private:
 	void FillLevelList()
 	{
 		check(!LevelList.Num());
-		check(World->TickTaskLevel);
-		LevelList.Add(World->TickTaskLevel);
-		for( int32 LevelIndex = 0; LevelIndex < World->GetNumLevels(); LevelIndex++ )
+		check(Context.World->TickTaskLevel);
+		LevelList.Add(Context.World->TickTaskLevel);
+		for( int32 LevelIndex = 0; LevelIndex < Context.World->GetNumLevels(); LevelIndex++ )
 		{
-			ULevel* Level = World->GetLevel(LevelIndex);
+			ULevel* Level = Context.World->GetLevel(LevelIndex);
 			if (Level->bIsVisible)
 			{
 				check(Level->TickTaskLevel);
@@ -1558,8 +1567,6 @@ private:
 
 	/** Global Sequencer														*/
 	FTickTaskSequencer&							TickTaskSequencer;
-	/** World currently ticking **/
-	UWorld*										World;
 	/** List of current levels **/
 	TArray<FTickTaskLevel*>						LevelList;
 	/** tick context **/
@@ -1589,8 +1596,9 @@ FTickFunction::FTickFunction()
 	, TickVisitedGFrameCounter(0)
 	, TickQueuedGFrameCounter(0)
 	, RelativeTickCooldown(0.f)
+	, LastTickGameTimeSeconds(-1.f)
 	, TickInterval(0.f)
-	, TickTaskLevel(NULL)
+	, TickTaskLevel(nullptr)
 {
 }
 
@@ -1880,6 +1888,31 @@ void FTickFunction::QueueTickFunctionParallel(const struct FTickContext& TickCon
 	}
 }
 
+float FTickFunction::CalculateDeltaTime(const FTickContext& TickContext) 
+{
+	float DeltaTimeForFunction = TickContext.DeltaSeconds;
+
+	if (TickInterval == 0.f)
+	{
+		// No tick interval. Return the world delta seconds, and make sure to mark that 
+		// we're not tracking last-tick-time for this object.
+		LastTickGameTimeSeconds = -1.f;
+	}
+	else
+	{
+		// We've got a tick interval. Mark last-tick-time. If we already had last-tick-time, return
+		// the time since then; otherwise, return the world delta seconds.
+		const float CurrentWorldTime = (bTickEvenWhenPaused ? TickContext.World->GetUnpausedTimeSeconds() : TickContext.World->GetTimeSeconds());
+		if (LastTickGameTimeSeconds >= 0.f)
+		{
+			DeltaTimeForFunction = CurrentWorldTime - LastTickGameTimeSeconds;
+		}
+		LastTickGameTimeSeconds = CurrentWorldTime;
+	}
+
+	return DeltaTimeForFunction;
+}
+
 /**
  * Singleton to retrieve the global tick task manager
  * @return Reference to the global tick task manager
diff --git a/Engine/Source/Runtime/Engine/Private/Timeline.cpp b/Engine/Source/Runtime/Engine/Private/Timeline.cpp
--- a/Engine/Source/Runtime/Engine/Private/Timeline.cpp
+++ b/Engine/Source/Runtime/Engine/Private/Timeline.cpp
@@ -72,11 +72,12 @@ void FTimeline::AddEvent(float Time, FOnTimelineEvent Event)
 	Events.Add(NewEntry);
 }
 
-void FTimeline::AddInterpVector(UCurveVector* VectorCurve, FOnTimelineVector InterpFunc, FName PropertyName)
+void FTimeline::AddInterpVector(UCurveVector* VectorCurve, FOnTimelineVector InterpFunc, FName PropertyName, FName TrackName)
 {
 	FTimelineVectorTrack NewEntry;
 	NewEntry.VectorCurve = VectorCurve;
 	NewEntry.InterpFunc = InterpFunc;
+	NewEntry.TrackName = TrackName;
 	NewEntry.VectorPropertyName = PropertyName;
 
 	InterpVectors.Add(NewEntry);
@@ -91,11 +92,12 @@ void FTimeline::AddInterpVector(UCurveVector* VectorCurve, FOnTimelineVectorStat
 	InterpVectors.Add(NewEntry);
 }
 
-void FTimeline::AddInterpFloat(UCurveFloat* FloatCurve, FOnTimelineFloat InterpFunc, FName PropertyName)
+void FTimeline::AddInterpFloat(UCurveFloat* FloatCurve, FOnTimelineFloat InterpFunc, FName PropertyName, FName TrackName)
 {
 	FTimelineFloatTrack NewEntry;
 	NewEntry.FloatCurve = FloatCurve;
 	NewEntry.InterpFunc = InterpFunc;
+	NewEntry.TrackName = TrackName;
 	NewEntry.FloatPropertyName = PropertyName;
 
 	InterpFloats.Add(NewEntry);
@@ -110,11 +112,12 @@ void FTimeline::AddInterpFloat(UCurveFloat* FloatCurve, FOnTimelineFloatStatic I
 	InterpFloats.Add(NewEntry);
 }
 
-void FTimeline::AddInterpLinearColor(UCurveLinearColor* LinearColorCurve, FOnTimelineLinearColor InterpFunc, FName PropertyName)
+void FTimeline::AddInterpLinearColor(UCurveLinearColor* LinearColorCurve, FOnTimelineLinearColor InterpFunc, FName PropertyName, FName TrackName)
 {
 	FTimelineLinearColorTrack NewEntry;
 	NewEntry.LinearColorCurve = LinearColorCurve;
 	NewEntry.InterpFunc = InterpFunc;
+	NewEntry.TrackName = TrackName;
 	NewEntry.LinearColorPropertyName = PropertyName;
 
 	InterpLinearColors.Add(NewEntry);
@@ -129,6 +132,73 @@ void FTimeline::AddInterpLinearColor(UCurveLinearColor* LinearColorCurve, FOnTim
 	InterpLinearColors.Add(NewEntry);
 }
 
+void FTimeline::SetFloatCurve(UCurveFloat* NewFloatCurve, FName FloatTrackName)
+{
+	bool bFoundTrack = false;
+	if (FloatTrackName != NAME_None)
+	{
+		for (FTimelineFloatTrack& FloatTrack : InterpFloats)
+		{
+			if (FloatTrack.TrackName == FloatTrackName)
+			{
+				FloatTrack.FloatCurve = NewFloatCurve;
+				bFoundTrack = true;
+				break;
+			}
+		}
+	}
+
+	if(!bFoundTrack)
+	{
+		UE_LOG(LogTimeline, Log, TEXT("SetFloatCurve: No float track with name %s!"), *FloatTrackName.ToString());
+	}
+}
+
+void FTimeline::SetVectorCurve(UCurveVector* NewVectorCurve, FName VectorTrackName)
+{
+	bool bFoundTrack = false;
+	if (VectorTrackName != NAME_None)
+	{
+		for (FTimelineVectorTrack& VectorTrack : InterpVectors)
+		{
+			if (VectorTrack.TrackName == VectorTrackName)
+			{
+				VectorTrack.VectorCurve = NewVectorCurve;
+				bFoundTrack = true;
+				break;
+			}
+		}
+	}
+
+	if (!bFoundTrack)
+	{
+		UE_LOG(LogTimeline, Log, TEXT("SetVectorCurve: No vector track with name %s!"), *VectorTrackName.ToString());
+	}
+}
+
+void FTimeline::SetLinearColorCurve(UCurveLinearColor* NewLinearColorCurve, FName LinearColorTrackName)
+{
+	bool bFoundTrack = false;
+	if (LinearColorTrackName != NAME_None)
+	{
+		for (FTimelineLinearColorTrack& ColorTrack : InterpLinearColors)
+		{
+			if (ColorTrack.TrackName == LinearColorTrackName)
+			{
+				ColorTrack.LinearColorCurve = NewLinearColorCurve;
+				bFoundTrack = true;
+				break;
+			}
+
+		}
+	}
+
+	if (!bFoundTrack)
+	{
+		UE_LOG(LogTimeline, Log, TEXT("SetLinearColorCurve: No color track with name %s!"), *LinearColorTrackName.ToString());
+	}
+}
+
 void FTimeline::SetPlaybackPosition(float NewPosition, bool bFireEvents, bool bFireUpdate)
 {
 	float OldPosition = Position;
@@ -655,19 +725,19 @@ void UTimelineComponent::AddEvent(float Time, FOnTimelineEvent Event)
 	TheTimeline.AddEvent(Time, Event);
 }
 
-void UTimelineComponent::AddInterpVector(UCurveVector* VectorCurve, FOnTimelineVector InterpFunc, FName PropertyName)
+void UTimelineComponent::AddInterpVector(UCurveVector* VectorCurve, FOnTimelineVector InterpFunc, FName PropertyName, FName TrackName)
 {
-	TheTimeline.AddInterpVector(VectorCurve, InterpFunc, PropertyName);
+	TheTimeline.AddInterpVector(VectorCurve, InterpFunc, PropertyName, TrackName);
 }
 
-void UTimelineComponent::AddInterpFloat(UCurveFloat* FloatCurve, FOnTimelineFloat InterpFunc, FName PropertyName)
+void UTimelineComponent::AddInterpFloat(UCurveFloat* FloatCurve, FOnTimelineFloat InterpFunc, FName PropertyName, FName TrackName)
 {
-	TheTimeline.AddInterpFloat(FloatCurve, InterpFunc, PropertyName);
+	TheTimeline.AddInterpFloat(FloatCurve, InterpFunc, PropertyName, TrackName);
 }
 
-void UTimelineComponent::AddInterpLinearColor(UCurveLinearColor* LinearColorCurve, FOnTimelineLinearColor InterpFunc, FName PropertyName)
+void UTimelineComponent::AddInterpLinearColor(UCurveLinearColor* LinearColorCurve, FOnTimelineLinearColor InterpFunc, FName PropertyName, FName TrackName)
 {
-	TheTimeline.AddInterpLinearColor(LinearColorCurve, InterpFunc, PropertyName);
+	TheTimeline.AddInterpLinearColor(LinearColorCurve, InterpFunc, PropertyName, TrackName);
 }
 
 void UTimelineComponent::SetPlaybackPosition(float NewPosition, bool bFireEvents, bool bFireUpdate)
@@ -735,6 +805,21 @@ bool UTimelineComponent::GetIgnoreTimeDilation() const
 	return bIgnoreTimeDilation;
 }
 
+void UTimelineComponent::SetFloatCurve(UCurveFloat* NewFloatCurve, FName FloatTrackName)
+{
+	TheTimeline.SetFloatCurve(NewFloatCurve, FloatTrackName);
+}
+
+void UTimelineComponent::SetVectorCurve(UCurveVector* NewVectorCurve, FName VectorTrackName)
+{
+	TheTimeline.SetVectorCurve(NewVectorCurve, VectorTrackName);
+}
+
+void UTimelineComponent::SetLinearColorCurve(UCurveLinearColor* NewLinearColorCurve, FName LinearColorTrackName)
+{
+	TheTimeline.SetLinearColorCurve(NewLinearColorCurve, LinearColorTrackName);
+}
+
 void UTimelineComponent::SetPropertySetObject(UObject* NewPropertySetObject)
 {
 	TheTimeline.SetPropertySetObject(NewPropertySetObject);
diff --git a/Engine/Source/Runtime/Engine/Private/TriggerBase.cpp b/Engine/Source/Runtime/Engine/Private/TriggerBase.cpp
--- a/Engine/Source/Runtime/Engine/Private/TriggerBase.cpp
+++ b/Engine/Source/Runtime/Engine/Private/TriggerBase.cpp
@@ -33,14 +33,12 @@ ATriggerBase::ATriggerBase(const FObjectInitializer& ObjectInitializer)
 		CollisionComponent->bHiddenInGame = false;
 	}
 
-	SpriteComponent = CreateDefaultSubobject<UBillboardComponent>(TEXT("Sprite"));
+	SpriteComponent = CreateEditorOnlyDefaultSubobject<UBillboardComponent>(TEXT("Sprite"));
 	if (SpriteComponent)
 	{
 		SpriteComponent->Sprite = ConstructorStatics.TriggerTextureObject.Get();
 		SpriteComponent->RelativeScale3D = FVector(0.5f, 0.5f, 0.5f);
 		SpriteComponent->bHiddenInGame = false;
-		SpriteComponent->AlwaysLoadOnClient = false;
-		SpriteComponent->AlwaysLoadOnServer = false;
 #if WITH_EDITORONLY_DATA
 		SpriteComponent->SpriteInfo.Category = ConstructorStatics.ID_Triggers;
 		SpriteComponent->SpriteInfo.DisplayName = ConstructorStatics.NAME_Triggers;
