[NO NON-PII COMMIT MESSAGE]

diff --git a/Engine/Source/Runtime/Engine/Private/GameInstance.cpp b/Engine/Source/Runtime/Engine/Private/GameInstance.cpp
--- a/Engine/Source/Runtime/Engine/Private/GameInstance.cpp
+++ b/Engine/Source/Runtime/Engine/Private/GameInstance.cpp
@@ -8,6 +8,7 @@
 #include "EnginePrivate.h"
 #include "Engine/GameInstance.h"
 #include "Engine/Engine.h"
+#include "Engine/GameEngine.h"
 #include "Engine/DemoNetDriver.h"
 #include "Engine/LatentActionManager.h"
 #include "Engine/NetworkObjectList.h"
