Temp change to try and isolate and issue

diff --git a/Engine/Source/Runtime/Engine/Private/HierarchicalInstancedStaticMesh.cpp b/Engine/Source/Runtime/Engine/Private/HierarchicalInstancedStaticMesh.cpp
--- a/Engine/Source/Runtime/Engine/Private/HierarchicalInstancedStaticMesh.cpp
+++ b/Engine/Source/Runtime/Engine/Private/HierarchicalInstancedStaticMesh.cpp
@@ -1784,6 +1784,7 @@ void UHierarchicalInstancedStaticMeshComponent::FlushAsyncBuildInstanceBufferTas
 void UHierarchicalInstancedStaticMeshComponent::Serialize(FArchive& Ar)
 {
 	// On save, if we have a pending async build we should wait for it to complete rather than saving an incomplete tree
+/*
 	if (Ar.IsSaving())
 	{
 		if (bIsAsyncBuilding)
@@ -1791,6 +1792,7 @@ void UHierarchicalInstancedStaticMeshComponent::Serialize(FArchive& Ar)
 			BuildTree();
 		}
 	}
+*/
 
 	Super::Serialize(Ar);
 
