Fixed some orion content issues
Added level name to invalid poly warning

diff --git a/Engine/Source/Runtime/Engine/Private/Brush.cpp b/Engine/Source/Runtime/Engine/Private/Brush.cpp
--- a/Engine/Source/Runtime/Engine/Private/Brush.cpp
+++ b/Engine/Source/Runtime/Engine/Private/Brush.cpp
@@ -202,7 +202,7 @@ void ABrush::PostLoad()
 			{
 				if(!Poly.Normal.Equals(Normal))
 				{
-					UE_LOG(LogBrush, Warning, TEXT("%s had invalid poly normals which have been fixed. Resave the level to remove this warning."), *Brush->GetName());
+					UE_LOG(LogBrush, Log, TEXT("%s had invalid poly normals which have been fixed. Resave the level '%s' to remove this warning."), *Brush->GetName(), *GetLevel()->GetOuter()->GetName());
 					if(IsStaticBrush())
 					{
 						UE_LOG(LogBrush, Log, TEXT("%s had invalid poly normals which have been fixed. Resave the level '%s' to remove this warning."), *Brush->GetName(), *GetLevel()->GetOuter()->GetName());
