Attempt to fix nonunity again. I'm having trouble testing it locally, will fully verify once I get local compiling working

diff --git a/Engine/Source/Runtime/Engine/Private/GameInstance.cpp b/Engine/Source/Runtime/Engine/Private/GameInstance.cpp
--- a/Engine/Source/Runtime/Engine/Private/GameInstance.cpp
+++ b/Engine/Source/Runtime/Engine/Private/GameInstance.cpp
@@ -14,6 +14,7 @@
 #include "Engine/NetworkObjectList.h"
 #include "GameFramework/OnlineSession.h"
 #include "GameFramework/PlayerState.h"
+#include "GameFramework/GameModeBase.h"
 
 #if WITH_EDITOR
 #include "UnrealEd.h"
