Copying //UE4/UT-Staging to //UE4/Main (Source: //UE4/UT-Staging @ 3140618)

==========================
MAJOR FEATURES + CHANGES
==========================

CL 3094723 - Matt.Oelfke

Make MatineeActor extensible

CL 3106243 - Steve.Polge

FBatchedElements::AddTranslucentLine

CL 3111035 and 3114239 - Steve.Polge

NetStats can be available Shipping

CL 3113964 - Matt.Oelfke

BeginSpectatingState() won't unpossess a pawn that when it's not authority

Timeline component replication adjustment

CL 3126827 - John.Pollard via Peter.Knepley

Flush async loading when loading replay checkpoint

==========================
MINOR FEATURES + CHANGES
==========================

CL 3140532 - Peter.Knepley

Undo CVar changes that were mistakenly made in 3113964

diff --git a/Engine/Source/Runtime/Engine/Classes/Engine/NetDriver.h b/Engine/Source/Runtime/Engine/Classes/Engine/NetDriver.h
--- a/Engine/Source/Runtime/Engine/Classes/Engine/NetDriver.h
+++ b/Engine/Source/Runtime/Engine/Classes/Engine/NetDriver.h
@@ -285,6 +285,8 @@ public:
 	double						StatUpdateTime;
 	/** Interval between gathering stats */
 	float						StatPeriod;
+	/** Collect net stats even if not FThreadStats::IsCollectingData(). */
+	bool bCollectNetStats;
 	/** Time of last netdriver cleanup pass */
 	double						LastCleanupTime;
 	/** Used to determine if checking for standby cheats should occur */
diff --git a/Engine/Source/Runtime/Engine/Classes/Matinee/MatineeActor.h b/Engine/Source/Runtime/Engine/Classes/Matinee/MatineeActor.h
--- a/Engine/Source/Runtime/Engine/Classes/Matinee/MatineeActor.h
+++ b/Engine/Source/Runtime/Engine/Classes/Matinee/MatineeActor.h
@@ -49,8 +49,8 @@ struct FInterpGroupActorInfo
 
 };
 
-UCLASS(MinimalAPI, NotBlueprintable, hidecategories=(Collision, Game, Input), showcategories=("Input|MouseInput", "Input|TouchInput", "Game|Damage"))
-class AMatineeActor : public AActor
+UCLASS(NotBlueprintable, hidecategories=(Collision, Game, Input), showcategories=("Input|MouseInput", "Input|TouchInput", "Game|Damage"))
+class ENGINE_API AMatineeActor : public AActor
 {
 	GENERATED_UCLASS_BODY()
 
@@ -212,13 +212,13 @@ private:
 public:
 
 	/** Increment the count to ignore internal matinee actor selection */
-	ENGINE_API static void PushIgnoreActorSelection();
+	static void PushIgnoreActorSelection();
 
 	/** Decrement the count to ignore internal matinee actor selection */
-	ENGINE_API static void PopIgnoreActorSelection();
+	static void PopIgnoreActorSelection();
 
 	/** @return Should we ignore internal matinee actor selection? */
-	ENGINE_API static bool IgnoreActorSelection();
+	static bool IgnoreActorSelection();
 
 	/**
 	 * Check if we should perform a network positional update of this matinee
@@ -253,7 +253,7 @@ public:
 	 * @param bJump if true, teleport to the new position (don't trigger any events between the old and new positions, etc)
 	 */
 	UFUNCTION(BlueprintCallable, Category="Cinematic")
-	ENGINE_API void SetPosition(float NewPosition, bool bJump = false);
+	void SetPosition(float NewPosition, bool bJump = false);
 
 	/** Changes the direction of playback (go in reverse if it was going forward, or vice versa) */
 	UFUNCTION(BlueprintCallable, Category="Cinematic")
@@ -265,7 +265,7 @@ public:
 
 #if WITH_EDITOR
 	/** Fix up our references to any objects that have been replaced (e.g. through blueprint compiling) */
-	ENGINE_API void OnObjectsReplaced(const TMap<UObject*,UObject*>& ReplacementMap);
+	void OnObjectsReplaced(const TMap<UObject*,UObject*>& ReplacementMap);
 #endif //WITH_EDITOR
 
 	virtual void EndPlay(const EEndPlayReason::Type EndPlayReason) override;
@@ -350,16 +350,16 @@ public:
 	//~ End UObject Interface
 
 	/** Increment track forwards by given timestep and iterate over each track updating any properties. */
-	ENGINE_API void StepInterp(float DeltaTime, bool bPreview=false);
+	void StepInterp(float DeltaTime, bool bPreview=false);
 
 	/** Move interpolation to new position and iterate over each track updating any properties. */
-	ENGINE_API void UpdateInterp(float NewPosition, bool bPreview=false, bool bJump=false);
+	void UpdateInterp(float NewPosition, bool bPreview=false, bool bJump=false);
 
 	/** For each InterGroup/Actor combination, create a UInterpGroupInst, assign  AActor  and initialise each track. */
-	ENGINE_API void InitInterp();
+	void InitInterp();
 
 	/** Destroy all UInterpGroupInst. */
-	ENGINE_API void TermInterp();
+	void TermInterp();
 
 	/** Scan the matinee data for camera cuts and sets up the CameraCut array. */
 	void SetupCameraCuts();
@@ -383,7 +383,7 @@ public:
 	void GetAffectedActors( TArray<AActor*>& OutActors, bool bMovementTrackOnly );
 
 	/** Get all actors controlled by this matinee. Iterates over GroupActorInfos and finds all referenced actors, regardless of matinee state. */
-	ENGINE_API void GetControlledActors(TArray<AActor*>& OutActors);
+	void GetControlledActors(TArray<AActor*>& OutActors);
 
 	/**
 	 *	Update the streaming system with the camera locations for the upcoming camera cuts, so
@@ -408,75 +408,75 @@ public:
 	bool IsMatineeCompatibleWithPlayer( APlayerController* InPC ) const;
 
 	/** @return a group instance referring to the supplied Actor. Returns NULL if no group can be found. */
-	ENGINE_API UInterpGroupInst* FindGroupInst(const AActor* InActor) const;
+	UInterpGroupInst* FindGroupInst(const AActor* InActor) const;
 
 
 
 #if WITH_EDITOR
 
 	/** Ensure that the group to actor mapping is consistent with the interp data groups */
-	ENGINE_API void EnsureActorGroupConsistency();
+	void EnsureActorGroupConsistency();
 
 	/** Validate All Actors in the group **/
 	void ValidateActorGroups();
 
 	/** Replace Group Actors */
-	ENGINE_API void ReplaceActorGroupInfo(class UInterpGroup * Group, AActor* OldActor, AActor* NewActor);
+	void ReplaceActorGroupInfo(class UInterpGroup * Group, AActor* OldActor, AActor* NewActor);
 
 	/** Delete Group Actors */
-	ENGINE_API void DeleteActorGroupInfo(class UInterpGroup * Group, AActor* ActorToDelete);
+	void DeleteActorGroupInfo(class UInterpGroup * Group, AActor* ActorToDelete);
 
 	/** Rename groupinfo **/
-	ENGINE_API void DeleteGroupinfo(class UInterpGroup * GroupToDelete);
+	void DeleteGroupinfo(class UInterpGroup * GroupToDelete);
 
 	/** Save whether or not the passed in actor is hidden so we can restore it later */
 	void SaveActorVisibility( AActor*  AActor  );
 
 	/** return if this actor is valid to add **/
-	ENGINE_API EActorAddWarningType IsValidActorToAdd(const AActor* NewActor) const;
+	EActorAddWarningType IsValidActorToAdd(const AActor* NewActor) const;
 
 	/**
 	 * Conditionally save state for the specified actor and its children
 	 */
-	ENGINE_API void ConditionallySaveActorState( UInterpGroupInst* GroupInst, AActor*  AActor  );
+	void ConditionallySaveActorState( UInterpGroupInst* GroupInst, AActor*  AActor  );
 
 	/**
 	 * Add the specified actor and any actors attached to it to the list
 	 * of saved actor transforms.  Does nothing if an actor has already
 	 * been saved.
 	 */
-	ENGINE_API void SaveActorTransforms( AActor*  AActor  );
+	void SaveActorTransforms( AActor*  AActor  );
 
 	/**
 	 * Apply the saved locations and rotations to all saved actors.
 	 */
-	ENGINE_API void RestoreActorTransforms();
+	void RestoreActorTransforms();
 	
 	/** Apply the saved visibility state for all saved actors */
-	ENGINE_API void RestoreActorVisibilities();
+	void RestoreActorVisibilities();
 
 	/**
 	 * Store the current scrub position, restores all saved actor transforms,
 	 * then saves off the transforms for actors referenced (directly or indirectly)
 	 * by group instances, and finally restores the scrub position.
 	 */
-	ENGINE_API void RecaptureActorState();
+	void RecaptureActorState();
 
 	/** Set up the group actor for the specified UInterpGroup. */
-	ENGINE_API void InitGroupActorForGroup(class UInterpGroup* InGroup, class AActor* GroupActor);
+	void InitGroupActorForGroup(class UInterpGroup* InGroup, class AActor* GroupActor);
 #endif
 
 	/** Find the first group instance based on the given UInterpGroup. */
-	ENGINE_API UInterpGroupInst* FindFirstGroupInst(class UInterpGroup* InGroup);
+	UInterpGroupInst* FindFirstGroupInst(class UInterpGroup* InGroup);
 
 	/** Find the first group instance based on the UInterpGroup with the given name. */
-	ENGINE_API UInterpGroupInst* FindFirstGroupInstByName(const FString& InGroupName);
+	UInterpGroupInst* FindFirstGroupInstByName(const FString& InGroupName);
 
 	/** 
 	 * If there is DirectorGroup, use it to find the viewed group at the current position, 
 	 * then the first instance of that group, and the  AActor  it is bound to. 
 	 */
-	ENGINE_API AActor* FindViewedActor();
+	AActor* FindViewedActor();
 
 	
 	void AddPlayerToDirectorTracks( class APlayerController* PC );
@@ -489,12 +489,12 @@ public:
 	virtual void NotifyEventTriggered(FName EventName, float EventTime, bool bUseCustomEventName=false);
 
 	/** Util to get the name of the function to find for the given event name */
-	ENGINE_API FName GetFunctionNameForEvent(FName EventName,bool bUseCustomEventName=false);
+	FName GetFunctionNameForEvent(FName EventName,bool bUseCustomEventName=false);
 
 public:
 #if WITH_EDITORONLY_DATA
 	/** Returns SpriteComponent subobject **/
-	ENGINE_API class UBillboardComponent* GetSpriteComponent() const;
+	class UBillboardComponent* GetSpriteComponent() const;
 #endif
 };
 
diff --git a/Engine/Source/Runtime/Engine/Private/BatchedElements.cpp b/Engine/Source/Runtime/Engine/Private/BatchedElements.cpp
--- a/Engine/Source/Runtime/Engine/Private/BatchedElements.cpp
+++ b/Engine/Source/Runtime/Engine/Private/BatchedElements.cpp
@@ -63,6 +63,39 @@ void FBatchedElements::AddLine(const FVector& Start, const FVector& End, const F
 	}
 }
 
+void FBatchedElements::AddTranslucentLine(const FVector& Start, const FVector& End, const FLinearColor& Color, FHitProxyId HitProxyId, float Thickness, float DepthBias, bool bScreenSpace)
+{
+	if (Thickness == 0.0f)
+	{
+		if (DepthBias == 0.0f)
+		{
+			new(LineVertices) FSimpleElementVertex(Start, FVector2D::ZeroVector, Color, HitProxyId);
+			new(LineVertices) FSimpleElementVertex(End, FVector2D::ZeroVector, Color, HitProxyId);
+		}
+		else
+		{
+			// Draw degenerate triangles in wireframe mode to support depth bias (d3d11 and opengl3 don't support depth bias on line primitives, but do on wireframes)
+			FBatchedWireTris* WireTri = new(WireTris) FBatchedWireTris();
+			WireTri->DepthBias = DepthBias;
+			new(WireTriVerts) FSimpleElementVertex(Start, FVector2D::ZeroVector, Color, HitProxyId);
+			new(WireTriVerts) FSimpleElementVertex(End, FVector2D::ZeroVector, Color, HitProxyId);
+			new(WireTriVerts) FSimpleElementVertex(End, FVector2D::ZeroVector, Color, HitProxyId);
+		}
+	}
+	else
+	{
+		FBatchedThickLines* ThickLine = new(ThickLines) FBatchedThickLines;
+		ThickLine->Start = Start;
+		ThickLine->End = End;
+		ThickLine->Thickness = Thickness;
+		ThickLine->Color = Color;
+		ThickLine->HitProxyId = HitProxyId;
+		ThickLine->DepthBias = DepthBias;
+		ThickLine->bScreenSpace = bScreenSpace;
+
+	}
+}
+
 void FBatchedElements::AddPoint(const FVector& Position,float Size,const FLinearColor& Color,FHitProxyId HitProxyId)
 {
 	// Ensure the point isn't masked out.  Some legacy code relies on Color.A being ignored.
@@ -811,6 +844,7 @@ bool FBatchedElements::Draw(FRHICommandList& RHICmdList, ERHIFeatureLevel::Type
 
 			if ( ThickLines.Num() > 0 )
 			{
+				PrepareShaders(RHICmdList, FeatureLevel, SE_BLEND_Translucent, Transform, bNeedToSwitchVerticalAxis, BatchedElementParameters, GWhiteTexture, bHitTesting, Gamma, NULL, View, DepthTexture);
 				float OrthoZoomFactor = 1.0f;
 
 				if( View )
diff --git a/Engine/Source/Runtime/Engine/Private/DemoNetDriver.cpp b/Engine/Source/Runtime/Engine/Private/DemoNetDriver.cpp
--- a/Engine/Source/Runtime/Engine/Private/DemoNetDriver.cpp
+++ b/Engine/Source/Runtime/Engine/Private/DemoNetDriver.cpp
@@ -2394,6 +2394,7 @@ void UDemoNetDriver::LoadCheckpoint( FArchive* GotoCheckpointArchive, int64 Goto
 	}
 
 	// Clean package map to prepare to restore it to the checkpoint state
+	FlushAsyncLoading();
 	GuidCache->ObjectLookup.Empty();
 	GuidCache->NetGUIDLookup.Empty();
 
diff --git a/Engine/Source/Runtime/Engine/Private/NetworkDriver.cpp b/Engine/Source/Runtime/Engine/Private/NetworkDriver.cpp
--- a/Engine/Source/Runtime/Engine/Private/NetworkDriver.cpp
+++ b/Engine/Source/Runtime/Engine/Private/NetworkDriver.cpp
@@ -166,6 +166,7 @@ UNetDriver::UNetDriver(const FObjectInitializer& ObjectInitializer)
 ,	OutOutOfOrderPackets(0)
 ,	StatUpdateTime(0.0)
 ,	StatPeriod(1.f)
+,	bCollectNetStats(false)
 ,	LastCleanupTime(0.0)
 ,	NetTag(0)
 ,	DebugRelevantActors(false)
@@ -324,32 +325,38 @@ void UNetDriver::TickFlush(float DeltaSeconds)
 
 	const double CurrentRealtimeSeconds = FPlatformTime::Seconds();
 
+	bool bCollectServerStats = false;
 #if USE_SERVER_PERF_COUNTERS || STATS
-	// Update network stats (only main game net driver for now) if stats or perf counters are used
-	if (NetDriverName == NAME_GameNetDriver && 
-		 CurrentRealtimeSeconds - StatUpdateTime > StatPeriod)
-	{
-		int32 ClientInBytesMax = 0;
-		int32 ClientInBytesMin = 0;
-		int32 ClientInBytesAvg = 0;
-		int32 ClientInPacketsMax = 0;
-		int32 ClientInPacketsMin = 0;
-		int32 ClientInPacketsAvg = 0;
-		int32 ClientOutBytesMax = 0;
-		int32 ClientOutBytesMin = 0;
-		int32 ClientOutBytesAvg = 0;
-		int32 ClientOutPacketsMax = 0;
-		int32 ClientOutPacketsMin = 0;
-		int32 ClientOutPacketsAvg = 0;
-		int NumClients = 0;
-		int32 MaxPacketOverhead = 0;
-		float RemoteSaturationMax = 0.0f;
-
-		// these need to be updated even if we are not collecting stats, since they get reported to analytics/QoS
-		for (UNetConnection * Client : ClientConnections)
-		{
-			if (Client)
-			{
+	bCollectServerStats = true;
+#endif
+
+	if (bCollectNetStats || bCollectServerStats)
+	{
+		// Update network stats (only main game net driver for now) if stats or perf counters are used
+		if (NetDriverName == NAME_GameNetDriver &&
+			CurrentRealtimeSeconds - StatUpdateTime > StatPeriod)
+		{
+			int32 ClientInBytesMax = 0;
+			int32 ClientInBytesMin = 0;
+			int32 ClientInBytesAvg = 0;
+			int32 ClientInPacketsMax = 0;
+			int32 ClientInPacketsMin = 0;
+			int32 ClientInPacketsAvg = 0;
+			int32 ClientOutBytesMax = 0;
+			int32 ClientOutBytesMin = 0;
+			int32 ClientOutBytesAvg = 0;
+			int32 ClientOutPacketsMax = 0;
+			int32 ClientOutPacketsMin = 0;
+			int32 ClientOutPacketsAvg = 0;
+			int NumClients = 0;
+			int32 MaxPacketOverhead = 0;
+			float RemoteSaturationMax = 0.0f;
+
+			// these need to be updated even if we are not collecting stats, since they get reported to analytics/QoS
+			for (UNetConnection * Client : ClientConnections)
+			{
+				if (Client)
+				{
 #define UpdatePerClientMinMaxAvg(VariableName) \
 				Client##VariableName##Max = FMath::Max(Client##VariableName##Max, Client->VariableName##PerSecond); \
 				if (Client##VariableName##Min == 0 || Client->VariableName##PerSecond < Client##VariableName##Min) \
@@ -358,283 +365,287 @@ void UNetDriver::TickFlush(float DeltaSeconds)
 				} \
 				Client##VariableName##Avg += Client->VariableName##PerSecond; \
 
-				UpdatePerClientMinMaxAvg(InBytes);
-				UpdatePerClientMinMaxAvg(OutBytes);
-				UpdatePerClientMinMaxAvg(InPackets);
-				UpdatePerClientMinMaxAvg(OutPackets);
+					UpdatePerClientMinMaxAvg(InBytes);
+					UpdatePerClientMinMaxAvg(OutBytes);
+					UpdatePerClientMinMaxAvg(InPackets);
+					UpdatePerClientMinMaxAvg(OutPackets);
 
-				MaxPacketOverhead = FMath::Max(Client->PacketOverhead, MaxPacketOverhead);
+					MaxPacketOverhead = FMath::Max(Client->PacketOverhead, MaxPacketOverhead);
 
 #undef UpdatePerClientMinMaxAvg
 
-				++NumClients;
+					++NumClients;
+				}
 			}
-		}
-
-		if (NumClients > 1)
-		{
-			ClientInBytesAvg /= NumClients;
-			ClientInPacketsAvg /= NumClients;
-			ClientOutBytesAvg /= NumClients;
-			ClientOutPacketsAvg /= NumClients;
-		}
 
-#if STATS
-		int32 Ping = 0;
-		int32 NumOpenChannels = 0;
-		int32 NumActorChannels = 0;
-		int32 NumDormantActors = 0;
-		int32 NumActors = 0;
-		int32 AckCount = 0;
-		int32 UnAckCount = 0;
-		int32 PendingCount = 0;
-		int32 NetSaturated = 0;
-
-		if (FThreadStats::IsCollectingData())
-		{
-			const float RealTime = CurrentRealtimeSeconds - StatUpdateTime;
-
-			// Use the elapsed time to keep things scaled to one measured unit
-			InBytes = FMath::TruncToInt(InBytes / RealTime);
-			OutBytes = FMath::TruncToInt(OutBytes / RealTime);
-
-			NetGUIDOutBytes = FMath::TruncToInt(NetGUIDOutBytes / RealTime);
-			NetGUIDInBytes = FMath::TruncToInt(NetGUIDInBytes / RealTime);
-
-			// Save off for stats later
-
-			InBytesPerSecond = InBytes;
-			OutBytesPerSecond = OutBytes;
-
-			InPackets = FMath::TruncToInt(InPackets / RealTime);
-			OutPackets = FMath::TruncToInt(OutPackets / RealTime);
-			InBunches = FMath::TruncToInt(InBunches / RealTime);
-			OutBunches = FMath::TruncToInt(OutBunches / RealTime);
-			OutPacketsLost = FMath::TruncToInt(100.f * OutPacketsLost / FMath::Max((float)OutPackets,1.f));
-			InPacketsLost = FMath::TruncToInt(100.f * InPacketsLost / FMath::Max((float)InPackets + InPacketsLost,1.f));
-			
-			if (ServerConnection != NULL && ServerConnection->PlayerController != NULL && ServerConnection->PlayerController->PlayerState != NULL)
+			if (NumClients > 1)
 			{
-				Ping = FMath::TruncToInt(ServerConnection->PlayerController->PlayerState->ExactPing);
+				ClientInBytesAvg /= NumClients;
+				ClientInPacketsAvg /= NumClients;
+				ClientOutBytesAvg /= NumClients;
+				ClientOutPacketsAvg /= NumClients;
 			}
 
-			if (ServerConnection != NULL)
-			{
-				NumOpenChannels = ServerConnection->OpenChannels.Num();
-				RemoteSaturationMax = FMath::Max( RemoteSaturationMax, ServerConnection->RemoteSaturation );
-			}
+			int32 Ping = 0;
+			int32 NumOpenChannels = 0;
+			int32 NumActorChannels = 0;
+			int32 NumDormantActors = 0;
+			int32 NumActors = 0;
+			int32 AckCount = 0;
+			int32 UnAckCount = 0;
+			int32 PendingCount = 0;
+			int32 NetSaturated = 0;
 
-			for (int32 i = 0; i < ClientConnections.Num(); i++)
+			if (
+#if STATS
+				FThreadStats::IsCollectingData() ||
+#endif
+				bCollectNetStats)
 			{
-				NumOpenChannels += ClientConnections[i]->OpenChannels.Num();
-				RemoteSaturationMax = FMath::Max( RemoteSaturationMax, ClientConnections[i]->RemoteSaturation );
-			}
+				const float RealTime = CurrentRealtimeSeconds - StatUpdateTime;
 
-			// Use the elapsed time to keep things scaled to one measured unit
-			VoicePacketsSent = FMath::TruncToInt(VoicePacketsSent / RealTime);
-			VoicePacketsRecv = FMath::TruncToInt(VoicePacketsRecv / RealTime);
-			VoiceBytesSent = FMath::TruncToInt(VoiceBytesSent / RealTime);
-			VoiceBytesRecv = FMath::TruncToInt(VoiceBytesRecv / RealTime);
+				// Use the elapsed time to keep things scaled to one measured unit
+				InBytes = FMath::TruncToInt(InBytes / RealTime);
+				OutBytes = FMath::TruncToInt(OutBytes / RealTime);
 
-			// Determine voice percentages
-			VoiceInPercent = (InBytes > 0) ? FMath::TruncToInt(100.f * (float)VoiceBytesRecv / (float)InBytes) : 0;
-			VoiceOutPercent = (OutBytes > 0) ? FMath::TruncToInt(100.f * (float)VoiceBytesSent / (float)OutBytes) : 0;
+				NetGUIDOutBytes = FMath::TruncToInt(NetGUIDOutBytes / RealTime);
+				NetGUIDInBytes = FMath::TruncToInt(NetGUIDInBytes / RealTime);
 
-			UNetConnection * Connection = (ServerConnection ? ServerConnection : (ClientConnections.Num() > 0 ? ClientConnections[0] : NULL));
-			if (Connection)
-			{
-				NumActorChannels = Connection->ActorChannels.Num();
-				NumDormantActors = Connection->DormantActors.Num();
+				// Save off for stats later
 
-				if (World)
+				InBytesPerSecond = InBytes;
+				OutBytesPerSecond = OutBytes;
+
+				InPackets = FMath::TruncToInt(InPackets / RealTime);
+				OutPackets = FMath::TruncToInt(OutPackets / RealTime);
+				InBunches = FMath::TruncToInt(InBunches / RealTime);
+				OutBunches = FMath::TruncToInt(OutBunches / RealTime);
+				OutPacketsLost = FMath::TruncToInt(100.f * OutPacketsLost / FMath::Max((float)OutPackets, 1.f));
+				InPacketsLost = FMath::TruncToInt(100.f * InPacketsLost / FMath::Max((float)InPackets + InPacketsLost, 1.f));
+
+				if (ServerConnection != NULL && ServerConnection->PlayerController != NULL && ServerConnection->PlayerController->PlayerState != NULL)
+				{
+					Ping = FMath::TruncToInt(ServerConnection->PlayerController->PlayerState->ExactPing);
+				}
+
+				if (ServerConnection != NULL)
+				{
+					NumOpenChannels = ServerConnection->OpenChannels.Num();
+					RemoteSaturationMax = FMath::Max(RemoteSaturationMax, ServerConnection->RemoteSaturation);
+				}
+
+				for (int32 i = 0; i < ClientConnections.Num(); i++)
 				{
-					NumActors = World->GetActorCount();
+					NumOpenChannels += ClientConnections[i]->OpenChannels.Num();
+					RemoteSaturationMax = FMath::Max(RemoteSaturationMax, ClientConnections[i]->RemoteSaturation);
+				}
+
+				// Use the elapsed time to keep things scaled to one measured unit
+				VoicePacketsSent = FMath::TruncToInt(VoicePacketsSent / RealTime);
+				VoicePacketsRecv = FMath::TruncToInt(VoicePacketsRecv / RealTime);
+				VoiceBytesSent = FMath::TruncToInt(VoiceBytesSent / RealTime);
+				VoiceBytesRecv = FMath::TruncToInt(VoiceBytesRecv / RealTime);
+
+				// Determine voice percentages
+				VoiceInPercent = (InBytes > 0) ? FMath::TruncToInt(100.f * (float)VoiceBytesRecv / (float)InBytes) : 0;
+				VoiceOutPercent = (OutBytes > 0) ? FMath::TruncToInt(100.f * (float)VoiceBytesSent / (float)OutBytes) : 0;
+
+				UNetConnection * Connection = (ServerConnection ? ServerConnection : (ClientConnections.Num() > 0 ? ClientConnections[0] : NULL));
+				if (Connection)
+				{
+					NumActorChannels = Connection->ActorChannels.Num();
+					NumDormantActors = Connection->DormantActors.Num();
+
+					if (World)
+					{
+						NumActors = World->GetActorCount();
+					}
+#if STATS
+					Connection->PackageMap->GetNetGUIDStats(AckCount, UnAckCount, PendingCount);
+#endif
+					NetSaturated = Connection->IsNetReady(false) ? 0 : 1;
 				}
+			}
 
-				Connection->PackageMap->GetNetGUIDStats(AckCount, UnAckCount, PendingCount);
-
-				NetSaturated = Connection->IsNetReady(false) ? 0 : 1;
-			}
-		}
-
-		// Copy the net status values over
-		SET_DWORD_STAT(STAT_Ping, Ping);
-		SET_DWORD_STAT(STAT_Channels, NumOpenChannels);
-		SET_DWORD_STAT(STAT_MaxPacketOverhead, MaxPacketOverhead);
-
-		SET_DWORD_STAT(STAT_OutLoss, OutPacketsLost);
-		SET_DWORD_STAT(STAT_InLoss, InPacketsLost);
-		SET_DWORD_STAT(STAT_InRate, InBytes);
-		SET_DWORD_STAT(STAT_OutRate, OutBytes);
-		SET_DWORD_STAT(STAT_OutSaturation, RemoteSaturationMax);
-		SET_DWORD_STAT(STAT_InRateClientMax, ClientInBytesMax);
-		SET_DWORD_STAT(STAT_InRateClientMin, ClientInBytesMin);
-		SET_DWORD_STAT(STAT_InRateClientAvg, ClientInBytesAvg);
-		SET_DWORD_STAT(STAT_InPacketsClientMax, ClientInPacketsMax);
-		SET_DWORD_STAT(STAT_InPacketsClientMin, ClientInPacketsMin);
-		SET_DWORD_STAT(STAT_InPacketsClientAvg, ClientInPacketsAvg);
-		SET_DWORD_STAT(STAT_OutRateClientMax, ClientOutBytesMax);
-		SET_DWORD_STAT(STAT_OutRateClientMin, ClientOutBytesMin);
-		SET_DWORD_STAT(STAT_OutRateClientAvg, ClientOutBytesAvg);
-		SET_DWORD_STAT(STAT_OutPacketsClientMax, ClientOutPacketsMax);
-		SET_DWORD_STAT(STAT_OutPacketsClientMin, ClientOutPacketsMin);
-		SET_DWORD_STAT(STAT_OutPacketsClientAvg, ClientOutPacketsAvg);
-
-		SET_DWORD_STAT(STAT_NetNumClients, NumClients);
-		SET_DWORD_STAT(STAT_InPackets, InPackets);
-		SET_DWORD_STAT(STAT_OutPackets, OutPackets);
-		SET_DWORD_STAT(STAT_InBunches, InBunches);
-		SET_DWORD_STAT(STAT_OutBunches, OutBunches);
-
-		SET_DWORD_STAT(STAT_NetGUIDInRate, NetGUIDInBytes);
-		SET_DWORD_STAT(STAT_NetGUIDOutRate, NetGUIDOutBytes);
-
-		SET_DWORD_STAT(STAT_VoicePacketsSent, VoicePacketsSent);
-		SET_DWORD_STAT(STAT_VoicePacketsRecv, VoicePacketsRecv);
-		SET_DWORD_STAT(STAT_VoiceBytesSent, VoiceBytesSent);
-		SET_DWORD_STAT(STAT_VoiceBytesRecv, VoiceBytesRecv);
-
-		SET_DWORD_STAT(STAT_PercentInVoice, VoiceInPercent);
-		SET_DWORD_STAT(STAT_PercentOutVoice, VoiceOutPercent);
-
-		SET_DWORD_STAT(STAT_NumActorChannels, NumActorChannels);
-		SET_DWORD_STAT(STAT_NumDormantActors, NumDormantActors);
-		SET_DWORD_STAT(STAT_NumActors, NumActors);
-		SET_DWORD_STAT(STAT_NumNetActors, GetNetworkObjectList().GetObjects().Num());
-		SET_DWORD_STAT(STAT_NumNetGUIDsAckd, AckCount);
-		SET_DWORD_STAT(STAT_NumNetGUIDsPending, UnAckCount);
-		SET_DWORD_STAT(STAT_NumNetGUIDsUnAckd, PendingCount);
-		SET_DWORD_STAT(STAT_NetSaturated, NetSaturated);
+#if STATS
+			// Copy the net status values over
+			SET_DWORD_STAT(STAT_Ping, Ping);
+			SET_DWORD_STAT(STAT_Channels, NumOpenChannels);
+			SET_DWORD_STAT(STAT_MaxPacketOverhead, MaxPacketOverhead);
+
+			SET_DWORD_STAT(STAT_OutLoss, OutPacketsLost);
+			SET_DWORD_STAT(STAT_InLoss, InPacketsLost);
+			SET_DWORD_STAT(STAT_InRate, InBytes);
+			SET_DWORD_STAT(STAT_OutRate, OutBytes);
+			SET_DWORD_STAT(STAT_OutSaturation, RemoteSaturationMax);
+			SET_DWORD_STAT(STAT_InRateClientMax, ClientInBytesMax);
+			SET_DWORD_STAT(STAT_InRateClientMin, ClientInBytesMin);
+			SET_DWORD_STAT(STAT_InRateClientAvg, ClientInBytesAvg);
+			SET_DWORD_STAT(STAT_InPacketsClientMax, ClientInPacketsMax);
+			SET_DWORD_STAT(STAT_InPacketsClientMin, ClientInPacketsMin);
+			SET_DWORD_STAT(STAT_InPacketsClientAvg, ClientInPacketsAvg);
+			SET_DWORD_STAT(STAT_OutRateClientMax, ClientOutBytesMax);
+			SET_DWORD_STAT(STAT_OutRateClientMin, ClientOutBytesMin);
+			SET_DWORD_STAT(STAT_OutRateClientAvg, ClientOutBytesAvg);
+			SET_DWORD_STAT(STAT_OutPacketsClientMax, ClientOutPacketsMax);
+			SET_DWORD_STAT(STAT_OutPacketsClientMin, ClientOutPacketsMin);
+			SET_DWORD_STAT(STAT_OutPacketsClientAvg, ClientOutPacketsAvg);
+
+			SET_DWORD_STAT(STAT_NetNumClients, NumClients);
+			SET_DWORD_STAT(STAT_InPackets, InPackets);
+			SET_DWORD_STAT(STAT_OutPackets, OutPackets);
+			SET_DWORD_STAT(STAT_InBunches, InBunches);
+			SET_DWORD_STAT(STAT_OutBunches, OutBunches);
+
+			SET_DWORD_STAT(STAT_NetGUIDInRate, NetGUIDInBytes);
+			SET_DWORD_STAT(STAT_NetGUIDOutRate, NetGUIDOutBytes);
+
+			SET_DWORD_STAT(STAT_VoicePacketsSent, VoicePacketsSent);
+			SET_DWORD_STAT(STAT_VoicePacketsRecv, VoicePacketsRecv);
+			SET_DWORD_STAT(STAT_VoiceBytesSent, VoiceBytesSent);
+			SET_DWORD_STAT(STAT_VoiceBytesRecv, VoiceBytesRecv);
+
+			SET_DWORD_STAT(STAT_PercentInVoice, VoiceInPercent);
+			SET_DWORD_STAT(STAT_PercentOutVoice, VoiceOutPercent);
+
+			SET_DWORD_STAT(STAT_NumActorChannels, NumActorChannels);
+			SET_DWORD_STAT(STAT_NumDormantActors, NumDormantActors);
+			SET_DWORD_STAT(STAT_NumActors, NumActors);
+			SET_DWORD_STAT(STAT_NumNetActors, GetNetworkObjectList().GetObjects().Num());
+			SET_DWORD_STAT(STAT_NumNetGUIDsAckd, AckCount);
+			SET_DWORD_STAT(STAT_NumNetGUIDsPending, UnAckCount);
+			SET_DWORD_STAT(STAT_NumNetGUIDsUnAckd, PendingCount);
+			SET_DWORD_STAT(STAT_NetSaturated, NetSaturated);
 #endif // STATS
 
 #if USE_SERVER_PERF_COUNTERS
-		IPerfCounters* PerfCounters = IPerfCountersModule::Get().GetPerformanceCounters();
-		if (PerfCounters)
-		{
-			// Update total connections
-			PerfCounters->Set(TEXT("NumConnections"), ClientConnections.Num());
-
-			const int kNumBuckets = 8;	// evenly spaced with increment of 30 ms; last bucket collects all off-scale pings as well
-			if (ClientConnections.Num() > 0)
+			IPerfCounters* PerfCounters = IPerfCountersModule::Get().GetPerformanceCounters();
+			if (PerfCounters)
 			{
-				// Update per connection statistics
-				float MinPing = MAX_FLT;
-				float AvgPing = 0;
-				float MaxPing = -MAX_FLT;
-				float PingCount = 0;
-
-				int32 Buckets[kNumBuckets] = { 0 };
+				// Update total connections
+				PerfCounters->Set(TEXT("NumConnections"), ClientConnections.Num());
 
-				for (int32 i = 0; i < ClientConnections.Num(); i++)
+				const int kNumBuckets = 8;	// evenly spaced with increment of 30 ms; last bucket collects all off-scale pings as well
+				if (ClientConnections.Num() > 0)
 				{
-					UNetConnection* Connection = ClientConnections[i];
+					// Update per connection statistics
+					float MinPing = MAX_FLT;
+					float AvgPing = 0;
+					float MaxPing = -MAX_FLT;
+					float PingCount = 0;
+
+					int32 Buckets[kNumBuckets] = { 0 };
 
-					if (Connection != nullptr)
+					for (int32 i = 0; i < ClientConnections.Num(); i++)
 					{
-						if (Connection->PlayerController != nullptr && Connection->PlayerController->PlayerState != nullptr)
+						UNetConnection* Connection = ClientConnections[i];
+
+						if (Connection != nullptr)
 						{
-							// Ping value calculated per client
-							float ConnPing = Connection->PlayerController->PlayerState->ExactPing;
+							if (Connection->PlayerController != nullptr && Connection->PlayerController->PlayerState != nullptr)
+							{
+								// Ping value calculated per client
+								float ConnPing = Connection->PlayerController->PlayerState->ExactPing;
 
-							int Bucket = FMath::Max(0, FMath::Min(kNumBuckets - 1, (static_cast<int>(ConnPing) / 30)));
-							++Buckets[Bucket];
+								int Bucket = FMath::Max(0, FMath::Min(kNumBuckets - 1, (static_cast<int>(ConnPing) / 30)));
+								++Buckets[Bucket];
 
-							if (ConnPing < MinPing)
-							{
-								MinPing = ConnPing;
-							}
+								if (ConnPing < MinPing)
+								{
+									MinPing = ConnPing;
+								}
 
-							if (ConnPing > MaxPing)
-							{
-								MaxPing = ConnPing;
-							}
+								if (ConnPing > MaxPing)
+								{
+									MaxPing = ConnPing;
+								}
 
-							AvgPing += ConnPing;
-							PingCount++;
+								AvgPing += ConnPing;
+								PingCount++;
+							}
 						}
 					}
-				}
 
-				if (PingCount > 0)
-				{
-					AvgPing /= static_cast<float>(PingCount);
-				}
+					if (PingCount > 0)
+					{
+						AvgPing /= static_cast<float>(PingCount);
+					}
 
-				PerfCounters->Set(TEXT("AvgPing"), AvgPing, IPerfCounters::Flags::Transient);
-				float CurrentMaxPing = PerfCounters->Get(TEXT("MaxPing"), MaxPing);
-				PerfCounters->Set(TEXT("MaxPing"), FMath::Max(MaxPing, CurrentMaxPing), IPerfCounters::Flags::Transient);
-				float CurrentMinPing = PerfCounters->Get(TEXT("MinPing"), MinPing);
-				PerfCounters->Set(TEXT("MinPing"), FMath::Min(MinPing, CurrentMinPing), IPerfCounters::Flags::Transient);
+					PerfCounters->Set(TEXT("AvgPing"), AvgPing, IPerfCounters::Flags::Transient);
+					float CurrentMaxPing = PerfCounters->Get(TEXT("MaxPing"), MaxPing);
+					PerfCounters->Set(TEXT("MaxPing"), FMath::Max(MaxPing, CurrentMaxPing), IPerfCounters::Flags::Transient);
+					float CurrentMinPing = PerfCounters->Get(TEXT("MinPing"), MinPing);
+					PerfCounters->Set(TEXT("MinPing"), FMath::Min(MinPing, CurrentMinPing), IPerfCounters::Flags::Transient);
 
-				// update buckets
-				for (int BucketIdx = 0; BucketIdx < ARRAY_COUNT(Buckets); ++BucketIdx)
-				{
-					PerfCountersIncrement(FString::Printf(TEXT("PingBucketInt%d"), BucketIdx), Buckets[BucketIdx], 0, IPerfCounters::Flags::Transient);
+					// update buckets
+					for (int BucketIdx = 0; BucketIdx < ARRAY_COUNT(Buckets); ++BucketIdx)
+					{
+						PerfCountersIncrement(FString::Printf(TEXT("PingBucketInt%d"), BucketIdx), Buckets[BucketIdx], 0, IPerfCounters::Flags::Transient);
+					}
 				}
-			}
-			else
-			{
-				PerfCounters->Set(TEXT("AvgPing"), 0.0f, IPerfCounters::Flags::Transient);
-				PerfCounters->Set(TEXT("MaxPing"), -FLT_MAX, IPerfCounters::Flags::Transient);
-				PerfCounters->Set(TEXT("MinPing"), FLT_MAX, IPerfCounters::Flags::Transient);
-
-				for (int BucketIdx = 0; BucketIdx < kNumBuckets; ++BucketIdx)
+				else
 				{
-					PerfCounters->Set(FString::Printf(TEXT("PingBucketInt%d"), BucketIdx), 0, IPerfCounters::Flags::Transient);
+					PerfCounters->Set(TEXT("AvgPing"), 0.0f, IPerfCounters::Flags::Transient);
+					PerfCounters->Set(TEXT("MaxPing"), -FLT_MAX, IPerfCounters::Flags::Transient);
+					PerfCounters->Set(TEXT("MinPing"), FLT_MAX, IPerfCounters::Flags::Transient);
+
+					for (int BucketIdx = 0; BucketIdx < kNumBuckets; ++BucketIdx)
+					{
+						PerfCounters->Set(FString::Printf(TEXT("PingBucketInt%d"), BucketIdx), 0, IPerfCounters::Flags::Transient);
+					}
 				}
-			}
 
-			// set the per connection stats (these are calculated earlier).
-			// Note that NumClients may be != NumConnections. Also, if NumClients is 0, the rest of counters should be 0 as well
-			PerfCounters->Set(TEXT("NumClients"), NumClients);
-			PerfCounters->Set(TEXT("MaxPacketOverhead"), MaxPacketOverhead);
-			PerfCounters->Set(TEXT("InRateClientMax"), ClientInBytesMax);
-			PerfCounters->Set(TEXT("InRateClientMin"), ClientInBytesMin);
-			PerfCounters->Set(TEXT("InRateClientAvg"), ClientInBytesAvg);
-			PerfCounters->Set(TEXT("InPacketsClientMax"), ClientInPacketsMax);
-			PerfCounters->Set(TEXT("InPacketsClientMin"), ClientInPacketsMin);
-			PerfCounters->Set(TEXT("InPacketsClientAvg"), ClientInPacketsAvg);
-			PerfCounters->Set(TEXT("OutRateClientMax"), ClientOutBytesMax);
-			PerfCounters->Set(TEXT("OutRateClientMin"), ClientOutBytesMin);
-			PerfCounters->Set(TEXT("OutRateClientAvg"), ClientOutBytesAvg);
-			PerfCounters->Set(TEXT("OutPacketsClientMax"), ClientOutPacketsMax);
-			PerfCounters->Set(TEXT("OutPacketsClientMin"), ClientOutPacketsMin);
-			PerfCounters->Set(TEXT("OutPacketsClientAvg"), ClientOutPacketsAvg);
-
-			PerfCounters->Set(TEXT("InRate"), InBytes);
-			PerfCounters->Set(TEXT("OutRate"), OutBytes);
-			PerfCounters->Set(TEXT("InPacketsLost"), InPacketsLost);
-			PerfCounters->Set(TEXT("OutPacketsLost"), OutPacketsLost);
-			PerfCounters->Set(TEXT("InPackets"), InPackets);
-			PerfCounters->Set(TEXT("OutPackets"), OutPackets);
-			PerfCounters->Set(TEXT("InBunches"), InBunches);
-			PerfCounters->Set(TEXT("OutBunches"), OutBunches);
-
-			PerfCounters->Set(TEXT("ServerReplicateActorsTimeMs"), ServerReplicateActorsTimeMs);
-			PerfCounters->Set(TEXT("OutSaturationMax"), RemoteSaturationMax);
-		}
+				// set the per connection stats (these are calculated earlier).
+				// Note that NumClients may be != NumConnections. Also, if NumClients is 0, the rest of counters should be 0 as well
+				PerfCounters->Set(TEXT("NumClients"), NumClients);
+				PerfCounters->Set(TEXT("MaxPacketOverhead"), MaxPacketOverhead);
+				PerfCounters->Set(TEXT("InRateClientMax"), ClientInBytesMax);
+				PerfCounters->Set(TEXT("InRateClientMin"), ClientInBytesMin);
+				PerfCounters->Set(TEXT("InRateClientAvg"), ClientInBytesAvg);
+				PerfCounters->Set(TEXT("InPacketsClientMax"), ClientInPacketsMax);
+				PerfCounters->Set(TEXT("InPacketsClientMin"), ClientInPacketsMin);
+				PerfCounters->Set(TEXT("InPacketsClientAvg"), ClientInPacketsAvg);
+				PerfCounters->Set(TEXT("OutRateClientMax"), ClientOutBytesMax);
+				PerfCounters->Set(TEXT("OutRateClientMin"), ClientOutBytesMin);
+				PerfCounters->Set(TEXT("OutRateClientAvg"), ClientOutBytesAvg);
+				PerfCounters->Set(TEXT("OutPacketsClientMax"), ClientOutPacketsMax);
+				PerfCounters->Set(TEXT("OutPacketsClientMin"), ClientOutPacketsMin);
+				PerfCounters->Set(TEXT("OutPacketsClientAvg"), ClientOutPacketsAvg);
+
+				PerfCounters->Set(TEXT("InRate"), InBytes);
+				PerfCounters->Set(TEXT("OutRate"), OutBytes);
+				PerfCounters->Set(TEXT("InPacketsLost"), InPacketsLost);
+				PerfCounters->Set(TEXT("OutPacketsLost"), OutPacketsLost);
+				PerfCounters->Set(TEXT("InPackets"), InPackets);
+				PerfCounters->Set(TEXT("OutPackets"), OutPackets);
+				PerfCounters->Set(TEXT("InBunches"), InBunches);
+				PerfCounters->Set(TEXT("OutBunches"), OutBunches);
+
+				PerfCounters->Set(TEXT("ServerReplicateActorsTimeMs"), ServerReplicateActorsTimeMs);
+				PerfCounters->Set(TEXT("OutSaturationMax"), RemoteSaturationMax);
+			}
 #endif // USE_SERVER_PERF_COUNTERS
 
-		// Reset everything
-		InBytes = 0;
-		OutBytes = 0;
-		NetGUIDOutBytes = 0;
-		NetGUIDInBytes = 0;
-		InPackets = 0;
-		OutPackets = 0;
-		InBunches = 0;
-		OutBunches = 0;
-		OutPacketsLost = 0;
-		InPacketsLost = 0;
-		VoicePacketsSent = 0;
-		VoiceBytesSent = 0;
-		VoicePacketsRecv = 0;
-		VoiceBytesRecv = 0;
-		VoiceInPercent = 0;
-		VoiceOutPercent = 0;
-		StatUpdateTime = CurrentRealtimeSeconds;
-	}
-#endif // (USE_SERVER_PERF_COUNTERS) || STATS
+			// Reset everything
+			InBytes = 0;
+			OutBytes = 0;
+			NetGUIDOutBytes = 0;
+			NetGUIDInBytes = 0;
+			InPackets = 0;
+			OutPackets = 0;
+			InBunches = 0;
+			OutBunches = 0;
+			OutPacketsLost = 0;
+			InPacketsLost = 0;
+			VoicePacketsSent = 0;
+			VoiceBytesSent = 0;
+			VoicePacketsRecv = 0;
+			VoiceBytesRecv = 0;
+			VoiceInPercent = 0;
+			VoiceOutPercent = 0;
+			StatUpdateTime = CurrentRealtimeSeconds;
+		}
+	} // bCollectNetStats ||(USE_SERVER_PERF_COUNTERS) || STATS
 
 	// Poll all sockets.
 	if( ServerConnection )
diff --git a/Engine/Source/Runtime/Engine/Private/PackageMapClient.cpp b/Engine/Source/Runtime/Engine/Private/PackageMapClient.cpp
--- a/Engine/Source/Runtime/Engine/Private/PackageMapClient.cpp
+++ b/Engine/Source/Runtime/Engine/Private/PackageMapClient.cpp
@@ -2295,15 +2295,22 @@ UObject* FNetGUIDCache::GetObjectFromNetGUID( const FNetworkGUID& NetGUID, const
 			//	2. We aren't already pending
 			//	3. We're actually suppose to load (levels don't load here for example)
 			//		(Refer to CanClientLoadObject, which is where we protect clients from trying to load levels)
-			check( !PendingAsyncPackages.Contains( CacheObjectPtr->PathName ) );
-
+			
 			if ( CVarAllowAsyncLoading.GetValueOnGameThread() > 0 )
 			{
-				PendingAsyncPackages.Add( CacheObjectPtr->PathName, NetGUID );
-				CacheObjectPtr->bIsPending = true;
-				LoadPackageAsync(CacheObjectPtr->PathName.ToString(), FLoadPackageAsyncDelegate::CreateRaw(this, &FNetGUIDCache::AsyncPackageCallback));
+				if (!PendingAsyncPackages.Contains(CacheObjectPtr->PathName))
+				{
+					PendingAsyncPackages.Add(CacheObjectPtr->PathName, NetGUID);
+					CacheObjectPtr->bIsPending = true;
+					LoadPackageAsync(CacheObjectPtr->PathName.ToString(), FLoadPackageAsyncDelegate::CreateRaw(this, &FNetGUIDCache::AsyncPackageCallback));
 
-				UE_LOG( LogNetPackageMap, Log, TEXT( "GetObjectFromNetGUID: Async loading package. Path: %s, NetGUID: %s" ), *CacheObjectPtr->PathName.ToString(), *NetGUID.ToString() );
+					UE_LOG(LogNetPackageMap, Log, TEXT("GetObjectFromNetGUID: Async loading package. Path: %s, NetGUID: %s"), *CacheObjectPtr->PathName.ToString(), *NetGUID.ToString());
+				}
+				else
+				{
+					check(PendingAsyncPackages[CacheObjectPtr->PathName] == NetGUID);
+					UE_LOG(LogNetPackageMap, Log, TEXT("GetObjectFromNetGUID: Already async loading package. Path: %s, NetGUID: %s"), *CacheObjectPtr->PathName.ToString(), *NetGUID.ToString());
+				}
 
 				// There is nothing else to do except wait on the delegate to tell us this package is done loading
 				return NULL;
diff --git a/Engine/Source/Runtime/Engine/Private/PlayerController.cpp b/Engine/Source/Runtime/Engine/Private/PlayerController.cpp
--- a/Engine/Source/Runtime/Engine/Private/PlayerController.cpp
+++ b/Engine/Source/Runtime/Engine/Private/PlayerController.cpp
@@ -4145,7 +4145,7 @@ void APlayerController::EndPlayingState()
 
 void APlayerController::BeginSpectatingState()
 {
-	if ( GetPawn() != NULL )
+	if (GetPawn() != NULL && Role == ROLE_Authority)
 	{
 		UnPossess();
 	}
diff --git a/Engine/Source/Runtime/Engine/Private/SkeletalMeshComponentPhysics.cpp b/Engine/Source/Runtime/Engine/Private/SkeletalMeshComponentPhysics.cpp
--- a/Engine/Source/Runtime/Engine/Private/SkeletalMeshComponentPhysics.cpp
+++ b/Engine/Source/Runtime/Engine/Private/SkeletalMeshComponentPhysics.cpp
@@ -2350,6 +2350,12 @@ void USkeletalMeshComponent::RecreateClothingActors()
 
 	ReleaseAllClothingResources();
 
+	// FIXME: before merging to engine bDisableClothSimulation needs to become protected w/accessors so toggling the flag works correctly
+	if (bDisableClothSimulation)
+	{
+		return;
+	}
+
 	const int32 NumAssets = SkeletalMesh->ClothingAssets.Num();
 	for(int32 AssetIdx=0; AssetIdx < NumAssets; AssetIdx++)
 	{
diff --git a/Engine/Source/Runtime/Engine/Private/Timeline.cpp b/Engine/Source/Runtime/Engine/Private/Timeline.cpp
--- a/Engine/Source/Runtime/Engine/Private/Timeline.cpp
+++ b/Engine/Source/Runtime/Engine/Private/Timeline.cpp
@@ -899,7 +899,12 @@ void UTimelineComponent::SetDirectionPropertyName(FName DirectionPropertyName)
 
 void UTimelineComponent::OnRep_Timeline()
 {
-
+	if (!TheTimeline.IsPlaying())
+	{
+		// make sure a final update call occurs on the client for the final position
+		// FIXME: this is incomplete, we need to compare vs the last simulated position for firing events and such
+		TheTimeline.SetPlaybackPosition(TheTimeline.GetPlaybackPosition(), false, true);
+	}
 }
 
 void UTimelineComponent::GetLifetimeReplicatedProps( TArray< FLifetimeProperty > & OutLifetimeProps ) const
diff --git a/Engine/Source/Runtime/Engine/Public/BatchedElements.h b/Engine/Source/Runtime/Engine/Public/BatchedElements.h
--- a/Engine/Source/Runtime/Engine/Public/BatchedElements.h
+++ b/Engine/Source/Runtime/Engine/Public/BatchedElements.h
@@ -106,6 +106,9 @@ public:
 	/** Adds a line to the batch. Note only SE_BLEND_Opaque will be used for batched line rendering. */
 	void AddLine(const FVector& Start,const FVector& End,const FLinearColor& Color,FHitProxyId HitProxyId, float Thickness = 0.0f, float DepthBias = 0.0f, bool bScreenSpace = false);
 
+	/** Adds a translucent line to the batch. */
+	void AddTranslucentLine(const FVector& Start, const FVector& End, const FLinearColor& Color, FHitProxyId HitProxyId, float Thickness = 0.0f, float DepthBias = 0.0f, bool bScreenSpace = false);
+
 	/** Adds a point to the batch. Note only SE_BLEND_Opaque will be used for batched point rendering. */
 	void AddPoint(const FVector& Position,float Size,const FLinearColor& Color,FHitProxyId HitProxyId);
 
diff --git a/Engine/Source/Runtime/Renderer/Private/TranslucentRendering.cpp b/Engine/Source/Runtime/Renderer/Private/TranslucentRendering.cpp
--- a/Engine/Source/Runtime/Renderer/Private/TranslucentRendering.cpp
+++ b/Engine/Source/Runtime/Renderer/Private/TranslucentRendering.cpp
@@ -491,7 +491,9 @@ bool FTranslucencyDrawingPolicyFactory::DrawMesh(
 	// Only render translucent materials
 	if (IsTranslucentBlendMode(BlendMode))
 	{
-		const bool bMeshUseSeparateTranslucency = Material->IsSeparateTranslucencyEnabled() && IsSeparateTranslucencyPossible(View);
+		// fix for materials on Canvas
+		// PrimitiveSceneProxy is NULL when rendering Canvas items
+		const bool bMeshUseSeparateTranslucency = Material->IsSeparateTranslucencyEnabled() && IsSeparateTranslucencyPossible(View) && PrimitiveSceneProxy != nullptr;
 
 		// Only draw meshes in the relevant pass
 		if (bMeshUseSeparateTranslucency == (DrawingContext.TranslucenyPassType == ETranslucencyPass::TPT_SeparateTranslucency)
