Make sure we return eTouch to physx during an overlap query. Fixes bad behavior when multiple objects blocked in an overlap query.

diff --git a/Engine/Source/Runtime/Engine/Private/Collision/PhysXCollision.cpp b/Engine/Source/Runtime/Engine/Private/Collision/PhysXCollision.cpp
--- a/Engine/Source/Runtime/Engine/Private/Collision/PhysXCollision.cpp
+++ b/Engine/Source/Runtime/Engine/Private/Collision/PhysXCollision.cpp
@@ -699,6 +699,11 @@ PxQueryHitType::Enum FPxQueryFilterCallback::preFilter(const PxFilterData& filte
 	}
 #endif // ENABLE_PREFILTER_LOGGING
 
+	if(bIsOverlapQuery && Result == PxQueryHitType::eBLOCK)
+	{
+		Result = PxQueryHitType::eTOUCH;	//In the case of overlaps, physx only understands touches. We do this at the end to ensure all filtering logic based on block vs overlap is correct
+	}
+
 	return  (PrefilterReturnValue = Result);
 }
 
@@ -1760,6 +1765,7 @@ bool GeomOverlapMultiImp_PhysX(const UWorld* World, const PxGeometry& PGeom, con
 		PxQueryFilterData PQueryFilterDataAny(PFilter, StaticDynamicQueryFlags(Params) | PxQueryFlag::ePREFILTER | PxQueryFlag::eANY_HIT);
 		FPxQueryFilterCallback PQueryCallback(Params);
 		PQueryCallback.bIgnoreTouches = (InfoType == EQueryInfo::IsBlocking); // pre-filter to ignore touches and only get blocking hits, if that's what we're after.
+		PQueryCallback.bIsOverlapQuery = true;
 
 		// Enable scene locks, in case they are required
 		FScopedMultiSceneReadLock SceneLocks;
diff --git a/Engine/Source/Runtime/Engine/Private/Collision/PhysXCollision.h b/Engine/Source/Runtime/Engine/Private/Collision/PhysXCollision.h
--- a/Engine/Source/Runtime/Engine/Private/Collision/PhysXCollision.h
+++ b/Engine/Source/Runtime/Engine/Private/Collision/PhysXCollision.h
@@ -96,6 +96,9 @@ public:
 	/** Result of PreFilter callback. */
 	PxQueryHitType::Enum PrefilterReturnValue;
 
+	/** Whether we are doing an overlap query. This is needed to ensure physx results are never blocking (even if they are in terms of unreal)*/
+	bool bIsOverlapQuery;
+
 	/** Whether to ignore touches (convert an eTOUCH result to eNONE). */
 	bool bIgnoreTouches;
 
@@ -107,6 +110,7 @@ public:
 		, IgnoreActors(InQueryParams.GetIgnoredActors())
 	{
 		PrefilterReturnValue = PxQueryHitType::eNONE;		
+		bIsOverlapQuery = false;
 		bIgnoreTouches = false;
 		bIgnoreBlocks = InQueryParams.bIgnoreBlocks;
 	}
