fix for notification warning when running DDC commandlet

diff --git a/Engine/Source/Runtime/Engine/Private/Kismet/BlueprintPlatformLibrary.cpp b/Engine/Source/Runtime/Engine/Private/Kismet/BlueprintPlatformLibrary.cpp
--- a/Engine/Source/Runtime/Engine/Private/Kismet/BlueprintPlatformLibrary.cpp
+++ b/Engine/Source/Runtime/Engine/Private/Kismet/BlueprintPlatformLibrary.cpp
@@ -74,7 +74,9 @@ UBlueprintPlatformLibrary::UBlueprintPlatformLibrary(const FObjectInitializer& O
 	}
 	else
 	{
+#if !WITH_EDITOR
 		UE_LOG(LogBlueprintUserMessages, Warning, TEXT("DEBUG: LocalNotification DefaultPlatformService NOT FOUND"));
+#endif
 	}
 }
 
