VREditor: Fixed compilation editor with non-editor targets (UE-36784)

diff --git a/Engine/Source/Runtime/Engine/Private/UnrealEngine.cpp b/Engine/Source/Runtime/Engine/Private/UnrealEngine.cpp
--- a/Engine/Source/Runtime/Engine/Private/UnrealEngine.cpp
+++ b/Engine/Source/Runtime/Engine/Private/UnrealEngine.cpp
@@ -10537,7 +10537,9 @@ void UEngine::DestroyWorldContext(UWorld * InWorld)
 	{
 		if (WorldList[idx].World() == InWorld)
 		{
+#if WITH_EDITOR
 			WorldContextDestroyedEvent.Broadcast(WorldList[idx]);
+#endif
 			// Set the current world to NULL so that any external referencers are cleaned up before we remove
 			WorldList[idx].SetCurrentWorld(NULL);
 			WorldList.RemoveAt(idx);
