Adding editor-only flag and function on SkeletalMeshComponent to allow playback of animations on skeletal meshes placed in the level.

diff --git a/Engine/Source/Runtime/Engine/Classes/Components/SkeletalMeshComponent.h b/Engine/Source/Runtime/Engine/Classes/Components/SkeletalMeshComponent.h
--- a/Engine/Source/Runtime/Engine/Classes/Components/SkeletalMeshComponent.h
+++ b/Engine/Source/Runtime/Engine/Classes/Components/SkeletalMeshComponent.h
@@ -592,6 +592,15 @@ public:
 	UPROPERTY()
 	uint32 bEnableLineCheckWithBounds:1;
 
+protected:
+#if WITH_EDITORONLY_DATA
+	/** If true, this will Tick until disabled */
+	UPROPERTY(AdvancedDisplay, EditInstanceOnly, transient, Category = SkeletalMesh)
+	uint32 bUpdateAnimationInEditor : 1;
+#endif
+
+public:
+
 	/** Cache AnimCurveUidVersion from Skeleton and this will be used to identify if it needs to be updated */
 	UPROPERTY(transient)
 	uint16 CachedAnimCurveUidVersion;
@@ -788,6 +797,13 @@ public:
 	UFUNCTION(BlueprintCallable, Category="Components|SkeletalMesh", meta=(UnsafeDuringActorConstruction="true"))
 	void UnbindClothFromMasterPoseComponent(bool bRestoreSimulationSpace = true);
 
+	/**
+	* Sets whether or not to force tick component in order to update animation and refresh transform for this component
+	* This is supported only in the editor
+	*/
+	UFUNCTION(BlueprintCallable, Category = "Components|SkeletalMesh", meta = (DevelopmentOnly, UnsafeDuringActorConstruction = "true"))
+	void SetUpdateAnimationInEditor(const bool NewUpdateState);
+
 	/** We detach the Component once we are done playing it.
 	 *
 	 * @param	ParticleSystemComponent that finished
diff --git a/Engine/Source/Runtime/Engine/Private/Components/SkeletalMeshComponent.cpp b/Engine/Source/Runtime/Engine/Private/Components/SkeletalMeshComponent.cpp
--- a/Engine/Source/Runtime/Engine/Private/Components/SkeletalMeshComponent.cpp
+++ b/Engine/Source/Runtime/Engine/Private/Components/SkeletalMeshComponent.cpp
@@ -747,12 +747,21 @@ bool USkeletalMeshComponent::UpdateLODStatus()
 
 bool USkeletalMeshComponent::ShouldUpdateTransform(bool bLODHasChanged) const
 {
+
 #if WITH_EDITOR
+	
 	// If we're in an editor world (Non running, WorldType will be PIE when simulating or in PIE) then we only want transform updates on LOD changes as the
 	// animation isn't running so it would just waste CPU time
-	if(GetWorld()->WorldType == EWorldType::Editor && !bLODHasChanged)
+	if(GetWorld()->WorldType == EWorldType::Editor)
 	{
-		return false;
+		if( bUpdateAnimationInEditor )
+		{
+			return true;
+		}
+		if( !bLODHasChanged )
+		{
+			return false;
+		}
 	}
 #endif
 
@@ -768,6 +777,17 @@ bool USkeletalMeshComponent::ShouldTickPose() const
 	// When we stop root motion we go back to ticking after CharacterMovement. Unfortunately that means that we could tick twice that frame.
 	// So only enforce a single tick per frame.
 	const bool bAlreadyTickedThisFrame = PoseTickedThisFrame();
+
+#ifdef WITH_EDITOR
+	if (GetWorld()->WorldType == EWorldType::Editor)
+	{
+		if (bUpdateAnimationInEditor)
+		{
+			return true;
+		}
+	}
+#endif 
+
 	return (Super::ShouldTickPose() && IsRegistered() && (AnimScriptInstance || PostProcessAnimInstance) && !bAutonomousTickPose && !bPauseAnims && GetWorld()->AreActorsInitialized() && !bNoSkeletonUpdate && !bAlreadyTickedThisFrame);
 }
 
@@ -2746,4 +2766,16 @@ void USkeletalMeshComponent::AddSlavePoseComponent(USkinnedMeshComponent* Skinne
 	bRequiredBonesUpToDate = false;
 }
 
+
+void USkeletalMeshComponent::SetUpdateAnimationInEditor(const bool NewUpdateState)
+{
+	#if WITH_EDITOR
+	if (IsRegistered())
+	{
+		bUpdateAnimationInEditor = NewUpdateState;
+	}
+	#endif
+}
+
+
 #undef LOCTEXT_NAMESPACE
