Fixing typo in #if WITH_EDITOR statement

diff --git a/Engine/Source/Runtime/Engine/Private/Components/SkeletalMeshComponent.cpp b/Engine/Source/Runtime/Engine/Private/Components/SkeletalMeshComponent.cpp
--- a/Engine/Source/Runtime/Engine/Private/Components/SkeletalMeshComponent.cpp
+++ b/Engine/Source/Runtime/Engine/Private/Components/SkeletalMeshComponent.cpp
@@ -778,7 +778,7 @@ bool USkeletalMeshComponent::ShouldTickPose() const
 	// So only enforce a single tick per frame.
 	const bool bAlreadyTickedThisFrame = PoseTickedThisFrame();
 
-#ifdef WITH_EDITOR
+#if WITH_EDITOR
 	if (GetWorld()->WorldType == EWorldType::Editor)
 	{
 		if (bUpdateAnimationInEditor)
