Copying //UE4/Dev-Networking to //UE4/Main (Source: //UE4/Dev-Networking [REDACTED_HANDLE] )

==========================
MAJOR FEATURES + CHANGES
==========================

Change 3098408 on 2016/08/23 by [REDACTED_PERSON]
	Add IsStandalone blueprint node which returns true if the game is running in standalone

	PR #2109: [KismetSystemLibrary ~ Networking] New BP Node! Is Stand Alone  ? Rama (Contributed by [REDACTED_PERSON])

	#tests added bluprint node to blueprint in test project and made sure it worked

Change 3101070 on 2016/08/25 by [REDACTED_PERSON]
	Exposed replication condition property flags to blueprints.

	PR #1355: Exposed "Replicate to Owner Only" for variables to Blueprints (dropdown) (Contributed by [REDACTED_PERSON])

Change 3102934 on 2016/08/26 by [REDACTED_PERSON]
	#UE4 - fixed automation mistake that removes test data from test builds and adds it to shipping builds

Change 3116738 on 2016/09/07 by [REDACTED_PERSON]

	Shared replication shadow state initial pass

	* Use a change list mgr to accumulate changes for a particular actor
	* Each connection will then re-use this information so that compares only happen at most once per frame/actor.
	* The enire history of changes for the actor are saved for the entire life of the actor
	* When the history starts to fill up, start merging older ones
	* JIP players will get entire history
	* Reset connection viewers after we're done with it

Change 3123101 on 2016/09/13 by [REDACTED_PERSON]

	More efficiently process replayout commands

	* We used to iterate over every cmd even with small changelists
	* Now we only touch the absolute necessary cmds
	* Introduced FChangelistIterator that makes it easy to iterate over changelists
	* Introduced FRepHandleIterator, which allows you to convert a changelist of handles to CmdIndex, ArrayIndex, etc
	* FRepHandleIterator handles array and cmd index list bounds checking
	* Converted SentProperties and SendProperties_BackwardsCompatible to use the new FRepHandleIterator
	* Made explicit SendAllProperties_BackwardsCompatible since it was no longer easy to re-use implementation of SendProperties_BackwardsCompatible for this
	* Optimized changelist merging to use pruning when possible

Change 3123589 on 2016/09/13 by [REDACTED_PERSON]

	Fix UE-35757 - Resources do not show a health bar when they're being damaged

Change 3130333 on 2016/09/19 by [REDACTED_PERSON]

	Fix replay backwards compatibility crash

Change 3134499 on 2016/09/21 by [REDACTED_PERSON]

	Engine changes to support multiple "collections" of levels in a single world, the ability to duplicate the dynamic levels in a world at load time, and the ability to play back a replay in the duplicated levels without interfering with the original ones.

Change 3134771 on 2016/09/21 by [REDACTED_PERSON]

	Missing file for 3134499.

Change 3134948 on 2016/09/21 by [REDACTED_PERSON]

	Fix CIS

Change 3135279 on 2016/09/21 by [REDACTED_PERSON]

	Fix initialization order warning

Change 3137140 on 2016/09/22 by [REDACTED_PERSON]

	Fix static analysis warning

Change 3137886 on 2016/09/23 by [REDACTED_PERSON]

	Fix mismatch in replay backwards compatibility property serialization

	* Fixed a bug where we were writing out an extra handle when the array size was 0
	* Added extra sanity checks when reading arrays to check index bounds
	* Added asserts to verify that we actually have elemements in replicated arrays

Change 3138081 on 2016/09/23 by [REDACTED_PERSON]

	Explicity remove a level from its collection in BeginDestroy.

Change 3138768 on 2016/09/23 by [REDACTED_PERSON]

	Enable world origin rebasing for multiplayer. It can be enabled in projects in network project settings.

	PR #2359: UE-13079: World Origin Shifting for Multiplayer (Contributed by [REDACTED_PERSON])

Change 3140413 on 2016/09/26 by [REDACTED_PERSON]

	Fix for crash during cooking

Change 3142515 on 2016/09/27 by [REDACTED_PERSON]

	Paragon code to use the single-world deathcam implementation and related fixes.

Change 3146799 on 2016/09/30 by [REDACTED_PERSON]

	Fix UE-36634: Logspam after rejoining LogOrion:Warning: AOrionPickupManager::TickPickups - SpawnedPickups

	* We now pass in bNetInitial to the compare function
	* We force compares when bNetInitial is true
	* We no longer compare properties that don't have the lifetime flag (weren't registered via DOREPLIFETIME

Change 3147049 on 2016/09/30 by [REDACTED_PERSON]

	Fix UE-36635 - Crash when slotting upgrades quickly in cardshop

Change 3148296 on 2016/10/02 by [REDACTED_PERSON]

	Merging //UE4/Dev-Main to Dev-Networking (//UE4/Dev-Networking)

Change 3148816 on 2016/10/03 by [REDACTED_PERSON]

	Implement "Play Together" PS4 feature in Shooter Game

Change 3149279 on 2016/10/03 by [REDACTED_PERSON]
	PS4 Play Together compile fix

Change 3150142 on 2016/10/04 by [REDACTED_PERSON]

	Remove reference to "deathcam" from a comment in engine code

diff --git a/Engine/Plugins/Online/OnlineSubsystem/Source/Public/OnlineSubsystem.h b/Engine/Plugins/Online/OnlineSubsystem/Source/Public/OnlineSubsystem.h
--- a/Engine/Plugins/Online/OnlineSubsystem/Source/Public/OnlineSubsystem.h
+++ b/Engine/Plugins/Online/OnlineSubsystem/Source/Public/OnlineSubsystem.h
@@ -79,6 +79,15 @@ typedef FOnConnectionStatusChanged::FDelegate FOnConnectionStatusChangedDelegate
 DECLARE_MULTICAST_DELEGATE_TwoParams(FOnOnlineEnvironmentChanged, EOnlineEnvironment::Type /*LastEnvironment*/, EOnlineEnvironment::Type /*Environment*/);
 typedef FOnOnlineEnvironmentChanged::FDelegate FOnOnlineEnvironmentChangedDelegate;
 
+/**
+* Delegate fired when the "Play Together" event is sent from the PS4 system
+*
+* @param UserIndex - User index of the player the event is for
+* @param UserIdList - list of other users in the PS4 party to send invites to
+*/
+DECLARE_MULTICAST_DELEGATE_TwoParams(FOnPlayTogetherEventReceived, int32, TArray<TSharedPtr<const FUniqueNetId>>);
+typedef FOnPlayTogetherEventReceived::FDelegate FOnPlayTogetherEventReceivedDelegate;
+
 /**
  *	OnlineSubsystem - Series of interfaces to support communicating with various web/platform layer services
  */
@@ -484,6 +493,14 @@ public:
 	 * @param Environment - current online environment
 	 */
 	DEFINE_ONLINE_DELEGATE_TWO_PARAM(OnOnlineEnvironmentChanged, EOnlineEnvironment::Type /*LastEnvironment*/, EOnlineEnvironment::Type /*Environment*/);
+
+	/**
+	* Delegate fired when the "Play Together" event is sent from the PS4 system
+	*
+	* @param UserIndex - User index of the player the event is for
+	* @param UserIdList - list of other users in the PS4 party to send invites to
+	*/
+	DEFINE_ONLINE_DELEGATE_TWO_PARAM(OnPlayTogetherEventReceived, int32, TArray<TSharedPtr<const FUniqueNetId>>);
 };
 
 /** Public references to the online subsystem pointer should use this */
diff --git a/Engine/Plugins/Online/OnlineSubsystemUtils/Source/OnlineSubsystemUtils/Classes/OnlineSessionClient.h b/Engine/Plugins/Online/OnlineSubsystemUtils/Source/OnlineSubsystemUtils/Classes/OnlineSessionClient.h
--- a/Engine/Plugins/Online/OnlineSubsystemUtils/Source/OnlineSubsystemUtils/Classes/OnlineSessionClient.h
+++ b/Engine/Plugins/Online/OnlineSubsystemUtils/Source/OnlineSubsystemUtils/Classes/OnlineSessionClient.h
@@ -27,6 +27,8 @@ protected:
 	FOnJoinSessionCompleteDelegate OnJoinSessionCompleteDelegate;
 	/** Delegate for accepting session invites */
 	FOnSessionUserInviteAcceptedDelegate OnSessionUserInviteAcceptedDelegate;
+	/** Delegate for handling the play together system event */
+	FOnPlayTogetherEventReceivedDelegate OnPlayTogetherEventReceivedDelegate;
 
 	// Handles to the above delegates
 	FDelegateHandle OnSessionInviteAcceptedDelegateHandle;
@@ -34,6 +36,7 @@ protected:
 	FDelegateHandle OnDestroyForJoinSessionCompleteDelegateHandle;
 	FDelegateHandle OnDestroyForMainMenuCompleteDelegateHandle;
 	FDelegateHandle OnJoinSessionCompleteDelegateHandle;
+	FDelegateHandle OnPlayTogetherEventReceivedDelegateHandle;
 
 	/** Handle to outstanding start session call */
 	FDelegateHandle StartSessionCompleteHandle;
diff --git a/Engine/Plugins/Online/OnlineSubsystemUtils/Source/OnlineSubsystemUtils/Private/IpNetDriver.cpp b/Engine/Plugins/Online/OnlineSubsystemUtils/Source/OnlineSubsystemUtils/Private/IpNetDriver.cpp
--- a/Engine/Plugins/Online/OnlineSubsystemUtils/Source/OnlineSubsystemUtils/Private/IpNetDriver.cpp
+++ b/Engine/Plugins/Online/OnlineSubsystemUtils/Source/OnlineSubsystemUtils/Private/IpNetDriver.cpp
@@ -189,6 +189,22 @@ void UIpNetDriver::TickDispatch( float DeltaTime )
 {
 	Super::TickDispatch( DeltaTime );
 
+	// Set the context on the world for this driver's level collection.
+	const FLevelCollection* FoundCollection = nullptr;
+	if (World)
+	{
+		for (const FLevelCollection& LC : World->GetLevelCollections())
+		{
+			if (LC.GetNetDriver() == this)
+			{
+				FoundCollection = &LC;
+				break;
+			}
+		}
+	}
+
+	FScopedLevelCollectionContextSwitch LCSwitch(FoundCollection, World);
+
 	ISocketSubsystem* SocketSubsystem = GetSocketSubsystem();
 
 	const double StartReceiveTime = FPlatformTime::Seconds();
diff --git a/Engine/Plugins/Online/OnlineSubsystemUtils/Source/OnlineSubsystemUtils/Private/OnlineSessionClient.cpp b/Engine/Plugins/Online/OnlineSubsystemUtils/Source/OnlineSubsystemUtils/Private/OnlineSessionClient.cpp
--- a/Engine/Plugins/Online/OnlineSubsystemUtils/Source/OnlineSubsystemUtils/Private/OnlineSessionClient.cpp
+++ b/Engine/Plugins/Online/OnlineSubsystemUtils/Source/OnlineSubsystemUtils/Private/OnlineSessionClient.cpp
@@ -8,6 +8,7 @@
 #include "OnlineSubsystemUtilsPrivatePCH.h"
 #include "OnlineSessionClient.h"
 #include "Engine/GameInstance.h"
+#include "OnlineSubsystem.h"
 
 UOnlineSessionClient::UOnlineSessionClient(const FObjectInitializer& ObjectInitializer)
 	: Super(ObjectInitializer)
@@ -45,6 +46,12 @@ void UOnlineSessionClient::RegisterOnlineDelegates()
 	OnDestroyForJoinSessionCompleteDelegate = FOnDestroySessionCompleteDelegate::CreateUObject(this, &ThisClass::OnDestroyForJoinSessionComplete);
 	OnDestroyForMainMenuCompleteDelegate	= FOnDestroySessionCompleteDelegate::CreateUObject(this, &ThisClass::OnDestroyForMainMenuComplete);
 	OnSessionUserInviteAcceptedDelegate     = FOnSessionUserInviteAcceptedDelegate::CreateUObject(this, &ThisClass::OnSessionUserInviteAccepted);
+	OnPlayTogetherEventReceivedDelegate		= FOnPlayTogetherEventReceivedDelegate::CreateUObject(this, &UOnlineSession::OnPlayTogetherEventReceived);
+
+	if (IOnlineSubsystem* const OnlineSubsystem = IOnlineSubsystem::Get())
+	{
+		OnPlayTogetherEventReceivedDelegateHandle = OnlineSubsystem->AddOnPlayTogetherEventReceivedDelegate_Handle(OnPlayTogetherEventReceivedDelegate);
+	}
 
 	IOnlineSessionPtr SessionInt = GetSessionInt();
 	if (SessionInt.IsValid())
@@ -60,6 +67,11 @@ void UOnlineSessionClient::ClearOnlineDelegates()
 	{
 		SessionInt->ClearOnSessionUserInviteAcceptedDelegate_Handle(OnSessionUserInviteAcceptedDelegateHandle);
 	}
+
+	if (IOnlineSubsystem* const OnlineSubsystem = IOnlineSubsystem::Get())
+	{
+		OnlineSubsystem->ClearOnPlayTogetherEventReceivedDelegate_Handle(OnPlayTogetherEventReceivedDelegateHandle);
+	}
 }
 
 void UOnlineSessionClient::OnSessionUserInviteAccepted(bool bWasSuccessful, int32 ControllerId, TSharedPtr<const FUniqueNetId> UserId, const FOnlineSessionSearchResult& SearchResult)
@@ -329,4 +341,4 @@ void UOnlineSessionClient::SetInviteFlags(UWorld* World, const FJoinabilitySetti
 			SessionInt->UpdateSession(Settings.SessionName, *GameSettings, false);
 		}
 	}
-}
\ No newline at end of file
+}
diff --git a/Engine/Plugins/Online/OnlineSubsystemUtils/Source/OnlineSubsystemUtils/Private/Tests/TestVoiceData.h b/Engine/Plugins/Online/OnlineSubsystemUtils/Source/OnlineSubsystemUtils/Private/Tests/TestVoiceData.h
--- a/Engine/Plugins/Online/OnlineSubsystemUtils/Source/OnlineSubsystemUtils/Private/Tests/TestVoiceData.h
+++ b/Engine/Plugins/Online/OnlineSubsystemUtils/Source/OnlineSubsystemUtils/Private/Tests/TestVoiceData.h
@@ -4,10 +4,6 @@
 
 #if WITH_DEV_AUTOMATION_TESTS
 
-uint8 RawVoiceTestData[] = { 0x0, 0x0 };
-
-#else
-
 /** 16khz wav mono voice data "this is a test" 2 sec duration */
 uint8 RawVoiceTestData[] = {
 	0xcc, 0xff, 0xb8, 0xff, 0x97, 0xff, 0x87, 0xff, 0x94, 0xff, 0xa0, 0xff, 0xad, 0xff, 0xa0, 0xff, 0xb7, 0xff, 0x9f, 0xff, 0x84, 0xff, 0x80, 0xff, 0x80, 0xff, 0x95, 0xff, 0x85, 0xff, 0x8e, 0xff, 
diff --git a/Engine/Source/Editor/Kismet/Private/BlueprintDetailsCustomization.cpp b/Engine/Source/Editor/Kismet/Private/BlueprintDetailsCustomization.cpp
--- a/Engine/Source/Editor/Kismet/Private/BlueprintDetailsCustomization.cpp
+++ b/Engine/Source/Editor/Kismet/Private/BlueprintDetailsCustomization.cpp
@@ -633,6 +633,37 @@ void FBlueprintVarActionDetails::CustomizeDetails( IDetailLayoutBuilder& DetailL
 		.ToolTip(ReplicationTooltip)
 	];
 
+	ReplicationConditionEnumTypeNames.Empty();
+	UEnum* Enum = FindObject<UEnum>(ANY_PACKAGE, TEXT("ELifetimeCondition"), true);
+	check(Enum);
+	
+	for (int32 i = 0; i < Enum->NumEnums(); i++)
+	{
+		if (!Enum->HasMetaData(TEXT("Hidden"), i))
+		{
+			ReplicationConditionEnumTypeNames.Add(MakeShareable(new FString(Enum->GetDisplayNameText(i).ToString())));
+		}
+	}
+
+	Category.AddCustomRow(LOCTEXT("VariableReplicationConditionsLabel", "Replication Condition"))
+	.Visibility(TAttribute<EVisibility>(this, &FBlueprintVarActionDetails::ReplicationVisibility))
+	.NameContent()
+	[
+		SNew(STextBlock)
+		.ToolTip(ReplicationTooltip)
+		.Text(LOCTEXT("VariableReplicationConditionsLabel", "Replication Condition"))
+		.Font(DetailFontInfo)
+	]
+	.ValueContent()
+	[
+		SNew(STextComboBox)
+		.OptionsSource(&ReplicationConditionEnumTypeNames)
+		.InitiallySelectedItem(GetVariableReplicationCondition())
+		.OnSelectionChanged(this, &FBlueprintVarActionDetails::OnChangeReplicationCondition)
+		.IsEnabled(this, &FBlueprintVarActionDetails::ReplicationConditionEnabled)
+	];
+
+
 	UBlueprint* BlueprintObj = GetBlueprintObj();
 
 	// Handle event generation
@@ -1668,6 +1699,64 @@ EVisibility FBlueprintVarActionDetails::ExposeToCinematicsVisibility() const
 	return EVisibility::Collapsed;
 }
 
+TSharedPtr<FString> FBlueprintVarActionDetails::GetVariableReplicationCondition() const
+{
+	ELifetimeCondition VariableRepCondition = COND_None;
+		
+	const UProperty* const Property = CachedVariableProperty.Get();
+
+	if (Property)
+	{
+		VariableRepCondition = Property->GetBlueprintReplicationCondition();
+	}
+
+	return ReplicationConditionEnumTypeNames[(uint8)VariableRepCondition];
+}
+
+void FBlueprintVarActionDetails::OnChangeReplicationCondition(TSharedPtr<FString> ItemSelected, ESelectInfo::Type SelectInfo)
+{
+	int32 NewSelection;
+	const bool bFound = ReplicationConditionEnumTypeNames.Find(ItemSelected, NewSelection);
+	check(bFound && NewSelection != INDEX_NONE);
+
+	const ELifetimeCondition NewRepCondition = (ELifetimeCondition)NewSelection;
+
+	UBlueprint* const BlueprintObj = GetBlueprintObj();
+	const FName VarName = CachedVariableName;
+
+	if (BlueprintObj && VarName != NAME_None)
+	{
+		const int32 VarIndex = FBlueprintEditorUtils::FindNewVariableIndex(BlueprintObj, VarName);
+
+		if (VarIndex != INDEX_NONE)
+		{
+			BlueprintObj->NewVariables[VarIndex].ReplicationCondition = NewRepCondition;
+		
+			FBlueprintEditorUtils::MarkBlueprintAsStructurallyModified(BlueprintObj);
+		}
+	}
+
+}
+
+bool FBlueprintVarActionDetails::ReplicationConditionEnabled() const
+{
+	const UProperty* const VariableProperty = CachedVariableProperty.Get();
+	if (VariableProperty)
+	{
+		const uint64 *PropFlagPtr = FBlueprintEditorUtils::GetBlueprintVariablePropertyFlags(GetBlueprintObj(), VariableProperty->GetFName());
+		uint64 PropFlags = 0;
+
+		if (PropFlagPtr != nullptr)
+		{
+			PropFlags = *PropFlagPtr;
+			return (PropFlags & CPF_Net) > 0;
+			
+		}
+	}
+
+	return false;
+}
+
 ECheckBoxState FBlueprintVarActionDetails::OnGetConfigVariableCheckboxState() const
 {
 	UBlueprint* BlueprintObj = GetPropertyOwnerBlueprint();
@@ -1949,6 +2038,14 @@ void FBlueprintVarActionDetails::OnChangeReplication(TSharedPtr<FString> ItemSel
 	
 	UProperty* VariableProperty = CachedVariableProperty.Get();
 
+	UBlueprint* const BlueprintObj = GetBlueprintObj();
+	const FName VarName = CachedVariableName;
+	int32 VarIndex = INDEX_NONE;
+	if (BlueprintObj && VarName != NAME_None)
+	{
+		VarIndex = FBlueprintEditorUtils::FindNewVariableIndex(BlueprintObj, VarName);
+	}
+
 	if (VariableProperty)
 	{
 		uint64 *PropFlagPtr = FBlueprintEditorUtils::GetBlueprintVariablePropertyFlags(GetBlueprintObj(), VariableProperty->GetFName());
@@ -1958,7 +2055,14 @@ void FBlueprintVarActionDetails::OnChangeReplication(TSharedPtr<FString> ItemSel
 			{
 			case EVariableReplication::None:
 				*PropFlagPtr &= ~CPF_Net;
-				ReplicationOnRepFuncChanged(FName(NAME_None).ToString());	
+				ReplicationOnRepFuncChanged(FName(NAME_None).ToString());
+
+				//set replication condition to none:
+				if (VarIndex != INDEX_NONE)
+				{
+					BlueprintObj->NewVariables[VarIndex].ReplicationCondition = COND_None;
+				}
+				
 				break;
 				
 			case EVariableReplication::Replicated:
@@ -1969,11 +2073,11 @@ void FBlueprintVarActionDetails::OnChangeReplication(TSharedPtr<FString> ItemSel
 			case EVariableReplication::RepNotify:
 				*PropFlagPtr |= CPF_Net;
 				FString NewFuncName = FString::Printf(TEXT("OnRep_%s"), *VariableProperty->GetName());
-				UEdGraph* FuncGraph = FindObject<UEdGraph>(GetBlueprintObj(), *NewFuncName);
+				UEdGraph* FuncGraph = FindObject<UEdGraph>(BlueprintObj, *NewFuncName);
 				if (!FuncGraph)
 				{
-					FuncGraph = FBlueprintEditorUtils::CreateNewGraph(GetBlueprintObj(), FName(*NewFuncName), UEdGraph::StaticClass(), UEdGraphSchema_K2::StaticClass());
-					FBlueprintEditorUtils::AddFunctionGraph<UClass>(GetBlueprintObj(), FuncGraph, false, NULL);
+					FuncGraph = FBlueprintEditorUtils::CreateNewGraph(BlueprintObj, FName(*NewFuncName), UEdGraph::StaticClass(), UEdGraphSchema_K2::StaticClass());
+					FBlueprintEditorUtils::AddFunctionGraph<UClass>(BlueprintObj, FuncGraph, false, NULL);
 				}
 
 				if (FuncGraph)
@@ -1983,7 +2087,7 @@ void FBlueprintVarActionDetails::OnChangeReplication(TSharedPtr<FString> ItemSel
 				break;
 			}
 
-			FBlueprintEditorUtils::MarkBlueprintAsStructurallyModified(GetBlueprintObj());
+			FBlueprintEditorUtils::MarkBlueprintAsStructurallyModified(BlueprintObj);
 		}
 	}
 }
diff --git a/Engine/Source/Editor/Kismet/Private/BlueprintDetailsCustomization.h b/Engine/Source/Editor/Kismet/Private/BlueprintDetailsCustomization.h
--- a/Engine/Source/Editor/Kismet/Private/BlueprintDetailsCustomization.h
+++ b/Engine/Source/Editor/Kismet/Private/BlueprintDetailsCustomization.h
@@ -168,6 +168,12 @@ private:
 	void ReplicationOnRepFuncChanged(const FString& NewOnRepFunc) const;
 	EVisibility ReplicationVisibility() const;
 
+	/** Array of replication conditions for the combo text box */
+	TArray<TSharedPtr<FString>> ReplicationConditionEnumTypeNames;
+	TSharedPtr<FString> GetVariableReplicationCondition() const;
+	void OnChangeReplicationCondition(TSharedPtr<FString> ItemSelected, ESelectInfo::Type SelectInfo);
+	bool ReplicationConditionEnabled() const;
+
 	EVisibility GetTransientVisibility() const;
 	ECheckBoxState OnGetTransientCheckboxState() const;
 	void OnTransientChanged(ECheckBoxState InNewState);
diff --git a/Engine/Source/Editor/KismetCompiler/Private/KismetCompiler.cpp b/Engine/Source/Editor/KismetCompiler/Private/KismetCompiler.cpp
--- a/Engine/Source/Editor/KismetCompiler/Private/KismetCompiler.cpp
+++ b/Engine/Source/Editor/KismetCompiler/Private/KismetCompiler.cpp
@@ -542,6 +542,7 @@ void FKismetCompilerContext::CreateClassVariablesFromBlueprint()
 			NewProperty->SetMetaData(TEXT("DisplayName"), *Variable.FriendlyName);
 			NewProperty->SetMetaData(TEXT("Category"), *Variable.Category.ToString());
 			NewProperty->RepNotifyFunc = Variable.RepNotifyFunc;
+			NewProperty->SetBlueprintReplicationCondition(Variable.ReplicationCondition);
 
 			if(!Variable.DefaultValue.IsEmpty())
 			{
diff --git a/Engine/Source/Editor/UnrealEd/Classes/Editor/EditorEngine.h b/Engine/Source/Editor/UnrealEd/Classes/Editor/EditorEngine.h
--- a/Engine/Source/Editor/UnrealEd/Classes/Editor/EditorEngine.h
+++ b/Engine/Source/Editor/UnrealEd/Classes/Editor/EditorEngine.h
@@ -784,7 +784,7 @@ private:
 	virtual void RemapGamepadControllerIdForPIE(class UGameViewportClient* GameViewport, int32 &ControllerId) override;
 	virtual TSharedPtr<SViewport> GetGameViewportWidget() const override;
 	virtual void TriggerStreamingDataRebuild() override;
-	virtual bool NetworkRemapPath(UWorld* InWorld, FString &Str, bool reading = true) override;
+	virtual bool NetworkRemapPath(UNetDriver* Driver, FString &Str, bool reading = true) override;
 	virtual bool NetworkRemapPath(UPendingNetGame* PendingNetGame, FString& Str, bool reading = true) override;
 	virtual bool AreEditorAnalyticsEnabled() const override;
 	virtual void CreateStartupAnalyticsAttributes(TArray<FAnalyticsEventAttribute>& StartSessionAttributes) const override;
diff --git a/Engine/Source/Editor/UnrealEd/Private/EditorEngine.cpp b/Engine/Source/Editor/UnrealEd/Private/EditorEngine.cpp
--- a/Engine/Source/Editor/UnrealEd/Private/EditorEngine.cpp
+++ b/Engine/Source/Editor/UnrealEd/Private/EditorEngine.cpp
@@ -6385,9 +6385,14 @@ FORCEINLINE bool NetworkRemapPath_local(FWorldContext &Context, FString &Str, bo
 	}
 }
 
-bool UEditorEngine::NetworkRemapPath( UWorld *InWorld, FString &Str, bool reading)
+bool UEditorEngine::NetworkRemapPath(UNetDriver* Driver, FString &Str, bool reading)
 {
-	FWorldContext &Context = GetWorldContextFromWorldChecked(InWorld);
+	if (Driver == nullptr)
+	{
+		return false;
+	}
+
+	FWorldContext &Context = GetWorldContextFromWorldChecked(Driver->GetWorld());
 	if (Context.PIEPrefix.IsEmpty() || Context.PIERemapPrefix.IsEmpty())
 	{
 		return false;
diff --git a/Engine/Source/Editor/UnrealEd/Private/Kismet2/BlueprintEditorUtils.cpp b/Engine/Source/Editor/UnrealEd/Private/Kismet2/BlueprintEditorUtils.cpp
--- a/Engine/Source/Editor/UnrealEd/Private/Kismet2/BlueprintEditorUtils.cpp
+++ b/Engine/Source/Editor/UnrealEd/Private/Kismet2/BlueprintEditorUtils.cpp
@@ -4243,6 +4243,7 @@ bool FBlueprintEditorUtils::AddMemberVariable(UBlueprint* Blueprint, const FName
 	{
 		PostSetupObjectPinType(Blueprint, NewVar);
 	}
+	NewVar.ReplicationCondition = COND_None;
 	NewVar.Category = K2Schema->VR_DefaultCategory;
 	NewVar.DefaultValue = DefaultValue;
 
diff --git a/Engine/Source/Runtime/Core/Public/GenericPlatform/GenericPlatformProperties.h b/Engine/Source/Runtime/Core/Public/GenericPlatform/GenericPlatformProperties.h
--- a/Engine/Source/Runtime/Core/Public/GenericPlatform/GenericPlatformProperties.h
+++ b/Engine/Source/Runtime/Core/Public/GenericPlatform/GenericPlatformProperties.h
@@ -259,4 +259,10 @@ struct FGenericPlatformProperties
 	{
 		return false;
 	}
+
+	// whether or not the platform supports shader compression
+	static FORCEINLINE bool SupportsShaderCompression()
+	{
+		return true;
+	}
 };
diff --git a/Engine/Source/Runtime/Core/Public/Math/IntVector.h b/Engine/Source/Runtime/Core/Public/Math/IntVector.h
--- a/Engine/Source/Runtime/Core/Public/Math/IntVector.h
+++ b/Engine/Source/Runtime/Core/Public/Math/IntVector.h
@@ -177,6 +177,12 @@ public:
 	 */
 	FIntVector operator-( const FIntVector& Other ) const;
 
+	/**
+	 * Is vector equal to zero.
+	 * @return is zero
+	*/
+	bool IsZero() const;
+
 public:
 
 	/**
@@ -415,6 +421,11 @@ FORCEINLINE int32 FIntVector::Size() const
 	return int32(FMath::Sqrt(float(X64 * X64 + Y64 * Y64 + Z64 * Z64)));
 }
 
+FORCEINLINE bool FIntVector::IsZero() const
+{
+	return *this == ZeroValue;
+}
+
 
 FORCEINLINE FString FIntVector::ToString() const
 {
diff --git a/Engine/Source/Runtime/CoreUObject/Private/Serialization/AsyncLoading.cpp b/Engine/Source/Runtime/CoreUObject/Private/Serialization/AsyncLoading.cpp
--- a/Engine/Source/Runtime/CoreUObject/Private/Serialization/AsyncLoading.cpp
+++ b/Engine/Source/Runtime/CoreUObject/Private/Serialization/AsyncLoading.cpp
@@ -1472,9 +1472,7 @@ EAsyncPackageState::Type FAsyncPackage::CreateLinker()
 
 		// Set package specific data 
 		Package->SetPackageFlags(Desc.PackageFlags);
-#if WITH_EDITOR
 		Package->PIEInstanceID = Desc.PIEInstanceID;
-#endif
 
 		// Always store package filename we loading from
 		Package->FileName = Desc.NameToLoad;
diff --git a/Engine/Source/Runtime/CoreUObject/Private/UObject/Package.cpp b/Engine/Source/Runtime/CoreUObject/Private/UObject/Package.cpp
--- a/Engine/Source/Runtime/CoreUObject/Private/UObject/Package.cpp
+++ b/Engine/Source/Runtime/CoreUObject/Private/UObject/Package.cpp
@@ -34,9 +34,7 @@ void UPackage::PostInitProperties()
 	MetaData = NULL;
 	LinkerPackageVersion = GPackageFileUE4Version;
 	LinkerLicenseeVersion = GPackageFileLicenseeUE4Version;
-#if WITH_EDITOR
 	PIEInstanceID = INDEX_NONE;
-#endif
 #if WITH_EDITORONLY_DATA
 	bIsCookedForEditor = false;
 	// Mark this package as editor-only by default. As soon as something in it is accessed through a non editor-only
diff --git a/Engine/Source/Runtime/CoreUObject/Private/UObject/UObjectGlobals.cpp b/Engine/Source/Runtime/CoreUObject/Private/UObject/UObjectGlobals.cpp
--- a/Engine/Source/Runtime/CoreUObject/Private/UObject/UObjectGlobals.cpp
+++ b/Engine/Source/Runtime/CoreUObject/Private/UObject/UObjectGlobals.cpp
@@ -881,16 +881,13 @@ UObject* StaticLoadObject(UClass* ObjectClass, UObject* InOuter, const TCHAR* In
 		FString ObjectName = InName;
 		ResolveName(InOuter, ObjectName, true, true, LoadFlags & LOAD_EditorOnly);
 
-		if (InOuter == nullptr || FLinkerLoad::IsKnownMissingPackage(FName(*InOuter->GetPathName())) == false)
-		{
-			// we haven't created or found the object, error
-			FFormatNamedArguments Arguments;
-			Arguments.Add(TEXT("ClassName"), FText::FromString(ObjectClass->GetName()));
-			Arguments.Add(TEXT("OuterName"), InOuter ? FText::FromString(InOuter->GetPathName()) : NSLOCTEXT("Core", "None", "None"));
-			Arguments.Add(TEXT("ObjectName"), FText::FromString(ObjectName));
-			const FString Error = FText::Format(NSLOCTEXT("Core", "ObjectNotFound", "Failed to find object '{ClassName} {OuterName}.{ObjectName}'"), Arguments).ToString();
-			SafeLoadError(InOuter, LoadFlags, *Error);
-		}
+		// we haven't created or found the object, error
+		FFormatNamedArguments Arguments;
+		Arguments.Add(TEXT("ClassName"), FText::FromString(ObjectClass->GetName()));
+		Arguments.Add(TEXT("OuterName"), InOuter ? FText::FromString(InOuter->GetPathName()) : NSLOCTEXT("Core", "None", "None"));
+		Arguments.Add(TEXT("ObjectName"), FText::FromString(ObjectName));
+		const FString Error = FText::Format(NSLOCTEXT("Core", "ObjectNotFound", "Failed to find object '{ClassName} {OuterName}.{ObjectName}'"), Arguments).ToString();
+		SafeLoadError(InOuter, LoadFlags, *Error);
 	}
 	return Result;
 }
@@ -2329,7 +2326,8 @@ UObject* StaticAllocateObject
 					// If we're not in the editor, and aren't doing something specifically destructive like reconstructing blueprints, this is fatal
 					if (!GIsEditor && FApp::IsGame() && !GIsReconstructingBlueprintInstances)
 					{
-						UE_LOG(LogUObjectGlobals, Fatal, TEXT("Gamethread hitch waiting for resource cleanup on a UObject (%s) overwrite. Fix the higher level code so that this does not happen."), *OldName );
+						// Switching to warning, investigate why level duplication triggers this
+						UE_LOG(LogUObjectGlobals, Warning, TEXT("Gamethread hitch waiting for resource cleanup on a UObject (%s) overwrite. Fix the higher level code so that this does not happen."), *OldName );
 					}
 					FPlatformProcess::Sleep(0);
 				}
diff --git a/Engine/Source/Runtime/CoreUObject/Public/Serialization/AsyncPackage.h b/Engine/Source/Runtime/CoreUObject/Public/Serialization/AsyncPackage.h
--- a/Engine/Source/Runtime/CoreUObject/Public/Serialization/AsyncPackage.h
+++ b/Engine/Source/Runtime/CoreUObject/Public/Serialization/AsyncPackage.h
@@ -22,10 +22,8 @@ struct FAsyncPackageDesc
 	EPackageFlags PackageFlags;
 	/** Package loading priority. Higher number is higher priority. */
 	TAsyncLoadPriority Priority;
-#if WITH_EDITOR
-	/** Editor only: PIE instance ID this package belongs to, INDEX_NONE otherwise */
+	/** PIE instance ID this package belongs to, INDEX_NONE otherwise */
 	int32 PIEInstanceID;
-#endif
 
 
 	FAsyncPackageDesc(int32 InRequestID, const FName& InName, FName InPackageToLoadFrom = NAME_None, const FGuid& InGuid = FGuid(), FLoadPackageAsyncDelegate InCompletionDelegate = FLoadPackageAsyncDelegate(), EPackageFlags InPackageFlags = PKG_None, int32 InPIEInstanceID = INDEX_NONE, TAsyncLoadPriority InPriority = 0)
@@ -36,9 +34,7 @@ struct FAsyncPackageDesc
 		, PackageLoadedDelegate(InCompletionDelegate)
 		, PackageFlags(InPackageFlags)
 		, Priority(InPriority)
-#if WITH_EDITOR
 		, PIEInstanceID(InPIEInstanceID)
-#endif
 	{
 		if (NameToLoad == NAME_None)
 		{
diff --git a/Engine/Source/Runtime/CoreUObject/Public/UObject/CoreNet.h b/Engine/Source/Runtime/CoreUObject/Public/UObject/CoreNet.h
--- a/Engine/Source/Runtime/CoreUObject/Public/UObject/CoreNet.h
+++ b/Engine/Source/Runtime/CoreUObject/Public/UObject/CoreNet.h
@@ -5,7 +5,7 @@
 #include "ObjectMacros.h"
 #include "WeakObjectPtr.h"
 #include "Object.h"
-
+#include "CoreNetTypes.h"
 
 DECLARE_DELEGATE_RetVal_OneParam( bool, FNetObjectIsDynamic, const UObject*);
 
@@ -220,33 +220,6 @@ struct FPropertyRetirement
 };
 
 
-/** Secondary condition to check before considering the replication of a lifetime property. */
-enum ELifetimeCondition
-{
-	COND_None						= 0,		// This property has no condition, and will send anytime it changes
-	COND_InitialOnly				= 1,		// This property will only attempt to send on the initial bunch
-	COND_OwnerOnly					= 2,		// This property will only send to the actor's owner
-	COND_SkipOwner					= 3,		// This property send to every connection EXCEPT the owner
-	COND_SimulatedOnly				= 4,		// This property will only send to simulated actors
-	COND_AutonomousOnly				= 5,		// This property will only send to autonomous actors
-	COND_SimulatedOrPhysics			= 6,		// This property will send to simulated OR bRepPhysics actors
-	COND_InitialOrOwner				= 7,		// This property will send on the initial packet, or to the actors owner
-	COND_Custom						= 8,		// This property has no particular condition, but wants the ability to toggle on/off via SetCustomIsActiveOverride
-	COND_ReplayOrOwner				= 9,		// This property will only send to the replay connection, or to the actors owner
-	COND_ReplayOnly					= 10,		// This property will only send to the replay connection
-	COND_SimulatedOnlyNoReplay		= 11,
-	COND_SimulatedOrPhysicsNoReplay	= 12,
-	COND_Max						= 13,
-};
-
-
-enum ELifetimeRepNotifyCondition
-{
-	REPNOTIFY_OnChanged		= 0,		// Only call the property's RepNotify function if it changes from the local value
-	REPNOTIFY_Always		= 1,		// Always Call the property's RepNotify function when it is received from the server
-};
-
-
 /** FLifetimeProperty
  *	This class is used to track a property that is marked to be replicated for the lifetime of the actor channel.
  *  This doesn't mean the property will necessarily always be replicated, it just means:
diff --git a/Engine/Source/Runtime/CoreUObject/Public/UObject/CoreNetTypes.h b/Engine/Source/Runtime/CoreUObject/Public/UObject/CoreNetTypes.h
new file mode 100644
--- /dev/null
+++ b/Engine/Source/Runtime/CoreUObject/Public/UObject/CoreNetTypes.h
@@ -0,0 +1,28 @@
+#pragma once
+
+/** Secondary condition to check before considering the replication of a lifetime property. */
+UENUM(BlueprintType)
+enum ELifetimeCondition
+{
+	COND_None = 0							UMETA(DisplayName = "None"),							// This property has no condition, and will send anytime it changes
+	COND_InitialOnly = 1					UMETA(DisplayName = "Initial Only"),					// This property will only attempt to send on the initial bunch
+	COND_OwnerOnly = 2						UMETA(DisplayName = "Owner Only"),						// This property will only send to the actor's owner
+	COND_SkipOwner = 3						UMETA(DisplayName = "Skip Owner"),						// This property send to every connection EXCEPT the owner
+	COND_SimulatedOnly = 4					UMETA(DisplayName = "Simulated Only"),					// This property will only send to simulated actors
+	COND_AutonomousOnly = 5					UMETA(DisplayName = "Autonomous Only"),					// This property will only send to autonomous actors
+	COND_SimulatedOrPhysics = 6				UMETA(DisplayName = "Simulated Or Physics"),			// This property will send to simulated OR bRepPhysics actors
+	COND_InitialOrOwner = 7					UMETA(DisplayName = "Initial Or Owner"),				// This property will send on the initial packet, or to the actors owner
+	COND_Custom = 8							UMETA(DisplayName = "Custom"),							// This property has no particular condition, but wants the ability to toggle on/off via SetCustomIsActiveOverride
+	COND_ReplayOrOwner = 9					UMETA(DisplayName = "Replay Or Owner"),					// This property will only send to the replay connection, or to the actors owner
+	COND_ReplayOnly = 10					UMETA(DisplayName = "Replay Only"),						// This property will only send to the replay connection
+	COND_SimulatedOnlyNoReplay = 11			UMETA(DisplayName = "Simulated Only No Replay"),		// This property will send to actors only, but not to replay connections
+	COND_SimulatedOrPhysicsNoReplay = 12	UMETA(DisplayName = "Simulated Or Physics No Replay"),	// This property will send to simulated Or bRepPhysics actors, but not to replay connections
+	COND_Max = 13							UMETA(Hidden)
+};
+
+
+enum ELifetimeRepNotifyCondition
+{
+	REPNOTIFY_OnChanged = 0,		// Only call the property's RepNotify function if it changes from the local value
+	REPNOTIFY_Always = 1,		// Always Call the property's RepNotify function when it is received from the server
+};
diff --git a/Engine/Source/Runtime/CoreUObject/Public/UObject/Package.h b/Engine/Source/Runtime/CoreUObject/Public/UObject/Package.h
--- a/Engine/Source/Runtime/CoreUObject/Public/UObject/Package.h
+++ b/Engine/Source/Runtime/CoreUObject/Public/UObject/Package.h
@@ -154,10 +154,8 @@ public:
 
 	TMap<FName, int32> ClassUniqueNameIndexMap;
 
-#if WITH_EDITOR
 	/** Editor only: PIE instance ID this package belongs to, INDEX_NONE otherwise */
 	int32 PIEInstanceID;
-#endif
 
 #if WITH_EDITORONLY_DATA
 	/** True if this packages has been cooked for the editor / opened cooked by the editor */
diff --git a/Engine/Source/Runtime/CoreUObject/Public/UObject/UnrealType.h b/Engine/Source/Runtime/CoreUObject/Public/UObject/UnrealType.h
--- a/Engine/Source/Runtime/CoreUObject/Public/UObject/UnrealType.h
+++ b/Engine/Source/Runtime/CoreUObject/Public/UObject/UnrealType.h
@@ -72,6 +72,8 @@ private:
 	// In memory variables (generated during Link()).
 	int32		Offset_Internal;
 
+	ELifetimeCondition BlueprintReplicationCondition;
+
 public:
 	/** In memory only: Linked list of properties from most-derived to base **/
 	UProperty*	PropertyLinkNext;
@@ -815,6 +817,9 @@ public:
 	/** returns true, if Other is property of exactly the same type */
 	virtual bool SameType(const UProperty* Other) const;
 
+	ELifetimeCondition GetBlueprintReplicationCondition() const { return BlueprintReplicationCondition; }
+	void SetBlueprintReplicationCondition(ELifetimeCondition InBlueprintReplicationCondition) { BlueprintReplicationCondition = InBlueprintReplicationCondition; }
+
 #if HACK_HEADER_GENERATOR
 	// Required by UHT makefiles for internal data serialization.
 	friend struct FPropertyArchiveProxy;
diff --git a/Engine/Source/Runtime/Engine/Classes/Components/SceneComponent.h b/Engine/Source/Runtime/Engine/Classes/Components/SceneComponent.h
--- a/Engine/Source/Runtime/Engine/Classes/Components/SceneComponent.h
+++ b/Engine/Source/Runtime/Engine/Classes/Components/SceneComponent.h
@@ -7,6 +7,8 @@
 #include "ComponentInstanceDataCache.h" // for FActorComponentInstanceData
 #include "SceneComponent.generated.h"
 
+struct FLevelCollection;
+
 /** Overlap info consisting of the primitive and the body that is overlapping */
 USTRUCT()
 struct ENGINE_API FOverlapInfo
@@ -241,6 +243,9 @@ private:
 	FName NetOldAttachSocketName;
 	USceneComponent *NetOldAttachParent;
 
+	/** Cached level collection that contains the level this component is registered in, for fast access in IsVisible(). */
+	const FLevelCollection* CachedLevelCollection;
+
 	UFUNCTION()
 	void OnRep_Transform();
 
@@ -755,6 +760,7 @@ public:
 
 	//~ Begin ActorComponent Interface
 	virtual void OnRegister() override;
+	virtual void OnUnregister() override;
 	/** Return true if CreateRenderState() should be called */
 	virtual bool ShouldCreateRenderState() const override
 	{
diff --git a/Engine/Source/Runtime/Engine/Classes/Engine/Blueprint.h b/Engine/Source/Runtime/Engine/Classes/Engine/Blueprint.h
--- a/Engine/Source/Runtime/Engine/Classes/Engine/Blueprint.h
+++ b/Engine/Source/Runtime/Engine/Classes/Engine/Blueprint.h
@@ -2,13 +2,13 @@
 
 #pragma once
 
+#include "CoreNetTypes.h"
 #include "EdGraph/EdGraphPin.h"
 #include "EdGraph/EdGraphNode.h"
 #include "Engine/EngineTypes.h"
 #include "BlueprintCore.h"
 #include "Blueprint.generated.h"
 
-
 /**
  * Enumerates states a blueprint can be in.
  */
@@ -187,6 +187,9 @@ struct FBPVariableDescription
 	UPROPERTY(EditAnywhere, Category=BPVariableRepNotify)
 	FName RepNotifyFunc;
 
+	UPROPERTY(EditAnywhere, Category=BPVariableDescription)
+	TEnumAsByte<ELifetimeCondition> ReplicationCondition;
+
 	/** Metadata information for this variable */
 	UPROPERTY(EditAnywhere, Category=BPVariableDescription)
 	TArray<struct FBPVariableMetaDataEntry> MetaDataArray;
diff --git a/Engine/Source/Runtime/Engine/Classes/Engine/DemoNetConnection.h b/Engine/Source/Runtime/Engine/Classes/Engine/DemoNetConnection.h
--- a/Engine/Source/Runtime/Engine/Classes/Engine/DemoNetConnection.h
+++ b/Engine/Source/Runtime/Engine/Classes/Engine/DemoNetConnection.h
@@ -28,7 +28,7 @@ public:
  * Simulated network connection for recording and playing back game sessions.
  */
 UCLASS(transient, config=Engine)
-class UDemoNetConnection
+class ENGINE_API UDemoNetConnection
 	: public UNetConnection
 {
 	GENERATED_UCLASS_BODY()
diff --git a/Engine/Source/Runtime/Engine/Classes/Engine/Engine.h b/Engine/Source/Runtime/Engine/Classes/Engine/Engine.h
--- a/Engine/Source/Runtime/Engine/Classes/Engine/Engine.h
+++ b/Engine/Source/Runtime/Engine/Classes/Engine/Engine.h
@@ -2257,6 +2257,12 @@ public:
 	// This should only ever be called for a EditorEngine
 	virtual UWorld* CreatePIEWorldByDuplication(FWorldContext &Context, UWorld* InWorld, FString &PlayWorldMapName) { check(false); return nullptr; }
 
+	/** 
+	 *	If this function returns true, the DynamicSourceLevels collection will be duplicated for the given map.
+	 *	This is necessary to do outside of the editor when we don't have the original editor world, and it's 
+	 *	not safe to copy the dynamic levels once they've been fully initialized, so we pre-duplicate them when the original levels are first created.
+	 */
+	virtual bool Experimental_ShouldPreDuplicateMap(const FName MapName) const { return false; }
 
 protected:
 
@@ -2454,7 +2460,7 @@ public:
 	void DestroyNamedNetDriver(UWorld *InWorld, FName NetDriverName);
 	void DestroyNamedNetDriver(UPendingNetGame *PendingNetGame, FName NetDriverName);
 
-	virtual bool NetworkRemapPath( UWorld *InWorld, FString &Str, bool reading=true) { return false; }
+	virtual bool NetworkRemapPath( UNetDriver* Driver, FString &Str, bool reading=true) { return false; }
 	virtual bool NetworkRemapPath( UPendingNetGame *PendingNetGame, FString &Str, bool reading=true) { return false; }
 
 	virtual bool HandleOpenCommand( const TCHAR* Cmd, FOutputDevice& Ar, UWorld * InWorld );
@@ -2665,6 +2671,12 @@ protected:
 	 */
 	virtual void HandleBrowseToDefaultMapFailure(FWorldContext& Context, const FString& TextURL, const FString& Error);
 
+	/**
+	 * Helper function that returns true if InWorld is the outer of a level in a collection of type DynamicDuplicatedLevels.
+	 * For internal engine use.
+	 */
+	bool IsWorldDuplicate(const UWorld* const InWorld);
+
 protected:
 
 	// Async map change/ persistent level transition code.
diff --git a/Engine/Source/Runtime/Engine/Classes/Engine/EngineTypes.h b/Engine/Source/Runtime/Engine/Classes/Engine/EngineTypes.h
--- a/Engine/Source/Runtime/Engine/Classes/Engine/EngineTypes.h
+++ b/Engine/Source/Runtime/Engine/Classes/Engine/EngineTypes.h
@@ -2902,9 +2902,9 @@ struct ENGINE_API FRepMovement
 		return true;
 	}
 
-	void FillFrom(const struct FRigidBodyState& RBState)
+	void FillFrom(const struct FRigidBodyState& RBState, const AActor* const Actor = nullptr)
 	{
-		Location = RBState.Position;
+		Location = RebaseOntoZeroOrigin(RBState.Position, Actor);
 		Rotation = RBState.Quaternion.Rotator();
 		LinearVelocity = RBState.LinVel;
 		AngularVelocity = RBState.AngVel;
@@ -2912,9 +2912,9 @@ struct ENGINE_API FRepMovement
 		bRepPhysics = true;
 	}
 
-	void CopyTo(struct FRigidBodyState& RBState)
+	void CopyTo(struct FRigidBodyState& RBState, const AActor* const Actor = nullptr)
 	{
-		RBState.Position = Location;
+		RBState.Position = RebaseOntoLocalOrigin(Location, Actor);
 		RBState.Quaternion = Rotation.Quaternion();
 		RBState.LinVel = LinearVelocity;
 		RBState.AngVel = AngularVelocity;
@@ -2960,6 +2960,26 @@ struct ENGINE_API FRepMovement
 	{
 		return !(*this == Other);
 	}
+
+	static int32 EnableMultiplayerWorldOriginRebasing;
+
+	/** Rebase zero-origin position onto local world origin value. */
+	static FVector RebaseOntoLocalOrigin(const struct FVector& Location, const struct FIntVector& LocalOrigin);
+
+	/** Rebase local-origin position onto zero world origin value. */
+	static FVector RebaseOntoZeroOrigin(const struct FVector& Location, const struct FIntVector& LocalOrigin);
+
+	/** Rebase zero-origin position onto an Actor's local world origin. */
+	static FVector RebaseOntoLocalOrigin(const struct FVector& Location, const AActor* const WorldContextActor);
+
+	/** Rebase an Actor's local-origin position onto zero world origin value. */
+	static FVector RebaseOntoZeroOrigin(const struct FVector& Location, const AActor* const WorldContextActor);
+
+	/** Rebase zero-origin position onto local world origin value based on an actor component's world. */
+	static FVector RebaseOntoLocalOrigin(const struct FVector& Location, const class UActorComponent* const WorldContextActorComponent);
+
+	/** Rebase local-origin position onto zero world origin value based on an actor component's world.*/
+	static FVector RebaseOntoZeroOrigin(const struct FVector& Location, const class UActorComponent* const WorldContextActorComponent);
 };
 
 
@@ -3719,4 +3739,22 @@ enum class EMeshBufferAccess: uint8
 
     /** Force access on both CPU and GPU. */
     ForceCPUAndGPU
-};
\ No newline at end of file
+};
+
+/** Indicates the type of a level collection, used in FLevelCollection. */
+enum class ELevelCollectionType
+{
+	/**
+	 * The dynamic levels that are used for normal gameplay and the source for any duplicated collections.
+	 * Will contain a world's persistent level and any streaming levels that contain dynamic or replicated gameplay actors.
+	 */
+	DynamicSourceLevels,
+	/** Gameplay relevant levels that have been duplicated from DynamicSourceLevels if requested by the game. */
+	DynamicDuplicatedLevels,
+	/**
+	 * These levels are shared between the source levels and the duplicated levels, and should contain
+	 * only static geometry and other visuals that are not replicated or affected by gameplay.
+	 * These will not be duplicated in order to save memory.
+	 */
+	StaticLevels
+};
diff --git a/Engine/Source/Runtime/Engine/Classes/Engine/GameEngine.h b/Engine/Source/Runtime/Engine/Classes/Engine/GameEngine.h
--- a/Engine/Source/Runtime/Engine/Classes/Engine/GameEngine.h
+++ b/Engine/Source/Runtime/Engine/Classes/Engine/GameEngine.h
@@ -105,7 +105,7 @@ public:
 	virtual float GetMaxTickRate( float DeltaTime, bool bAllowFrameRateSmoothing = true ) const override;
 	virtual void ProcessToggleFreezeCommand( UWorld* InWorld ) override;
 	virtual void ProcessToggleFreezeStreamingCommand( UWorld* InWorld ) override;
-	virtual bool NetworkRemapPath(UWorld* InWorld, FString& Str, bool bReading = true) override;
+	virtual bool NetworkRemapPath(UNetDriver* Driver, FString& Str, bool bReading = true) override;
 
 public:
 
diff --git a/Engine/Source/Runtime/Engine/Classes/Engine/GameInstance.h b/Engine/Source/Runtime/Engine/Classes/Engine/GameInstance.h
--- a/Engine/Source/Runtime/Engine/Classes/Engine/GameInstance.h
+++ b/Engine/Source/Runtime/Engine/Classes/Engine/GameInstance.h
@@ -54,6 +54,9 @@ protected:
 
 	virtual bool HandleOpenCommand(const TCHAR* Cmd, FOutputDevice& Ar, UWorld* InWorld);
 
+	/** Delegate for handling PS4 play together system events */
+	void OnPlayTogetherEventReceived(int32 UserIndex, const TArray<const FUniqueNetId&>& UserList);
+
 	UPROPERTY()
 	TArray<ULocalPlayer*> LocalPlayers;		// List of locally participating players in this game instance
 	
@@ -64,6 +67,9 @@ protected:
 	/** Listeners to PreClientTravel call */
 	FOnPreClientTravel NotifyPreClientTravelDelegates;
 
+	/** Handle for delegate for handling PS4 play together system events */
+	FDelegateHandle OnPlayTogetherEventReceivedDelegateHandle;
+
 public:
 
 	FString PIEMapName;
diff --git a/Engine/Source/Runtime/Engine/Classes/Engine/Level.h b/Engine/Source/Runtime/Engine/Classes/Engine/Level.h
--- a/Engine/Source/Runtime/Engine/Classes/Engine/Level.h
+++ b/Engine/Source/Runtime/Engine/Classes/Engine/Level.h
@@ -12,6 +12,7 @@ class AInstancedFoliageActor;
 class AWorldSettings;
 class UWorld;
 class FSceneInterface;
+struct FLevelCollection;
 
 /**
  * Structure containing all information needed for determining the screen space
@@ -533,6 +534,9 @@ private:
 	UPROPERTY()
 	AWorldSettings* WorldSettings;
 
+	/** Cached level collection that this level is contained in, for faster access than looping through the collections in the world. */
+	FLevelCollection* CachedLevelCollection;
+
 protected:
 
 	/** Array of user data stored with the asset */
@@ -701,6 +705,12 @@ public:
 	 */
 	ENGINE_API class ALevelScriptActor* GetLevelScriptActor() const;
 
+	/** Returns the cached collection that contains this level, if any. May be null. */
+	FLevelCollection* GetCachedLevelCollection() const { return CachedLevelCollection; }
+
+	/** Sets the cached level collection that contains this level. Should only be called by FLevelCollection. */
+	void SetCachedLevelCollection(FLevelCollection* const InCachedLevelCollection) { CachedLevelCollection = InCachedLevelCollection; }
+
 	/**
 	 * Utility searches this level's actor list for any actors of the specified type.
 	 */
@@ -719,6 +729,11 @@ public:
 	 */
 	ENGINE_API class ULevelScriptBlueprint* GetLevelScriptBlueprint(bool bDontCreate=false);
 
+	/**
+	 * Nulls certain references related to the LevelScriptBlueprint. Called by UWorld::CleanupWorld.
+	 */
+	ENGINE_API void CleanupLevelScriptBlueprint();
+
 	/**
 	 *  Returns a list of all blueprints contained within the level
 	 */
diff --git a/Engine/Source/Runtime/Engine/Classes/Engine/LevelStreaming.h b/Engine/Source/Runtime/Engine/Classes/Engine/LevelStreaming.h
--- a/Engine/Source/Runtime/Engine/Classes/Engine/LevelStreaming.h
+++ b/Engine/Source/Runtime/Engine/Classes/Engine/LevelStreaming.h
@@ -126,6 +126,13 @@ public:
 	UPROPERTY(EditAnywhere, Category=LevelStreaming, BlueprintReadWrite)
 	uint32 bShouldBeVisible:1;
 
+	/**
+	 * Whether this level only contains static actors that aren't affected by gameplay or replication.
+	 * If true, the engine can make certain optimizations and will add this level to the StaticLevels collection.
+	 */
+	UPROPERTY(EditDefaultsOnly, Category=LevelStreaming)
+	uint32 bIsStatic:1;
+
 	/** Whether we want to force a blocking load																				*/
 	UPROPERTY(Category=LevelStreaming, BlueprintReadWrite)
 	uint32 bShouldBlockOnLoad:1;
@@ -173,8 +180,8 @@ public:
 
 	//~ Begin UObject Interface
 	virtual void PostLoad() override;
-#if WITH_EDITOR
 	virtual void Serialize( FArchive& Ar ) override;
+#if WITH_EDITOR
 	virtual void PostEditChangeProperty(FPropertyChangedEvent& PropertyChangedEvent) override;
 	
 	/** Remove duplicates in EditorStreamingVolumes list*/
diff --git a/Engine/Source/Runtime/Engine/Classes/Engine/NetDriver.h b/Engine/Source/Runtime/Engine/Classes/Engine/NetDriver.h
--- a/Engine/Source/Runtime/Engine/Classes/Engine/NetDriver.h
+++ b/Engine/Source/Runtime/Engine/Classes/Engine/NetDriver.h
@@ -15,6 +15,7 @@ class FRepLayout;
 class FObjectReplicator;
 class FNetworkObjectList;
 struct FNetworkObjectInfo;
+class FReplicationChangelistMgr;
 
 //
 // Whether to support net lag and packet loss testing.
@@ -334,6 +335,9 @@ public:
 	/** Maps FRepLayout to the respective UClass */
 	TMap< TWeakObjectPtr< UObject >, TSharedPtr< FRepLayout > >					RepLayoutMap;
 
+	/** Maps an object to the respective FReplicationChangelistMgr */
+	TMap< TWeakObjectPtr< UObject >, TSharedPtr< FReplicationChangelistMgr > >	ReplicationChangeListMap;
+
 	/** Creates if necessary, and returns a FRepLayout that maps to the passed in UClass */
 	TSharedPtr< FRepLayout >	GetObjectClassRepLayout( UClass * InClass );
 
@@ -343,6 +347,9 @@ public:
 	/** Creates if necessary, and returns a FRepLayout that maps to the passed in UStruct */
 	TSharedPtr<FRepLayout>		GetStructRepLayout( UStruct * Struct );
 
+	/** Returns the FReplicationChangelistMgr that is associated with the passed in object */
+	TSharedPtr< FReplicationChangelistMgr > GetReplicationChangeListMgr( UObject* Object );
+
 	TMap< FNetworkGUID, TSet< FObjectReplicator* > >	GuidToReplicatorMap;
 	int32												TotalTrackedGuidMemoryBytes;
 	TSet< FObjectReplicator* >							UnmappedReplicators;
@@ -660,6 +667,11 @@ public:
 	/** Stop adaptive replication for the given actor if it's currently throttled. It maybe be allowed to throttle again later. */
 	ENGINE_API void CancelAdaptiveReplication(FNetworkObjectInfo& InNetworkActor);
 
+	/** Returns the level ID/PIE instance ID for this netdriver to use. */
+	ENGINE_API int32 GetDuplicateLevelID() const { return DuplicateLevelID; }
+
+	/** Sets the level ID/PIE instance ID for this netdriver to use. */
+	ENGINE_API void SetDuplicateLevelID(const int32 InDuplicateLevelID) { DuplicateLevelID = InDuplicateLevelID; }
 
 protected:
 
@@ -689,4 +701,7 @@ private:
 
 	/** Set to "Lagging" on the server when all client connections are near timing out. We are lagging on the client when the server connection is near timed out. */
 	ENetworkLagState::Type LagState;
+
+	/** Duplicate level instance to use for playback (PIE instance ID) */
+	int32 DuplicateLevelID;
 };
diff --git a/Engine/Source/Runtime/Engine/Classes/Engine/NetworkSettings.h b/Engine/Source/Runtime/Engine/Classes/Engine/NetworkSettings.h
--- a/Engine/Source/Runtime/Engine/Classes/Engine/NetworkSettings.h
+++ b/Engine/Source/Runtime/Engine/Classes/Engine/NetworkSettings.h
@@ -18,6 +18,11 @@ class ENGINE_API UNetworkSettings : public UDeveloperSettings
 		ToolTip="If true, libcurl authenticates the peer's certificate. Disable to allow self-signed certificates."))
 	uint32 bVerifyPeer:1;
 
+	UPROPERTY(config, EditAnywhere, Category=World, meta = (
+		ConsoleVariable = "p.EnableMultiplayerWorldOriginRebasing", DisplayName = "Enable Multiplayer World Origin Rebasing",
+		ToolTip="If true, origin rebasing is enabled in multiplayer games, meaning that servers and clients can have different local world origins."))
+	uint32 bEnableMultiplayerWorldOriginRebasing : 1;
+
 public:
 
 	//~ Begin UObject Interface
diff --git a/Engine/Source/Runtime/Engine/Classes/Engine/World.h b/Engine/Source/Runtime/Engine/Classes/Engine/World.h
--- a/Engine/Source/Runtime/Engine/Classes/Engine/World.h
+++ b/Engine/Source/Runtime/Engine/Classes/Engine/World.h
@@ -507,6 +507,148 @@ public:
 };
 #endif
 
+/**
+ * Contains a group of levels of a particular ELevelCollectionType within a UWorld
+ * and the context required to properly tick/update those levels. This object is move-only.
+ */
+USTRUCT()
+struct ENGINE_API FLevelCollection
+{
+	GENERATED_BODY()
+
+	FLevelCollection();
+
+	FLevelCollection(const FLevelCollection&) = delete;
+	FLevelCollection& operator=(const FLevelCollection&) = delete;
+
+	FLevelCollection(FLevelCollection&& Other);
+	FLevelCollection& operator=(FLevelCollection&& Other);
+
+	/** The destructor will clear the cached collection pointers in this collection's levels. */
+	~FLevelCollection();
+
+	/** Gets the type of this collection. */
+	ELevelCollectionType GetType() const { return CollectionType; }
+
+	/** Sets the type of this collection. */
+	void SetType(const ELevelCollectionType InType) { CollectionType = InType; }
+
+	/** Gets the game state for this collection. */
+	AGameStateBase* GetGameState() const { return GameState; }
+
+	/** Sets the game state for this collection. */
+	void SetGameState(AGameStateBase* const InGameState) { GameState = InGameState; }
+
+	/** Gets the net driver for this collection. */
+	UNetDriver* GetNetDriver() const { return NetDriver; }
+	
+	/** Sets the net driver for this collection. */
+	void SetNetDriver(UNetDriver* const InNetDriver) { NetDriver = InNetDriver; }
+
+	/** Gets the demo net driver for this collection. */
+	UDemoNetDriver* GetDemoNetDriver() const { return DemoNetDriver; }
+
+	/** Sets the demo net driver for this collection. */
+	void SetDemoNetDriver(UDemoNetDriver* const InDemoNetDriver) { DemoNetDriver = InDemoNetDriver; }
+
+	/** Returns the set of levels in this collection. */
+	const TSet<ULevel*>& GetLevels() const { return Levels; }
+
+	/** Adds a level to this collection and caches the collection pointer on the level for fast access. */
+	void AddLevel(ULevel* const Level);
+
+	/** Removes a level from this collection and clears the cached collection pointer on the level. */
+	void RemoveLevel(ULevel* const Level);
+
+	/** Sets this collection's PersistentLevel and adds it to the Levels set. */
+	void SetPersistentLevel(ULevel* const Level);
+
+	/** Returns this collection's PersistentLevel. */
+	ULevel* GetPersistentLevel() const { return PersistentLevel; }
+
+	/** Gets whether this collection is currently visible. */
+	bool IsVisible() const { return bIsVisible; }
+
+	/** Sets whether this collection is currently visible. */
+	void SetIsVisible(const bool bInIsVisible) { bIsVisible = bInIsVisible; }
+
+private:
+	/** The type of this collection. */
+	ELevelCollectionType CollectionType;
+
+	/**
+	 * The GameState associated with this collection. This may be different than the UWorld's GameState
+	 * since the source collection and the duplicated collection will have their own instances.
+	 */
+	UPROPERTY()
+	class AGameStateBase* GameState;
+
+	/**
+	 * The network driver associated with this collection.
+	 * The source collection and the duplicated collection will have their own instances.
+	 */
+	UPROPERTY()
+	class UNetDriver* NetDriver;
+
+	/**
+	 * The demo network driver associated with this collection.
+	 * The source collection and the duplicated collection will have their own instances.
+	 */
+	UPROPERTY()
+	class UDemoNetDriver* DemoNetDriver;
+
+	/**
+	 * The persistent level associated with this collection.
+	 * The source collection and the duplicated collection will have their own instances.
+	 */
+	UPROPERTY()
+	class ULevel* PersistentLevel;
+
+	/** All the levels in this collection. */
+	UPROPERTY()
+	TSet<ULevel*> Levels;
+
+	/**
+	 * Whether or not this collection is currently visible. While invisible, actors in this collection's
+	 * levels will not be rendered and sounds originating from levels in this collection will not be played.
+	 */
+	bool bIsVisible;
+};
+
+template<>
+struct TStructOpsTypeTraits<FLevelCollection> : public TStructOpsTypeTraitsBase
+{
+	enum
+	{
+		WithCopy = false
+	};
+};
+
+/**
+ * A helper RAII class to set the relevant context on a UWorld for a particular FLevelCollection within a scope.
+ * The constructor will set the PersistentLevel, GameState, NetDriver, and DemoNetDriver on the world and
+ * the destructor will restore the original values.
+ */
+class ENGINE_API FScopedLevelCollectionContextSwitch
+{
+public:
+	/**
+	 * Constructor that will save the current relevant values of InWorld
+	 * and set the collection's context values for InWorld.
+	 *
+	 * @param InLevelCollection The collection's context to use
+	 * @param InWorld The world on which to set the context.
+	 */
+	FScopedLevelCollectionContextSwitch(const FLevelCollection* const InLevelCollection, UWorld* const InWorld);
+	
+	/** The destructor restores the context on the world that was saved in the constructor. */
+	~FScopedLevelCollectionContextSwitch();
+
+private:
+	class UWorld* World;
+	const FLevelCollection* SavedTickingCollection;
+};
+
 /** 
  * The World is the top level object representing a map or a sandbox in which Actors and Components will exist and be rendered.  
  *
@@ -644,6 +786,16 @@ private:
 	UPROPERTY(Transient)
 	TArray<class ULevel*>						Levels;
 
+	/** Array of level collections currently in this world. */
+	UPROPERTY(Transient)
+	TArray<FLevelCollection>					LevelCollections;
+
+	/** Pointer to the level collection that's currently ticking. */
+	const FLevelCollection*						ActiveLevelCollection;
+
+	/** Creates the dynamic source and static level collections if they don't already exist. */
+	void ConditionallyCreateDefaultLevelCollections();
+
 public:
 
 #if WITH_EDITOR
@@ -2309,6 +2461,35 @@ public:
 	 */
 	bool RemoveLevel( ULevel* InLevel );
 
+	/** Returns the FLevelCollection for the given InType. If one does not exist, it is created. */
+	FLevelCollection& FindOrAddCollectionByType(const ELevelCollectionType InType);
+
+	/** Returns the FLevelCollection for the given InType, or null if a collection of that type hasn't been created yet. */
+	FLevelCollection* FindCollectionByType(const ELevelCollectionType InType);
+
+	/** Returns the FLevelCollection for the given InType, or null if a collection of that type hasn't been created yet. */
+	const FLevelCollection* FindCollectionByType(const ELevelCollectionType InType) const;
+
+	/**
+	 * Returns the level collection which currently has its context set on this world. May be null.
+	 * If non-null, this implies that execution is currently within the scope of an FScopedLevelCollectionContextSwitch for this world.
+	 */
+	const FLevelCollection* GetActiveLevelCollection() const { return ActiveLevelCollection; }
+
+	/** Sets the level collection and its context on this world. Should only be called by FScopedLevelCollectionContextSwitch. */
+	void SetActiveLevelCollection(const FLevelCollection* InCollection);
+
+	/** Returns a read-only reference to the list of level collections in this world. */
+	const TArray<FLevelCollection>& GetLevelCollections() const { return LevelCollections; }
+
+	/**
+	 * Creates a new level collection of type DynamicDuplicatedLevels by duplicating the levels in DynamicSourceLevels.
+	 * Should only be called by engine.
+	 *
+	 * @param MapName The name of the soure map, used as a parameter to UEngine::Experimental_ShouldPreDuplicateMap
+	 */
+	void DuplicateRequestedLevels(const FName MapName);
+
 	/** Handle Exec/Console Commands related to the World */
 	bool Exec( UWorld* InWorld, const TCHAR* Cmd, FOutputDevice& Ar=*GLog );
 
@@ -2555,7 +2736,7 @@ public:
 	/** Returns the current GameState instance. */
 	AGameStateBase* GetGameState() const { return GameState; }
 
-	/** Sets the active GameState */
+	/** Sets the current GameState instance on this world and the game state's level collection. */
 	void SetGameState(AGameStateBase* NewGameState);
 
 	/** Copies GameState properties from the GameMode. */
diff --git a/Engine/Source/Runtime/Engine/Classes/GameFramework/Character.h b/Engine/Source/Runtime/Engine/Classes/GameFramework/Character.h
--- a/Engine/Source/Runtime/Engine/Classes/GameFramework/Character.h
+++ b/Engine/Source/Runtime/Engine/Classes/GameFramework/Character.h
@@ -436,7 +436,6 @@ public:
 	virtual void OnRep_ReplicatedMovement() override;
 	virtual void PostNetReceiveLocationAndRotation() override;
 	virtual void GetSimpleCollisionCylinder(float& CollisionRadius, float& CollisionHalfHeight) const override;
-	virtual void ApplyWorldOffset(const FVector& InOffset, bool bWorldShift) override;
 	virtual UActorComponent* FindComponentByClass(const TSubclassOf<UActorComponent> ComponentClass) const override;
 	virtual void TornOff() override;
 	//~ End AActor Interface
diff --git a/Engine/Source/Runtime/Engine/Classes/GameFramework/CharacterMovementComponent.h b/Engine/Source/Runtime/Engine/Classes/GameFramework/CharacterMovementComponent.h
--- a/Engine/Source/Runtime/Engine/Classes/GameFramework/CharacterMovementComponent.h
+++ b/Engine/Source/Runtime/Engine/Classes/GameFramework/CharacterMovementComponent.h
@@ -1009,6 +1009,7 @@ public:
 	virtual void BeginDestroy() override;
 	virtual void PostLoad() override;
 	virtual void RegisterComponentTickFunctions(bool bRegister) override;
+	virtual void ApplyWorldOffset(const FVector& InOffset, bool bWorldShift) override;
 	//End UActorComponent Interface
 
 	//BEGIN UMovementComponent Interface
diff --git a/Engine/Source/Runtime/Engine/Classes/GameFramework/OnlineSession.h b/Engine/Source/Runtime/Engine/Classes/GameFramework/OnlineSession.h
--- a/Engine/Source/Runtime/Engine/Classes/GameFramework/OnlineSession.h
+++ b/Engine/Source/Runtime/Engine/Classes/GameFramework/OnlineSession.h
@@ -33,6 +33,9 @@ public:
 
 	/** Called when a user accepts an invite */
 	virtual void OnSessionUserInviteAccepted(const bool bWasSuccess, const int32 ControllerId, TSharedPtr< const FUniqueNetId > UserId, const FOnlineSessionSearchResult & InviteResult) {};
+
+	/** Called when the play together system event is received on PS4 */
+	virtual void OnPlayTogetherEventReceived(int32 UserIndex, TArray<TSharedPtr<const FUniqueNetId>> UserIdList) {};
 };
 
 
diff --git a/Engine/Source/Runtime/Engine/Classes/Kismet/GameplayStatics.h b/Engine/Source/Runtime/Engine/Classes/Kismet/GameplayStatics.h
--- a/Engine/Source/Runtime/Engine/Classes/Kismet/GameplayStatics.h
+++ b/Engine/Source/Runtime/Engine/Classes/Kismet/GameplayStatics.h
@@ -868,6 +868,14 @@ class ENGINE_API UGameplayStatics : public UBlueprintFunctionLibrary
 	UFUNCTION(BlueprintCallable, Category="Game", meta=(WorldContext="WorldContextObject"))
 	static void SetWorldOriginLocation(const UObject* WorldContextObject, FIntVector NewLocation);
 
+	/** Returns origin based position for local world location. */
+	UFUNCTION(BlueprintPure, Category = "Game", meta = (WorldContext = "WorldContextObject"))
+	static FVector RebaseLocalOriginOntoZero(UObject* WorldContextObject, FVector WorldLocation);
+
+	/** Returns local location for origin based position. */
+	UFUNCTION(BlueprintPure, Category = "Game", meta = (WorldContext = "WorldContextObject"))
+	static FVector RebaseZeroOriginOntoLocal(UObject* WorldContextObject, FVector WorldLocation);
+
 	/**
 	* Counts how many grass foliage instances overlap a given sphere.
 	*
diff --git a/Engine/Source/Runtime/Engine/Classes/Kismet/KismetSystemLibrary.h b/Engine/Source/Runtime/Engine/Classes/Kismet/KismetSystemLibrary.h
--- a/Engine/Source/Runtime/Engine/Classes/Kismet/KismetSystemLibrary.h
+++ b/Engine/Source/Runtime/Engine/Classes/Kismet/KismetSystemLibrary.h
@@ -137,6 +137,10 @@ class ENGINE_API UKismetSystemLibrary : public UBlueprintFunctionLibrary
 	UFUNCTION(BlueprintPure, Category="Networking", meta=(WorldContext="WorldContextObject"))
 	static bool IsDedicatedServer(UObject* WorldContextObject);
 
+	/** Returns whether this game instance is stand alone (no networking). */
+	UFUNCTION(BlueprintPure, Category="Networking", meta=(WorldContext="WorldContextObject"))
+	static bool IsStandalone(UObject* WorldContextObject);
+	
 	/** Returns whether this is a build that is packaged for distribution */
 	UFUNCTION(BlueprintPure, Category="Development")
 	static bool IsPackagedForDistribution();
diff --git a/Engine/Source/Runtime/Engine/Private/ActorReplication.cpp b/Engine/Source/Runtime/Engine/Private/ActorReplication.cpp
--- a/Engine/Source/Runtime/Engine/Private/ActorReplication.cpp
+++ b/Engine/Source/Runtime/Engine/Private/ActorReplication.cpp
@@ -115,7 +115,8 @@ float AActor::GetReplayPriority(const FVector& ViewPos, const FVector& ViewDir,
 		}
 	}
 
-	return Time;
+	// Use NetPriority here to be compatible with live networking.
+	return NetPriority * Time;
 }
 
 bool AActor::GetNetDormancy(const FVector& ViewPos, const FVector& ViewDir, AActor* Viewer, AActor* ViewTarget, UActorChannel* InChannel, float Time, bool bLowBandwidth)
@@ -203,9 +204,11 @@ void AActor::OnRep_ReplicatedMovement()
 
 void AActor::PostNetReceiveLocationAndRotation()
 {
-	if( RootComponent && RootComponent->IsRegistered() && (ReplicatedMovement.Location != GetActorLocation() || ReplicatedMovement.Rotation != GetActorRotation()) )
+	FVector NewLocation = FRepMovement::RebaseOntoLocalOrigin(ReplicatedMovement.Location, this);
+
+	if( RootComponent && RootComponent->IsRegistered() && (NewLocation != GetActorLocation() || ReplicatedMovement.Rotation != GetActorRotation()) )
 	{
-		SetActorLocationAndRotation(ReplicatedMovement.Location, ReplicatedMovement.Rotation, /*bSweep=*/ false);
+		SetActorLocationAndRotation(NewLocation, ReplicatedMovement.Rotation, /*bSweep=*/ false);
 	}
 }
 
@@ -219,7 +222,7 @@ void AActor::PostNetReceivePhysicState()
 	if (RootPrimComp)
 	{
 		FRigidBodyState NewState;
-		ReplicatedMovement.CopyTo(NewState);
+		ReplicatedMovement.CopyTo(NewState, this);
 
 		FVector DeltaPos(FVector::ZeroVector);
 		RootPrimComp->ConditionalApplyRigidBodyState(NewState, GEngine->PhysicErrorCorrection, DeltaPos);
@@ -289,7 +292,7 @@ void AActor::GatherCurrentMovement()
 		FRigidBodyState RBState;
 		RootPrimComp->GetRigidBodyState(RBState);
 
-		ReplicatedMovement.FillFrom(RBState);
+		ReplicatedMovement.FillFrom(RBState, this);
 		ReplicatedMovement.bRepPhysics = true;
 	}
 	else if (RootComponent != nullptr)
@@ -311,7 +314,7 @@ void AActor::GatherCurrentMovement()
 		}
 		else
 		{
-			ReplicatedMovement.Location = RootComponent->GetComponentLocation();
+			ReplicatedMovement.Location = FRepMovement::RebaseOntoZeroOrigin(RootComponent->GetComponentLocation(), this);
 			ReplicatedMovement.Rotation = RootComponent->GetComponentRotation();
 			ReplicatedMovement.LinearVelocity = GetVelocity();
 			ReplicatedMovement.AngularVelocity = FVector::ZeroVector;
diff --git a/Engine/Source/Runtime/Engine/Private/AudioDevice.cpp b/Engine/Source/Runtime/Engine/Private/AudioDevice.cpp
--- a/Engine/Source/Runtime/Engine/Private/AudioDevice.cpp
+++ b/Engine/Source/Runtime/Engine/Private/AudioDevice.cpp
@@ -3794,6 +3794,12 @@ void FAudioDevice::PlaySoundAtLocation(USoundBase* Sound, UWorld* World, float V
 		return;
 	}
 
+	// Not audible if the ticking level collection is not visible
+	if (World && World->GetActiveLevelCollection() && !World->GetActiveLevelCollection()->IsVisible())
+	{
+		return;
+	}
+
 	const FAttenuationSettings* AttenuationSettingsToApply = (AttenuationSettings ? &AttenuationSettings->Attenuation : Sound->GetAttenuationSettingsToApply());
 	float MaxDistance = 0.0f;
 	float FocusFactor = 0.0f;
diff --git a/Engine/Source/Runtime/Engine/Private/BlueprintGeneratedClass.cpp b/Engine/Source/Runtime/Engine/Private/BlueprintGeneratedClass.cpp
--- a/Engine/Source/Runtime/Engine/Private/BlueprintGeneratedClass.cpp
+++ b/Engine/Source/Runtime/Engine/Private/BlueprintGeneratedClass.cpp
@@ -1133,7 +1133,8 @@ void UBlueprintGeneratedClass::GetLifetimeBlueprintReplicationList(TArray<FLifet
 		if (Prop != NULL && Prop->GetPropertyFlags() & CPF_Net)
 		{
 			PropertiesLeft--;
-			OutLifetimeProps.Add(FLifetimeProperty(Prop->RepIndex));
+			
+			OutLifetimeProps.AddUnique(FLifetimeProperty(Prop->RepIndex, Prop->GetBlueprintReplicationCondition()));
 		}
 	}
 
diff --git a/Engine/Source/Runtime/Engine/Private/Character.cpp b/Engine/Source/Runtime/Engine/Private/Character.cpp
--- a/Engine/Source/Runtime/Engine/Private/Character.cpp
+++ b/Engine/Source/Runtime/Engine/Private/Character.cpp
@@ -192,17 +192,6 @@ void ACharacter::UpdateNavigationRelevance()
 	}
 }
 
-void ACharacter::ApplyWorldOffset(const FVector& InOffset, bool bWorldShift)
-{
-	Super::ApplyWorldOffset(InOffset, bWorldShift);
-
-	// Update cached location values in movement component
-	if (CharacterMovement)
-	{
-		CharacterMovement->OldBaseLocation+= InOffset;
-	}
-}
-
 float ACharacter::GetDefaultHalfHeight() const
 {
 	UCapsuleComponent* DefaultCapsule = GetClass()->GetDefaultObject<ACharacter>()->CapsuleComponent;
@@ -1257,7 +1246,8 @@ bool ACharacter::RestoreReplicatedMove(const FSimulatedRootMotionReplicatedMove&
 	// Absolute position
 	else
 	{
-		SetActorLocationAndRotation(RootMotionRepMove.RootMotion.Location, RootMotionRepMove.RootMotion.Rotation);
+		FVector LocalLocation = FRepMovement::RebaseOntoLocalOrigin(RootMotionRepMove.RootMotion.Location, this);
+		SetActorLocationAndRotation(LocalLocation, RootMotionRepMove.RootMotion.Rotation);
 	}
 
 	CharacterMovement->bJustTeleported = true;
@@ -1324,10 +1314,11 @@ void ACharacter::PostNetReceiveLocationAndRotation()
 		if (!ReplicatedBasedMovement.HasRelativeLocation())
 		{
 			const FVector OldLocation = GetActorLocation();
+			const FVector NewLocation = FRepMovement::RebaseOntoLocalOrigin(ReplicatedMovement.Location, this);
 			const FQuat OldRotation = GetActorQuat();
 		
 			CharacterMovement->bNetworkSmoothingComplete = false;
-			CharacterMovement->SmoothCorrection(OldLocation, OldRotation, ReplicatedMovement.Location, ReplicatedMovement.Rotation.Quaternion());
+			CharacterMovement->SmoothCorrection(OldLocation, OldRotation, NewLocation, ReplicatedMovement.Rotation.Quaternion());
 			OnUpdateSimulatedPosition(OldLocation, OldRotation);
 		}
 	}
@@ -1345,7 +1336,7 @@ void ACharacter::PreReplication( IRepChangedPropertyTracker & ChangedPropertyTra
 		// Is position stored in local space?
 		RepRootMotion.bRelativePosition = BasedMovement.HasRelativeLocation();
 		RepRootMotion.bRelativeRotation = BasedMovement.HasRelativeRotation();
-		RepRootMotion.Location			= RepRootMotion.bRelativePosition ? BasedMovement.Location : GetActorLocation();
+		RepRootMotion.Location			= RepRootMotion.bRelativePosition ? BasedMovement.Location : FRepMovement::RebaseOntoZeroOrigin(GetActorLocation(), GetWorld()->OriginLocation);
 		RepRootMotion.Rotation			= RepRootMotion.bRelativeRotation ? BasedMovement.Rotation : GetActorRotation();
 		RepRootMotion.MovementBase		= BasedMovement.MovementBase;
 		RepRootMotion.MovementBaseBoneName = BasedMovement.BoneName;
diff --git a/Engine/Source/Runtime/Engine/Private/Components/CharacterMovementComponent.cpp b/Engine/Source/Runtime/Engine/Private/Components/CharacterMovementComponent.cpp
--- a/Engine/Source/Runtime/Engine/Private/Components/CharacterMovementComponent.cpp
+++ b/Engine/Source/Runtime/Engine/Private/Components/CharacterMovementComponent.cpp
@@ -7923,7 +7923,7 @@ void UCharacterMovementComponent::ServerMoveHandleClientError(float ClientTimeSt
 		ServerData->PendingAdjustment.NewVel = Velocity;
 		ServerData->PendingAdjustment.NewBase = MovementBase;
 		ServerData->PendingAdjustment.NewBaseBoneName = CharacterOwner->GetBasedMovement().BoneName;
-		ServerData->PendingAdjustment.NewLoc = UpdatedComponent->GetComponentLocation();
+		ServerData->PendingAdjustment.NewLoc = FRepMovement::RebaseOntoZeroOrigin(UpdatedComponent->GetComponentLocation(), this);
 		ServerData->PendingAdjustment.NewRot = UpdatedComponent->GetComponentRotation();
 
 		ServerData->PendingAdjustment.bBaseRelativePosition = MovementBaseUtility::UseRelativeLocation(MovementBase);
@@ -8292,31 +8292,36 @@ void UCharacterMovementComponent::ClientAdjustPosition_Implementation
 		return;
 	}
 	ClientData->AckMove(MoveIndex);
-		
+	
+	FVector WorldShiftedNewLocation;
 	//  Received Location is relative to dynamic base
 	if (bBaseRelativePosition)
 	{
 		FVector BaseLocation;
 		FQuat BaseRotation;
-		MovementBaseUtility::GetMovementBaseTransform(NewBase, NewBaseBoneName, BaseLocation, BaseRotation); // TODO: error handling if returns false
-		NewLocation += BaseLocation;
+		MovementBaseUtility::GetMovementBaseTransform(NewBase, NewBaseBoneName, BaseLocation, BaseRotation); // TODO: error handling if returns false		
+		WorldShiftedNewLocation = NewLocation + BaseLocation;
+	}
+	else
+	{
+		WorldShiftedNewLocation = FRepMovement::RebaseOntoLocalOrigin(NewLocation, this);
 	}
 
 #if !UE_BUILD_SHIPPING
 	if (CharacterMovementCVars::NetShowCorrections != 0)
 	{
-		const FVector LocDiff = UpdatedComponent->GetComponentLocation() - NewLocation;
+		const FVector LocDiff = UpdatedComponent->GetComponentLocation() - WorldShiftedNewLocation;
 		const FString NewBaseString = NewBase ? NewBase->GetPathName(NewBase->GetOutermost()) : TEXT("None");
 		UE_LOG(LogNetPlayerMovement, Warning, TEXT("*** Client: Error for %s at Time=%.3f is %3.3f LocDiff(%s) ClientLoc(%s) ServerLoc(%s) NewBase: %s NewBone: %s ClientVel(%s) ServerVel(%s) SavedMoves %d"),
-			*GetNameSafe(CharacterOwner), TimeStamp, LocDiff.Size(), *LocDiff.ToString(), *UpdatedComponent->GetComponentLocation().ToString(), *NewLocation.ToString(), *NewBaseString, *NewBaseBoneName.ToString(), *Velocity.ToString(), *NewVelocity.ToString(), ClientData->SavedMoves.Num());
+			*GetNameSafe(CharacterOwner), TimeStamp, LocDiff.Size(), *LocDiff.ToString(), *UpdatedComponent->GetComponentLocation().ToString(), *WorldShiftedNewLocation.ToString(), *NewBaseString, *NewBaseBoneName.ToString(), *Velocity.ToString(), *NewVelocity.ToString(), ClientData->SavedMoves.Num());
 		const float DebugLifetime = CharacterMovementCVars::NetCorrectionLifetime;
 		DrawDebugCapsule(GetWorld(), UpdatedComponent->GetComponentLocation()	, CharacterOwner->GetSimpleCollisionHalfHeight(), CharacterOwner->GetSimpleCollisionRadius(), FQuat::Identity, FColor(255, 100, 100), true, DebugLifetime);
-		DrawDebugCapsule(GetWorld(), NewLocation						, CharacterOwner->GetSimpleCollisionHalfHeight(), CharacterOwner->GetSimpleCollisionRadius(), FQuat::Identity, FColor(100, 255, 100), true, DebugLifetime);
+		DrawDebugCapsule(GetWorld(), WorldShiftedNewLocation, CharacterOwner->GetSimpleCollisionHalfHeight(), CharacterOwner->GetSimpleCollisionRadius(), FQuat::Identity, FColor(100, 255, 100), true, DebugLifetime);
 	}
 #endif //!UE_BUILD_SHIPPING
 
 	// Trust the server's positioning.
-	UpdatedComponent->SetWorldLocation(NewLocation, false);	
+	UpdatedComponent->SetWorldLocation(WorldShiftedNewLocation, false);
 	Velocity = NewVelocity;
 
 	// Trust the server's movement mode
@@ -8397,7 +8402,7 @@ void UCharacterMovementComponent::ClientAdjustRootMotionPosition_Implementation(
 	// We're going to replay Root Motion. This is relative to the Pawn's rotation, so we need to reset that as well.
 	FRotator DecompressedRot(ServerRotation.X * 180.f, ServerRotation.Y * 180.f, ServerRotation.Z * 180.f);
 	CharacterOwner->SetActorRotation(DecompressedRot);
-	const FVector ServerLocation(ServerLoc);
+	const FVector ServerLocation(FRepMovement::RebaseOntoLocalOrigin(ServerLoc, UpdatedComponent));
 	UE_LOG(LogRootMotion, Log,  TEXT("ClientAdjustRootMotionPosition_Implementation TimeStamp: %f, ServerMontageTrackPosition: %f, ServerLocation: %s, ServerRotation: %s, ServerVelZ: %f, ServerBase: %s"),
 		TimeStamp, ServerMontageTrackPosition, *ServerLocation.ToCompactString(), *DecompressedRot.ToCompactString(), ServerVelZ, *GetNameSafe(ServerBase) );
 
@@ -8463,7 +8468,7 @@ void UCharacterMovementComponent::ClientAdjustRootMotionSourcePosition_Implement
 	// We're going to replay Root Motion. This can be relative to the Pawn's rotation, so we need to reset that as well.
 	FRotator DecompressedRot(ServerRotation.X * 180.f, ServerRotation.Y * 180.f, ServerRotation.Z * 180.f);
 	CharacterOwner->SetActorRotation(DecompressedRot);
-	const FVector ServerLocation(ServerLoc);
+	const FVector ServerLocation(FRepMovement::RebaseOntoLocalOrigin(ServerLoc, UpdatedComponent));
 	UE_LOG(LogRootMotion, Log,  TEXT("ClientAdjustRootMotionSourcePosition_Implementation TimeStamp: %f, NumRootMotionSources: %d, ServerLocation: %s, ServerRotation: %s, ServerVelZ: %f, ServerBase: %s"),
 		TimeStamp, ServerRootMotion.RootMotionSources.Num(), *ServerLocation.ToCompactString(), *DecompressedRot.ToCompactString(), ServerVelZ, *GetNameSafe(ServerBase) );
 
@@ -8829,6 +8834,38 @@ void UCharacterMovementComponent::RegisterComponentTickFunctions(bool bRegister)
 	}
 }
 
+void UCharacterMovementComponent::ApplyWorldOffset(const FVector& InOffset, bool bWorldShift)
+{
+	OldBaseLocation += InOffset;
+	LastUpdateLocation += InOffset;
+
+	if (CharacterOwner != nullptr)
+	{
+		FNetworkPredictionData_Client_Character* ClientData = GetPredictionData_Client_Character();
+		if (ClientData != nullptr)
+		{
+			const int32 NumSavedMoves = ClientData->SavedMoves.Num();
+			for (int32 i = 0; i < NumSavedMoves - 1; i++)
+			{
+				const FSavedMovePtr& CurrentMove = ClientData->SavedMoves[i];
+				CurrentMove->StartLocation += InOffset;
+				CurrentMove->SavedLocation += InOffset;
+			}
+
+			if (ClientData->PendingMove.IsValid())
+			{
+				ClientData->PendingMove->StartLocation += InOffset;
+				ClientData->PendingMove->SavedLocation += InOffset;
+			}
+
+			for (int32 i = 0; i < ClientData->ReplaySamples.Num(); i++)
+			{
+				ClientData->ReplaySamples[i].Location += InOffset;
+			}
+		}
+	}
+}
+
 void UCharacterMovementComponent::TickCharacterPose(float DeltaTime)
 {
 	check(CharacterOwner && CharacterOwner->GetMesh());
diff --git a/Engine/Source/Runtime/Engine/Private/Components/SceneComponent.cpp b/Engine/Source/Runtime/Engine/Private/Components/SceneComponent.cpp
--- a/Engine/Source/Runtime/Engine/Private/Components/SceneComponent.cpp
+++ b/Engine/Source/Runtime/Engine/Private/Components/SceneComponent.cpp
@@ -53,6 +53,7 @@ FName USceneComponent::GetDefaultSceneRootVariableName()
 
 USceneComponent::USceneComponent(const FObjectInitializer& ObjectInitializer /*= FObjectInitializer::Get()*/)
 	: Super(ObjectInitializer)
+	, CachedLevelCollection(nullptr)
 {
 	Mobility = EComponentMobility::Movable;
 	RelativeScale3D = FVector(1.0f, 1.0f, 1.0f);
@@ -544,6 +545,14 @@ void USceneComponent::OnRegister()
 		}
 	}
 	
+	// Cache the level collection that contains the level in which this component is registered for fast access in IsVisible().
+	const UWorld* const World = GetWorld();
+	if (World)
+	{
+		const ULevel* const CachedLevel = GetComponentLevel();
+		CachedLevelCollection = CachedLevel ? CachedLevel->GetCachedLevelCollection() : nullptr;
+	}
+
 	Super::OnRegister();
 
 #if WITH_EDITORONLY_DATA
@@ -568,6 +577,13 @@ void USceneComponent::OnRegister()
 #endif
 }
 
+void USceneComponent::OnUnregister()
+{
+	CachedLevelCollection = nullptr;
+
+	Super::OnUnregister();
+}
+
 void USceneComponent::PropagateTransformUpdate(bool bTransformChanged, EUpdateTransformFlags UpdateTransformFlags, ETeleportType Teleport)
 {
 	//QUICK_SCOPE_CYCLE_COUNTER(STAT_USceneComponent_PropagateTransformUpdate);
@@ -2689,8 +2705,8 @@ bool USceneComponent::IsVisible() const
 	{
 		return false;
 	}
-
-	return ( bVisible ); 
+	
+	return ( bVisible && (!CachedLevelCollection || CachedLevelCollection->IsVisible()) ); 
 }
 
 void USceneComponent::ToggleVisibility(bool bPropagateToChildren)
diff --git a/Engine/Source/Runtime/Engine/Private/DataChannel.cpp b/Engine/Source/Runtime/Engine/Private/DataChannel.cpp
--- a/Engine/Source/Runtime/Engine/Private/DataChannel.cpp
+++ b/Engine/Source/Runtime/Engine/Private/DataChannel.cpp
@@ -1762,6 +1762,20 @@ void UActorChannel::SetChannelActor( AActor* InActor )
 	check(!Closing);
 	check(Actor==NULL);
 
+	// Sanity check that the actor is in the same level collection as the channel's driver.
+	const UWorld* const World = Connection->Driver ? Connection->Driver->GetWorld() : nullptr;
+	if (World && InActor)
+	{
+		const ULevel* const CachedLevel = InActor->GetLevel();
+		const FLevelCollection* const ActorCollection = CachedLevel ? CachedLevel->GetCachedLevelCollection() : nullptr;
+		if (ActorCollection &&
+			ActorCollection->GetNetDriver() != Connection->Driver &&
+			ActorCollection->GetDemoNetDriver() != Connection->Driver)
+		{
+			UE_LOG(LogNet, Verbose, TEXT("UActorChannel::SetChannelActor: actor %s is not in the same level collection as the net driver (%s)!"), *GetFullNameSafe(InActor), *GetFullNameSafe(Connection->Driver));
+		}
+	}
+
 	// Set stuff.
 	Actor = InActor;
 
@@ -2849,6 +2863,9 @@ UObject* UActorChannel::ReadContentBlockHeader( FInBunch & Bunch, bool& bObjectD
 
 		// Track which sub-object guids we are creating
 		CreateSubObjects.AddUnique( SubObj );
+
+		// Add this sub-object to the ImportedNetGuids list so we can possibly map this object if needed
+		Connection->Driver->GuidCache->ImportedNetGuids.Add( NetGUID );
 	}
 
 	return SubObj;
diff --git a/Engine/Source/Runtime/Engine/Private/DataReplication.cpp b/Engine/Source/Runtime/Engine/Private/DataReplication.cpp
--- a/Engine/Source/Runtime/Engine/Private/DataReplication.cpp
+++ b/Engine/Source/Runtime/Engine/Private/DataReplication.cpp
@@ -14,6 +14,8 @@
 #include "Engine/DemoNetDriver.h"
 
 static TAutoConsoleVariable<int32> CVarMaxRPCPerNetUpdate( TEXT( "net.MaxRPCPerNetUpdate" ), 2, TEXT( "Maximum number of RPCs allowed per net update" ) );
+static TAutoConsoleVariable<int32> CVarShareShadowState( TEXT( "net.ShareShadowState" ), 1, TEXT( "If true, work done to compare properties will be shared across connections" ) );
+static TAutoConsoleVariable<float> CVarMaxUpdateDelay( TEXT( "net.MaxSharedShadowStateUpdateDelayInSeconds" ), 1.0f / 4.0f, TEXT( "When a new changelist is available for a particular connection (using shared shadow state), but too much time has passed, force another compare against all the properties" ) );
 
 class FNetSerializeCB : public INetSerializeCB
 {
@@ -50,7 +52,7 @@ public:
 				if ( Ar.IsSaving() )
 				{
 					TArray< uint16 > Changed;
-					RepLayout->SendProperties_BackwardsCompatible( nullptr, (uint8*)Data, PackageMapClient->GetConnection(), static_cast< FNetBitWriter& >( Ar ), Changed );
+					RepLayout->SendProperties_BackwardsCompatible( nullptr, nullptr, (uint8*)Data, PackageMapClient->GetConnection(), static_cast< FNetBitWriter& >( Ar ), Changed );
 				}
 				else
 				{
@@ -290,6 +292,49 @@ void FObjectReplicator::StartReplicating( class UActorChannel * InActorChannel )
 			}
 		}
 	}
+
+	// Prefer the changelist manager on the main net driver (so we share across net drivers if possible)
+	if ( Connection->Driver->GetWorld() && Connection->Driver->GetWorld()->NetDriver )
+	{
+		ChangelistMgr = Connection->Driver->GetWorld()->NetDriver->GetReplicationChangeListMgr( GetObject() );
+	}
+	else
+	{
+		ChangelistMgr = Connection->Driver->GetReplicationChangeListMgr( GetObject() );
+	}
+}
+
+FReplicationChangelistMgr::FReplicationChangelistMgr( UNetDriver* InDriver, UObject* InObject ) : Driver( InDriver ), LastReplicationFrame( 0 )
+{
+	RepLayout = Driver->GetObjectClassRepLayout( InObject->GetClass() );
+
+	RepChangelistState = TUniquePtr< FRepChangelistState >( new FRepChangelistState );
+
+	RepLayout->InitShadowData( RepChangelistState->StaticBuffer, InObject->GetClass(), (uint8*)InObject->GetArchetype() );
+	RepChangelistState->RepLayout = RepLayout;
+}
+
+FReplicationChangelistMgr::~FReplicationChangelistMgr()
+{
+}
+
+void FReplicationChangelistMgr::Update( const UObject* InObject, const uint32 ReplicationFrame, const int32 LastCompareIndex, const FReplicationFlags& RepFlags )
+{
+	// See if we can re-use the work already done on a previous connection
+	// Rules:
+	//	1. We always compare once per frame (i.e. check LastReplicationFrame == ReplicationFrame)
+	//	2. We check LastCompareIndex > 1 so we can do at least one pass per connection to compare all properties
+	//		This is necessary due to how RemoteRole is manipulated per connection, so we need to give all connections a chance to see if it changed
+	//	3. We ALWAYS compare on bNetInitial to make sure we have a fresh changelist of net initial properties in this case
+	if ( CVarShareShadowState.GetValueOnGameThread() && !RepFlags.bNetInitial && LastCompareIndex > 1 && LastReplicationFrame == ReplicationFrame )
+	{
+		INC_DWORD_STAT_BY( STAT_NetSkippedDynamicProps, 1 );
+		return;
+	}
+
+	RepLayout->CompareProperties( RepChangelistState.Get(), (const uint8*)InObject, RepFlags );
+
+	LastReplicationFrame = ReplicationFrame;
 }
 
 static FORCEINLINE void ValidateRetirementHistory( const FPropertyRetirement & Retire, const UObject* Object )
@@ -1082,8 +1127,11 @@ bool FObjectReplicator::ReplicateProperties( FOutBunch & Bunch, FReplicationFlag
 
 	FNetBitWriter Writer( Bunch.PackageMap, 0 );
 
+	// Update change list (this will re-use work done by previous connections)
+	ChangelistMgr->Update( Object, Connection->Driver->ReplicationFrame, RepState->LastCompareIndex, RepFlags );
+
 	// Replicate properties in the layout
-	const bool bHasRepLayout = RepLayout->ReplicateProperties( RepState, ( uint8* )Object, ObjectClass, OwningChannel, Writer, RepFlags );
+	const bool bHasRepLayout = RepLayout->ReplicateProperties( RepState, ChangelistMgr->GetRepChangelistState(), ( uint8* )Object, ObjectClass, OwningChannel, Writer, RepFlags );
 
 	// Replicate all the custom delta properties (fast arrays, etc)
 	ReplicateCustomDeltaProperties( Writer, RepFlags );
diff --git a/Engine/Source/Runtime/Engine/Private/DemoNetDriver.cpp b/Engine/Source/Runtime/Engine/Private/DemoNetDriver.cpp
--- a/Engine/Source/Runtime/Engine/Private/DemoNetDriver.cpp
+++ b/Engine/Source/Runtime/Engine/Private/DemoNetDriver.cpp
@@ -30,6 +30,7 @@
 #include "GameFramework/GameStateBase.h"
 #include "GameFramework/PlayerState.h"
 #include "NetworkVersion.h"
+#include "Engine/Level.h"
 
 DEFINE_LOG_CATEGORY( LogDemo );
 
@@ -507,54 +508,79 @@ bool UDemoNetDriver::InitConnectInternal( FString& Error )
 		bAsyncLoadWorld = FCString::ToBool(AsyncLoadWorldOverrideOption);
 	}
 
-	if ( bAsyncLoadWorld )
+	const TCHAR* const LevelPrefixOverrideOption = DemoURL.GetOption(TEXT("LevelPrefixOverride="), nullptr);
+	if (LevelPrefixOverrideOption)
 	{
-		TArray<AController*> Controllers;
-		for ( FConstControllerIterator Iterator = World->GetControllerIterator(); Iterator; ++Iterator )
+		SetDuplicateLevelID(FCString::Atoi(LevelPrefixOverrideOption));
+	}
+
+	if (GetDuplicateLevelID() == -1)
+	{
+		// Set this driver as the demo net driver for the source level collection.
+		FLevelCollection* const SourceCollection = GetWorld()->FindCollectionByType(ELevelCollectionType::DynamicSourceLevels);
+		if (SourceCollection)
 		{
-			Controllers.Add( *Iterator );
+			SourceCollection->SetDemoNetDriver(this);
 		}
 
-		for ( int i = 0; i < Controllers.Num(); i++ )
+		if ( bAsyncLoadWorld )
 		{
-			if ( Controllers[i] )
+			TArray<AController*> Controllers;
+			for ( FConstControllerIterator Iterator = World->GetControllerIterator(); Iterator; ++Iterator )
 			{
-				Controllers[i]->Destroy();
+				Controllers.Add( *Iterator );
 			}
-		}
 
-		// FIXME: Test for failure!!!
-		World->SeamlessTravel( PlaybackDemoHeader.LevelName, true );
-	}
-	else
-	{
-		// Bypass UDemoPendingNetLevel
-		FString LoadMapError;
+			for ( int i = 0; i < Controllers.Num(); i++ )
+			{
+				if ( Controllers[i] )
+				{
+					Controllers[i]->Destroy();
+				}
+			}
 
-		FURL LocalDemoURL;
-		LocalDemoURL.Map = PlaybackDemoHeader.LevelName;
+			// FIXME: Test for failure!!!
+			World->SeamlessTravel( PlaybackDemoHeader.LevelName, true );
+		}
+		else
+		{
+			// Bypass UDemoPendingNetLevel
+			FString LoadMapError;
 
-		FWorldContext * WorldContext = GEngine->GetWorldContextFromWorld( GetWorld() );
+			FURL LocalDemoURL;
+			LocalDemoURL.Map = PlaybackDemoHeader.LevelName;
 
-		if ( WorldContext == NULL )
-		{
-			Error = FString::Printf( TEXT( "No world context" ) );
-			UE_LOG( LogDemo, Error, TEXT( "UDemoNetDriver::InitConnect: %s" ), *Error );
-			GameInstance->HandleDemoPlaybackFailure( EDemoPlayFailure::Generic, FString( TEXT( "No world context" ) ) );
-			return false;
-		}
+			FWorldContext * WorldContext = GEngine->GetWorldContextFromWorld( GetWorld() );
 
-		GetWorld()->DemoNetDriver = NULL;
-		SetWorld( NULL );
+			if ( WorldContext == NULL )
+			{
+				Error = FString::Printf( TEXT( "No world context" ) );
+				UE_LOG( LogDemo, Error, TEXT( "UDemoNetDriver::InitConnect: %s" ), *Error );
+				GameInstance->HandleDemoPlaybackFailure( EDemoPlayFailure::Generic, FString( TEXT( "No world context" ) ) );
+				return false;
+			}
+
+			GetWorld()->DemoNetDriver = NULL;
+			SetWorld( NULL );
 
-		auto NewPendingNetGame = NewObject<UDemoPendingNetGame>();
+			auto NewPendingNetGame = NewObject<UDemoPendingNetGame>();
 
-		// Set up the pending net game so that the engine can call LoadMap on the next tick.
-		NewPendingNetGame->DemoNetDriver = this;
-		NewPendingNetGame->URL = LocalDemoURL;
-		NewPendingNetGame->bSuccessfullyConnected = true;
+			// Set up the pending net game so that the engine can call LoadMap on the next tick.
+			NewPendingNetGame->DemoNetDriver = this;
+			NewPendingNetGame->URL = LocalDemoURL;
+			NewPendingNetGame->bSuccessfullyConnected = true;
 
-		WorldContext->PendingNetGame = NewPendingNetGame;
+			WorldContext->PendingNetGame = NewPendingNetGame;
+		}
+	}
+	else
+	{
+		// Set this driver as the demo net driver for the duplicate level collection.
+		FLevelCollection* const DuplicateCollection = GetWorld()->FindCollectionByType(ELevelCollectionType::DynamicDuplicatedLevels);
+		if (DuplicateCollection)
+		{
+			DuplicateCollection->SetDemoNetDriver(this);
+		}
 	}
 
 	return true;
@@ -648,6 +674,22 @@ bool UDemoNetDriver::IsPlaying() const
 
 void UDemoNetDriver::TickFlush( float DeltaSeconds )
 {
+	// Set the context on the world for this driver's level collection.
+	const FLevelCollection* FoundCollection = nullptr;
+	if (GetWorld())
+	{
+		for (const FLevelCollection& LC : GetWorld()->GetLevelCollections())
+		{
+			if (LC.GetDemoNetDriver() == this)
+			{
+				FoundCollection = &LC;
+				break;
+			}
+		}
+	}
+
+	FScopedLevelCollectionContextSwitch LCSwitch(FoundCollection, GetWorld());
+
 	Super::TickFlush( DeltaSeconds );
 
 	if ( !IsRecording() || bPauseRecording )
@@ -729,6 +771,22 @@ static float GetClampedDeltaSeconds( UWorld* World, const float DeltaSeconds )
 
 void UDemoNetDriver::TickDispatch( float DeltaSeconds )
 {
+	// Set the context on the world for this driver's level collection.
+	const FLevelCollection* FoundCollection = nullptr;
+	if (GetWorld())
+	{
+		for (const FLevelCollection& LC : GetWorld()->GetLevelCollections())
+		{
+			if (LC.GetDemoNetDriver() == this)
+			{
+				FoundCollection = &LC;
+				break;
+			}
+		}
+	}
+
+	FScopedLevelCollectionContextSwitch LCSwitch(FoundCollection, GetWorld());
+
 	Super::TickDispatch( DeltaSeconds );
 
 	if ( !IsPlaying() )
@@ -2283,6 +2341,19 @@ void UDemoNetDriver::LoadCheckpoint( FArchive* GotoCheckpointArchive, int64 Goto
 
 	bIsLoadingCheckpoint = true;
 
+	struct FPreservedNetworkGUIDEntry
+	{
+		FPreservedNetworkGUIDEntry( const FNetworkGUID InNetGUID, const AActor* const InActor )
+			: NetGUID( InNetGUID ), Actor( InActor ) {}
+
+		FNetworkGUID NetGUID;
+		const AActor* Actor;
+	};
+
+	// Store GUIDs for the spectator controller and any of its owned actors, so we can find them when we process the checkpoint.
+	// For the spectator controller, this allows the state and position to persist.
+	TArray<FPreservedNetworkGUIDEntry> NetGUIDsToPreserve;
+
 #if 1
 	// Destroy all non startup actors. They will get restored with the checkpoint
 	for ( FActorIterator It( GetWorld() ); It; ++It )
@@ -2299,13 +2370,17 @@ void UDemoNetDriver::LoadCheckpoint( FArchive* GotoCheckpointArchive, int64 Goto
 		
 		if ( SpectatorController != nullptr )
 		{
-			if ( *It == SpectatorController || *It == SpectatorController->GetSpectatorPawn() )
-			{
-				continue;
-			}
-
-			if ( It->GetOwner() == SpectatorController )
+			if ( *It == SpectatorController || *It == SpectatorController->GetSpectatorPawn() || It->GetOwner() == SpectatorController )
 			{
+				// If an non-startup actor that we don't destroy has an entry in the GuidCache, preserve that entry so
+				// that the object will be re-used after loading the checkpoint. Otherwise, a new copy
+				// of the object will be created each time a checkpoint is loaded, causing a leak.
+				const FNetworkGUID FoundGUID = GuidCache->NetGUIDLookup.FindRef( *It );
+				
+				if ( FoundGUID.IsValid() )
+				{
+					NetGUIDsToPreserve.Emplace( FoundGUID, *It );
+				}
 				continue;
 			}
 		}
@@ -2381,18 +2456,15 @@ void UDemoNetDriver::LoadCheckpoint( FArchive* GotoCheckpointArchive, int64 Goto
 	// Create fake control channel
 	ServerConnection->CreateChannel( CHTYPE_Control, 1 );
 
-	// Remember the spectator network guid, so we can persist the spectator player across the checkpoint
-	// (so the state and position persists)
-	FNetworkGUID SpectatorGUID = GuidCache->NetGUIDLookup.FindRef( SpectatorController );
-
-	// Handle a rare case where the spectator controller is null, but a valid GUID is
+	// Catch a rare case where the spectator controller is null, but a valid GUID is
 	// found on the GuidCache. The weak pointers in the NetGUIDLookup map are probably
 	// going null, and we want catch these cases and investigate further.
-	if ( !ensure( SpectatorGUID.IsValid() == ( SpectatorController != nullptr ) ) )
+	if (!ensure( GuidCache->NetGUIDLookup.FindRef( SpectatorController ).IsValid() == ( SpectatorController != nullptr ) ))
 	{
-		SpectatorGUID.Reset();
+		UE_LOG(LogDemo, Log, TEXT("LoadCheckpoint: SpectatorController is null and a valid GUID for null was found in the GuidCache. SpectatorController = %s"),
+			*GetFullNameSafe(SpectatorController));
 	}
-
+	
 	// Clean package map to prepare to restore it to the checkpoint state
 	FlushAsyncLoading();
 	GuidCache->ObjectLookup.Empty();
@@ -2402,15 +2474,17 @@ void UDemoNetDriver::LoadCheckpoint( FArchive* GotoCheckpointArchive, int64 Goto
 	GuidCache->NetFieldExportGroupPathToIndex.Empty();
 	GuidCache->NetFieldExportGroupIndexToPath.Empty();
 
-	// Restore the spectator controller packagemap entry (so we find it when we process the checkpoint)
-	if ( SpectatorGUID.IsValid() )
+	// Restore preserved packagemap entries
+	for ( const FPreservedNetworkGUIDEntry& PreservedEntry : NetGUIDsToPreserve )
 	{
-		FNetGuidCacheObject& CacheObject = GuidCache->ObjectLookup.FindOrAdd( SpectatorGUID );
+		check( PreservedEntry.NetGUID.IsValid() );
+		
+		FNetGuidCacheObject& CacheObject = GuidCache->ObjectLookup.FindOrAdd( PreservedEntry.NetGUID );
 
-		CacheObject.Object = SpectatorController;
+		CacheObject.Object = PreservedEntry.Actor;
 		check( CacheObject.Object != NULL );
 		CacheObject.bNoLoad = true;
-		GuidCache->NetGUIDLookup.Add( SpectatorController, SpectatorGUID );
+		GuidCache->NetGUIDLookup.Add( PreservedEntry.Actor, PreservedEntry.NetGUID );
 	}
 
 	if ( GotoCheckpointArchive->TotalSize() == 0 || GotoCheckpointArchive->TotalSize() == INDEX_NONE )
@@ -2446,7 +2520,7 @@ void UDemoNetDriver::LoadCheckpoint( FArchive* GotoCheckpointArchive, int64 Goto
 		*GotoCheckpointArchive << CacheObject.NetworkChecksum;
 
 		// Remap the pathname to handle client-recorded replays
-		GEngine->NetworkRemapPath(GetWorld(), PathName, true);
+		GEngine->NetworkRemapPath(this, PathName, true);
 
 		CacheObject.PathName = FName( *PathName );
 
@@ -2738,6 +2812,13 @@ TSharedPtr<FObjectReplicator> UDemoNetConnection::CreateReplicatorForNewActorCha
 	if (bIsCheckpointStartupActor && NewReplicator->RepLayout.IsValid() && Object->GetClass())
 	{
 		NewReplicator->RepLayout->DiffProperties(nullptr, Object, Object->GetClass()->GetDefaultObject(), true);
+
+		// Need to swap roles for the startup actor since in the CDO they aren't swapped, and the CDO just
+		// overwrote the actor state.
+		if (Actor->Role == ROLE_Authority)
+		{
+			Actor->SwapRolesForReplay();
+		}
 	}
 
 	return NewReplicator;
diff --git a/Engine/Source/Runtime/Engine/Private/EngineTypes.cpp b/Engine/Source/Runtime/Engine/Private/EngineTypes.cpp
--- a/Engine/Source/Runtime/Engine/Private/EngineTypes.cpp
+++ b/Engine/Source/Runtime/Engine/Private/EngineTypes.cpp
@@ -11,6 +11,16 @@ FAttachmentTransformRules FAttachmentTransformRules::SnapToTargetIncludingScale(
 FDetachmentTransformRules FDetachmentTransformRules::KeepRelativeTransform(EDetachmentRule::KeepRelative, true);
 FDetachmentTransformRules FDetachmentTransformRules::KeepWorldTransform(EDetachmentRule::KeepWorld, true);
 
+/** If true, origin rebasing is enabled in multiplayer games, meaning that servers and clients can have different local world origins. */
+int32 FRepMovement::EnableMultiplayerWorldOriginRebasing = 0;
+
+/** Console variable ref to enable multiplayer world origin rebasing. */
+FAutoConsoleVariableRef CVarEnableMultiplayerWorldOriginRebasing(
+	TEXT("p.EnableMultiplayerWorldOriginRebasing"),
+	FRepMovement::EnableMultiplayerWorldOriginRebasing,
+	TEXT("Enable world origin rebasing for multiplayer, meaning that servers and clients can have different world origin locations."),
+	ECVF_ReadOnly);
+
 void FMeshProxySettings::PostLoadDeprecated()
 {
 	FMeshProxySettings DefaultObject;
@@ -282,3 +292,69 @@ FRepMovement::FRepMovement()
 	, RotationQuantizationLevel(ERotatorQuantization::ByteComponents)
 {
 }
+
+/** Rebase zero-origin position onto local world origin value. */
+FVector FRepMovement::RebaseOntoLocalOrigin(const struct FVector& Location, const struct FIntVector& LocalOrigin)
+{
+	if (CVarEnableMultiplayerWorldOriginRebasing->GetInt() <= 0 || LocalOrigin == FIntVector::ZeroValue)
+	{
+		return Location;
+	}
+
+	return FVector(Location.X - LocalOrigin.X, Location.Y - LocalOrigin.Y, Location.Z - LocalOrigin.Z);
+}
+
+/** Rebase local-origin position onto zero world origin value. */
+FVector FRepMovement::RebaseOntoZeroOrigin(const struct FVector& Location, const struct FIntVector& LocalOrigin)
+{
+	if (CVarEnableMultiplayerWorldOriginRebasing->GetInt() <= 0 || LocalOrigin == FIntVector::ZeroValue)
+	{
+		return Location;
+	}
+
+	return FVector(Location.X + LocalOrigin.X, Location.Y + LocalOrigin.Y, Location.Z + LocalOrigin.Z);
+}
+
+/** Rebase zero-origin position onto local world origin value based on an actor's world. */
+FVector FRepMovement::RebaseOntoLocalOrigin(const struct FVector& Location, const AActor* const WorldContextActor)
+{
+	if (WorldContextActor == nullptr || EnableMultiplayerWorldOriginRebasing <= 0)
+	{
+		return Location;
+	}
+
+	return RebaseOntoLocalOrigin(Location, WorldContextActor->GetWorld()->OriginLocation);
+}
+
+/** Rebase local-origin position onto zero world origin value based on an actor's world.*/
+FVector FRepMovement::RebaseOntoZeroOrigin(const struct FVector& Location, const AActor* const WorldContextActor)
+{
+	if (WorldContextActor == nullptr || EnableMultiplayerWorldOriginRebasing <= 0)
+	{
+		return Location;
+	}
+
+	return RebaseOntoZeroOrigin(Location, WorldContextActor->GetWorld()->OriginLocation);
+}
+
+/** Rebase zero-origin position onto local world origin value based on an actor component's world. */
+FVector FRepMovement::RebaseOntoLocalOrigin(const struct FVector& Location, const UActorComponent* const WorldContextActorComponent)
+{
+	if (WorldContextActorComponent == nullptr || EnableMultiplayerWorldOriginRebasing <= 0)
+	{
+		return Location;
+	}
+
+	return RebaseOntoLocalOrigin(Location, WorldContextActorComponent->GetWorld()->OriginLocation);
+}
+
+/** Rebase local-origin position onto zero world origin value based on an actor component's world.*/
+FVector FRepMovement::RebaseOntoZeroOrigin(const struct FVector& Location, const UActorComponent* const WorldContextActorComponent)
+{
+	if (WorldContextActorComponent == nullptr || EnableMultiplayerWorldOriginRebasing <= 0)
+	{
+		return Location;
+	}
+
+	return RebaseOntoZeroOrigin(Location, WorldContextActorComponent->GetWorld()->OriginLocation);
+}
\ No newline at end of file
diff --git a/Engine/Source/Runtime/Engine/Private/GameEngine.cpp b/Engine/Source/Runtime/Engine/Private/GameEngine.cpp
--- a/Engine/Source/Runtime/Engine/Private/GameEngine.cpp
+++ b/Engine/Source/Runtime/Engine/Private/GameEngine.cpp
@@ -41,6 +41,7 @@
 #include "GameDelegates.h"
 #include "Engine/CoreSettings.h"
 #include "EngineAnalytics.h"
+#include "Engine/DemoNetDriver.h"
 
 #include "Tickable.h"
 
@@ -669,11 +670,42 @@ void UGameEngine::FinishDestroy()
 	Super::FinishDestroy();
 }
 
-bool UGameEngine::NetworkRemapPath(UWorld* InWorld, FString& Str, bool bReading /*= true*/)
+bool UGameEngine::NetworkRemapPath(UNetDriver* Driver, FString& Str, bool bReading /*= true*/)
 {
+	if (Driver == nullptr)
+	{
+		return false;
+	}
+
+	UWorld* const World = Driver->GetWorld();
+
+	// If the driver is using a duplicate level ID, find the level collection using the driver
+	// and see if any of its levels match the prefixed name. If so, remap Str to that level's
+	// prefixed name.
+	if (Driver->GetDuplicateLevelID() != INDEX_NONE && bReading)
+	{
+		const FName PrefixedName = *UWorld::ConvertToPIEPackageName(Str, Driver->GetDuplicateLevelID());
+
+		for (const FLevelCollection& Collection : World->GetLevelCollections())
+		{
+			if (Collection.GetNetDriver() == Driver || Collection.GetDemoNetDriver() == Driver)
+			{
+				for (const ULevel* Level : Collection.GetLevels())
+				{
+					const UPackage* const CachedOutermost = Level ? Level->GetOutermost() : nullptr;
+					if (CachedOutermost && CachedOutermost->GetFName() == PrefixedName)
+					{
+						Str = PrefixedName.ToString();
+						return true;
+					}
+				}
+			}
+		}
+	}
+
 	// If the game has created multiple worlds, some of them may have prefixed package names,
 	// so we need to remap the world package and streaming levels for replay playback to work correctly.
-	FWorldContext& Context = GetWorldContextFromWorldChecked(InWorld);
+	FWorldContext& Context = GetWorldContextFromWorldChecked(World);
 	if (Context.PIEInstance == INDEX_NONE || !bReading)
 	{
 		return false;
@@ -682,14 +714,14 @@ bool UGameEngine::NetworkRemapPath(UWorld* InWorld, FString& Str, bool bReading
 	// If the prefixed path matches the world package name or the name of a streaming level,
 	// return the prefixed name.
 	const FString PrefixedName = UWorld::ConvertToPIEPackageName(Str, Context.PIEInstance);
-	const FString WorldPackageName = InWorld->GetOutermost()->GetName();
+	const FString WorldPackageName = World->GetOutermost()->GetName();
 	if (WorldPackageName == PrefixedName)
 	{
 		Str = PrefixedName;
 		return true;
 	}
 
-	for( ULevelStreaming* StreamingLevel : InWorld->StreamingLevels)
+	for( ULevelStreaming* StreamingLevel : World->StreamingLevels)
 	{
 		if (StreamingLevel != nullptr)
 		{
diff --git a/Engine/Source/Runtime/Engine/Private/GameInstance.cpp b/Engine/Source/Runtime/Engine/Private/GameInstance.cpp
--- a/Engine/Source/Runtime/Engine/Private/GameInstance.cpp
+++ b/Engine/Source/Runtime/Engine/Private/GameInstance.cpp
@@ -21,6 +21,7 @@ UGameInstance::UGameInstance(const FObjectInitializer& ObjectInitializer)
 , TimerManager(new FTimerManager())
 , LatentActionManager(new FLatentActionManager())
 {
+	TimerManager->SetGameInstance(this);
 }
 
 void UGameInstance::FinishDestroy()
@@ -785,6 +786,13 @@ void UGameInstance::StartRecordingReplay(const FString& Name, const FString& Fri
 
 	CurrentWorld->DemoNetDriver->SetWorld( CurrentWorld );
 
+	// Set the new demo driver as the current collection's driver
+	FLevelCollection* CurrentLevelCollection = CurrentWorld->FindCollectionByType(ELevelCollectionType::DynamicSourceLevels);
+	if (CurrentLevelCollection)
+	{
+		CurrentLevelCollection->SetDemoNetDriver(CurrentWorld->DemoNetDriver);
+	}
+
 	FString Error;
 
 	if ( !CurrentWorld->DemoNetDriver->InitListen( CurrentWorld, DemoURL, false, Error ) )
diff --git a/Engine/Source/Runtime/Engine/Private/GameplayStatics.cpp b/Engine/Source/Runtime/Engine/Private/GameplayStatics.cpp
--- a/Engine/Source/Runtime/Engine/Private/GameplayStatics.cpp
+++ b/Engine/Source/Runtime/Engine/Private/GameplayStatics.cpp
@@ -2018,6 +2018,16 @@ void UGameplayStatics::SetWorldOriginLocation(const UObject* WorldContextObject,
 	}
 }
 
+FVector UGameplayStatics::RebaseLocalOriginOntoZero(UObject* WorldContextObject, FVector WorldLocation)
+{
+	return FRepMovement::RebaseOntoZeroOrigin(WorldLocation, GetWorldOriginLocation(WorldContextObject));
+}
+
+FVector UGameplayStatics::RebaseZeroOriginOntoLocal(UObject* WorldContextObject, FVector WorldLocation)
+{
+	return FRepMovement::RebaseOntoLocalOrigin(WorldLocation, GetWorldOriginLocation(WorldContextObject));
+}
+
 int32 UGameplayStatics::GrassOverlappingSphereCount(const UObject* WorldContextObject, const UStaticMesh* Mesh, FVector CenterPosition, float Radius)
 {
 	int32 Count = 0;
diff --git a/Engine/Source/Runtime/Engine/Private/KismetSystemLibrary.cpp b/Engine/Source/Runtime/Engine/Private/KismetSystemLibrary.cpp
--- a/Engine/Source/Runtime/Engine/Private/KismetSystemLibrary.cpp
+++ b/Engine/Source/Runtime/Engine/Private/KismetSystemLibrary.cpp
@@ -126,6 +126,12 @@ bool UKismetSystemLibrary::IsDedicatedServer(UObject* WorldContextObject)
 	return IsRunningDedicatedServer();
 }
 
+bool UKismetSystemLibrary::IsStandalone(UObject* WorldContextObject)
+{
+	UWorld* World = GEngine->GetWorldFromContextObject( WorldContextObject );
+	return World ? (World->GetNetMode() == NM_Standalone) : false;
+}
+
 bool UKismetSystemLibrary::IsPackagedForDistribution()
 {
 	return FPlatformMisc::IsPackagedForDistribution();
diff --git a/Engine/Source/Runtime/Engine/Private/Level.cpp b/Engine/Source/Runtime/Engine/Private/Level.cpp
--- a/Engine/Source/Runtime/Engine/Private/Level.cpp
+++ b/Engine/Source/Runtime/Engine/Private/Level.cpp
@@ -657,6 +657,12 @@ void ULevel::BeginDestroy()
 
 	Super::BeginDestroy();
 
+	// Remove this level from its OwningWorld's collection
+	if (CachedLevelCollection)
+	{
+		CachedLevelCollection->RemoveLevel(this);
+	}
+
 	if (OwningWorld && IsPersistentLevel() && OwningWorld->Scene)
 	{
 		OwningWorld->Scene->SetPrecomputedVisibility(NULL);
@@ -1739,6 +1745,22 @@ ULevelScriptBlueprint* ULevel::GetLevelScriptBlueprint(bool bDontCreate)
 	return LevelScriptBlueprint;
 }
 
+void ULevel::CleanupLevelScriptBlueprint()
+{
+	if( LevelScriptBlueprint )
+	{
+		if( LevelScriptBlueprint->SkeletonGeneratedClass )
+		{
+			LevelScriptBlueprint->SkeletonGeneratedClass->ClassGeneratedBy = nullptr; 
+		}
+
+		if( LevelScriptBlueprint->GeneratedClass )
+		{
+			LevelScriptBlueprint->GeneratedClass->ClassGeneratedBy = nullptr; 
+		}
+	}
+}
+
 void ULevel::OnLevelScriptBlueprintChanged(ULevelScriptBlueprint* InBlueprint)
 {
 	if( !InBlueprint->bIsRegeneratingOnLoad )
diff --git a/Engine/Source/Runtime/Engine/Private/LevelStreaming.cpp b/Engine/Source/Runtime/Engine/Private/LevelStreaming.cpp
--- a/Engine/Source/Runtime/Engine/Private/LevelStreaming.cpp
+++ b/Engine/Source/Runtime/Engine/Private/LevelStreaming.cpp
@@ -264,7 +264,6 @@ UWorld* ULevelStreaming::GetWorld() const
 	}
 }
 
-#if WITH_EDITOR
 void ULevelStreaming::Serialize( FArchive& Ar )
 {
 	Super::Serialize(Ar);
@@ -277,7 +276,6 @@ void ULevelStreaming::Serialize( FArchive& Ar )
 		}
 	}
 }
-#endif
 
 FName ULevelStreaming::GetLODPackageName() const
 {
@@ -313,6 +311,13 @@ void ULevelStreaming::SetLoadedLevel(class ULevel* Level)
 
 	// Cancel unloading for this level, in case it was queued for it
 	FLevelStreamingGCHelper::CancelUnloadRequest(LoadedLevel);
+
+	// Add this level to the correct collection
+	const ELevelCollectionType CollectionType =	bIsStatic ? ELevelCollectionType::StaticLevels : ELevelCollectionType::DynamicSourceLevels;
+
+	FLevelCollection& LC = GetWorld()->FindOrAddCollectionByType(CollectionType);
+	LC.RemoveLevel(PendingUnloadLevel);
+	LC.AddLevel(LoadedLevel);
 }
 
 void ULevelStreaming::DiscardPendingUnloadLevel(UWorld* PersistentWorld)
@@ -418,13 +423,12 @@ bool ULevelStreaming::RequestLevel(UWorld* PersistentWorld, bool bAllowLevelLoad
 	// (the world is already loaded for the editor, just find it and copy it)
 	if ( PersistentWorld->IsPlayInEditor() )
 	{
-#if WITH_EDITOR
 		if (PersistentWorld->GetOutermost()->HasAnyPackageFlags(PKG_PlayInEditor))
 		{
 			PackageFlags |= PKG_PlayInEditor;
 		}
 		PIEInstanceID = PersistentWorld->GetOutermost()->PIEInstanceID;
-#endif
+
 		const FString NonPrefixedLevelName = UWorld::StripPIEPrefixFromPackageName(DesiredPackageName.ToString(), PersistentWorld->StreamingLevelsPrefix);
 		UPackage* EditorLevelPackage = FindObjectFast<UPackage>(nullptr, FName(*NonPrefixedLevelName));
 		
@@ -927,6 +931,7 @@ void ULevelStreaming::RemoveStreamingVolumeDuplicates()
 
 ULevelStreaming::ULevelStreaming(const FObjectInitializer& ObjectInitializer)
 	: Super(ObjectInitializer)
+	, bIsStatic(false)
 {
 	bShouldBeVisibleInEditor = true;
 	LevelColor = FLinearColor::White;
diff --git a/Engine/Source/Runtime/Engine/Private/LevelTick.cpp b/Engine/Source/Runtime/Engine/Private/LevelTick.cpp
--- a/Engine/Source/Runtime/Engine/Private/LevelTick.cpp
+++ b/Engine/Source/Runtime/Engine/Private/LevelTick.cpp
@@ -1268,122 +1268,155 @@ void UWorld::Tick( ELevelTick TickType, float DeltaSeconds )
 
 	// Reset the list of objects the LatentActionManager has processed this frame
 	CurrentLatentActionManager.BeginFrame();
-	// If caller wants time update only, or we are paused, skip the rest.
+	
 	if (bDoingActorTicks)
 	{
 		// Reset Async Trace before Tick starts 
-		{
-			SCOPE_CYCLE_COUNTER(STAT_ResetAsyncTraceTickTime);
-			ResetAsyncTrace();
-		}
-		SetupPhysicsTickFunctions(DeltaSeconds);
-		TickGroup = TG_PrePhysics; // reset this to the start tick group
-		FTickTaskManagerInterface::Get().StartFrame(this, DeltaSeconds, TickType);
+		SCOPE_CYCLE_COUNTER(STAT_ResetAsyncTraceTickTime);
+		ResetAsyncTrace();
+	}
 
-		SCOPE_CYCLE_COUNTER(STAT_TickTime);
-		{
-			SCOPE_CYCLE_COUNTER(STAT_TG_PrePhysics);
-			RunTickGroup(TG_PrePhysics);
-		}
-        bInTick = false;
-        EnsureCollisionTreeIsBuilt();
-        bInTick = true;
+	for (FLevelCollection& LC : LevelCollections)
+	{
+		// Build a list of levels from the collection that are also in the world's Levels array.
+		// Collections may contain levels that aren't loaded in the world at the moment.
+		TArray<ULevel*> LevelsToTick;
+		for (ULevel* CollectionLevel : LC.GetLevels())
 		{
-			SCOPE_CYCLE_COUNTER(STAT_TG_StartPhysics);
-			RunTickGroup(TG_StartPhysics); 
+			if (Levels.Contains(CollectionLevel))
+			{
+				LevelsToTick.Add(CollectionLevel);
+			}
 		}
+
+		// Set up context on the world for this level collection
+		FScopedLevelCollectionContextSwitch LevelContext(&LC, this);
+
+		// If caller wants time update only, or we are paused, skip the rest.
+		if (bDoingActorTicks)
 		{
-			SCOPE_CYCLE_COUNTER(STAT_TG_DuringPhysics);
-			RunTickGroup(TG_DuringPhysics, false); // No wait here, we should run until idle though. We don't care if all of the async ticks are done before we start running post-phys stuff
+			// Actually tick actors now that context is set up
+			SetupPhysicsTickFunctions(DeltaSeconds);
+			TickGroup = TG_PrePhysics; // reset this to the start tick group
+			FTickTaskManagerInterface::Get().StartFrame(this, DeltaSeconds, TickType, LevelsToTick);
+
+			SCOPE_CYCLE_COUNTER(STAT_TickTime);
+			{
+				SCOPE_CYCLE_COUNTER(STAT_TG_PrePhysics);
+				RunTickGroup(TG_PrePhysics);
+			}
+			bInTick = false;
+			EnsureCollisionTreeIsBuilt();
+			bInTick = true;
+			{
+				SCOPE_CYCLE_COUNTER(STAT_TG_StartPhysics);
+				RunTickGroup(TG_StartPhysics); 
+			}
+			{
+				SCOPE_CYCLE_COUNTER(STAT_TG_DuringPhysics);
+				RunTickGroup(TG_DuringPhysics, false); // No wait here, we should run until idle though. We don't care if all of the async ticks are done before we start running post-phys stuff
+			}
+			TickGroup = TG_EndPhysics; // set this here so the current tick group is correct during collision notifies, though I am not sure it matters. 'cause of the false up there^^^
+			{
+				SCOPE_CYCLE_COUNTER(STAT_TG_EndPhysics);
+				RunTickGroup(TG_EndPhysics);
+			}
+			{
+				SCOPE_CYCLE_COUNTER(STAT_TG_PostPhysics);
+				RunTickGroup(TG_PostPhysics);
+			}
+	
 		}
-		TickGroup = TG_EndPhysics; // set this here so the current tick group is correct during collision notifies, though I am not sure it matters. 'cause of the false up there^^^
+		else if( bIsPaused )
 		{
-			SCOPE_CYCLE_COUNTER(STAT_TG_EndPhysics);
-			RunTickGroup(TG_EndPhysics);
+			FTickTaskManagerInterface::Get().RunPauseFrame(this, DeltaSeconds, LEVELTICK_PauseTick, LevelsToTick);
 		}
+		
+		// We only want to run the following once, so only run it for the source level collection.
+		if (LC.GetType() == ELevelCollectionType::DynamicSourceLevels)
 		{
-			SCOPE_CYCLE_COUNTER(STAT_TG_PostPhysics);
-			RunTickGroup(TG_PostPhysics);
-		}
-	}
-	else if( bIsPaused )
-	{
-		FTickTaskManagerInterface::Get().RunPauseFrame(this, DeltaSeconds, LEVELTICK_PauseTick);
-	}
-	// Process any remaining latent actions
-	if( !bIsPaused )
-	{
-		// This will process any latent actions that have not been processed already
-		CurrentLatentActionManager.ProcessLatentActions(NULL, DeltaSeconds);
-	}
+			// Process any remaining latent actions
+			if( !bIsPaused )
+			{
+				// This will process any latent actions that have not been processed already
+				CurrentLatentActionManager.ProcessLatentActions(NULL, DeltaSeconds);
+			}
 #if 0 // if you need to debug physics drawing in editor, use this. If you type pxvis collision, it will work. 
-	else
-	{
-		// Tick our async work (physics, etc.) and tick with no elapsed time for playersonly
-		TickAsyncWork(0.f);
-		// Wait for async work to come back
-		WaitForAsyncWork();
-	}
+			else
+			{
+				// Tick our async work (physics, etc.) and tick with no elapsed time for playersonly
+				TickAsyncWork(0.f);
+				// Wait for async work to come back
+				WaitForAsyncWork();
+			}
 #endif
 
-	{
-		SCOPE_CYCLE_COUNTER(STAT_TickableTickTime);
+			{
+				SCOPE_CYCLE_COUNTER(STAT_TickableTickTime);
 
-		if (TickType != LEVELTICK_TimeOnly && !bIsPaused)
-		{
-			STAT(FScopeCycleCounter Context(GetTimerManager().GetStatId());)
-			GetTimerManager().Tick(DeltaSeconds);
-		}
+				if (TickType != LEVELTICK_TimeOnly && !bIsPaused)
+				{
+					STAT(FScopeCycleCounter Context(GetTimerManager().GetStatId());)
+					GetTimerManager().Tick(DeltaSeconds);
+				}
 
-		FTickableGameObject::TickObjects(this, TickType, bIsPaused, DeltaSeconds);
-	}
-	
-	// Update cameras and streaming volumes
-	{
-		SCOPE_CYCLE_COUNTER(STAT_UpdateCameraTime);
-		// Update cameras last. This needs to be done before NetUpdates, and after all actors have been ticked.
-		for( FConstPlayerControllerIterator Iterator = GetPlayerControllerIterator(); Iterator; ++Iterator )
-		{
-			APlayerController* PlayerController = *Iterator;			
-			if (!bIsPaused || PlayerController->ShouldPerformFullTickWhenPaused())
-			{
-				PlayerController->UpdateCameraManager(DeltaSeconds);
+				FTickableGameObject::TickObjects(this, TickType, bIsPaused, DeltaSeconds);
 			}
-		}
 
-		if( !bIsPaused )
-		{
-			// Issues level streaming load/unload requests based on local players being inside/outside level streaming volumes.
-			if (IsGameWorld())
+			// Update cameras and streaming volumes
 			{
-				ProcessLevelStreamingVolumes();
+				SCOPE_CYCLE_COUNTER(STAT_UpdateCameraTime);
+				// Update cameras last. This needs to be done before NetUpdates, and after all actors have been ticked.
+				for( FConstPlayerControllerIterator Iterator = GetPlayerControllerIterator(); Iterator; ++Iterator )
+				{
+					APlayerController* PlayerController = *Iterator;			
+					if (!bIsPaused || PlayerController->ShouldPerformFullTickWhenPaused())
+					{
+						PlayerController->UpdateCameraManager(DeltaSeconds);
+					}
+				}
 
-				if (WorldComposition)
+				if( !bIsPaused )
 				{
-					WorldComposition->UpdateStreamingState();
+					// Issues level streaming load/unload requests based on local players being inside/outside level streaming volumes.
+					if (IsGameWorld())
+					{
+						ProcessLevelStreamingVolumes();
+
+						if (WorldComposition)
+						{
+							WorldComposition->UpdateStreamingState();
+						}
+					}
 				}
 			}
 		}
+
+		if (bDoingActorTicks)
+		{
+			SCOPE_CYCLE_COUNTER(STAT_TickTime);
+			{
+				SCOPE_CYCLE_COUNTER(STAT_TG_PostUpdateWork);
+				RunTickGroup(TG_PostUpdateWork);
+			}
+			{
+				SCOPE_CYCLE_COUNTER(STAT_TG_LastDemotable);
+				RunTickGroup(TG_LastDemotable);
+			}
+
+			FTickTaskManagerInterface::Get().EndFrame(); 
+		}
 	}
 
 	if (bDoingActorTicks)
 	{
 		SCOPE_CYCLE_COUNTER(STAT_TickTime);
-		{
-			SCOPE_CYCLE_COUNTER(STAT_TG_PostUpdateWork);
-			RunTickGroup(TG_PostUpdateWork);
-		}
-		{
-			SCOPE_CYCLE_COUNTER(STAT_TG_LastDemotable);
-			RunTickGroup(TG_LastDemotable);
-		}
+
 		if ( PhysicsScene != NULL )
 		{
 			GPhysCommandHandler->Flush();
 		}
 		
-		FTickTaskManagerInterface::Get().EndFrame(); 
-
 		// All tick is done, execute async trace
 		{
 			SCOPE_CYCLE_COUNTER(STAT_FinishAsyncTraceTickTime);
diff --git a/Engine/Source/Runtime/Engine/Private/Model.cpp b/Engine/Source/Runtime/Engine/Private/Model.cpp
--- a/Engine/Source/Runtime/Engine/Private/Model.cpp
+++ b/Engine/Source/Runtime/Engine/Private/Model.cpp
@@ -879,8 +879,6 @@ int32 UModel::BuildVertexBuffers()
 
 #endif
 
-#if WITH_EDITOR
-
 /**
 * Clears local (non RHI) data associated with MaterialIndexBuffers
 */
@@ -894,8 +892,6 @@ void UModel::ClearLocalMaterialIndexBuffersData()
 	}
 }
 
-#endif // WITH_EDITOR
-
 void UModel::ReleaseVertices()
 {
 	ENQUEUE_UNIQUE_RENDER_COMMAND_ONEPARAMETER(
diff --git a/Engine/Source/Runtime/Engine/Private/NetConnection.cpp b/Engine/Source/Runtime/Engine/Private/NetConnection.cpp
--- a/Engine/Source/Runtime/Engine/Private/NetConnection.cpp
+++ b/Engine/Source/Runtime/Engine/Private/NetConnection.cpp
@@ -1773,9 +1773,9 @@ void UNetConnection::HandleClientPlayer( APlayerController *PC, UNetConnection*
 		break;
 	}
 
-	// Detach old player if same world.
+	// Detach old player if it's in the same level.
 	check(LocalPlayer);
-	if( LocalPlayer->PlayerController && LocalPlayer->PlayerController->GetWorld() == PC->GetWorld() )
+	if( LocalPlayer->PlayerController && LocalPlayer->PlayerController->GetLevel() == PC->GetLevel())
 	{
 		if (LocalPlayer->PlayerController->Role == ROLE_Authority)
 		{
@@ -1820,7 +1820,7 @@ void UNetConnection::HandleClientPlayer( APlayerController *PC, UNetConnection*
 				// Remap packagename for PIE networking before sending out to server
 				FName PackageName = Level->GetOutermost()->GetFName();
 				FString PackageNameStr = PackageName.ToString();
-				if (GEngine->NetworkRemapPath(Driver->GetWorld(), PackageNameStr, false))
+				if (GEngine->NetworkRemapPath(Driver, PackageNameStr, false))
 				{
 					PackageName = FName(*PackageNameStr);
 				}
diff --git a/Engine/Source/Runtime/Engine/Private/NetworkDriver.cpp b/Engine/Source/Runtime/Engine/Private/NetworkDriver.cpp
--- a/Engine/Source/Runtime/Engine/Private/NetworkDriver.cpp
+++ b/Engine/Source/Runtime/Engine/Private/NetworkDriver.cpp
@@ -173,6 +173,7 @@ UNetDriver::UNetDriver(const FObjectInitializer& ObjectInitializer)
 ,	ProcessQueuedBunchesCurrentFrameMilliseconds(0.0f)
 ,	NetworkObjects(new FNetworkObjectList)
 ,	LagState(ENetworkLagState::NotLagging)
+,	DuplicateLevelID(INDEX_NONE)
 {
 }
 
@@ -770,6 +771,14 @@ void UNetDriver::TickFlush(float DeltaSeconds)
 			}
 		}
 
+		for ( auto It = ReplicationChangeListMap.CreateIterator(); It; ++It )
+		{
+			if ( !It.Key().IsValid() )
+			{
+				It.RemoveCurrent();
+			}
+		}
+
 		LastCleanupTime = CurrentRealtimeSeconds;
 	}
 
@@ -3077,6 +3086,8 @@ int32 UNetDriver::ServerReplicateActors(float DeltaSeconds)
 				}
 			}
 			RelevantActorMark.Pop();
+
+			ConnectionViewers.Reset();
 		}
 	}
 
@@ -3404,6 +3415,18 @@ TSharedPtr<FRepLayout> UNetDriver::GetStructRepLayout( UStruct * Struct )
 	return *RepLayoutPtr;
 }
 
+TSharedPtr< FReplicationChangelistMgr > UNetDriver::GetReplicationChangeListMgr( UObject* Object )
+{
+	TSharedPtr< FReplicationChangelistMgr >* ReplicationChangeListMgrPtr = ReplicationChangeListMap.Find( Object );
+
+	if ( !ReplicationChangeListMgrPtr )
+	{
+		ReplicationChangeListMgrPtr = &ReplicationChangeListMap.Add( Object, TSharedPtr< FReplicationChangelistMgr >( new FReplicationChangelistMgr( this, Object ) ) );
+	}
+
+	return *ReplicationChangeListMgrPtr;
+}
+
 FAutoConsoleCommandWithWorld	DumpRelevantActorsCommand(
 	TEXT("net.DumpRelevantActors"), 
 	TEXT( "Dumps information on relevant actors during next network update" ), 
diff --git a/Engine/Source/Runtime/Engine/Private/PackageMapClient.cpp b/Engine/Source/Runtime/Engine/Private/PackageMapClient.cpp
--- a/Engine/Source/Runtime/Engine/Private/PackageMapClient.cpp
+++ b/Engine/Source/Runtime/Engine/Private/PackageMapClient.cpp
@@ -25,6 +25,7 @@ static const int INTERNAL_LOAD_OBJECT_RECURSION_LIMIT = 16;
 
 static TAutoConsoleVariable<int32> CVarAllowAsyncLoading( TEXT( "net.AllowAsyncLoading" ), 0, TEXT( "Allow async loading" ) );
 static TAutoConsoleVariable<int32> CVarIgnoreNetworkChecksumMismatch( TEXT( "net.IgnoreNetworkChecksumMismatch" ), 0, TEXT( "" ) );
+extern FAutoConsoleVariableRef CVarEnableMultiplayerWorldOriginRebasing;
 
 void BroadcastNetFailure(UNetDriver* Driver, ENetworkFailure::Type FailureType, const FString& ErrorStr)
 {
@@ -289,6 +290,7 @@ bool UPackageMapClient::SerializeNewActor(FArchive& Ar, class UActorChannel *Cha
 	{
 		UObject* Archetype = NULL;
 		FVector_NetQuantize10 Location;
+		FVector_NetQuantize10 LocalLocation;
 		FVector_NetQuantize10 Scale;
 		FVector_NetQuantize10 Velocity;
 		FRotator Rotation;
@@ -304,7 +306,15 @@ bool UPackageMapClient::SerializeNewActor(FArchive& Ar, class UActorChannel *Cha
 
 			const USceneComponent* RootComponent = Actor->GetRootComponent();
 
-			Location = RootComponent ? Actor->GetActorLocation() : FVector::ZeroVector;
+			if (RootComponent)
+			{
+				LocalLocation = Actor->GetActorLocation();
+				Location = FRepMovement::RebaseOntoZeroOrigin(Actor->GetActorLocation(), Actor);
+			} 
+			else
+			{
+				Location = LocalLocation = FVector::ZeroVector;
+			}
 			Rotation = RootComponent ? Actor->GetActorRotation() : FRotator::ZeroRotator;
 			Scale = RootComponent ? Actor->GetActorScale() : FVector::ZeroVector;
 			Velocity = RootComponent ? Actor->GetVelocity() : FVector::ZeroVector;
@@ -393,7 +403,7 @@ bool UPackageMapClient::SerializeNewActor(FArchive& Ar, class UActorChannel *Cha
 				uint8* Recent = RepData && RepData->RepState != NULL && RepData->RepState->StaticBuffer.Num() ? RepData->RepState->StaticBuffer.GetData() : NULL;
 				if ( Recent )
 				{
-					((AActor*)Recent)->ReplicatedMovement.Location = Location;
+					((AActor*)Recent)->ReplicatedMovement.Location = LocalLocation;
 					((AActor*)Recent)->ReplicatedMovement.Rotation = Rotation;
 					((AActor*)Recent)->ReplicatedMovement.LinearVelocity = Velocity;
 				}
@@ -412,7 +422,11 @@ bool UPackageMapClient::SerializeNewActor(FArchive& Ar, class UActorChannel *Cha
 					SpawnInfo.SpawnCollisionHandlingOverride = ESpawnActorCollisionHandlingMethod::AlwaysSpawn;
 					SpawnInfo.bRemoteOwned = true;
 					SpawnInfo.bNoFail = true;
-					Actor = Connection->Driver->GetWorld()->SpawnActorAbsolute(Archetype->GetClass(), FTransform(Rotation, Location), SpawnInfo );
+
+					UWorld* World = Connection->Driver->GetWorld();
+					FVector SpawnLocation = FRepMovement::RebaseOntoLocalOrigin( Location, World->OriginLocation );
+					Actor = World->SpawnActorAbsolute(Archetype->GetClass(), FTransform( Rotation, SpawnLocation ), SpawnInfo );
+
 					// Velocity was serialized by the server
 					if (bSerializeVelocity)
 					{
@@ -572,7 +586,7 @@ void UPackageMapClient::InternalWriteObject( FArchive & Ar, FNetworkGUID NetGUID
 
 		InternalWriteObject( Ar, OuterNetGUID, ObjectOuter, TEXT( "" ), NULL );
 
-		GEngine->NetworkRemapPath(Connection->Driver->GetWorld(), ObjectPathName, false);
+		GEngine->NetworkRemapPath(Connection->Driver, ObjectPathName, false);
 
 		// Serialize Name of object
 		Ar << ObjectPathName;
@@ -747,7 +761,7 @@ FNetworkGUID UPackageMapClient::InternalLoadObject( FArchive & Ar, UObject *& Ob
 		}
 
 		// Remap name for PIE
-		GEngine->NetworkRemapPath( Connection->Driver->GetWorld(), PathName, true );
+		GEngine->NetworkRemapPath( Connection->Driver, PathName, true );
 
 		if ( Object != NULL )
 		{
diff --git a/Engine/Source/Runtime/Engine/Private/Pawn.cpp b/Engine/Source/Runtime/Engine/Private/Pawn.cpp
--- a/Engine/Source/Runtime/Engine/Private/Pawn.cpp
+++ b/Engine/Source/Runtime/Engine/Private/Pawn.cpp
@@ -987,7 +987,8 @@ void APawn::PostNetReceiveVelocity(const FVector& NewVelocity)
 void APawn::PostNetReceiveLocationAndRotation()
 {
 	// always consider Location as changed if we were spawned this tick as in that case our replicated Location was set as part of spawning, before PreNetReceive()
-	if( (ReplicatedMovement.Location == GetActorLocation() && ReplicatedMovement.Rotation == GetActorRotation()) && (CreationTime != GetWorld()->TimeSeconds) )
+	if( (FRepMovement::RebaseOntoLocalOrigin(ReplicatedMovement.Location, this) == GetActorLocation() 
+		&& ReplicatedMovement.Rotation == GetActorRotation()) && (CreationTime != GetWorld()->TimeSeconds) )
 	{
 		return;
 	}
@@ -999,12 +1000,13 @@ void APawn::PostNetReceiveLocationAndRotation()
 
 		const FVector OldLocation = GetActorLocation();
 		const FQuat OldRotation = GetActorQuat();
-		SetActorLocationAndRotation(ReplicatedMovement.Location, ReplicatedMovement.Rotation, /*bSweep=*/ false);
+		const FVector NewLocation = FRepMovement::RebaseOntoLocalOrigin(ReplicatedMovement.Location, this);
+		SetActorLocationAndRotation(NewLocation, ReplicatedMovement.Rotation, /*bSweep=*/ false);
 
 		INetworkPredictionInterface* PredictionInterface = Cast<INetworkPredictionInterface>(GetMovementComponent());
 		if (PredictionInterface)
 		{
-			PredictionInterface->SmoothCorrection(OldLocation, OldRotation, ReplicatedMovement.Location, ReplicatedMovement.Rotation.Quaternion());
+			PredictionInterface->SmoothCorrection(OldLocation, OldRotation, NewLocation, ReplicatedMovement.Rotation.Quaternion());
 		}
 	}
 }
diff --git a/Engine/Source/Runtime/Engine/Private/PlayerCameraManager.cpp b/Engine/Source/Runtime/Engine/Private/PlayerCameraManager.cpp
--- a/Engine/Source/Runtime/Engine/Private/PlayerCameraManager.cpp
+++ b/Engine/Source/Runtime/Engine/Private/PlayerCameraManager.cpp
@@ -830,7 +830,8 @@ void APlayerCameraManager::UpdateCamera(float DeltaTime)
 			int32 const ShortPitch = FRotator::CompressAxisToShort(CameraCache.POV.Rotation.Pitch);
 			int32 const CompressedRotation = (ShortYaw << 16) | ShortPitch;
 
-			PCOwner->ServerUpdateCamera(CameraCache.POV.Location, CompressedRotation);
+			FVector ClientCameraPosition = FRepMovement::RebaseOntoZeroOrigin(CameraCache.POV.Location, this);
+			PCOwner->ServerUpdateCamera(ClientCameraPosition, CompressedRotation);
 			bShouldSendClientSideCameraUpdate = false;
 		}
 	}
diff --git a/Engine/Source/Runtime/Engine/Private/PlayerController.cpp b/Engine/Source/Runtime/Engine/Private/PlayerController.cpp
--- a/Engine/Source/Runtime/Engine/Private/PlayerController.cpp
+++ b/Engine/Source/Runtime/Engine/Private/PlayerController.cpp
@@ -192,7 +192,7 @@ void APlayerController::ClientUpdateLevelStreamingStatus_Implementation(FName Pa
 {
 	// For PIE Networking: remap the packagename to our local PIE packagename
 	FString PackageNameStr = PackageName.ToString();
-	if (GEngine->NetworkRemapPath(GetWorld(), PackageNameStr, true))
+	if (GEngine->NetworkRemapPath(GetNetDriver(), PackageNameStr, true))
 	{
 		PackageName = FName(*PackageNameStr);
 	}
@@ -206,8 +206,6 @@ void APlayerController::ClientUpdateLevelStreamingStatus_Implementation(FName Pa
 		}
 	}
 
-	//GEngine->NetworkRemapPath(GetWorld(), RealName, false);
-
 	// if we're about to commit a map change, we assume that the streaming update is based on the to be loaded map and so defer it until that is complete
 	if (GEngine->ShouldCommitPendingMapChange(GetWorld()))
 	{
@@ -1586,7 +1584,7 @@ void APlayerController::ServerUpdateCamera_Implementation(FVector_NetQuantize Ca
 	}
 
 	FPOV NewPOV;
-	NewPOV.Location = CamLoc;
+	NewPOV.Location = FRepMovement::RebaseOntoLocalOrigin(CamLoc, this);
 	
 	NewPOV.Rotation.Yaw = FRotator::DecompressAxisFromShort( (CamPitchAndYaw >> 16) & 65535 );
 	NewPOV.Rotation.Pitch = FRotator::DecompressAxisFromShort(CamPitchAndYaw & 65535);
@@ -3835,9 +3833,9 @@ void APlayerController::SetPlayer( UPlayer* InPlayer )
 {
 	check(InPlayer!=NULL);
 
-	const bool bIsSameWorld = InPlayer->PlayerController && (InPlayer->PlayerController->GetWorld() == GetWorld());
-	// Detach old player if same world.
-	if (bIsSameWorld)
+	const bool bIsSameLevel = InPlayer->PlayerController && (InPlayer->PlayerController->GetLevel() == GetLevel());
+	// Detach old player if it's in the same level.
+	if (bIsSameLevel)
 	{
 		InPlayer->PlayerController->Player = NULL;
 	}
diff --git a/Engine/Source/Runtime/Engine/Private/PrimitiveSceneProxy.cpp b/Engine/Source/Runtime/Engine/Private/PrimitiveSceneProxy.cpp
--- a/Engine/Source/Runtime/Engine/Private/PrimitiveSceneProxy.cpp
+++ b/Engine/Source/Runtime/Engine/Private/PrimitiveSceneProxy.cpp
@@ -34,7 +34,7 @@ FPrimitiveSceneProxy::FPrimitiveSceneProxy(const UPrimitiveComponent* InComponen
 ,	LevelColor(FLinearColor::White)
 ,	PropertyColor(FLinearColor::White)
 ,	Mobility(InComponent->Mobility)
-,	DrawInGame(InComponent->bVisible && !InComponent->bHiddenInGame)
+,	DrawInGame(InComponent->IsVisible())
 ,	DrawInEditor(InComponent->bVisible)
 ,	bReceivesDecals(InComponent->bReceivesDecals)
 ,	bOnlyOwnerSee(InComponent->bOnlyOwnerSee)
diff --git a/Engine/Source/Runtime/Engine/Private/RepLayout.cpp b/Engine/Source/Runtime/Engine/Private/RepLayout.cpp
--- a/Engine/Source/Runtime/Engine/Private/RepLayout.cpp
+++ b/Engine/Source/Runtime/Engine/Private/RepLayout.cpp
@@ -354,30 +354,49 @@ public:
 uint16 FRepLayout::CompareProperties_r(
 	const int32				CmdStart,
 	const int32				CmdEnd,
-	const uint8* RESTRICT	CompareData, 
+	const uint8* RESTRICT	CompareData,
 	const uint8* RESTRICT	Data,
 	TArray< uint16 > &		Changed,
-	uint16					Handle
-) const
+	uint16					Handle,
+	const bool				bIsInitial,
+	const bool				bForceFail ) const
 {
+	check( CompareData );
+
 	for ( int32 CmdIndex = CmdStart; CmdIndex < CmdEnd; CmdIndex++ )
 	{
 		const FRepLayoutCmd& Cmd = Cmds[CmdIndex];
+		const FRepParentCmd& ParentCmd = Parents[Cmd.ParentIndex];
 
 		check( Cmd.Type != REPCMD_Return );
 
 		Handle++;
 
+		const bool bIsLifetime = ( ParentCmd.Flags & PARENT_IsLifetime ) ? true : false;
+		const bool bShouldSkip = !bIsLifetime || ( ParentCmd.Condition == COND_InitialOnly && !bIsInitial );
+
 		if ( Cmd.Type == REPCMD_DynamicArray )
 		{
+			if ( bShouldSkip )
+			{
+				CmdIndex = Cmd.EndCmd - 1;		// The -1 to handle the ++ in the for loop
+				continue;
+			}
+
 			// Once we hit an array, start using a stack based approach
-			CompareProperties_Array_r( CompareData ? CompareData + Cmd.Offset : NULL, (const uint8*)Data + Cmd.Offset, Changed, CmdIndex, Handle );
+			CompareProperties_Array_r( CompareData + Cmd.Offset, ( const uint8* )Data + Cmd.Offset, Changed, CmdIndex, Handle, bIsInitial, bForceFail );
 			CmdIndex = Cmd.EndCmd - 1;		// The -1 to handle the ++ in the for loop
 			continue;
 		}
 
-		if ( CompareData == NULL || !PropertiesAreIdentical( Cmd, (const void*)( CompareData + Cmd.Offset ), (const void*)( Data + Cmd.Offset ) ) )
+		if ( bShouldSkip )
 		{
+			continue;
+		}
+
+		if ( bForceFail || !PropertiesAreIdentical( Cmd, ( const void* )( CompareData + Cmd.Offset ), ( const void* )( Data + Cmd.Offset ) ) )
+		{
+			StoreProperty( Cmd, ( void* )( CompareData + Cmd.Offset ), ( const void* )( Data + Cmd.Offset ) );
 			Changed.Add( Handle );
 		}
 	}
@@ -385,35 +404,42 @@ uint16 FRepLayout::CompareProperties_r(
 	return Handle;
 }
 
-void FRepLayout::CompareProperties_Array_r( 
-	const uint8* RESTRICT	CompareData, 
+void FRepLayout::CompareProperties_Array_r(
+	const uint8* RESTRICT	CompareData,
 	const uint8* RESTRICT	Data,
 	TArray< uint16 > &		Changed,
 	const uint16			CmdIndex,
-	const uint16			Handle
-) const
+	const uint16			Handle,
+	const bool				bIsInitial,
+	const bool				bForceFail
+	) const
 {
 	const FRepLayoutCmd& Cmd = Cmds[CmdIndex];
 
-	FScriptArray * CompareArray = (FScriptArray *)CompareData;
-	FScriptArray * Array = (FScriptArray *)Data;
+	FScriptArray * CompareArray = ( FScriptArray * )CompareData;
+	FScriptArray * Array = ( FScriptArray * )Data;
 
 	const uint16 ArrayNum			= Array->Num();
-	const uint16 CompareArrayNum	= CompareArray ? CompareArray->Num() : 0;
+	const uint16 CompareArrayNum	= CompareArray->Num();
+
+	// Make the shadow state match the actual state at the time of compare
+	FScriptArrayHelper StoredArrayHelper( ( UArrayProperty * )Cmd.Property, CompareData );
+	StoredArrayHelper.Resize( ArrayNum );
 
 	TArray< uint16 > ChangedLocal;
 
 	uint16 LocalHandle = 0;
 
-	Data = (uint8*)Array->GetData();
-	CompareData = CompareData ? (uint8*)CompareArray->GetData() : (uint8*)NULL;
+	Data = ( uint8* )Array->GetData();
+	CompareData = ( uint8* )CompareArray->GetData();
 
 	for ( int32 i = 0; i < ArrayNum; i++ )
 	{
 		const int32 ElementOffset = i * Cmd.ElementSize;
-		const uint8* LocalCompareData = ( i < CompareArrayNum ) ? ( CompareData + ElementOffset ) : NULL;
 
-		LocalHandle = CompareProperties_r( CmdIndex + 1, Cmd.EndCmd - 1, LocalCompareData, Data + ElementOffset, ChangedLocal, LocalHandle );
+		const bool bNewForceFail = bForceFail || i >= CompareArrayNum;
+
+		LocalHandle = CompareProperties_r( CmdIndex + 1, Cmd.EndCmd - 1, CompareData + ElementOffset, Data + ElementOffset, ChangedLocal, LocalHandle, bIsInitial, bNewForceFail );
 	}
 
 	if ( ChangedLocal.Num() )
@@ -435,176 +461,88 @@ void FRepLayout::CompareProperties_Array_r(
 	}
 }
 
-bool FRepLayout::CompareProperties( 
-	FRepState * RESTRICT				RepState, 
-	const uint8* RESTRICT				CompareData,
-	const uint8* RESTRICT				Data, 
-	TArray< FRepChangedParent > &		OutChangedParents,
-	const TArray< uint16 > &			PropertyList ) const
+bool FRepLayout::CompareProperties(
+	FRepChangelistState* RESTRICT	RepChangelistState,
+	const uint8* RESTRICT			Data,
+	const FReplicationFlags&		RepFlags ) const
 {
-	bool PropertyChanged = false;
-
-	const uint16* RESTRICT FirstProp	= PropertyList.GetData();
-	const uint16* RESTRICT LastProp	= FirstProp + PropertyList.Num();
-
-	for ( const uint16* RESTRICT pLifeProp = FirstProp; pLifeProp < LastProp; ++pLifeProp )
-	{
-		const uint16 LifeProp = *pLifeProp;
-		const FRepParentCmd& ParentCmd = Parents[LifeProp];
-
-		// We store changed properties on each parent, so we can build a final sorted change list later
-		TArray< uint16 > & Changed = OutChangedParents[LifeProp].Changed;
+	RepChangelistState->CompareIndex++;
 
-		// some games ship checks() in Shipping so we cannot rely on DO_CHECK here, and this check is in an extremely hot path
-		check( UE_BUILD_SHIPPING || UE_BUILD_TEST || Changed.Num() == 0 );
+	check( RepChangelistState->HistoryEnd - RepChangelistState->HistoryStart < FRepChangelistState::MAX_CHANGE_HISTORY );
+	const int32 HistoryIndex = RepChangelistState->HistoryEnd % FRepChangelistState::MAX_CHANGE_HISTORY;
 
-		// Loop over the block of child properties that are children of this parent
-		for ( int32 i = ParentCmd.CmdStart, ParentCmd_CmdEnd = ParentCmd.CmdEnd; i < ParentCmd_CmdEnd; ++i )
-		{
-			const FRepLayoutCmd& Cmd = Cmds[i];
-			const int32 Cmd_Offset = Cmd.Offset;
+	FRepChangedHistory& NewHistoryItem = RepChangelistState->ChangeHistory[HistoryIndex];
 
-			// some games ship checks() in Shipping so we cannot rely on DO_CHECK here, and this check is in an extremely hot path
-			check( UE_BUILD_SHIPPING || UE_BUILD_TEST || Cmd.Type != REPCMD_Return );		// REPCMD_Return's are markers we shouldn't hit
+	TArray<uint16>& Changed = NewHistoryItem.Changed;
+	Changed.Empty();
 
-			if ( UNLIKELY( Cmd.Type == REPCMD_DynamicArray ) )
-			{
-				// Once we hit an array, start using a recursive based approach
-				CompareProperties_Array_r( CompareData + Cmd_Offset, Data + Cmd_Offset, Changed, i, Cmd.RelativeHandle );
-				i = Cmd.EndCmd - 1;		// Jump past properties under array, we've already checked them (the -1 because of the ++ in the for loop)
-				continue;
-			}
+	CompareProperties_r( 0, Cmds.Num() - 1, RepChangelistState->StaticBuffer.GetData(), Data, Changed, 0, RepFlags.bNetInitial, false );
 
-			if ( !PropertiesAreIdentical( Cmd, (void*)( CompareData + Cmd_Offset ), (const void*)( Data + Cmd_Offset ) ) )
-			{
-				// Add this properties handle to the change list
-				Changed.Add( Cmd.RelativeHandle );
-			}
-		}
-
-		if ( Changed.Num() > 0 )
-		{
-			// Something changed on this parent property
-			PropertyChanged = true;
-		}
+	if ( Changed.Num() == 0 )
+	{
+		return false;
 	}
 
-	return PropertyChanged;
-}
-
-void FRepLayout::LogChangeListMismatches( 
-	const uint8* 						Data, 
-	UActorChannel *						OwningChannel, 
-	const TArray< uint16 > &			PropertyList,
-	FRepState *							RepState1, 
-	FRepState * 						RepState2, 
-	const TArray< FRepChangedParent > & ChangedParents1,
-	const TArray< FRepChangedParent > & ChangedParents2 ) const
-{
-	FRepChangedPropertyTracker *	ChangeTracker	= RepState1->RepChangedPropertyTracker.Get();
-	UObject *						Object			= (UObject*)Data;
-	const UNetDriver *				NetDriver		= OwningChannel->Connection->Driver;
+	//
+	// We produced a new change list, copy it to the history
+	//
 
-	UE_LOG( LogRep, Warning, TEXT( "LogChangeListMismatches: %s" ), *Object->GetName() );
+	// Null terminator
+	Changed.Add( 0 );
 
-	UE_LOG( LogRep, Warning, TEXT( "  %i, %i, %i, %i" ), NetDriver->ReplicationFrame, ChangeTracker->LastReplicationFrame, ChangeTracker->LastReplicationGroupFrame, RepState1->LastReplicationFrame );
+	// Move end pointer
+	RepChangelistState->HistoryEnd++;
 
-	for ( int32 i = 0; i < PropertyList.Num(); i++ )
+	// If we're full, merge the oldest up, so we always have room for a new entry
+	if ( RepChangelistState->HistoryEnd - RepChangelistState->HistoryStart == FRepChangelistState::MAX_CHANGE_HISTORY )
 	{
-		const int32 Index		= PropertyList[i];
-		const int32 Changed1	= ChangedParents1[Index].Changed.Num();
-		const int32 Changed2	= ChangedParents2[Index].Changed.Num();
-
-		if ( Changed1 || Changed2 )
-		{
-			FString Changed1Str = Changed1 ? TEXT( "TRUE" ) : TEXT( "FALSE" );
-			FString Changed2Str = Changed2 ? TEXT( "TRUE" ) : TEXT( "FALSE" );
-			FString ExtraStr	= Changed1 != Changed2 ? TEXT( "<--- MISMATCH!" ) : TEXT( "<--- SAME" );
-
-			UE_LOG( LogRep, Warning, TEXT( "    Property changed: %s (%s / %s) %s" ), *Parents[Index].Property->GetName(), *Changed1Str, *Changed2Str, *ExtraStr );
+		const int32 FirstHistoryIndex = RepChangelistState->HistoryStart % FRepChangelistState::MAX_CHANGE_HISTORY;
 
-			FString StringValue1;
-			FString StringValue2;
-			FString StringValue3;
+		RepChangelistState->HistoryStart++;
 
-			Parents[Index].Property->ExportText_InContainer( Parents[Index].ArrayIndex, StringValue1, Data, NULL, NULL, PPF_DebugDump );
-			Parents[Index].Property->ExportText_InContainer( Parents[Index].ArrayIndex, StringValue2, RepState1->StaticBuffer.GetData(), NULL, NULL, PPF_DebugDump );
-			Parents[Index].Property->ExportText_InContainer( Parents[Index].ArrayIndex, StringValue3, RepState2->StaticBuffer.GetData(), NULL, NULL, PPF_DebugDump );
+		const int32 SecondHistoryIndex = RepChangelistState->HistoryStart % FRepChangelistState::MAX_CHANGE_HISTORY;
 
-			UE_LOG( LogRep, Warning, TEXT( "      Values: %s, %s, %s" ), *StringValue1, *StringValue2, *StringValue3 );
-		}
-	}
-}
+		TArray< uint16 >& FirstChangelistRef = RepChangelistState->ChangeHistory[FirstHistoryIndex].Changed;
+		TArray< uint16 > SecondChangelistCopy = RepChangelistState->ChangeHistory[SecondHistoryIndex].Changed;
 
-static bool ChangedParentsHasChanged( const TArray< uint16 > & PropertyList, const TArray< FRepChangedParent > & ChangedParents )
-{
-	for ( int32 i = 0; i < PropertyList.Num(); i++ )
-	{
-		if ( ChangedParents[PropertyList[i]].Changed.Num() )
-		{
-			return true;
-		}
+		MergeChangeList( Data, FirstChangelistRef, SecondChangelistCopy, RepChangelistState->ChangeHistory[SecondHistoryIndex].Changed );
 	}
 
-	return false;
+	return true;
 }
 
-void FRepLayout::SanityCheckShadowStateAgainstChangeList( 
-	FRepState *							RepState, 
-	const uint8 *						Data, 
-	UActorChannel *						OwningChannel, 
-	const TArray< uint16 > &			PropertyList,
-	FRepState *							OtherRepState,
-	const TArray< FRepChangedParent > & OtherChangedParents ) const
+static FORCEINLINE void WritePropertyHandle( FNetBitWriter & Writer, uint16 Handle, bool bDoChecksum )
 {
-	const uint8 *	CompareData	= RepState->StaticBuffer.GetData();
-	UObject *		Object		= (UObject*)Data;
-
-	TArray< FRepChangedParent > LocalChangedParents;
-	LocalChangedParents.SetNum( OtherChangedParents.Num() );
+	const int NumStartingBits = Writer.GetNumBits();
 
-	// Do an actual check to see which properties are different from our internal shadow state, and make sure the passed in change list matches
-	const bool Result = CompareProperties( RepState, CompareData, Data, LocalChangedParents, PropertyList );
+	uint32 LocalHandle = Handle;
+	Writer.SerializeIntPacked( LocalHandle );
 
-	if ( Result != ChangedParentsHasChanged( PropertyList, OtherChangedParents ) )
+#ifdef ENABLE_PROPERTY_CHECKSUMS
+	if ( bDoChecksum )
 	{
-		LogChangeListMismatches( Data, OwningChannel, PropertyList, RepState, OtherRepState, LocalChangedParents, OtherChangedParents );
-		UE_LOG( LogRep, Fatal, TEXT( "ReplicateProperties: Compare result mismatch: %s" ), *Object->GetName() );
+		SerializeGenericChecksum( Writer );
 	}
+#endif
 
-	for ( int32 i = 0; i < PropertyList.Num(); i++ )
-	{
-		const int32 Index = PropertyList[i];
-
-		if ( OtherChangedParents[Index].Changed.Num() != LocalChangedParents[Index].Changed.Num() )
-		{
-			LogChangeListMismatches( Data, OwningChannel, PropertyList, RepState, OtherRepState, LocalChangedParents, OtherChangedParents );
-			UE_LOG( LogRep, Fatal, TEXT( "ReplicateProperties: Compare count mismatch: %s" ), *Object->GetName() );
-		}
-
-		for ( int32 j = 0; j < OtherChangedParents[Index].Changed.Num(); j++ )
-		{
-			if ( OtherChangedParents[Index].Changed[j] != LocalChangedParents[Index].Changed[j] )
-			{
-				LogChangeListMismatches( Data, OwningChannel, PropertyList, RepState, OtherRepState, LocalChangedParents, OtherChangedParents );
-				UE_LOG( LogRep, Fatal, TEXT( "ReplicateProperties: Compare changelist value mismatch: %s" ), *Object->GetName() );
-			}
-		}
-	}
+	NETWORK_PROFILER( GNetworkProfiler.TrackWritePropertyHandle( Writer.GetNumBits() - NumStartingBits, nullptr ) );
 }
 
-bool FRepLayout::ReplicateProperties( 
-	FRepState * RESTRICT		RepState, 
-	const uint8* RESTRICT		Data, 
-	UClass *					ObjectClass,
-	UActorChannel *				OwningChannel,
-	FNetBitWriter&				Writer,
-	const FReplicationFlags &	RepFlags ) const
+bool FRepLayout::ReplicateProperties(
+	FRepState* RESTRICT				RepState,
+	FRepChangelistState* RESTRICT	RepChangelistState,
+	const uint8* RESTRICT			Data,
+	UClass*							ObjectClass,
+	UActorChannel*					OwningChannel,
+	FNetBitWriter&					Writer,
+	const FReplicationFlags &		RepFlags ) const
 {
 	SCOPE_CYCLE_COUNTER( STAT_NetReplicateDynamicPropTime );
 
 	check( ObjectClass == Owner );
 
+	FRepChangedPropertyTracker*	ChangeTracker = RepState->RepChangedPropertyTracker.Get();
+
 	if ( OwningChannel->Connection->bResendAllDataSinceOpen )
 	{
 		check( OwningChannel->Connection->InternalAck );
@@ -615,144 +553,79 @@ bool FRepLayout::ReplicateProperties(
 			TArray< uint16 > Pruned;
 			PruneChangeList( RepState, Data, RepState->LifetimeChangelist, Pruned );
 			RepState->LifetimeChangelist = MoveTemp( Pruned );
-			SendProperties_BackwardsCompatible( nullptr, Data, OwningChannel->Connection, Writer, RepState->LifetimeChangelist );
+			SendProperties_BackwardsCompatible( RepState, ChangeTracker, Data, OwningChannel->Connection, Writer, RepState->LifetimeChangelist );
 			return true;
 		}
 
 		return false;
 	}
 
-	UObject *						Object			= (UObject*)Data;
-	const UNetDriver *				NetDriver		= OwningChannel->Connection->Driver;
-	FRepChangedPropertyTracker *	ChangeTracker	= RepState->RepChangedPropertyTracker.Get();
-	const uint8 *					CompareData		= RepState->StaticBuffer.GetData();
-
-	// Rebuild conditional properties if needed
-	if ( RepState->RepFlags.Value != RepFlags.Value || RepState->ActiveStatusChanged != ChangeTracker->ActiveStatusChanged )
-	{
-		RebuildConditionalProperties( RepState, *ChangeTracker, RepFlags );
+	check( RepState->HistoryEnd >= RepState->HistoryStart );
+	check( RepState->HistoryEnd - RepState->HistoryStart < FRepState::MAX_CHANGE_HISTORY );
 
-		RepState->RepFlags.Value		= RepFlags.Value;
-		RepState->ActiveStatusChanged	= ChangeTracker->ActiveStatusChanged;
-	}
+	const bool bFlushPreOpenAckHistory = RepState->OpenAckedCalled && RepState->PreOpenAckHistory.Num() > 0;
 
-	bool PropertyChanged = false;
+	const bool bCompareIndexSame = RepState->LastCompareIndex == RepChangelistState->CompareIndex;
 
-#ifdef ENABLE_SUPER_CHECKSUMS
-	const bool bIsAllAcked = AllAcked( RepState );
+	RepState->LastCompareIndex = RepChangelistState->CompareIndex;
 
-	if ( bIsAllAcked || !RepState->OpenAckedCalled )
-#endif
+	// We can early out if we know for sure there are no new changelists to send
+	if ( bCompareIndexSame || RepState->LastChangelistIndex == RepChangelistState->HistoryEnd )
 	{
-		const int32	AllowSkipping = CVarAllowPropertySkipping.GetValueOnGameThread();
-		
-		const bool bCanSkip =	AllowSkipping > 0 && 
-								RepState->LastReplicationFrame != 0 &&
-								ChangeTracker->LastReplicationFrame == NetDriver->ReplicationFrame &&
-								ChangeTracker->LastReplicationGroupFrame == RepState->LastReplicationFrame;
-
-		if ( bCanSkip )
+		if ( RepState->NumNaks == 0 && !bFlushPreOpenAckHistory )
 		{
-			INC_DWORD_STAT_BY( STAT_NetSkippedDynamicProps, UnconditionalLifetime.Num() );
-
-			if ( AllowSkipping == 2 )
-			{
-				// Sanity check results
-				check( ChangeTracker->UnconditionalPropChanged == ChangedParentsHasChanged( UnconditionalLifetime, ChangeTracker->Parents ) );
-
-				SanityCheckShadowStateAgainstChangeList( RepState, Data, OwningChannel, UnconditionalLifetime, ChangeTracker->LastRepState, ChangeTracker->Parents );
-			}
+			// Nothing changed and there are no nak's, so just do normal housekeeping and remove acked history items
+			UpdateChangelistHistory( RepState, ObjectClass, Data, OwningChannel->Connection->OutAckPacketId, NULL );
+			return false;
 		}
-		else
-		{
-			// FRepState group changed, force this group to compare again this frame
-			// This happens either once a frame, which is normal, or multiple times a frame 
-			// when multiple connections of the same actor aren't updated at the same time
-			ChangeTracker->LastReplicationFrame			= NetDriver->ReplicationFrame;
-			ChangeTracker->LastReplicationGroupFrame	= RepState->LastReplicationFrame;
-			ChangeTracker->LastRepState					= RepState;
-
-			// Reset changed list if anything changed last time
-			if ( ChangeTracker->UnconditionalPropChanged )
-			{
-				for ( int32 i = UnconditionalLifetime.Num() - 1; i >= 0; i-- )
-				{
-					ChangeTracker->Parents[UnconditionalLifetime[i]].Changed.Empty();
-				}
-			}
+	}
 
-			// Loop over all unconditional lifetime properties
-			ChangeTracker->UnconditionalPropChanged = CompareProperties( RepState, CompareData, Data, ChangeTracker->Parents, UnconditionalLifetime );
-		}
+	// Rebuild conditional state if needed
+	if ( RepState->RepFlags.Value != RepFlags.Value )
+	{
+		RebuildConditionalProperties( RepState, *ChangeTracker, RepFlags );
 
-		// Remember the last frame this FRepState was replicated, so we can note above when the FRepState replication group changes
-		RepState->LastReplicationFrame = NetDriver->ReplicationFrame;
+		RepState->RepFlags.Value = RepFlags.Value;
+	}
 
-		if ( ChangeTracker->UnconditionalPropChanged )
+	// Clamp to the valid history range (and warn if we end up sending entire history, this should only happen if we get really far behind)
+	//	NOTE - The RepState->LastChangelistIndex != 0 should handle/ignore the JIP case
+	if ( RepState->LastChangelistIndex <= RepChangelistState->HistoryStart )
+	{
+		if ( RepState->LastChangelistIndex != 0 )
 		{
-			PropertyChanged	= true;
+			UE_LOG( LogRep, Warning, TEXT( "FRepLayout::ReplicatePropertiesUsingChangelistState: Entire history sent for: %s" ), *GetNameSafe( ObjectClass ) );
 		}
 
-		// Loop over all the conditional properties
-		if ( CompareProperties( RepState, CompareData, Data, ChangeTracker->Parents, RepState->ConditionalLifetime ) )
-		{
-			PropertyChanged = true;
-		}
-	}
-#ifdef ENABLE_SUPER_CHECKSUMS
-	else
-	{
-		// If we didn't compare this frame, make sure to reset out replication frame
-		// This is to force a compare next time it comes up
-		RepState->LastReplicationFrame = 0;		
+		RepState->LastChangelistIndex = RepChangelistState->HistoryStart;
 	}
-#endif
 
-	// PreOpenAckHistory are all the properties sent before we got our first open ack
-	const bool bFlushPreOpenAckHistory = RepState->OpenAckedCalled && RepState->PreOpenAckHistory.Num() > 0;
+	const int32 PossibleNewHistoryIndex = RepState->HistoryEnd % FRepState::MAX_CHANGE_HISTORY;
 
-	if ( PropertyChanged || RepState->NumNaks > 0 || bFlushPreOpenAckHistory )
-	{
-		// Use the first inactive history item to build this change list on
-		check( RepState->HistoryEnd - RepState->HistoryStart < FRepState::MAX_CHANGE_HISTORY );
-		const int32 HistoryIndex = RepState->HistoryEnd % FRepState::MAX_CHANGE_HISTORY;
+	FRepChangedHistory& PossibleNewHistoryItem = RepState->ChangeHistory[PossibleNewHistoryIndex];
 
-		FRepChangedHistory & NewHistoryItem = RepState->ChangeHistory[ HistoryIndex ];
+	TArray< uint16 >& Changed = PossibleNewHistoryItem.Changed;
 
-		RepState->HistoryEnd++;
+	check( Changed.Num() == 0 );		// Make sure this history item is actually inactive
 
-		TArray<uint16> & Changed = NewHistoryItem.Changed;
-
-		check( Changed.Num() == 0 );		// Make sure this history item is actually inactive
-
-		if ( PropertyChanged )
-		{
-			// Initialize the history item change list with the parent change lists
-			// We do it in the order of the parents so that the final change list will be fully sorted
-			for ( int32 i = 0, Parents_Num = Parents.Num(); i < Parents_Num; ++i )
-			{
-				TArray<uint16>& Parents_Changed = ChangeTracker->Parents[i].Changed;
+	// Gather all change lists that are new since we last looked, and merge them all together into a single CL
+	for ( int32 i = RepState->LastChangelistIndex; i < RepChangelistState->HistoryEnd; i++ )
+	{
+		const int32 HistoryIndex = i % FRepChangelistState::MAX_CHANGE_HISTORY;
 
-				if ( Parents_Changed.Num() > 0 )
-				{
-					Changed.Append( Parents_Changed );
+		FRepChangedHistory& HistoryItem = RepChangelistState->ChangeHistory[HistoryIndex];
 
-					if ( Parents[i].Flags & PARENT_IsConditional )
-					{
-						// Reset properties that don't share information across connections
-						Parents_Changed.Empty();
-					}
-				}
-			}
+		TArray< uint16 > Temp = Changed;
+		MergeChangeList( Data, HistoryItem.Changed, Temp, Changed );
+	}
 
-			Changed.Add( 0 );
+	// We're all caught up now
+	RepState->LastChangelistIndex = RepChangelistState->HistoryEnd;
 
-#ifdef SANITY_CHECK_MERGES
-			SanityCheckChangeList( Data, Changed );
-#endif
-		}
+	if ( Changed.Num() > 0 || RepState->NumNaks > 0 || bFlushPreOpenAckHistory )
+	{
+		RepState->HistoryEnd++;
 
-		// Update the history, and merge in any nak'd change lists
 		UpdateChangelistHistory( RepState, ObjectClass, Data, OwningChannel->Connection->OutAckPacketId, &Changed );
 
 		// Merge in the PreOpenAckHistory (unreliable properties sent before the bunch was initially acked)
@@ -762,48 +635,48 @@ bool FRepLayout::ReplicateProperties(
 			{
 				TArray< uint16 > Temp = Changed;
 				Changed.Empty();
-				MergeDirtyList( RepState, (void*)Data, Temp, RepState->PreOpenAckHistory[i].Changed, Changed );
+				MergeChangeList( Data, RepState->PreOpenAckHistory[i].Changed, Temp, Changed );
 			}
 			RepState->PreOpenAckHistory.Empty();
 		}
+	}
+	else
+	{
+		// Nothing changed and there are no nak's, so just do normal housekeeping and remove acked history items
+		UpdateChangelistHistory( RepState, ObjectClass, Data, OwningChannel->Connection->OutAckPacketId, NULL );
+		return false;		// Nothing to send
+	}
 
-		// At this point we should have a non empty change list
-		check( Changed.Num() > 0 );
+	// At this point we should have a non empty change list
+	check( Changed.Num() > 0 );
 
-#ifdef SANITY_CHECK_MERGES
-		SanityCheckChangeList( Data, Changed );
-#endif
+	const int32 NumBits = Writer.GetNumBits();
 
-		// Send the final merged change list
-		if ( OwningChannel->Connection->InternalAck )
-		{
-			// Remember all properties that have changed since this channel was first opened in case we need it (for bResendAllDataSinceOpen)
-			TArray< uint16 > Temp = RepState->LifetimeChangelist;
-			MergeDirtyList( RepState, ( void* )Data, Temp, Changed, RepState->LifetimeChangelist );
+	// Send the final merged change list
+	if ( OwningChannel->Connection->InternalAck )
+	{
+		// Remember all properties that have changed since this channel was first opened in case we need it (for bResendAllDataSinceOpen)
+		TArray< uint16 > Temp = RepState->LifetimeChangelist;
+		MergeChangeList( Data, Changed, Temp, RepState->LifetimeChangelist );
 
-			SendProperties_BackwardsCompatible( RepState, Data, OwningChannel->Connection, Writer, Changed );
-		}
-		else
-		{
-			SendProperties( RepState, RepFlags, Data, ObjectClass, OwningChannel, Writer, Changed );
-		}
+		SendProperties_BackwardsCompatible( RepState, ChangeTracker, Data, OwningChannel->Connection, Writer, Changed );
+	}
+	else
+	{
+		SendProperties( RepState, ChangeTracker, Data, ObjectClass, Writer, Changed );
+	}
 
-#ifdef ENABLE_SUPER_CHECKSUMS
-		Writer.WriteBit( bIsAllAcked ? 1 : 0 );
+	// See if something actually sent (this may be false due to conditional checks inside the send properties function
+	const bool bSomethingSent = NumBits != Writer.GetNumBits();
 
-		if ( bIsAllAcked )
-		{
-			ValidateWithChecksum( RepState->StaticBuffer.GetData(), Writer, false );
-		}
-#endif
-		
-		return true;
+	if ( !bSomethingSent )
+	{
+		// We need to revert the change list in the history if nothing really sent (can happen due to condition checks)
+		Changed.Empty();
+		RepState->HistoryEnd--;
 	}
 
-	// Nothing changed and there are no nak's, so just do normal housekeeping and remove acked history items
-	UpdateChangelistHistory( RepState, ObjectClass, Data, OwningChannel->Connection->OutAckPacketId, NULL );
-
-	return false;
+	return bSomethingSent;
 }
 
 void FRepLayout::UpdateChangelistHistory( FRepState * RepState, UClass * ObjectClass, const uint8* RESTRICT Data, const int32 AckPacketId, TArray< uint16 > * OutMerged ) const
@@ -840,7 +713,7 @@ void FRepLayout::UpdateChangelistHistory( FRepState * RepState, UClass * ObjectC
 				check( OutMerged != NULL );
 				TArray< uint16 > Temp = *OutMerged;
 				OutMerged->Empty();
-				MergeDirtyList( RepState, (void*)Data, Temp, HistoryItem.Changed, *OutMerged );
+				MergeChangeList( Data, HistoryItem.Changed, Temp, *OutMerged );
 				HistoryItem.Changed.Empty();
 
 #ifdef SANITY_CHECK_MERGES
@@ -958,180 +831,399 @@ bool FRepLayout::ReadyForDormancy( FRepState * RepState ) const
 	return AllAcked( RepState );
 }
 
-static FORCEINLINE void WritePropertyHandle( FNetBitWriter & Writer, uint16 Handle, bool bDoChecksum )
+void FRepLayout::SerializeObjectReplicatedProperties(UObject* Object, FArchive & Ar) const
 {
-	const int NumStartingBits = Writer.GetNumBits();
+	for (int32 i = 0; i < Parents.Num(); i++)
+	{
+        UStructProperty* StructProperty = Cast<UStructProperty>(Parents[i].Property);
+        UObjectProperty* ObjectProperty = Cast<UObjectProperty>(Parents[i].Property);
 
-	uint32 LocalHandle = Handle;
-	Writer.SerializeIntPacked( LocalHandle );
+		// We're only able to easily serialize non-object/struct properties, so just do those.
+		if (ObjectProperty == nullptr && StructProperty == nullptr)
+		{
+			bool bHasUnmapped = false;
+			SerializeProperties_r(Ar, NULL, Parents[i].CmdStart, Parents[i].CmdEnd, (uint8*)Object, bHasUnmapped);
+		}
+	}
+}
 
-#ifdef ENABLE_PROPERTY_CHECKSUMS
-	if ( bDoChecksum )
+bool FRepHandleIterator::NextHandle()
+{
+	CmdIndex = INDEX_NONE;
+
+	Handle = ChangelistIterator.Changed[ChangelistIterator.ChangedIndex];
+
+	if ( Handle == 0 )
 	{
-		SerializeGenericChecksum( Writer );
+		return false;		// Done
 	}
-#endif
 
-	NETWORK_PROFILER(GNetworkProfiler.TrackWritePropertyHandle( Writer.GetNumBits() - NumStartingBits, nullptr ));
-}
+	ChangelistIterator.ChangedIndex++;
 
-static bool ShouldSendProperty( FRepWriterState & WriterState, const uint16 Handle )
-{
-	if ( Handle == WriterState.Changed[WriterState.CurrentChanged] )
+	if ( !ensure( ChangelistIterator.ChangedIndex < ChangelistIterator.Changed.Num() ) )
 	{
-		// Write out the handle
-		WritePropertyHandle( WriterState.Writer, Handle, WriterState.bDoChecksum );
+		return false;
+	}
 
-		// Advance to the next expected handle
-		WriterState.CurrentChanged++;
+	const int32 HandleMinusOne = Handle - 1;
 
-		return true;
+	ArrayIndex = ( ArrayElementSize > 0 && NumHandlesPerElement > 0 ) ? HandleMinusOne / NumHandlesPerElement : 0;
+
+	if ( ArrayIndex >= MaxArrayIndex )
+	{
+		return false;
 	}
 
-	return false;
-}
+	ArrayOffset	= ArrayIndex * ArrayElementSize;
 
-void FRepLayout::SendProperties_DynamicArray_r( 
-	FRepState *	RESTRICT		RepState, 
-	const FReplicationFlags &	RepFlags,
-	FRepWriterState &			WriterState,
-	const int32					CmdIndex, 
-	const uint8* RESTRICT		StoredData, 
-	const uint8* RESTRICT		Data, 
-	uint16						Handle ) const
-{
-	const FRepLayoutCmd& Cmd = Cmds[ CmdIndex ];
+	const int32 RelativeHandle = HandleMinusOne - ArrayIndex * NumHandlesPerElement;
 
-	FScriptArray * Array = (FScriptArray *)Data;
-	FScriptArray * StoredArray = (FScriptArray *)StoredData;
+	CmdIndex = HandleToCmdIndex[RelativeHandle].CmdIndex;
 
-	// Write array num
-	uint16 ArrayNum = Array->Num();
-	WriterState.Writer << ArrayNum;
+	if ( !ensure( CmdIndex >= MinCmdIndex && CmdIndex < MaxCmdIndex ) )
+	{
+		return false;
+	}
 
-	// Make the shadow state match the actual state at the time of send
-	FScriptArrayHelper StoredArrayHelper( (UArrayProperty *)Cmd.Property, StoredData );
-	StoredArrayHelper.Resize( ArrayNum );
+	const FRepLayoutCmd& Cmd = Cmds[CmdIndex];
 
-	// Read the jump offset
-	// We won't need to actually jump over anything because we expect the change list to be pruned once we get here
-	// But we can use it to verify we read the correct amount.
-	const int32 ArrayChangedCount = WriterState.Changed[WriterState.CurrentChanged++];
+	if ( !ensure( Cmd.RelativeHandle - 1 == RelativeHandle ) )
+	{
+		return false;
+	}
 
-	const int32 OldChangedIndex = WriterState.CurrentChanged;
+	if ( !ensure( Cmd.Type != REPCMD_Return ) )
+	{
+		return false;
+	}
 
-	Data = (uint8*)Array->GetData();
-	StoredData = (uint8*)StoredArray->GetData();
+	return true;
+}
 
-	uint16 LocalHandle = 0;
+bool FRepHandleIterator::JumpOverArray()
+{
+	const int32 ArrayChangedCount = ChangelistIterator.Changed[ChangelistIterator.ChangedIndex++];
+	ChangelistIterator.ChangedIndex += ArrayChangedCount;
 
-	for ( int32 i = 0; i < Array->Num(); i++ )
+	if ( !ensure( ChangelistIterator.Changed[ChangelistIterator.ChangedIndex] == 0 ) )
 	{
-		const int32 ElementOffset = i * Cmd.ElementSize;
-		LocalHandle = SendProperties_r( RepState, RepFlags, WriterState, CmdIndex + 1, Cmd.EndCmd - 1, StoredData + ElementOffset, Data + ElementOffset, LocalHandle );
+		return false;
 	}
 
-	check( WriterState.CurrentChanged - OldChangedIndex == ArrayChangedCount );	// Make sure we read correct amount
-	check( WriterState.Changed[WriterState.CurrentChanged] == 0 );				// Make sure we are at the end
+	ChangelistIterator.ChangedIndex++;
 
-	WriterState.CurrentChanged++;
+	return true;
+}
 
-	WritePropertyHandle( WriterState.Writer, 0, WriterState.bDoChecksum );		// Signify end of dynamic array
+int32 FRepHandleIterator::PeekNextHandle() const
+{
+	return ChangelistIterator.Changed[ChangelistIterator.ChangedIndex];
 }
 
-void FRepLayout::SerializeObjectReplicatedProperties(UObject* Object, FArchive & Ar) const
+class FScopedIteratorArrayTracker
 {
-	for (int32 i = 0; i < Parents.Num(); i++)
+public:
+	FScopedIteratorArrayTracker( FRepHandleIterator* InCmdIndexIterator )
 	{
-        UStructProperty* StructProperty = Cast<UStructProperty>(Parents[i].Property);
-        UObjectProperty* ObjectProperty = Cast<UObjectProperty>(Parents[i].Property);
+		CmdIndexIterator = InCmdIndexIterator;
 
-		// We're only able to easily serialize non-object/struct properties, so just do those.
-		if (ObjectProperty == nullptr && StructProperty == nullptr)
+		if ( CmdIndexIterator )
 		{
-			bool bHasUnmapped = false;
-			SerializeProperties_r(Ar, NULL, Parents[i].CmdStart, Parents[i].CmdEnd, (uint8*)Object, bHasUnmapped);
+			ArrayChangedCount	= CmdIndexIterator->ChangelistIterator.Changed[CmdIndexIterator->ChangelistIterator.ChangedIndex++];
+			OldChangedIndex		= CmdIndexIterator->ChangelistIterator.ChangedIndex;
 		}
 	}
-}
 
-uint16 FRepLayout::SendProperties_r( 
-	FRepState *	RESTRICT		RepState, 
-	const FReplicationFlags &	RepFlags,
-	FRepWriterState &			WriterState,
-	const int32					CmdStart, 
-	const int32					CmdEnd, 
-	const uint8* RESTRICT		StoredData, 
-	const uint8* RESTRICT		Data, 
-	uint16						Handle ) const
+	~FScopedIteratorArrayTracker()
+	{
+		if ( CmdIndexIterator )
+		{
+			check( CmdIndexIterator->ChangelistIterator.ChangedIndex - OldChangedIndex <= ArrayChangedCount );
+			CmdIndexIterator->ChangelistIterator.ChangedIndex = OldChangedIndex + ArrayChangedCount;
+			check( CmdIndexIterator->PeekNextHandle() == 0 );
+			CmdIndexIterator->ChangelistIterator.ChangedIndex++;
+		}
+	}
+
+	FRepHandleIterator*	CmdIndexIterator;
+	int32				ArrayChangedCount;
+	int32				OldChangedIndex;
+};
+
+void FRepLayout::MergeChangeList_r(
+	FRepHandleIterator&		RepHandleIterator1,
+	FRepHandleIterator&		RepHandleIterator2,
+	const uint8* RESTRICT	SourceData,
+	TArray< uint16 >&		OutChanged ) const
 {
-	for ( int32 CmdIndex = CmdStart; CmdIndex < CmdEnd; CmdIndex++ )
+	while ( true )
 	{
-		const FRepLayoutCmd& Cmd = Cmds[ CmdIndex ];
+		const int32 NextHandle1 = RepHandleIterator1.PeekNextHandle();
+		const int32 NextHandle2 = RepHandleIterator2.PeekNextHandle();
 
-		check( Cmd.Type != REPCMD_Return );
+		if ( NextHandle1 == 0 && NextHandle2 == 0 )
+		{
+			break;		// Done
+		}
 
-		Handle++;
+		if ( NextHandle2 == 0 )
+		{
+			PruneChangeList_r( RepHandleIterator1, SourceData, OutChanged );
+			return;
+		}
+		else if ( NextHandle1 == 0 )
+		{
+			PruneChangeList_r( RepHandleIterator2, SourceData, OutChanged );
+			return;
+		}
 
-		if ( Cmd.Type == REPCMD_DynamicArray )
+		FRepHandleIterator* ActiveIterator1 = nullptr;
+		FRepHandleIterator* ActiveIterator2 = nullptr;
+
+		int32 CmdIndex		= INDEX_NONE;
+		int32 ArrayOffset	= INDEX_NONE;
+
+		if ( NextHandle1 < NextHandle2 )
 		{
-			if ( ShouldSendProperty( WriterState, Handle ) )
+			if ( !RepHandleIterator1.NextHandle() )
 			{
-				SendProperties_DynamicArray_r( RepState, RepFlags, WriterState, CmdIndex, StoredData + Cmd.Offset, Data + Cmd.Offset, Handle );
+				break;		// Array overflow
 			}
-			CmdIndex = Cmd.EndCmd - 1;	// Jump past children of this array (-1 for the ++ in the for loop)
-			continue;
+
+			OutChanged.Add( NextHandle1 );
+
+			CmdIndex	= RepHandleIterator1.CmdIndex;
+			ArrayOffset = RepHandleIterator1.ArrayOffset;
+
+			ActiveIterator1 = &RepHandleIterator1;
 		}
+		else if ( NextHandle2 < NextHandle1 )
+		{
+			if ( !RepHandleIterator2.NextHandle() )
+			{
+				break;		// Array overflow
+			}
+
+			OutChanged.Add( NextHandle2 );
+
+			CmdIndex	= RepHandleIterator2.CmdIndex;
+			ArrayOffset = RepHandleIterator2.ArrayOffset;
 
-		if ( ShouldSendProperty( WriterState, Handle ) )
+			ActiveIterator2 = &RepHandleIterator2;
+		}
+		else
 		{
-#if !(UE_BUILD_SHIPPING || UE_BUILD_TEST)
-			if (CVarDoReplicationContextString->GetInt() > 0)
+			check( NextHandle1 == NextHandle2 );
+
+			if ( !RepHandleIterator1.NextHandle() )
 			{
-				WriterState.Writer.PackageMap->SetDebugContextString( FString::Printf(TEXT("%s - %s"), *Owner->GetPathName(), *Cmd.Property->GetPathName() ) );
+				break;		// Array overflow
 			}
-#endif
 
-			const int32 NumStartBits = WriterState.Writer.GetNumBits();
-			
-			// This property changed, so send it
-			Cmd.Property->NetSerializeItem( WriterState.Writer, WriterState.Writer.PackageMap, (void*)( Data + Cmd.Offset ) );
+			if ( !ensure( RepHandleIterator2.NextHandle() ) )
+			{
+				break;		// Array overflow
+			}
 
-			const int32 NumEndBits = WriterState.Writer.GetNumBits();
+			check( RepHandleIterator1.CmdIndex == RepHandleIterator2.CmdIndex );
 
-			const FRepParentCmd& ParentCmd = Parents[Cmd.ParentIndex];
+			OutChanged.Add( NextHandle1 );
 
-			NETWORK_PROFILER( GNetworkProfiler.TrackReplicateProperty( ParentCmd.Property, NumEndBits - NumStartBits, nullptr ) );
+			CmdIndex	= RepHandleIterator1.CmdIndex;
+			ArrayOffset = RepHandleIterator1.ArrayOffset;
 
-			// Make the shadow state match the actual state at the time of send
-			StoreProperty( Cmd, (void*)( StoredData + Cmd.Offset ), (const void*)( Data + Cmd.Offset ) );
+			ActiveIterator1 = &RepHandleIterator1;
+			ActiveIterator2 = &RepHandleIterator2;
+		}
 
-#if !(UE_BUILD_SHIPPING || UE_BUILD_TEST)
-			if (CVarDoReplicationContextString->GetInt() > 0)
+		const FRepLayoutCmd& Cmd = Cmds[CmdIndex];
+
+		if ( Cmd.Type == REPCMD_DynamicArray )
+		{
+			const uint8* Data = SourceData + ArrayOffset + Cmd.Offset;
+
+			const FScriptArray* Array = ( FScriptArray * )Data;
+
+			FScopedIteratorArrayTracker ArrayTracker1( ActiveIterator1 );
+			FScopedIteratorArrayTracker ArrayTracker2( ActiveIterator2 );
+
+			const int32 OriginalChangedNum	= OutChanged.AddUninitialized();
+
+			const uint8* NewData = ( uint8* )Array->GetData();
+
+			TArray< FHandleToCmdIndex >& ArrayHandleToCmdIndex = ActiveIterator1 ? *ActiveIterator1->HandleToCmdIndex[Cmd.RelativeHandle - 1].HandleToCmdIndex : *ActiveIterator2->HandleToCmdIndex[Cmd.RelativeHandle - 1].HandleToCmdIndex;
+
+			if ( !ActiveIterator1 )
 			{
-				WriterState.Writer.PackageMap->ClearDebugContextString();
+				FRepHandleIterator ArrayIterator2( ActiveIterator2->ChangelistIterator, Cmds, ArrayHandleToCmdIndex, Cmd.ElementSize, Array->Num(), CmdIndex + 1, Cmd.EndCmd - 1 );
+				PruneChangeList_r( ArrayIterator2, NewData, OutChanged );
 			}
-#endif
+			else if ( !ActiveIterator2 )
+			{
+				FRepHandleIterator ArrayIterator1( ActiveIterator1->ChangelistIterator, Cmds, ArrayHandleToCmdIndex, Cmd.ElementSize, Array->Num(), CmdIndex + 1, Cmd.EndCmd - 1 );
+				PruneChangeList_r( ArrayIterator1, NewData, OutChanged );
+			}
+			else
+			{
+				FRepHandleIterator ArrayIterator1( ActiveIterator1->ChangelistIterator, Cmds, ArrayHandleToCmdIndex, Cmd.ElementSize, Array->Num(), CmdIndex + 1, Cmd.EndCmd - 1 );
+				FRepHandleIterator ArrayIterator2( ActiveIterator2->ChangelistIterator, Cmds, ArrayHandleToCmdIndex, Cmd.ElementSize, Array->Num(), CmdIndex + 1, Cmd.EndCmd - 1 );
 
-#ifdef ENABLE_PROPERTY_CHECKSUMS
-			if ( WriterState.bDoChecksum )
+				MergeChangeList_r( ArrayIterator1, ArrayIterator2, NewData, OutChanged );
+			}
+
+			// Patch in the jump offset
+			OutChanged[OriginalChangedNum] = OutChanged.Num() - ( OriginalChangedNum + 1 );
+
+			// Add the array terminator
+			OutChanged.Add( 0 );
+		}
+	}
+}
+
+void FRepLayout::PruneChangeList_r(
+	FRepHandleIterator&		RepHandleIterator,
+	const uint8* RESTRICT	SourceData,
+	TArray< uint16 >&		OutChanged ) const
+{
+	while ( RepHandleIterator.NextHandle() )
+	{
+		OutChanged.Add( RepHandleIterator.Handle );
+
+		const int32 CmdIndex	= RepHandleIterator.CmdIndex;
+		const int32 ArrayOffset = RepHandleIterator.ArrayOffset;
+
+		const FRepLayoutCmd& Cmd = Cmds[CmdIndex];
+
+		if ( Cmd.Type == REPCMD_DynamicArray )
+		{
+			const uint8* Data = SourceData + ArrayOffset + Cmd.Offset;
+
+			const FScriptArray* Array = ( FScriptArray * )Data;
+
+			FScopedIteratorArrayTracker ArrayTracker( &RepHandleIterator );
+
+			const int32 OriginalChangedNum = OutChanged.AddUninitialized();
+
+			const uint8* NewData = ( uint8* )Array->GetData();
+
+			TArray< FHandleToCmdIndex >& ArrayHandleToCmdIndex = *RepHandleIterator.HandleToCmdIndex[Cmd.RelativeHandle - 1].HandleToCmdIndex;
+
+			FRepHandleIterator ArrayIterator( RepHandleIterator.ChangelistIterator, Cmds, ArrayHandleToCmdIndex, Cmd.ElementSize, Array->Num(), CmdIndex + 1, Cmd.EndCmd - 1 );
+			PruneChangeList_r( ArrayIterator, NewData, OutChanged );
+
+			// Patch in the jump offset
+			OutChanged[OriginalChangedNum] = OutChanged.Num() - ( OriginalChangedNum + 1 );
+
+			// Add the array terminator
+			OutChanged.Add( 0 );
+		}
+	}
+}
+
+void FRepLayout::SendProperties_r(
+	FRepState*	RESTRICT				RepState,
+	FRepChangedPropertyTracker*			ChangedTracker,
+	FNetBitWriter&						Writer,
+	const bool							bDoChecksum,
+	FRepHandleIterator&					HandleIterator,
+	const uint8* RESTRICT				SourceData ) const
+{
+	while ( HandleIterator.NextHandle() )
+	{
+		const FRepLayoutCmd& Cmd = Cmds[HandleIterator.CmdIndex];
+
+		const FRepParentCmd& ParentCmd = Parents[Cmd.ParentIndex];
+
+		if ( !RepState->ConditionMap[ParentCmd.Condition] || !ChangedTracker->Parents[Cmd.ParentIndex].Active )
+		{
+			if ( Cmd.Type == REPCMD_DynamicArray )
 			{
-				SerializeReadWritePropertyChecksum( Cmd, CmdIndex, Data + Cmd.Offset, WriterState.Writer );
+				if ( !HandleIterator.JumpOverArray() )
+				{
+					break;
+				}
 			}
+
+			continue;
+		}
+
+		const uint8* Data = SourceData + HandleIterator.ArrayOffset + Cmd.Offset;
+
+		WritePropertyHandle( Writer, HandleIterator.Handle, bDoChecksum );
+
+		if ( Cmd.Type == REPCMD_DynamicArray )
+		{
+			const FScriptArray* Array = ( FScriptArray * )Data;
+
+			// Write array num
+			uint16 ArrayNum = Array->Num();
+			Writer << ArrayNum;
+
+			// Read the jump offset
+			// We won't need to actually jump over anything because we expect the change list to be pruned once we get here
+			// But we can use it to verify we read the correct amount.
+			const int32 ArrayChangedCount = HandleIterator.ChangelistIterator.Changed[HandleIterator.ChangelistIterator.ChangedIndex++];
+
+			const int32 OldChangedIndex = HandleIterator.ChangelistIterator.ChangedIndex;
+
+			const uint8* NewData = ( uint8* )Array->GetData();
+
+			TArray< FHandleToCmdIndex >& ArrayHandleToCmdIndex = *HandleIterator.HandleToCmdIndex[Cmd.RelativeHandle - 1].HandleToCmdIndex;
+
+			FRepHandleIterator ArrayHandleIterator( HandleIterator.ChangelistIterator, Cmds, ArrayHandleToCmdIndex, Cmd.ElementSize, ArrayNum, HandleIterator.CmdIndex + 1, Cmd.EndCmd - 1 );
+
+			check( ArrayHandleIterator.ArrayElementSize > 0 );
+			check( ArrayHandleIterator.NumHandlesPerElement > 0 );
+
+			SendProperties_r( RepState, ChangedTracker, Writer, bDoChecksum, ArrayHandleIterator, NewData );
+
+			check( HandleIterator.ChangelistIterator.ChangedIndex - OldChangedIndex == ArrayChangedCount );				// Make sure we read correct amount
+			check( HandleIterator.ChangelistIterator.Changed[HandleIterator.ChangelistIterator.ChangedIndex] == 0 );	// Make sure we are at the end
+
+			HandleIterator.ChangelistIterator.ChangedIndex++;
+
+			WritePropertyHandle( Writer, 0, bDoChecksum );		// Signify end of dynamic array
+			continue;
+		}
+
+#if !(UE_BUILD_SHIPPING || UE_BUILD_TEST)
+		if ( CVarDoReplicationContextString->GetInt() > 0 )
+		{
+			Writer.PackageMap->SetDebugContextString( FString::Printf( TEXT( "%s - %s" ), *Owner->GetPathName(), *Cmd.Property->GetPathName() ) );
+		}
 #endif
+
+		const int32 NumStartBits = Writer.GetNumBits();
+
+		// This property changed, so send it
+		Cmd.Property->NetSerializeItem( Writer, Writer.PackageMap, ( void* )Data );
+
+		const int32 NumEndBits = Writer.GetNumBits();
+
+		NETWORK_PROFILER( GNetworkProfiler.TrackReplicateProperty( ParentCmd.Property, NumEndBits - NumStartBits, nullptr ) );
+
+#if !(UE_BUILD_SHIPPING || UE_BUILD_TEST)
+		if ( CVarDoReplicationContextString->GetInt() > 0 )
+		{
+			Writer.PackageMap->ClearDebugContextString();
 		}
-	}
+#endif
 
-	return Handle;
+#ifdef ENABLE_PROPERTY_CHECKSUMS
+		if ( bDoChecksum )
+		{
+			SerializeReadWritePropertyChecksum( Cmd, HandleIterator.CmdIndex, Data, Writer );
+		}
+#endif
+	}
 }
 
-void FRepLayout::SendProperties( 
-	FRepState *	RESTRICT		RepState, 
-	const FReplicationFlags &	RepFlags,
-	const uint8* RESTRICT		Data, 
+void FRepLayout::SendProperties(
+	FRepState *	RESTRICT		RepState,
+	FRepChangedPropertyTracker* ChangedTracker,
+	const uint8* RESTRICT		Data,
 	UClass *					ObjectClass,
-	UActorChannel	*			OwningChannel,
-	FNetBitWriter&				Writer, 
+	FNetBitWriter&				Writer,
 	TArray< uint16 > &			Changed ) const
 {
 #ifdef ENABLE_PROPERTY_CHECKSUMS
@@ -1140,15 +1232,28 @@ void FRepLayout::SendProperties(
 	const bool bDoChecksum = false;
 #endif
 
-	FRepWriterState WriterState( Writer, Changed, bDoChecksum );
+	FBitWriterMark Mark( Writer );
 
 #ifdef ENABLE_PROPERTY_CHECKSUMS
 	Writer.WriteBit( bDoChecksum ? 1 : 0 );
 #endif
 
-	SendProperties_r( RepState, RepFlags, WriterState, 0, Cmds.Num() - 1, RepState->StaticBuffer.GetData(), Data, 0 );
+	const int32 NumBits = Writer.GetNumBits();
+
+	FChangelistIterator ChangelistIterator( Changed, 0 );
+	FRepHandleIterator HandleIterator( ChangelistIterator, Cmds, BaseHandleToCmdIndex, 0, 1, 0, Cmds.Num() - 1 );
 
-	WritePropertyHandle( Writer, 0, bDoChecksum );
+	SendProperties_r( RepState, ChangedTracker, Writer, bDoChecksum, HandleIterator, Data );
+
+	if ( NumBits != Writer.GetNumBits() )
+	{
+		// We actually wrote stuff
+		WritePropertyHandle( Writer, 0, bDoChecksum );
+	}
+	else
+	{
+		Mark.Pop( Writer );
+	}
 }
 
 static FORCEINLINE void WritePropertyHandle_BackwardsCompatible( FNetBitWriter & Writer, uint32 NetFieldExportHandle, bool bDoChecksum )
@@ -1167,88 +1272,177 @@ static FORCEINLINE void WritePropertyHandle_BackwardsCompatible( FNetBitWriter &
 	NETWORK_PROFILER( GNetworkProfiler.TrackWritePropertyHandle( Writer.GetNumBits() - NumStartingBits, nullptr ) );
 }
 
-static bool ShouldSendProperty_BackwardsCompatible( FRepWriterState & WriterState, FNetBitWriter & Writer, const uint16 Handle, const uint32 NetFieldExportHandle )
+TSharedPtr< FNetFieldExportGroup > FRepLayout::CreateNetfieldExportGroup() const
 {
-	if ( WriterState.Changed.Num() == 0 || Handle == WriterState.Changed[WriterState.CurrentChanged] )
-	{
-		// Write out the handle
-		WritePropertyHandle_BackwardsCompatible( Writer, NetFieldExportHandle, WriterState.bDoChecksum );
+	TSharedPtr< FNetFieldExportGroup > NetFieldExportGroup = TSharedPtr< FNetFieldExportGroup >( new FNetFieldExportGroup() );
 
-		// Advance to the next expected handle
-		WriterState.CurrentChanged++;
+	NetFieldExportGroup->PathName = Owner->GetPathName();
+	NetFieldExportGroup->NetFieldExports.SetNum( Cmds.Num() );
 
-		return true;
+	for ( int32 i = 0; i < Cmds.Num(); i++ )
+	{
+		FNetFieldExport NetFieldExport(
+			i,
+			Cmds[i].CompatibleChecksum,
+			Cmds[i].Property ? Cmds[i].Property->GetName() : TEXT( "" ),
+			Cmds[i].Property ? Cmds[i].Property->GetCPPType( nullptr, 0 ) : TEXT( "" ) );
+
+		NetFieldExportGroup->NetFieldExports[i] = NetFieldExport;
 	}
 
-	return false;
+	return NetFieldExportGroup;
 }
 
-void FRepLayout::SendProperties_BackwardsCompatible_DynamicArray_r(
-	FRepState *	RESTRICT		RepState,
-	FRepWriterState &			WriterState,
-	FNetBitWriter &				Writer,
-	UPackageMapClient*			PackageMapClient,
-	FNetFieldExportGroup*		NetFieldExportGroup,
-	const int32					CmdIndex,
-	const uint8* RESTRICT		StoredData,
-	const uint8* RESTRICT		Data ) const
+static FORCEINLINE void WriteProperty_BackwardsCompatible( FNetBitWriter& Writer, const FRepLayoutCmd& Cmd, const int32 CmdIndex, const UObject* Owner, const uint8* Data, const bool bDoChecksum )
 {
-	const FRepLayoutCmd& Cmd = Cmds[CmdIndex];
+#if !(UE_BUILD_SHIPPING || UE_BUILD_TEST)
+	if ( CVarDoReplicationContextString->GetInt() > 0 )
+	{
+		Writer.PackageMap->SetDebugContextString( FString::Printf( TEXT( "%s - %s" ), *Owner->GetPathName(), *Cmd.Property->GetPathName() ) );
+	}
+#endif
 
-	FScriptArray * Array = ( FScriptArray * )Data;
-	FScriptArray * StoredArray = ( FScriptArray * )StoredData;
+	const int32 NumStartBits = Writer.GetNumBits();
 
-	// Write array num
-	uint32 ArrayNum = Array->Num();
-	Writer.SerializeIntPacked( ArrayNum );
+	FNetBitWriter TempWriter( Writer.PackageMap, 0 );
 
-	// Make the shadow state match the actual state at the time of send
-	if ( StoredData != nullptr )
-	{
-		FScriptArrayHelper StoredArrayHelper( ( UArrayProperty * )Cmd.Property, StoredData );
-		StoredArrayHelper.Resize( ArrayNum );
-	}
+	// This property changed, so send it
+	Cmd.Property->NetSerializeItem( TempWriter, TempWriter.PackageMap, ( void* )Data );
 
-	// Read the jump offset
-	// We won't need to actually jump over anything because we expect the change list to be pruned once we get here
-	// But we can use it to verify we read the correct amount.
-	const int32 ArrayChangedCount = WriterState.Changed.Num() > 0 ? WriterState.Changed[WriterState.CurrentChanged++] : -1;
+	uint32 NumBits = TempWriter.GetNumBits();
+	Writer.SerializeIntPacked( NumBits );
+	Writer.SerializeBits( TempWriter.GetData(), NumBits );
 
-	const int32 OldChangedIndex = WriterState.CurrentChanged;
+	const int32 NumEndBits = Writer.GetNumBits();
 
-	Data = ( uint8* )Array->GetData();
-	StoredData = StoredArray != nullptr ? ( uint8* )StoredArray->GetData() : nullptr;
+	NETWORK_PROFILER( GNetworkProfiler.TrackReplicateProperty( Cmd.Property, NumEndBits - NumStartBits, nullptr ) );
 
-	uint16 LocalHandle = 0;
+#if !(UE_BUILD_SHIPPING || UE_BUILD_TEST)
+	if ( CVarDoReplicationContextString->GetInt() > 0 )
+	{
+		Writer.PackageMap->ClearDebugContextString();
+	}
+#endif
 
-	for ( int32 i = 0; i < Array->Num(); i++ )
+#ifdef ENABLE_PROPERTY_CHECKSUMS
+	if ( bDoChecksum )
 	{
-		uint32 Index = i + 1;
-		Writer.SerializeIntPacked( Index );
-		const int32 ElementOffset = i * Cmd.ElementSize;
-		LocalHandle = SendProperties_BackwardsCompatible_r( RepState, WriterState, Writer, PackageMapClient, NetFieldExportGroup, CmdIndex + 1, Cmd.EndCmd - 1, StoredData ? StoredData + ElementOffset : nullptr, Data + ElementOffset, LocalHandle );
+		SerializeReadWritePropertyChecksum( Cmd, CmdIndex, Data, Writer );
 	}
+#endif
+}
+
+void FRepLayout::SendProperties_BackwardsCompatible_r(
+	FRepState* RESTRICT					RepState,
+	UPackageMapClient*					PackageMapClient,
+	FNetFieldExportGroup*				NetFieldExportGroup,
+	FRepChangedPropertyTracker*			ChangedTracker,
+	FNetBitWriter&						Writer,
+	const bool							bDoChecksum,
+	FRepHandleIterator&					HandleIterator,
+	const uint8* RESTRICT				SourceData ) const
+{
+	int32 OldIndex = -1;
+
+	while ( HandleIterator.NextHandle() )
+	{
+		const FRepLayoutCmd& Cmd = Cmds[HandleIterator.CmdIndex];
+
+		const FRepParentCmd& ParentCmd = Parents[Cmd.ParentIndex];
+
+		const bool bConditionMatches = ChangedTracker == nullptr || ( RepState->ConditionMap[ParentCmd.Condition] && ChangedTracker->Parents[Cmd.ParentIndex].Active );
+
+		if ( !bConditionMatches )
+		{
+			if ( Cmd.Type == REPCMD_DynamicArray )
+			{
+				if ( !HandleIterator.JumpOverArray() )
+				{
+					break;
+				}
+			}
+
+			continue;
+		}
+
+		const uint8* Data = SourceData + HandleIterator.ArrayOffset + Cmd.Offset;
+
+		PackageMapClient->TrackNetFieldExport( NetFieldExportGroup, HandleIterator.CmdIndex );
+
+		if ( HandleIterator.ArrayElementSize > 0 && HandleIterator.ArrayIndex != OldIndex )
+		{
+			if ( OldIndex != -1 )
+			{
+				WritePropertyHandle_BackwardsCompatible( Writer, 0, bDoChecksum );
+			}
+
+			uint32 Index = HandleIterator.ArrayIndex + 1;
+			Writer.SerializeIntPacked( Index );
+			OldIndex = HandleIterator.ArrayIndex;
+		}
+
+		WritePropertyHandle_BackwardsCompatible( Writer, HandleIterator.CmdIndex + 1, bDoChecksum );
+
+		if ( Cmd.Type == REPCMD_DynamicArray )
+		{
+			const FScriptArray* Array = ( FScriptArray * )Data;
+
+			uint32 ArrayNum = Array->Num();
+
+			// Read the jump offset
+			// We won't need to actually jump over anything because we expect the change list to be pruned once we get here
+			// But we can use it to verify we read the correct amount.
+			const int32 ArrayChangedCount = HandleIterator.ChangelistIterator.Changed[HandleIterator.ChangelistIterator.ChangedIndex++];
+
+			const int32 OldChangedIndex = HandleIterator.ChangelistIterator.ChangedIndex;
+
+			const uint8* NewData = ( uint8* )Array->GetData();
+
+			TArray< FHandleToCmdIndex >& ArrayHandleToCmdIndex = *HandleIterator.HandleToCmdIndex[Cmd.RelativeHandle - 1].HandleToCmdIndex;
+
+			FRepHandleIterator ArrayHandleIterator( HandleIterator.ChangelistIterator, Cmds, ArrayHandleToCmdIndex, Cmd.ElementSize, ArrayNum, HandleIterator.CmdIndex + 1, Cmd.EndCmd - 1 );
+
+			check( ArrayHandleIterator.ArrayElementSize > 0 );
+			check( ArrayHandleIterator.NumHandlesPerElement > 0 );
+
+			FNetBitWriter TempWriter( Writer.PackageMap, 0 );
+
+			// Write array num
+			TempWriter.SerializeIntPacked( ArrayNum );
+
+			if ( ArrayNum > 0 )
+			{
+				SendProperties_BackwardsCompatible_r( RepState, PackageMapClient, NetFieldExportGroup, ChangedTracker, TempWriter, bDoChecksum, ArrayHandleIterator, NewData );
+			}
+
+			uint32 EndArrayIndex = 0;
+			TempWriter.SerializeIntPacked( EndArrayIndex );
+
+			uint32 NumBits = TempWriter.GetNumBits();
+			Writer.SerializeIntPacked( NumBits );
+			Writer.SerializeBits( TempWriter.GetData(), NumBits );
 
-	uint32 Index = 0;
-	Writer.SerializeIntPacked( Index );
+			check( HandleIterator.ChangelistIterator.ChangedIndex - OldChangedIndex == ArrayChangedCount );				// Make sure we read correct amount
+			check( HandleIterator.ChangelistIterator.Changed[HandleIterator.ChangelistIterator.ChangedIndex] == 0 );	// Make sure we are at the end
+
+			HandleIterator.ChangelistIterator.ChangedIndex++;
+			continue;
+		}
 
-	check( WriterState.Changed.Num() == 0 || WriterState.CurrentChanged - OldChangedIndex == ArrayChangedCount );							// Make sure we read correct amount
-	check( WriterState.Changed.Num() == 0 || WriterState.Changed[WriterState.CurrentChanged] == 0 );	// Make sure we are at the end
+		WriteProperty_BackwardsCompatible( Writer, Cmd, HandleIterator.CmdIndex, Owner, Data, bDoChecksum );
+	}
 
-	WriterState.CurrentChanged++;
+	WritePropertyHandle_BackwardsCompatible( Writer, 0, bDoChecksum );
 }
 
-uint16 FRepLayout::SendProperties_BackwardsCompatible_r(
-	FRepState *	RESTRICT		RepState,
-	FRepWriterState &			WriterState,
-	FNetBitWriter &				Writer,
-	UPackageMapClient*			PackageMapClient,
-	FNetFieldExportGroup*		NetFieldExportGroup,
-	const int32					CmdStart,
-	const int32					CmdEnd,
-	const uint8* RESTRICT		StoredData,
-	const uint8* RESTRICT		Data,
-	uint16						Handle ) const
+void FRepLayout::SendAllProperties_BackwardsCompatible_r( 
+	FNetBitWriter&						Writer,
+	const bool							bDoChecksum,
+	UPackageMapClient*					PackageMapClient,
+	FNetFieldExportGroup*				NetFieldExportGroup,
+	const int32							CmdStart,
+	const int32							CmdEnd, 
+	const uint8*						SourceData ) const
 {
 	for ( int32 CmdIndex = CmdStart; CmdIndex < CmdEnd; CmdIndex++ )
 	{
@@ -1256,107 +1450,57 @@ uint16 FRepLayout::SendProperties_BackwardsCompatible_r(
 
 		check( Cmd.Type != REPCMD_Return );
 
-		Handle++;
+		PackageMapClient->TrackNetFieldExport( NetFieldExportGroup, CmdIndex );
+
+		WritePropertyHandle_BackwardsCompatible( Writer, CmdIndex + 1, bDoChecksum );
+
+		const uint8* Data = SourceData + Cmd.Offset;
 
 		if ( Cmd.Type == REPCMD_DynamicArray )
-		{
-			if ( ShouldSendProperty_BackwardsCompatible( WriterState, Writer, Handle, CmdIndex + 1 ) )
-			{
-				PackageMapClient->TrackNetFieldExport( NetFieldExportGroup, CmdIndex );
+		{			
+			const FScriptArray* Array = ( FScriptArray * )Data;
 
-				FNetBitWriter TempWriter( Writer.PackageMap, 0 );
-				SendProperties_BackwardsCompatible_DynamicArray_r( RepState, WriterState, TempWriter, PackageMapClient, NetFieldExportGroup, CmdIndex, StoredData ? StoredData + Cmd.Offset : nullptr, Data + Cmd.Offset );
+			FNetBitWriter TempWriter( Writer.PackageMap, 0 );
 
-				uint32 NumBits = TempWriter.GetNumBits();
-				Writer.SerializeIntPacked( NumBits );
-				Writer.SerializeBits( TempWriter.GetData(), NumBits );
-			}
-			CmdIndex = Cmd.EndCmd - 1;	// Jump past children of this array (-1 for the ++ in the for loop)
-			continue;
-		}
+			// Write array num
+			uint32 ArrayNum = Array->Num();
+			TempWriter.SerializeIntPacked( ArrayNum );
 
-		if ( ShouldSendProperty_BackwardsCompatible( WriterState, Writer, Handle, CmdIndex + 1 ) )
-		{
-#if !(UE_BUILD_SHIPPING || UE_BUILD_TEST)
-			if ( CVarDoReplicationContextString->GetInt() > 0 )
+			for ( int32 i = 0; i < Array->Num(); i++ )
 			{
-				Writer.PackageMap->SetDebugContextString( FString::Printf( TEXT( "%s - %s" ), *Owner->GetPathName(), *Cmd.Property->GetPathName() ) );
-			}
-#endif
-			PackageMapClient->TrackNetFieldExport( NetFieldExportGroup, CmdIndex );
-
-			const int32 NumStartBits = Writer.GetNumBits();
+				uint32 ArrayIndex = i + 1;
+				TempWriter.SerializeIntPacked( ArrayIndex );
 
-			// This property changed, so send it
-			FNetBitWriter TempWriter( Writer.PackageMap, 0 );
+				SendAllProperties_BackwardsCompatible_r( TempWriter, bDoChecksum, PackageMapClient, NetFieldExportGroup, CmdIndex + 1, Cmd.EndCmd - 1, ((const uint8*)Array->GetData()) + Cmd.ElementSize * i );
+			}
 
-			Cmd.Property->NetSerializeItem( TempWriter, TempWriter.PackageMap, ( void* )( Data + Cmd.Offset ) );
+			uint32 EndArrayIndex = 0;
+			TempWriter.SerializeIntPacked( EndArrayIndex );
 
 			uint32 NumBits = TempWriter.GetNumBits();
 			Writer.SerializeIntPacked( NumBits );
 			Writer.SerializeBits( TempWriter.GetData(), NumBits );
 
-			const int32 NumEndBits = Writer.GetNumBits();
-
-			const FRepParentCmd& ParentCmd = Parents[Cmd.ParentIndex];
-
-			NETWORK_PROFILER( GNetworkProfiler.TrackReplicateProperty( ParentCmd.Property, NumEndBits - NumStartBits, nullptr ) );
-
-			// Make the shadow state match the actual state at the time of send
-			if ( StoredData != nullptr )
-			{
-				StoreProperty( Cmd, ( void* )( StoredData + Cmd.Offset ), ( const void* )( Data + Cmd.Offset ) );
-			}
-
-#if !(UE_BUILD_SHIPPING || UE_BUILD_TEST)
-			if ( CVarDoReplicationContextString->GetInt() > 0 )
-			{
-				Writer.PackageMap->ClearDebugContextString();
-			}
-#endif
-
-#ifdef ENABLE_PROPERTY_CHECKSUMS
-			if ( WriterState.bDoChecksum )
-			{
-				SerializeReadWritePropertyChecksum( Cmd, CmdIndex, Data + Cmd.Offset, Writer );
-			}
-#endif
+			CmdIndex = Cmd.EndCmd - 1;		// The -1 to handle the ++ in the for loop
+			continue;
 		}
-	}
-
-	WritePropertyHandle_BackwardsCompatible( Writer, 0, WriterState.bDoChecksum );
-
-	return Handle;
-}
 
-TSharedPtr< FNetFieldExportGroup > FRepLayout::CreateNetfieldExportGroup() const
-{
-	TSharedPtr< FNetFieldExportGroup > NetFieldExportGroup = TSharedPtr< FNetFieldExportGroup >( new FNetFieldExportGroup() );
-
-	NetFieldExportGroup->PathName = Owner->GetPathName();
-	NetFieldExportGroup->NetFieldExports.SetNum( Cmds.Num() );
-
-	for ( int32 i = 0; i < Cmds.Num(); i++ )
-	{
-		FNetFieldExport NetFieldExport(
-			i,
-			Cmds[i].CompatibleChecksum,
-			Cmds[i].Property ? Cmds[i].Property->GetName() : TEXT( "" ),
-			Cmds[i].Property ? Cmds[i].Property->GetCPPType( nullptr, 0 ) : TEXT( "" ) );
-
-		NetFieldExportGroup->NetFieldExports[i] = NetFieldExport;
+		WriteProperty_BackwardsCompatible( Writer, Cmd, CmdIndex, Owner, Data, bDoChecksum );
 	}
 
-	return NetFieldExportGroup;
+	WritePropertyHandle_BackwardsCompatible( Writer, 0, bDoChecksum );
 }
 
 void FRepLayout::SendProperties_BackwardsCompatible(
 	FRepState* RESTRICT			RepState,
+	FRepChangedPropertyTracker* ChangedTracker,
 	const uint8* RESTRICT		Data,
 	UNetConnection*				Connection,
 	FNetBitWriter&				Writer,
 	TArray< uint16 >&			Changed ) const
 {
+	FBitWriterMark Mark( Writer );
+
 #ifdef ENABLE_PROPERTY_CHECKSUMS
 	const bool bDoChecksum = CVarDoPropertyChecksum.GetValueOnGameThread() == 1;
 	Writer.WriteBit( bDoChecksum ? 1 : 0 );
@@ -1375,9 +1519,24 @@ void FRepLayout::SendProperties_BackwardsCompatible(
 		PackageMapClient->AddNetFieldExportGroup( Owner->GetPathName(), NetFieldExportGroup );
 	}
 
-	FRepWriterState WriterState( Writer, Changed, bDoChecksum );
+	const int32 NumBits = Writer.GetNumBits();
+
+	if ( Changed.Num() == 0 )
+	{
+		SendAllProperties_BackwardsCompatible_r( Writer, bDoChecksum, PackageMapClient, NetFieldExportGroup.Get(), 0, Cmds.Num() - 1, Data );
+	}
+	else
+	{
+		FChangelistIterator ChangelistIterator( Changed, 0 );
+		FRepHandleIterator HandleIterator( ChangelistIterator, Cmds, BaseHandleToCmdIndex, 0, 1, 0, Cmds.Num() - 1 );
+
+		SendProperties_BackwardsCompatible_r( RepState, PackageMapClient, NetFieldExportGroup.Get(), ChangedTracker, Writer, bDoChecksum, HandleIterator, Data );
+	}
 
-	SendProperties_BackwardsCompatible_r( RepState, WriterState, Writer, PackageMapClient, NetFieldExportGroup.Get(), 0, Cmds.Num() - 1, RepState ? RepState->StaticBuffer.GetData() : nullptr, Data, 0 );
+	if ( NumBits == Writer.GetNumBits() )
+	{
+		Mark.Pop( Writer );
+	}
 }
 
 class FReceivedPropertiesStackState : public FCmdIteratorBaseStackState
@@ -1529,9 +1688,9 @@ static FGuidReferencesMap* PrepReceivedArray(
 		{
 			NewGuidReferencesArray = &ParentGuidReferences->FindOrAdd( AbsOffset );
 
-			NewGuidReferencesArray->Array			= new FGuidReferencesMap;
+			NewGuidReferencesArray->Array		= new FGuidReferencesMap;
 			NewGuidReferencesArray->ParentIndex	= Cmd.ParentIndex;
-			NewGuidReferencesArray->CmdIndex		= CmdIndex;
+			NewGuidReferencesArray->CmdIndex	= CmdIndex;
 		}
 
 		check( NewGuidReferencesArray != NULL );
@@ -1914,7 +2073,16 @@ bool FRepLayout::ReceiveProperties_BackwardsCompatible_r(
 					break;
 				}
 
-				const int32 ElementOffset = ( Index - 1 ) * Cmd.ElementSize;
+				// Shift all indexes down since 0 represents null handle
+				Index--;
+
+				if ( !ensure( Index < ArrayNum ) )
+				{
+					UE_LOG( LogRep, Warning, TEXT( "ReceiveProperties_BackwardsCompatible_r: Array index out of bounds. Index: %i, ArrayNum: %i, Owner: %s, Name: %s, Type: %s, NetFieldExportHandle: %i, Checksum: %u" ), Index, ArrayNum, *Owner->GetName(), *NetFieldExportGroup->NetFieldExports[NetFieldExportHandle].Name, *NetFieldExportGroup->NetFieldExports[NetFieldExportHandle].Type, NetFieldExportHandle, Checksum );
+					return false;
+				}
+
+				const int32 ElementOffset = Index * Cmd.ElementSize;
 
 				uint8* ElementData			= LocalData + ElementOffset;
 				uint8* ElementShadowData	= LocalShadowData ? LocalShadowData + ElementOffset : nullptr;
@@ -2257,266 +2425,43 @@ uint32 FRepLayout::GenerateChecksum( const FRepState* RepState ) const
 	return FCrc::MemCrc32( Writer.GetData(), Writer.GetNumBytes(), 0 );
 }
 
-class FPruneChangeListImpl : public FRepLayoutCmdIterator< FPruneChangeListImpl, FCmdIteratorBaseStackState >
-{
-public:
-	FPruneChangeListImpl( const TArray< uint16 >& InChanged, TArray< uint16 >& OutChanged, const TArray< FRepParentCmd >& InParents, const TArray< FRepLayoutCmd >& InCmds ) :
-        FRepLayoutCmdIterator( InParents, InCmds ),
-		Changed( InChanged ),
-		ChangeIndex( 0 ),
-		CurrentHandle( 0 ),
-		PrunedChanged( OutChanged ),
-		bLastChangedMatches( false )
-	{}
-
-	INIT_STACK( FCmdIteratorBaseStackState ) { }
-
-	SHOULD_PROCESS_NEXT_CMD() 
-	{ 
-		CurrentHandle++;
-
-		check( CurrentHandle != 0 );
-
-		bLastChangedMatches = Changed[ChangeIndex] == CurrentHandle;
-
-		return bLastChangedMatches;
-	}
-
-	PROCESS_ARRAY_CMD( FCmdIteratorBaseStackState ) 
-	{
-		check( bLastChangedMatches )
-
-		const bool bDirty1Matches = bLastChangedMatches;
-
-		// This will be a new pruned entry (i.e. clamped to new array boundary)
-		PrunedChanged.Add( CurrentHandle );
-
-		const int32 OriginalChangeIndex = PrunedChanged.AddUninitialized();
-		check( OriginalChangeIndex == PrunedChanged.Num() - 1 );
-
-		// Advance the change list index since we matched
-		if ( bLastChangedMatches )
-		{
-			ChangeIndex++;
-		}
-
-		const int32 JumpToIndex = bLastChangedMatches ? Changed[ChangeIndex++] : -1;
-
-		const int32 OldChangeListIndex = ChangeIndex;
-
-		const int32 OldHandle = CurrentHandle;
-		CurrentHandle = 0;
-
-		// Process the array elements
-		ProcessDataArrayElements_r( StackState, Cmd );
-
-		// Restore the handle
-		CurrentHandle = OldHandle;
-
-		if ( bDirty1Matches )
-		{
-			check( ChangeIndex - OldChangeListIndex <= JumpToIndex );
-			ChangeIndex = OldChangeListIndex + JumpToIndex;
-			check( Changed[ChangeIndex] == 0 );
-			ChangeIndex++;
-		}
-
-		// Patch in the jump offset
-		PrunedChanged[OriginalChangeIndex] = PrunedChanged.Num() - ( OriginalChangeIndex + 1 );
-
-		// Add the array terminator
-		PrunedChanged.Add( 0 );
-	}
-
-	PROCESS_CMD( FCmdIteratorBaseStackState ) 
-	{
-		check( bLastChangedMatches )
-
-		// This will be a new merged dirty entry
-		PrunedChanged.Add( CurrentHandle );
-
-		// Advance matching dirty indices
-		ChangeIndex++;
-	}
-
-	const TArray< uint16 >& Changed;
-	int32					ChangeIndex;
-	uint16					CurrentHandle;
-	TArray< uint16 >& 		PrunedChanged;
-	bool					bLastChangedMatches;
-};
-
 void FRepLayout::PruneChangeList( FRepState* RepState, const void* RESTRICT Data, const TArray< uint16 >& Changed, TArray< uint16 >& PrunedChanged ) const
 {
 	check( Changed.Num() > 0 );
 
 	PrunedChanged.Empty();
 
-	FPruneChangeListImpl PrunePropertiesImpl( Changed, PrunedChanged, Parents, Cmds );
+	FChangelistIterator ChangelistIterator( Changed, 0 );
+	FRepHandleIterator HandleIterator( ChangelistIterator, Cmds, BaseHandleToCmdIndex, 0, 1, 0, Cmds.Num() - 1 );
+	PruneChangeList_r( HandleIterator, ( uint8* )Data, PrunedChanged );
 
-	// Prune it
-	PrunePropertiesImpl.ProcessCmds( ( uint8* )Data, ( uint8* )RepState->StaticBuffer.GetData() );
-
-	PrunePropertiesImpl.PrunedChanged.Add( 0 );
+	PrunedChanged.Add( 0 );
 }
 
-class FMergeDirtyListImpl : public FRepLayoutCmdIterator< FMergeDirtyListImpl, FCmdIteratorBaseStackState >
+void FRepLayout::MergeChangeList( const uint8* RESTRICT Data, const TArray< uint16 > & Dirty1, const TArray< uint16 > & Dirty2, TArray< uint16 > & MergedDirty ) const
 {
-public:
-	FMergeDirtyListImpl( const TArray< uint16 > & InDirty1, const TArray< uint16 > & InDirty2, TArray< uint16 > & OutMergedDirty, const TArray< FRepParentCmd >& InParents, const TArray< FRepLayoutCmd >& InCmds ) : 
-        FRepLayoutCmdIterator( InParents, InCmds ),
-		DirtyList1( InDirty1 ),
-		DirtyList2( InDirty2 ), 
-		DirtyListIndex1( 0 ),
-		DirtyListIndex2( 0 ),
-		CurrentHandle( 0 ),
-		bDirtyValid1( true ),
-		bDirtyValid2( true ),
-		MergedDirtyList( OutMergedDirty ),
-		bLastDirty1Matches( false ), 
-		bLastDirty2Matches( false )
-	{}
-
-	INIT_STACK( FCmdIteratorBaseStackState ) { }
-
-	SHOULD_PROCESS_NEXT_CMD() 
-	{ 
-		CurrentHandle++;
+	check( Dirty1.Num() > 0 );
 
-		check( CurrentHandle != 0 );
-
-		bLastDirty1Matches = bDirtyValid1 && DirtyList1[DirtyListIndex1] == CurrentHandle;
-		bLastDirty2Matches = bDirtyValid2 && DirtyList2[DirtyListIndex2] == CurrentHandle;
-
-		return bLastDirty1Matches || bLastDirty2Matches;
-	}
+	MergedDirty.Empty();
 
-	PROCESS_ARRAY_CMD( FCmdIteratorBaseStackState ) 
+	if ( Dirty2.Num() == 0 )
 	{
-		// At least one of the list should be valid to be here
-		check( bDirtyValid1 || bDirtyValid2 );
-		check( bLastDirty1Matches || bLastDirty2Matches )
-
-		const bool bDirty1Matches = bLastDirty1Matches;
-		const bool bDirty2Matches = bLastDirty2Matches;
-
-		// This will be a new merged dirty entry
-		MergedDirtyList.Add( CurrentHandle );
-
-		const int32 OriginalMergedDirtyListIndex = MergedDirtyList.AddUninitialized();
-		check( OriginalMergedDirtyListIndex == MergedDirtyList.Num() - 1 );
-
-		// Advance the matching dirty lists
-		if ( bLastDirty1Matches )
-		{
-			DirtyListIndex1++;
-		}
-
-		if ( bLastDirty2Matches )
-		{
-			DirtyListIndex2++;
-		}
-
-		// Remember valid dirty lists
-		const bool bOldDirtyValid1 = bDirtyValid1;
-		const bool bOldDirtyValid2 = bDirtyValid2;
-
-		// Update which lists are still valid from this point
-		bDirtyValid1 = bLastDirty1Matches;
-		bDirtyValid2 = bLastDirty2Matches;
-
-		const int32 JumpToIndex1 = bLastDirty1Matches ? DirtyList1[DirtyListIndex1++] : -1;
-		const int32 JumpToIndex2 = bLastDirty2Matches ? DirtyList2[DirtyListIndex2++] : -1;
-
-		const int32 OldDirtyListIndex1 = DirtyListIndex1;
-		const int32 OldDirtyListIndex2 = DirtyListIndex2;
-
-		const int32 OldHandle = CurrentHandle;
-		CurrentHandle = 0;
-
-		// Process the array elements
-		ProcessDataArrayElements_r( StackState, Cmd );
-
-		// Restore the handle
-		CurrentHandle = OldHandle;
-
-		if ( bDirty1Matches )
-		{
-			check( DirtyListIndex1 - OldDirtyListIndex1 <= JumpToIndex1 );
-			DirtyListIndex1 = OldDirtyListIndex1 + JumpToIndex1;
-			check( DirtyList1[DirtyListIndex1] == 0 );
-			DirtyListIndex1++;
-		}
-
-		if ( bDirty2Matches )
-		{
-			check( DirtyListIndex2 - OldDirtyListIndex2 <= JumpToIndex2 );
-			DirtyListIndex2 = OldDirtyListIndex2 + JumpToIndex2;
-			check( DirtyList2[DirtyListIndex2] == 0 );
-			DirtyListIndex2++;
-		}
-
-		bDirtyValid1 = bOldDirtyValid1;
-		bDirtyValid2 = bOldDirtyValid2;
-
-		// Patch in the jump offset
-		MergedDirtyList[OriginalMergedDirtyListIndex] = MergedDirtyList.Num() - ( OriginalMergedDirtyListIndex + 1 );
-
-		// Add the array terminator
-		MergedDirtyList.Add( 0 );	
+		FChangelistIterator ChangelistIterator( Dirty1, 0 );
+		FRepHandleIterator HandleIterator( ChangelistIterator, Cmds, BaseHandleToCmdIndex, 0, 1, 0, Cmds.Num() - 1 );
+		PruneChangeList_r( HandleIterator, ( uint8* )Data, MergedDirty );
 	}
-
-	PROCESS_CMD( FCmdIteratorBaseStackState ) 
+	else
 	{
-		check( bLastDirty1Matches || bLastDirty2Matches )
+		FChangelistIterator ChangelistIterator1( Dirty1, 0 );
+		FRepHandleIterator HandleIterator1( ChangelistIterator1, Cmds, BaseHandleToCmdIndex, 0, 1, 0, Cmds.Num() - 1 );
 
-		// This will be a new merged dirty entry
-		MergedDirtyList.Add( CurrentHandle );
+		FChangelistIterator ChangelistIterator2( Dirty2, 0 );
+		FRepHandleIterator HandleIterator2( ChangelistIterator2, Cmds, BaseHandleToCmdIndex, 0, 1, 0, Cmds.Num() - 1 );
 
-		// Advance matching dirty indices
-		if ( bLastDirty1Matches )
-		{
-			DirtyListIndex1++;
-		}
-
-		if ( bLastDirty2Matches )
-		{
-			DirtyListIndex2++;
-		}
+		MergeChangeList_r( HandleIterator1, HandleIterator2, ( uint8* )Data, MergedDirty );
 	}
 
-	const TArray< uint16 >& DirtyList1;
-	const TArray< uint16 >& DirtyList2;
-
-	int32					DirtyListIndex1;
-	int32					DirtyListIndex2;
-	uint16					CurrentHandle;
-
-	bool					bDirtyValid1;
-	bool					bDirtyValid2;
-
-	TArray< uint16 >& 		MergedDirtyList;
-
-	bool					bLastDirty1Matches;
-	bool					bLastDirty2Matches;
-};
-
-void FRepLayout::MergeDirtyList( FRepState * RepState, const void* RESTRICT Data, const TArray< uint16 > & Dirty1, const TArray< uint16 > & Dirty2, TArray< uint16 > & MergedDirty ) const
-{
-	check( Dirty1.Num() > 0 || Dirty2.Num() > 0 );
-
-	MergedDirty.Empty();
-
-	FMergeDirtyListImpl MergePropertiesImpl( Dirty1, Dirty2, MergedDirty, Parents, Cmds );
-
-	// Even though one of these can be empty, we need to send the single one through, so we can prune it to the current shape of the tree
-	MergePropertiesImpl.bDirtyValid1 = Dirty1.Num() > 0;
-	MergePropertiesImpl.bDirtyValid2 = Dirty2.Num() > 0;
-
-	// Merge lists
-	MergePropertiesImpl.ProcessCmds( (uint8*)Data, (uint8*)RepState->StaticBuffer.GetData() );
-
-	MergePropertiesImpl.MergedDirtyList.Add( 0 );
+	MergedDirty.Add( 0 );
 }
 
 void FRepLayout::SanityCheckChangeList_DynamicArray_r( 
@@ -3017,49 +2962,37 @@ void FRepLayout::InitFromObjectClass( UClass * InObjectClass )
 
 	Object->GetLifetimeReplicatedProps( LifetimeProps );
 
-	if ( LifetimeProps.Num() > 0 )
+	// Setup lifetime replicated properties
+	for ( int32 i = 0; i < LifetimeProps.Num(); i++ )
 	{
-		UnconditionalLifetime.Empty();
-		ConditionalLifetime.Empty();
+		// Store the condition on the parent in case we need it
+		Parents[LifetimeProps[i].RepIndex].Condition			= LifetimeProps[i].Condition;
+		Parents[LifetimeProps[i].RepIndex].RepNotifyCondition	= LifetimeProps[i].RepNotifyCondition;
 
-		// Fix Lifetime props to have the proper index to the parent
-		for ( int32 i = 0; i < LifetimeProps.Num(); i++ )
+		if ( Parents[LifetimeProps[i].RepIndex].Flags & PARENT_IsCustomDelta )
 		{
-			// Store the condition on the parent in case we need it
-			Parents[LifetimeProps[i].RepIndex].Condition = LifetimeProps[i].Condition;
-			Parents[LifetimeProps[i].RepIndex].RepNotifyCondition = LifetimeProps[i].RepNotifyCondition;
-
-			if ( Parents[LifetimeProps[i].RepIndex].Flags & PARENT_IsCustomDelta )
-			{
-				continue;		// We don't handle custom properties in the FRepLayout class
-			}
-
-			Parents[LifetimeProps[i].RepIndex].Flags |= PARENT_IsLifetime;
+			continue;		// We don't handle custom properties in the FRepLayout class
+		}
 
-			if ( LifetimeProps[i].RepIndex == RemoteRoleIndex )
-			{
-				// We handle remote role specially, since it can change between connections when downgraded
-				// So we force it on the conditional list
-				check( LifetimeProps[i].Condition == COND_None );
-				LifetimeProps[i].Condition = COND_Custom;
-				ConditionalLifetime.AddUnique( LifetimeProps[i] );
-				continue;		
-			}
+		Parents[LifetimeProps[i].RepIndex].Flags |= PARENT_IsLifetime;
 
-			if ( LifetimeProps[i].Condition == COND_None )
-			{
-				// These properties are simple, and can benefit from property compare sharing
-				UnconditionalLifetime.AddUnique( LifetimeProps[i].RepIndex );
-				Parents[LifetimeProps[i].RepIndex].Flags &= ~PARENT_IsConditional;
-				continue;
-			}
-			
-			// The rest are conditional
-			// These properties are eventually conditionally copied to the RepState
-			ConditionalLifetime.AddUnique( LifetimeProps[i] );
+		if ( LifetimeProps[i].RepIndex == RemoteRoleIndex )
+		{
+			// We handle remote role specially, since it can change between connections when downgraded
+			// So we force it on the conditional list
+			check( LifetimeProps[i].Condition == COND_None );
+			LifetimeProps[i].Condition = COND_Custom;
+			continue;		
 		}
+
+		if ( LifetimeProps[i].Condition == COND_None )
+		{
+			Parents[LifetimeProps[i].RepIndex].Flags &= ~PARENT_IsConditional;
+		}			
 	}
 
+	BuildHandleToCmdIndexTable_r( 0, Cmds.Num() - 1, BaseHandleToCmdIndex );
+
 	Owner = InObjectClass;
 }
 
@@ -3080,6 +3013,8 @@ void FRepLayout::InitFromFunction( UFunction * InFunction )
 
 	AddReturnCmd();
 
+	BuildHandleToCmdIndexTable_r( 0, Cmds.Num() - 1, BaseHandleToCmdIndex );
+
 	Owner = InFunction;
 }
 
@@ -3105,6 +3040,8 @@ void FRepLayout::InitFromStruct( UStruct * InStruct )
 
 	AddReturnCmd();
 
+	BuildHandleToCmdIndexTable_r( 0, Cmds.Num() - 1, BaseHandleToCmdIndex );
+
 	Owner = InStruct;
 }
 
@@ -3198,7 +3135,7 @@ void FRepLayout::SerializeProperties_r(
 	}
 }
 
-void FRepLayout::BuildChangeList_r( const int32 CmdStart, const int32 CmdEnd, uint8* Data, const int32 HandleOffset, TArray< uint16 >& Changed ) const
+void FRepLayout::BuildChangeList_r( const TArray< FHandleToCmdIndex >& HandleToCmdIndex, const int32 CmdStart, const int32 CmdEnd, uint8* Data, const int32 HandleOffset, TArray< uint16 >& Changed ) const
 {
 	for ( int32 CmdIndex = CmdStart; CmdIndex < CmdEnd; CmdIndex++ )
 	{
@@ -3212,13 +3149,17 @@ void FRepLayout::BuildChangeList_r( const int32 CmdStart, const int32 CmdEnd, ui
 
 			TArray< uint16 > ChangedLocal;
 
+			TArray< FHandleToCmdIndex >& ArrayHandleToCmdIndex = *HandleToCmdIndex[Cmd.RelativeHandle - 1].HandleToCmdIndex;
+
 			const int32 ArrayCmdStart			= CmdIndex + 1;
 			const int32 ArrayCmdEnd				= Cmd.EndCmd - 1;
-			const int32 NumHandlesPerElement	= ArrayCmdEnd - ArrayCmdStart;
+			const int32 NumHandlesPerElement	= ArrayHandleToCmdIndex.Num();
+
+			check( NumHandlesPerElement > 0 );
 
 			for ( int32 i = 0; i < Array->Num(); i++ )
 			{
-				BuildChangeList_r( ArrayCmdStart, ArrayCmdEnd, ((uint8*)Array->GetData()) + Cmd.ElementSize * i, i * NumHandlesPerElement, ChangedLocal );
+				BuildChangeList_r( ArrayHandleToCmdIndex, ArrayCmdStart, ArrayCmdEnd, ((uint8*)Array->GetData()) + Cmd.ElementSize * i, i * NumHandlesPerElement, ChangedLocal );
 			}
 
 			if ( ChangedLocal.Num() )
@@ -3249,13 +3190,13 @@ void FRepLayout::SendPropertiesForRPC( UObject* Object, UFunction * Function, UA
 		{
 			if ( !Parents[i].Property->Identical_InContainer( Data, NULL, Parents[i].ArrayIndex ) )
 			{
-				BuildChangeList_r( Parents[i].CmdStart, Parents[i].CmdEnd, ( uint8* )Data, 0, Changed );
+				BuildChangeList_r( BaseHandleToCmdIndex, Parents[i].CmdStart, Parents[i].CmdEnd, ( uint8* )Data, 0, Changed );
 			}
 		}
 
 		Changed.Add( 0 ); // Null terminator
 
-		SendProperties_BackwardsCompatible( nullptr, ( uint8* )Data, Channel->Connection, Writer, Changed );
+		SendProperties_BackwardsCompatible( nullptr, nullptr, ( uint8* )Data, Channel->Connection, Writer, Changed );
 
 		return;
 	}
@@ -3359,56 +3300,57 @@ void FRepLayout::SerializePropertiesForStruct( UStruct * Struct, FArchive & Ar,
 	}
 }
 
+void FRepLayout::BuildHandleToCmdIndexTable_r( const int32 CmdStart, const int32 CmdEnd, TArray< FHandleToCmdIndex >& HandleToCmdIndex )
+{
+	for ( int32 CmdIndex = CmdStart; CmdIndex < CmdEnd; CmdIndex++ )
+	{
+		const FRepLayoutCmd& Cmd = Cmds[CmdIndex];
+
+		check( Cmd.Type != REPCMD_Return );
+
+		const int32 Index = HandleToCmdIndex.Add( FHandleToCmdIndex( CmdIndex ) );
+
+		if ( Cmd.Type == REPCMD_DynamicArray )
+		{
+			HandleToCmdIndex[Index].HandleToCmdIndex = TUniquePtr< TArray< FHandleToCmdIndex > >( new TArray< FHandleToCmdIndex >() );
+
+			TArray< FHandleToCmdIndex >& ArrayHandleToCmdIndex = *HandleToCmdIndex[Index].HandleToCmdIndex;
+
+			BuildHandleToCmdIndexTable_r( CmdIndex + 1, Cmd.EndCmd - 1, ArrayHandleToCmdIndex );
+			CmdIndex = Cmd.EndCmd - 1;		// The -1 to handle the ++ in the for loop
+		}
+	}
+}
+
 void FRepLayout::RebuildConditionalProperties( FRepState * RESTRICT	RepState, const FRepChangedPropertyTracker& ChangedTracker, const FReplicationFlags& RepFlags ) const
 {
 	SCOPE_CYCLE_COUNTER( STAT_NetRebuildConditionalTime );
-
-	//UE_LOG( LogRep, Warning, TEXT( "Rebuilding custom properties [%s]" ), *Owner->GetName() );
 	
 	// Setup condition map
-	bool ConditionMap[COND_Max];
-
 	const bool bIsInitial	= RepFlags.bNetInitial		? true : false;
 	const bool bIsOwner		= RepFlags.bNetOwner		? true : false;
 	const bool bIsSimulated	= RepFlags.bNetSimulated	? true : false;
 	const bool bIsPhysics	= RepFlags.bRepPhysics		? true : false;
 	const bool bIsReplay	= RepFlags.bReplay			? true : false;
 
-	ConditionMap[COND_None]					= true;
-	ConditionMap[COND_InitialOnly]			= bIsInitial;
+	RepState->ConditionMap[COND_None]						= true;
+	RepState->ConditionMap[COND_InitialOnly]				= bIsInitial;
 
-	ConditionMap[COND_OwnerOnly]			= bIsOwner;
-	ConditionMap[COND_SkipOwner]			= !bIsOwner;
+	RepState->ConditionMap[COND_OwnerOnly]					= bIsOwner;
+	RepState->ConditionMap[COND_SkipOwner]					= !bIsOwner;
 
-	ConditionMap[COND_SimulatedOnly]			= bIsSimulated;
-	ConditionMap[COND_SimulatedOnlyNoReplay]	= bIsSimulated && !bIsReplay;
-	ConditionMap[COND_AutonomousOnly]			= !bIsSimulated;
+	RepState->ConditionMap[COND_SimulatedOnly]				= bIsSimulated;
+	RepState->ConditionMap[COND_SimulatedOnlyNoReplay]		= bIsSimulated && !bIsReplay;
+	RepState->ConditionMap[COND_AutonomousOnly]				= !bIsSimulated;
 
-	ConditionMap[COND_SimulatedOrPhysics]			= bIsSimulated || bIsPhysics;
-	ConditionMap[COND_SimulatedOrPhysicsNoReplay]	= ( bIsSimulated || bIsPhysics ) && !bIsReplay;
+	RepState->ConditionMap[COND_SimulatedOrPhysics]			= bIsSimulated || bIsPhysics;
+	RepState->ConditionMap[COND_SimulatedOrPhysicsNoReplay]	= ( bIsSimulated || bIsPhysics ) && !bIsReplay;
 
-	ConditionMap[COND_InitialOrOwner]		= bIsInitial || bIsOwner;
-	ConditionMap[COND_ReplayOrOwner]		= bIsReplay || bIsOwner;
-	ConditionMap[COND_ReplayOnly]			= bIsReplay;
-
-	ConditionMap[COND_Custom]				= true;
-
-	RepState->ConditionalLifetime.Empty();
-
-	for ( int32 i = 0; i < ConditionalLifetime.Num(); i++ )
-	{
-		if ( !ConditionMap[ConditionalLifetime[i].Condition] )
-		{
-			continue;
-		}
-		
-		if ( !ChangedTracker.Parents[ConditionalLifetime[i].RepIndex].Active )
-		{
-			continue;
-		}
+	RepState->ConditionMap[COND_InitialOrOwner]				= bIsInitial || bIsOwner;
+	RepState->ConditionMap[COND_ReplayOrOwner]				= bIsReplay || bIsOwner;
+	RepState->ConditionMap[COND_ReplayOnly]					= bIsReplay;
 
-		RepState->ConditionalLifetime.Add( ConditionalLifetime[i].RepIndex );
-	}
+	RepState->ConditionMap[COND_Custom]						= true;
 
 	RepState->RepFlags = RepFlags;
 }
@@ -3423,21 +3365,29 @@ void FRepLayout::InitChangedTracker( FRepChangedPropertyTracker * ChangedTracker
 	}
 }
 
-void FRepLayout::InitRepState( 
-	FRepState *									RepState,
-	UClass *									InObjectClass, 
-	uint8 *										Src, 
-	TSharedPtr< FRepChangedPropertyTracker > &	InRepChangedPropertyTracker ) const
+void FRepLayout::InitShadowData(
+	TArray< uint8, TAlignedHeapAllocator<16> >&	ShadowData,
+	UClass *									InObjectClass,
+	uint8 *										Src ) const
 {
-	RepState->StaticBuffer.Empty();
-	RepState->StaticBuffer.AddZeroed( InObjectClass->GetDefaultsCount() );
+	ShadowData.Empty();
+	ShadowData.AddZeroed( InObjectClass->GetDefaultsCount() );
 
 	// Construct the properties
-	ConstructProperties( RepState );
+	ConstructProperties( ShadowData );
 
 	// Init the properties
-	InitProperties( RepState, Src );
-	
+	InitProperties( ShadowData, Src );
+}
+
+void FRepLayout::InitRepState(
+	FRepState*									RepState,
+	UClass *									InObjectClass, 
+	uint8 *										Src, 
+	TSharedPtr< FRepChangedPropertyTracker > &	InRepChangedPropertyTracker ) const
+{
+	InitShadowData( RepState->StaticBuffer, InObjectClass, Src );
+
 	RepState->RepChangedPropertyTracker = InRepChangedPropertyTracker;
 
 	check( RepState->RepChangedPropertyTracker->Parents.Num() == Parents.Num() );
@@ -3447,9 +3397,9 @@ void FRepLayout::InitRepState(
 	RebuildConditionalProperties( RepState, *InRepChangedPropertyTracker.Get(), FReplicationFlags() );
 }
 
-void FRepLayout::ConstructProperties( FRepState * RepState ) const
+void FRepLayout::ConstructProperties( TArray< uint8, TAlignedHeapAllocator<16> >& ShadowData ) const
 {
-	uint8* StoredData = RepState->StaticBuffer.GetData();
+	uint8* StoredData = ShadowData.GetData();
 
 	// Construct all items
 	for ( int32 i = 0; i < Parents.Num(); i++ )
@@ -3458,16 +3408,16 @@ void FRepLayout::ConstructProperties( FRepState * RepState ) const
 		if ( Parents[i].ArrayIndex == 0 )
 		{
 			PTRINT Offset = Parents[i].Property->ContainerPtrToValuePtr<uint8>( StoredData ) - StoredData;
-			check( Offset >= 0 && Offset < RepState->StaticBuffer.Num() );
+			check( Offset >= 0 && Offset < ShadowData.Num() );
 
 			Parents[i].Property->InitializeValue( StoredData + Offset );
 		}
 	}
 }
 
-void FRepLayout::InitProperties( FRepState * RepState, uint8* Src ) const
+void FRepLayout::InitProperties( TArray< uint8, TAlignedHeapAllocator<16> >& ShadowData, uint8* Src ) const
 {
-	uint8* StoredData = RepState->StaticBuffer.GetData();
+	uint8* StoredData = ShadowData.GetData();
 
 	// Init all items
 	for ( int32 i = 0; i < Parents.Num(); i++ )
@@ -3476,7 +3426,7 @@ void FRepLayout::InitProperties( FRepState * RepState, uint8* Src ) const
 		if ( Parents[i].ArrayIndex == 0 )
 		{
 			PTRINT Offset = Parents[i].Property->ContainerPtrToValuePtr<uint8>( StoredData ) - StoredData;
-			check( Offset >= 0 && Offset < RepState->StaticBuffer.Num() );
+			check( Offset >= 0 && Offset < ShadowData.Num() );
 
 			Parents[i].Property->CopyCompleteValue( StoredData + Offset, Src + Offset );
 		}
diff --git a/Engine/Source/Runtime/Engine/Private/TickTaskManager.cpp b/Engine/Source/Runtime/Engine/Private/TickTaskManager.cpp
--- a/Engine/Source/Runtime/Engine/Private/TickTaskManager.cpp
+++ b/Engine/Source/Runtime/Engine/Private/TickTaskManager.cpp
@@ -1319,14 +1319,15 @@ public:
 	}
 
 	/**
-	 * Ticks the world's dynamic actors based upon their tick group. This function
+	 * Ticks the dynamic actors in the given levels based upon their tick group. This function
 	 * is called once for each ticking group
 	 *
 	 * @param World	- World currently ticking
 	 * @param DeltaSeconds - time in seconds since last tick
 	 * @param TickType - type of tick (viewports only, time only, etc)
+	 * @param LevelsToTick - the levels to tick, may be a subset of InWorld->Levels
 	 */
-	virtual void StartFrame(UWorld* InWorld, float InDeltaSeconds, ELevelTick InTickType) override
+	virtual void StartFrame(UWorld* InWorld, float InDeltaSeconds, ELevelTick InTickType, const TArray<ULevel*>& LevelsToTick) override
 	{
 		SCOPE_CYCLE_COUNTER(STAT_QueueTicks);
 #if !UE_BUILD_SHIPPING
@@ -1344,7 +1345,7 @@ public:
 
 		bTickNewlySpawned = true;
 		TickTaskSequencer.StartFrame();
-		FillLevelList();
+		FillLevelList(LevelsToTick);
 		
 		int32 NumWorkerThread = 0;
 		bool bConcurrentQueue = false;
@@ -1406,7 +1407,7 @@ public:
 	 * @param DeltaSeconds - time in seconds since last tick
 	 * @param TickType - type of tick (viewports only, time only, etc)
 	 */
-	virtual void RunPauseFrame(UWorld* InWorld, float InDeltaSeconds, ELevelTick InTickType) override
+	virtual void RunPauseFrame(UWorld* InWorld, float InDeltaSeconds, ELevelTick InTickType, const TArray<ULevel*>& LevelsToTick) override
 	{
 		bTickNewlySpawned = true; // we don't support new spawns, but lets at least catch them.
 		Context.TickGroup = ETickingGroup(0); // reset this to the start tick group
@@ -1414,7 +1415,7 @@ public:
 		Context.TickType = InTickType;
 		Context.Thread = ENamedThreads::GameThread;
 		Context.World = InWorld;
-		FillLevelList();
+		FillLevelList(LevelsToTick);
 		for( int32 LevelIndex = 0; LevelIndex < LevelList.Num(); LevelIndex++ )
 		{
 			LevelList[LevelIndex]->RunPauseFrame(Context);
@@ -1515,14 +1516,17 @@ private:
 	}
 
 	/** Fill the level list **/
-	void FillLevelList()
+	void FillLevelList(const TArray<ULevel*>& Levels)
 	{
 		check(!LevelList.Num());
-		check(Context.World->TickTaskLevel);
-		LevelList.Add(Context.World->TickTaskLevel);
-		for( int32 LevelIndex = 0; LevelIndex < Context.World->GetNumLevels(); LevelIndex++ )
+		if (!Context.World->GetActiveLevelCollection() || Context.World->GetActiveLevelCollection()->GetType() == ELevelCollectionType::DynamicSourceLevels)
 		{
-			ULevel* Level = Context.World->GetLevel(LevelIndex);
+			check(Context.World->TickTaskLevel);
+			LevelList.Add(Context.World->TickTaskLevel);
+		}
+		for( int32 LevelIndex = 0; LevelIndex < Levels.Num(); LevelIndex++ )
+		{
+			ULevel* Level = Levels[LevelIndex];
 			if (Level->bIsVisible)
 			{
 				check(Level->TickTaskLevel);
diff --git a/Engine/Source/Runtime/Engine/Private/TimerManager.cpp b/Engine/Source/Runtime/Engine/Private/TimerManager.cpp
--- a/Engine/Source/Runtime/Engine/Private/TimerManager.cpp
+++ b/Engine/Source/Runtime/Engine/Private/TimerManager.cpp
@@ -30,6 +30,7 @@ namespace
 FTimerManager::FTimerManager()
 	: InternalTime(0.0)
 	, LastTickedFrame(static_cast<uint64>(-1))
+	, OwningGameInstance(nullptr)
 {
 	if (IsRunningDedicatedServer())
 	{
@@ -242,6 +243,13 @@ void FTimerManager::InternalSetTimer(FTimerData& NewTimerData, float InRate, boo
 		NewTimerData.bLoop = InbLoop;
 		NewTimerData.bRequiresDelegate = NewTimerData.TimerDelegate.IsBound();
 
+		// Set level collection
+		const UWorld* const OwningWorld = OwningGameInstance ? OwningGameInstance->GetWorld() : nullptr;
+		if (OwningWorld && OwningWorld->GetActiveLevelCollection())
+		{
+			NewTimerData.LevelCollection = OwningWorld->GetActiveLevelCollection()->GetType();
+		}
+
 		const float FirstDelay = (InFirstDelay >= 0.f) ? InFirstDelay : InRate;
 
 		if (HasBeenTickedThisFrame())
@@ -272,6 +280,15 @@ void FTimerManager::InternalSetTimerForNextTick(FTimerUnifiedDelegate const& InD
 	NewTimerData.TimerDelegate = InDelegate;
 	NewTimerData.ExpireTime = InternalTime;
 	NewTimerData.Status = ETimerStatus::Active;
+
+	// Set level collection
+	const UWorld* const OwningWorld = OwningGameInstance ? OwningGameInstance->GetWorld() : nullptr;
+	if (OwningWorld && OwningWorld->GetActiveLevelCollection())
+	{
+		NewTimerData.LevelCollection = OwningWorld->GetActiveLevelCollection()->GetType();
+	}
+
+
 	ActiveTimerHeap.HeapPush(NewTimerData);
 }
 
@@ -523,6 +540,18 @@ void FTimerManager::Tick(float DeltaTime)
 		{
 			// Timer has expired! Fire the delegate, then handle potential looping.
 
+			// Set the relevant level context for this timer
+			const FLevelCollection* LevelCollection = nullptr;
+			UWorld* LevelCollectionWorld = nullptr;
+			UWorld* const OwningWorld = OwningGameInstance ? OwningGameInstance->GetWorld() : nullptr;
+			if (OwningGameInstance && OwningWorld)
+			{
+				LevelCollection = OwningWorld->FindCollectionByType(Top.LevelCollection);
+				LevelCollectionWorld = OwningWorld;
+			}
+
+			FScopedLevelCollectionContextSwitch LevelContext(LevelCollection, LevelCollectionWorld);
+
 			// Remove it from the heap and store it while we're executing
 			ActiveTimerHeap.HeapPop(CurrentlyExecutingTimer, /*bAllowShrinking=*/ false);
 			CurrentlyExecutingTimer.Status = ETimerStatus::Executing;
diff --git a/Engine/Source/Runtime/Engine/Private/UnrealEngine.cpp b/Engine/Source/Runtime/Engine/Private/UnrealEngine.cpp
--- a/Engine/Source/Runtime/Engine/Private/UnrealEngine.cpp
+++ b/Engine/Source/Runtime/Engine/Private/UnrealEngine.cpp
@@ -521,9 +521,7 @@ namespace
 
 		check(NewWorld);
 
-#if WITH_EDITOR
 		OutPackage->PIEInstanceID = WorldContext.PIEInstance;
-#endif
 
 		// Rename streaming levels to PIE
 		for (ULevelStreaming* StreamingLevel : NewWorld->StreamingLevels)
@@ -8693,6 +8691,23 @@ void DestroyNamedNetDriver_Local(FWorldContext &Context, FName NetDriverName)
 			NetDriver->Shutdown();
 			NetDriver->LowLevelDestroy();
 			Context.ActiveNetDrivers.RemoveAtSwap(Index);
+
+			// Remove this driver from the main level collection
+			const ELevelCollectionType DriverType = NetDriver->GetDuplicateLevelID() == INDEX_NONE ? ELevelCollectionType::DynamicSourceLevels : ELevelCollectionType::DynamicDuplicatedLevels;
+			FLevelCollection* const LevelCollection = Context.World()->FindCollectionByType(DriverType);
+			if (LevelCollection)
+			{
+				if (LevelCollection->GetNetDriver() == NetDriver)
+				{
+					LevelCollection->SetNetDriver(nullptr);
+				}
+
+				if (LevelCollection->GetDemoNetDriver() == NetDriver)
+				{
+					LevelCollection->SetDemoNetDriver(nullptr);
+				}
+			}
+
 			break;
 		}
 	}
@@ -9908,6 +9923,12 @@ bool UEngine::LoadMap( FWorldContext& WorldContext, FURL URL, class UPendingNetG
 
 		// Make sure "always loaded" sub-levels are fully loaded
 		WorldContext.World()->FlushLevelStreaming(EFlushLevelStreamingType::Visibility);
+
+		if (!GIsEditor && !IsRunningDedicatedServer())
+		{
+			// If requested, duplicate dynamic levels here after the source levels are created.
+			WorldContext.World()->DuplicateRequestedLevels(FName(*URL.Map));
+		}
 	}
 	
 	// Note that AI system will be created only if ai-system-creation conditions are met
@@ -10105,6 +10126,9 @@ void UEngine::MovePendingLevel(FWorldContext &Context)
 		// The pending net driver is renamed to the current "game net driver"
 		NetDriver->NetDriverName = NAME_GameNetDriver;
 		NetDriver->SetWorld(Context.World());
+
+		FLevelCollection& SourceLevels = Context.World()->FindOrAddCollectionByType(ELevelCollectionType::DynamicSourceLevels);
+		SourceLevels.SetNetDriver(NetDriver);
 	}
 
 	// Attach the DemoNetDriver to the world if there is one
@@ -10112,6 +10136,9 @@ void UEngine::MovePendingLevel(FWorldContext &Context)
 	{
 		Context.PendingNetGame->DemoNetDriver->SetWorld( Context.World() );
 		Context.World()->DemoNetDriver = Context.PendingNetGame->DemoNetDriver;
+
+		FLevelCollection& MainLevels = Context.World()->FindOrAddCollectionByType(ELevelCollectionType::DynamicSourceLevels);
+		MainLevels.SetDemoNetDriver(Context.PendingNetGame->DemoNetDriver);
 	}
 
 	// Reset the Navigation System
@@ -10553,6 +10580,30 @@ bool UEngine::WorldHasValidContext(UWorld *InWorld)
 	return (GetWorldContextFromWorld(InWorld) != NULL);
 }
 
+bool UEngine::IsWorldDuplicate(const UWorld* const InWorld)
+{
+	// World is considered a duplicate if it's the outer of a level in a duplicate levels collection
+	for (const FWorldContext& Context : WorldList)
+	{
+		if (Context.World())
+		{	
+			const FLevelCollection* const Collection = Context.World()->FindCollectionByType(ELevelCollectionType::DynamicDuplicatedLevels);
+			if (Collection)
+			{
+				for (const ULevel* const Level : Collection->GetLevels())
+				{
+					if (Level->GetOuter() == InWorld)
+					{
+						return true;
+					}
+				}
+			}
+		}
+	}
+
+	return false;
+}
+
 void UEngine::VerifyLoadMapWorldCleanup()
 {
 	// All worlds at this point should be the CurrentWorld of some context, preview worlds, or streaming level
@@ -10563,7 +10614,7 @@ void UEngine::VerifyLoadMapWorldCleanup()
 		const bool bIsPersistantWorldType = (World->WorldType == EWorldType::Inactive) || (World->WorldType == EWorldType::Preview);
 		if (!bIsPersistantWorldType && !WorldHasValidContext(World))
 		{
-			if (World->PersistentLevel == nullptr || !WorldHasValidContext(World->PersistentLevel->OwningWorld))
+			if ((World->PersistentLevel == nullptr || !WorldHasValidContext(World->PersistentLevel->OwningWorld)) && !IsWorldDuplicate(World))
 			{
 				// Print some debug information...
 				UE_LOG(LogLoad, Log, TEXT("%s not cleaned up by garbage collection! "), *World->GetFullName());
diff --git a/Engine/Source/Runtime/Engine/Private/World.cpp b/Engine/Source/Runtime/Engine/Private/World.cpp
--- a/Engine/Source/Runtime/Engine/Private/World.cpp
+++ b/Engine/Source/Runtime/Engine/Private/World.cpp
@@ -134,6 +134,111 @@ FActorSpawnParameters::FActorSpawnParameters()
 
 PRAGMA_ENABLE_DEPRECATION_WARNINGS
 
+
+
+FLevelCollection::FLevelCollection()
+	: CollectionType(ELevelCollectionType::DynamicSourceLevels)
+	, GameState(nullptr)
+	, NetDriver(nullptr)
+	, DemoNetDriver(nullptr)
+	, PersistentLevel(nullptr)
+	, bIsVisible(true)
+{
+}
+
+FLevelCollection::~FLevelCollection()
+{
+	for (ULevel* Level : Levels)
+	{
+		check(Level->GetCachedLevelCollection() == this);
+		Level->SetCachedLevelCollection(nullptr);
+	}
+}
+
+FLevelCollection::FLevelCollection(FLevelCollection&& Other)
+	: CollectionType(Other.CollectionType)
+	, GameState(Other.GameState)
+	, NetDriver(Other.NetDriver)
+	, DemoNetDriver(Other.DemoNetDriver)
+	, PersistentLevel(Other.PersistentLevel)
+	, Levels(MoveTemp(Other.Levels))
+	, bIsVisible(Other.bIsVisible)
+{
+	for (ULevel* Level : Levels)
+	{
+		check(Level->GetCachedLevelCollection() == &Other);
+		Level->SetCachedLevelCollection(this);
+	}
+}
+
+FLevelCollection& FLevelCollection::operator=(FLevelCollection&& Other)
+{
+	if (this != &Other)
+	{
+		CollectionType = Other.CollectionType;
+		GameState = Other.GameState;
+		NetDriver = Other.NetDriver;
+		DemoNetDriver = Other.DemoNetDriver;
+		PersistentLevel = Other.PersistentLevel;
+		Levels = MoveTemp(Other.Levels);
+		bIsVisible = Other.bIsVisible;
+
+		for (ULevel* Level : Levels)
+		{
+			check(Level->GetCachedLevelCollection() == &Other);
+			Level->SetCachedLevelCollection(this);
+		}
+	}
+
+	return *this;
+}
+
+void FLevelCollection::SetPersistentLevel(ULevel* const Level)
+{
+	PersistentLevel = Level;
+}
+
+void FLevelCollection::AddLevel(ULevel* const Level)
+{
+	if (Level)
+	{
+		// Sanity check that Level isn't already in another collection.
+		check(Level->GetCachedLevelCollection() == nullptr);
+		Levels.Add(Level);
+		Level->SetCachedLevelCollection(this);
+	}
+}
+
+void FLevelCollection::RemoveLevel(ULevel* const Level)
+{
+	if (Level)
+	{
+		check(Level->GetCachedLevelCollection() == this);
+		Level->SetCachedLevelCollection(nullptr);
+		Levels.Remove(Level);
+	}
+}
+
+FScopedLevelCollectionContextSwitch::FScopedLevelCollectionContextSwitch(const FLevelCollection* const InLevelCollection, UWorld* const InWorld)
+	: World(InWorld)
+	, SavedTickingCollection(InWorld ? InWorld->GetActiveLevelCollection() : nullptr)
+{
+	if (InLevelCollection && World)
+	{
+		World->SetActiveLevelCollection(InLevelCollection);
+	}
+}
+	
+FScopedLevelCollectionContextSwitch::~FScopedLevelCollectionContextSwitch()
+{
+	if (World)
+	{
+		World->SetActiveLevelCollection(SavedTickingCollection);
+	}
+}
+
+
+
 /*-----------------------------------------------------------------------------
 	UWorld implementation.
 -----------------------------------------------------------------------------*/
@@ -163,6 +268,7 @@ FWorldDelegates::FRefreshLevelScriptActionsEvent FWorldDelegates::RefreshLevelSc
 
 UWorld::UWorld( const FObjectInitializer& ObjectInitializer )
 : UObject(ObjectInitializer)
+, ActiveLevelCollection(nullptr)
 #if WITH_EDITOR
 , HierarchicalLODBuilder(new FHierarchicalLODBuilder(this))
 #endif
@@ -1104,6 +1210,8 @@ void UWorld::InitWorld(const InitializationValues IVS)
 	// update it's bIsDefaultLevel
 	bIsDefaultLevel = (FPaths::GetBaseFilename(GetMapName()) == FPaths::GetBaseFilename(UGameMapsSettings::GetGameDefaultMap()));
 
+	ConditionallyCreateDefaultLevelCollections();
+
 	// We're initialized now.
 	bIsWorldInitialized = true;
 
@@ -1117,6 +1225,33 @@ void UWorld::InitWorld(const InitializationValues IVS)
 	BroadcastLevelsChanged();
 }
 
+void UWorld::ConditionallyCreateDefaultLevelCollections()
+{
+	// Create main level collection. The persistent level will always be considered dynamic.
+	if (!FindCollectionByType(ELevelCollectionType::DynamicSourceLevels))
+	{
+		FLevelCollection& DynamicCollection = FindOrAddCollectionByType(ELevelCollectionType::DynamicSourceLevels);
+		DynamicCollection.SetPersistentLevel(PersistentLevel);
+		
+		// Don't add the persistent level if it is already a member of another collection.
+		// This may be the case if, for example, this world is the outer of a streaming level,
+		// in which case the persistent level may be in one of the collections in the streaming level's OwningWorld.
+		if (PersistentLevel->GetCachedLevelCollection() == nullptr)
+		{
+			DynamicCollection.AddLevel(PersistentLevel);
+		}
+
+		// Default to the dynamic source collection
+		ActiveLevelCollection = &DynamicCollection;
+	}
+
+	if (!FindCollectionByType(ELevelCollectionType::StaticLevels))
+	{
+		FLevelCollection& StaticCollection = FindOrAddCollectionByType(ELevelCollectionType::StaticLevels);
+		StaticCollection.SetPersistentLevel(PersistentLevel);
+	}
+}
+
 void UWorld::InitializeNewWorld(const InitializationValues IVS)
 {
 	if (!IVS.bTransactional)
@@ -1979,7 +2114,7 @@ void UWorld::AddToWorld( ULevel* Level, const FTransform& LevelTransform )
 					// Remap packagename for PIE networking before sending out to server
 					FName PackageName = Level->GetOutermost()->GetFName();
 					FString PackageNameStr = PackageName.ToString();
-					if (GEngine->NetworkRemapPath(this, PackageNameStr, false))
+					if (GEngine->NetworkRemapPath(LocalPlayerController->GetNetDriver(), PackageNameStr, false))
 					{
 						PackageName = FName(*PackageNameStr);
 					}
@@ -2098,7 +2233,7 @@ void UWorld::RemoveFromWorld( ULevel* Level )
 					// Remap packagename for PIE networking before sending out to server
 					FName PackageName = Level->GetOutermost()->GetFName();
 					FString PackageNameStr = PackageName.ToString();
-					if (GEngine->NetworkRemapPath(this, PackageNameStr, false))
+					if (GEngine->NetworkRemapPath(LocalPlayerController->GetNetDriver(), PackageNameStr, false))
 					{
 						PackageName = FName(*PackageNameStr);
 					}
@@ -2274,9 +2409,8 @@ void UWorld::RenameToPIEWorld(int32 PIEInstanceID)
 {
 	UPackage* WorldPackage = GetOutermost();
 
-#if WITH_EDITOR
 	WorldPackage->PIEInstanceID = PIEInstanceID;
-#endif
+
 	const FString PIEPackageName = *UWorld::ConvertToPIEPackageName(WorldPackage->GetName(), PIEInstanceID);
 	WorldPackage->Rename(*PIEPackageName);
 
@@ -2367,9 +2501,7 @@ UWorld* UWorld::DuplicateWorldForPIE(const FString& PackageName, UWorld* OwningW
 	UWorld::WorldTypePreLoadMap.FindOrAdd(PrefixedLevelFName) = EWorldType::PIE;
 	UPackage* PIELevelPackage = CastChecked<UPackage>(CreatePackage(NULL,*PrefixedLevelName));
 	PIELevelPackage->SetPackageFlags(PKG_PlayInEditor);
-#if WITH_EDITOR
 	PIELevelPackage->PIEInstanceID = WorldContext.PIEInstance;
-#endif
 	PIELevelPackage->SetGuid( EditorLevelPackage->GetGuid() );
 
 	// Set up string asset reference fixups
@@ -2410,7 +2542,7 @@ UWorld* UWorld::DuplicateWorldForPIE(const FString& PackageName, UWorld* OwningW
 	{
 		ULevel* EditorLevel = EditorLevelWorld->PersistentLevel;
 		ULevel* PIELevel = PIELevelWorld->PersistentLevel;
-#if WITH_EDITOR
+
 		// Fixup model components. The index buffers have been created for the components in the EditorWorld and the order
 		// in which components were post-loaded matters. So don't try to guarantee a particular order here, just copy the
 		// elements over.
@@ -2428,7 +2560,7 @@ UWorld* UWorld::DuplicateWorldForPIE(const FString& PackageName, UWorld* OwningW
 				DestComponent->CopyElementsFrom(SrcComponent);
 			}
 		}
-#endif
+
 		// We have to place PIELevel at his local position in case EditorLevel was visible
 		// Correct placement will occur during UWorld::AddToWorld
 		if (EditorLevel->OwningWorld->WorldComposition && EditorLevel->bIsVisible)
@@ -2441,7 +2573,7 @@ UWorld* UWorld::DuplicateWorldForPIE(const FString& PackageName, UWorld* OwningW
 	PIELevelWorld->ClearFlags(RF_Standalone);
 	EditorLevelWorld->TransferBlueprintDebugReferences(PIELevelWorld);
 
-	UE_LOG(LogWorld, Log, TEXT("PIE: Copying PIE streaming level from %s to %s. OwningWorld: %s"),
+	UE_LOG(LogWorld, Verbose, TEXT("PIE: Copying PIE streaming level from %s to %s. OwningWorld: %s"),
 		*EditorLevelWorld->GetPathName(),
 		*PIELevelWorld->GetPathName(),
 		OwningWorld ? *OwningWorld->GetPathName() : TEXT("<null>"));
@@ -2723,6 +2855,10 @@ bool UWorld::AreAlwaysLoadedLevelsLoaded() const
 
 void UWorld::AsyncLoadAlwaysLoadedLevelsForSeamlessTravel()
 {
+	// Need to do this now so that data can be set correctly on the loaded world's collections.
+	// This normally happens in InitWorld but that happens too late for seamless travel.
+	ConditionallyCreateDefaultLevelCollections();
+
 	for (int32 LevelIndex = 0; LevelIndex < StreamingLevels.Num(); LevelIndex++)
 	{
 		ULevelStreaming* LevelStreaming = StreamingLevels[LevelIndex];
@@ -3097,7 +3233,7 @@ void UWorld::InitializeActorsForPlay(const FURL& InURL, bool bResetTime)
 			// Remap packagename for PIE networking before sending out to server
 			FName PackageName = SubLevel->GetOutermost()->GetFName();
 			FString PackageNameStr = PackageName.ToString();
-			if (GEngine->NetworkRemapPath(this, PackageNameStr, false))
+			if (GEngine->NetworkRemapPath(GetNetDriver(), PackageNameStr, false))
 			{
 				PackageName = FName(*PackageNameStr);
 			}
@@ -3216,18 +3352,7 @@ void UWorld::CleanupWorld(bool bSessionEnded, bool bCleanupResources, UWorld* Ne
 	// to avoid a GC crash for not cleaning up the gameinfo referenced by levelscript
 	if (IsGameWorld() && !GIsEditor && !IsRunningCommandlet() && bSessionEnded && bCleanupResources && PersistentLevel)
 	{
-		if( ULevelScriptBlueprint* LSBP = PersistentLevel->GetLevelScriptBlueprint(true) )
-		{
-			if( LSBP->SkeletonGeneratedClass )
-			{
-				LSBP->SkeletonGeneratedClass->ClassGeneratedBy = NULL; 
-			}
-
-			if( LSBP->GeneratedClass )
-			{
-				LSBP->GeneratedClass->ClassGeneratedBy = NULL; 
-			}
-		}
+		PersistentLevel->CleanupLevelScriptBlueprint();
 	}
 #endif //WITH_EDITOR
 
@@ -3275,6 +3400,27 @@ void UWorld::CleanupWorld(bool bSessionEnded, bool bCleanupResources, UWorld* Ne
 			World->CleanupWorld(bSessionEnded, bCleanupResources, NewWorld);
 		}
 	}
+
+	// Clean up any duplicated levels.
+	const FLevelCollection* const DuplicateCollection = FindCollectionByType(ELevelCollectionType::DynamicDuplicatedLevels);
+	if (DuplicateCollection)
+	{
+		for (const ULevel* Level : DuplicateCollection->GetLevels())
+		{
+			if (!Level)
+			{
+				continue;
+			}
+
+			UWorld* const LevelWorld = CastChecked<UWorld>(Level->GetOuter());
+			if (LevelWorld != this)
+			{
+				LevelWorld->CleanupWorld(bSessionEnded, bCleanupResources, NewWorld);
+			}
+		}
+	}
+
+	LevelCollections.Empty();
 }
 
 UGameViewportClient* UWorld::GetGameViewport() const
@@ -4160,6 +4306,11 @@ bool UWorld::Listen( FURL& InURL )
 	{
 		NetDriver = GEngine->FindNamedNetDriver(this, NAME_GameNetDriver);
 		NetDriver->SetWorld(this);
+		FLevelCollection* const SourceCollection = FindCollectionByType(ELevelCollectionType::DynamicSourceLevels);
+		if (SourceCollection)
+		{
+			SourceCollection->SetNetDriver(NetDriver);
+		}
 	}
 
 	if (NetDriver == NULL)
@@ -4175,6 +4326,11 @@ bool UWorld::Listen( FURL& InURL )
 		UE_LOG(LogWorld, Log,  TEXT("Failed to listen: %s"), *Error );
 		NetDriver->SetWorld(NULL);
 		NetDriver = NULL;
+		FLevelCollection* SourceCollection = FindCollectionByType(ELevelCollectionType::DynamicSourceLevels);
+		if (SourceCollection)
+		{
+			SourceCollection->SetNetDriver(nullptr);
+		}
 		return false;
 	}
 	static bool LanPlay = FParse::Param(FCommandLine::Get(),TEXT("lanplay"));
@@ -4671,6 +4827,9 @@ void FSeamlessTravelHandler::StartLoadingDestination()
 
 void FSeamlessTravelHandler::CopyWorldData()
 {
+	FLevelCollection* const CurrentCollection = CurrentWorld->FindCollectionByType(ELevelCollectionType::DynamicSourceLevels);
+	FLevelCollection* const LoadedCollection = LoadedWorld->FindCollectionByType(ELevelCollectionType::DynamicSourceLevels);
+
 	// If we are doing seamless travel for replay playback, then make sure to transfer the replay driver over to the new world
 	if ( CurrentWorld->DemoNetDriver && CurrentWorld->DemoNetDriver->IsPlaying() )
 	{
@@ -4678,14 +4837,32 @@ void FSeamlessTravelHandler::CopyWorldData()
 		CurrentWorld->DemoNetDriver = nullptr;
 		OldDriver->SetWorld( LoadedWorld );
 		LoadedWorld->DemoNetDriver = OldDriver;
+
+		if (CurrentCollection && LoadedCollection)
+		{
+			LoadedCollection->SetDemoNetDriver(OldDriver);
+			CurrentCollection->SetDemoNetDriver(nullptr);
+		}
 	}
 	else
 	{
 		CurrentWorld->DestroyDemoNetDriver();
+
+		if (CurrentCollection)
+		{
+			CurrentCollection->SetNetDriver(nullptr);
+		}
 	}
 
 	UNetDriver* const NetDriver = CurrentWorld->GetNetDriver();
 	LoadedWorld->SetNetDriver(NetDriver);
+
+	if (CurrentCollection && LoadedCollection)
+	{
+		LoadedCollection->SetNetDriver(NetDriver);
+		CurrentCollection->SetNetDriver(nullptr);
+	}
+
 	if (NetDriver != NULL)
 	{
 		CurrentWorld->SetNetDriver(NULL);
@@ -4758,6 +4935,12 @@ UWorld* FSeamlessTravelHandler::Tick()
 		{
 			// Make sure there are no pending visibility requests.
 			CurrentWorld->FlushLevelStreaming(EFlushLevelStreamingType::Visibility);
+				
+			if (!GIsEditor && !IsRunningDedicatedServer() && bSwitchedToDefaultMap)
+			{
+				// If requested, duplicate dynamic levels here after the source levels are created.
+				LoadedWorld->DuplicateRequestedLevels(LoadedWorld->GetOuter()->GetFName());
+			}
 
 			if (CurrentWorld->GetGameState())
 			{
@@ -5680,12 +5863,27 @@ void UWorld::SetGameState(AGameStateBase* NewGameState)
 	}
 
 	GameState = NewGameState;
+
+	// Set the GameState on the LevelCollection it's associated with.
+	const ULevel* const CachedLevel = NewGameState->GetLevel();
+	FLevelCollection* const FoundCollection = NewGameState ? CachedLevel->GetCachedLevelCollection() : nullptr;
+	if (FoundCollection)
+	{
+		FoundCollection->SetGameState(NewGameState);
+
+		// For now the static levels use the same GameState as the source dynamic levels.
+		if (FoundCollection->GetType() == ELevelCollectionType::DynamicSourceLevels)
+		{
+			FLevelCollection& StaticLevels = FindOrAddCollectionByType(ELevelCollectionType::StaticLevels);
+			StaticLevels.SetGameState(NewGameState);
+		}
+	}
 }
 
 void UWorld::CopyGameState(AGameModeBase* FromGameMode, AGameStateBase* FromGameState)
 {
 	AuthorityGameMode = FromGameMode;
-	GameState = FromGameState;
+	SetGameState(FromGameState);
 }
 
 void UWorld::GetLightMapsAndShadowMaps(ULevel* Level, TArray<UTexture2D*>& OutLightMapsAndShadowMaps)
@@ -5755,6 +5953,180 @@ void UWorld::CreateFXSystem()
 	}
 }
 
+FLevelCollection& UWorld::FindOrAddCollectionByType(const ELevelCollectionType InType)
+{
+	for (FLevelCollection& LC : LevelCollections)
+	{
+		if (LC.GetType() == InType)
+		{
+			return LC;
+		}
+	}
+
+	// Not found, add a new one.
+	FLevelCollection NewLC;
+	NewLC.SetType(InType);
+	LevelCollections.Add(MoveTemp(NewLC));
+	return LevelCollections.Last();
+}
+
+FLevelCollection* UWorld::FindCollectionByType(const ELevelCollectionType InType)
+{
+	for (FLevelCollection& LC : LevelCollections)
+	{
+		if (LC.GetType() == InType)
+		{
+			return &LC;
+		}
+	}
+
+	return nullptr;
+}
+
+const FLevelCollection* UWorld::FindCollectionByType(const ELevelCollectionType InType) const
+{
+	for (const FLevelCollection& LC : LevelCollections)
+	{
+		if (LC.GetType() == InType)
+		{
+			return &LC;
+		}
+	}
+
+	return nullptr;
+}
+
+void UWorld::SetActiveLevelCollection(const FLevelCollection* InCollection)
+{
+	ActiveLevelCollection = InCollection;
+
+	PersistentLevel = InCollection->GetPersistentLevel();
+	SetCurrentLevel(InCollection->GetPersistentLevel());
+	GameState = InCollection->GetGameState();
+	NetDriver = InCollection->GetNetDriver();
+	DemoNetDriver = InCollection->GetDemoNetDriver();
+}
+
+static ULevel* DuplicateLevelWithPrefix(ULevel* InLevel, int32 InstanceID )
+{
+	if (!InLevel || !InLevel->GetOutermost())
+	{
+		return nullptr;
+	}
+
+	const double DuplicateStart = FPlatformTime::Seconds();
+
+	UWorld* OriginalOwningWorld = CastChecked<UWorld>(InLevel->GetOuter());
+	UPackage* OriginalPackage = InLevel->GetOutermost();
+
+	const FString OriginalPackageName = OriginalPackage->GetName();
+
+	// Use a PIE prefix for the new package
+	const FString PrefixedPackageName = UWorld::ConvertToPIEPackageName( OriginalPackageName, InstanceID );
+
+	// Create a package for duplicated level
+	UPackage* NewPackage = CastChecked<UPackage>( CreatePackage( NULL, *PrefixedPackageName ) );
+	NewPackage->SetPackageFlags( PKG_PlayInEditor );
+	NewPackage->PIEInstanceID = InstanceID;
+	NewPackage->FileName = OriginalPackage->FileName;
+	NewPackage->SetGuid( OriginalPackage->GetGuid() );
+
+	GPlayInEditorID = InstanceID;
+
+	// Create "vestigial" world for the persistent level - it's OwningWorld will still be the main world,
+	// but we're treating it like a streaming level (even though it's a duplicate of the persistent level).
+	UWorld* NewWorld = NewObject<UWorld>(NewPackage, OriginalOwningWorld->GetFName());
+	NewWorld->SetFlags(RF_Transactional);
+	NewWorld->WorldType = EWorldType::Game;
+	NewWorld->FeatureLevel = ERHIFeatureLevel::Num;
+
+	TArray<FString> PackageNamesBeingDuplicatedForPIE;
+	PackageNamesBeingDuplicatedForPIE.Add( PrefixedPackageName );
+
+	FStringAssetReference::SetPackageNamesBeingDuplicatedForPIE( PackageNamesBeingDuplicatedForPIE );
+
+	ULevel::StreamedLevelsOwningWorld.Add(NewPackage->GetFName(), OriginalOwningWorld);
+
+	FObjectDuplicationParameters Parameters( InLevel, NewWorld );
+		
+	Parameters.DestName			= InLevel->GetFName();
+	Parameters.DestClass		= InLevel->GetClass();
+	Parameters.FlagMask			= RF_AllFlags;
+	Parameters.InternalFlagMask = EInternalObjectFlags::AllFlags;
+	Parameters.PortFlags		= PPF_DuplicateForPIE;
+
+	ULevel* NewLevel = CastChecked<ULevel>( StaticDuplicateObjectEx( Parameters ) );
+
+	ULevel::StreamedLevelsOwningWorld.Remove(NewPackage->GetFName());
+	FStringAssetReference::ClearPackageNamesBeingDuplicatedForPIE();
+
+	// Fixup model components. The index buffers have been created for the components in the source world and the order
+	// in which components were post-loaded matters. So don't try to guarantee a particular order here, just copy the
+	// elements over.
+	if ( NewLevel->Model != NULL
+			&& NewLevel->Model == InLevel->Model
+			&& NewLevel->ModelComponents.Num() == InLevel->ModelComponents.Num() )
+	{
+		NewLevel->Model->ClearLocalMaterialIndexBuffersData();
+		for ( int32 ComponentIndex = 0; ComponentIndex < NewLevel->ModelComponents.Num(); ++ComponentIndex )
+		{
+			UModelComponent* SrcComponent = InLevel->ModelComponents[ComponentIndex];
+			UModelComponent* DestComponent = NewLevel->ModelComponents[ComponentIndex];
+			DestComponent->CopyElementsFrom( SrcComponent );
+		}
+	}
+
+	const double DuplicateEnd = FPlatformTime::Seconds();
+	const float TotalSeconds = ( DuplicateEnd - DuplicateStart );
+
+	UE_LOG( LogNet, Log, TEXT( "DuplicateLevelWithPrefix. TotalSeconds: %2.2f" ), TotalSeconds );
+
+	GPlayInEditorID = -1;
+
+	return NewLevel;
+}
+
+void UWorld::DuplicateRequestedLevels(const FName MapName)
+{
+	if (GEngine->Experimental_ShouldPreDuplicateMap(MapName))
+	{
+		// Duplicate the persistent level and only dynamic levels, but don't add them to the world.
+		FLevelCollection DuplicateLevels;
+		DuplicateLevels.SetType(ELevelCollectionType::DynamicDuplicatedLevels);
+		DuplicateLevels.SetIsVisible(false);
+		ULevel* const DuplicatePersistentLevel = DuplicateLevelWithPrefix(PersistentLevel, 1);
+		if (!DuplicatePersistentLevel)
+		{
+			UE_LOG(LogWorld, Warning, TEXT("UWorld::DuplicateRequestedLevels: failed to duplicate persistent level %s. No duplicate level collection will be created."),
+				*GetFullNameSafe(PersistentLevel));
+			return;
+		}
+		// Don't tell the server about this level
+		DuplicatePersistentLevel->bClientOnlyVisible = true;
+		DuplicateLevels.SetPersistentLevel(DuplicatePersistentLevel);
+		DuplicateLevels.AddLevel(DuplicatePersistentLevel);
+
+		for (ULevelStreaming* StreamingLevel : StreamingLevels)
+		{
+			if (StreamingLevel && !StreamingLevel->bIsStatic)
+			{
+				ULevel* DuplicatedLevel = DuplicateLevelWithPrefix(StreamingLevel->LoadedLevel, 1);
+				if (!DuplicatedLevel)
+				{
+					UE_LOG(LogWorld, Warning, TEXT("UWorld::DuplicateRequestedLevels: failed to duplicate streaming level %s. No duplicate level collection will be created."),
+						*GetFullNameSafe(StreamingLevel->LoadedLevel));
+					return;
+				}
+				// Don't tell the server about these levels
+				DuplicatedLevel->bClientOnlyVisible = true;
+				DuplicateLevels.AddLevel(DuplicatedLevel);
+			}
+		}
+
+		LevelCollections.Add(MoveTemp(DuplicateLevels));
+	}
+}
+
 #if WITH_EDITOR
 void UWorld::ChangeFeatureLevel(ERHIFeatureLevel::Type InFeatureLevel, bool bShowSlowProgressDialog )
 {
@@ -5882,6 +6254,34 @@ static FAutoConsoleCommand DumpVisibleActorsCmd(
 	FConsoleCommandDelegate::CreateStatic(DumpVisibleActors)
 	);
 
+static void DumpLevelCollections(UWorld* InWorld)
+{
+	if (!InWorld)
+	{
+		return;
+	}
+
+	UE_LOG(LogWorld, Log, TEXT("--- Dumping LevelCollections ---"));
+
+	for(const FLevelCollection& LC : InWorld->GetLevelCollections())
+	{
+		UE_LOG(LogWorld, Log, TEXT("%d: %d levels."), static_cast<int32>(LC.GetType()), LC.GetLevels().Num());
+		UE_LOG(LogWorld, Log, TEXT("  PersistentLevel: %s"), *GetFullNameSafe(LC.GetPersistentLevel()));
+		UE_LOG(LogWorld, Log, TEXT("  GameState: %s"), *GetFullNameSafe(LC.GetGameState()));
+		UE_LOG(LogWorld, Log, TEXT("  Levels:"));
+		for (const ULevel* Level : LC.GetLevels())
+		{
+			UE_LOG(LogWorld, Log, TEXT("    %s"), *GetFullNameSafe(Level));
+		}
+	}
+}
+
+static FAutoConsoleCommandWithWorld DumpLevelCollectionsCmd(
+	TEXT("DumpLevelCollections"),
+	TEXT("Dump level collections in the current world."),
+	FConsoleCommandWithWorldDelegate::CreateStatic(DumpLevelCollections)
+	);
+
 #if WITH_EDITOR
 FAsyncPreRegisterDDCRequest::~FAsyncPreRegisterDDCRequest()
 {
diff --git a/Engine/Source/Runtime/Engine/Private/WorldComposition.cpp b/Engine/Source/Runtime/Engine/Private/WorldComposition.cpp
--- a/Engine/Source/Runtime/Engine/Private/WorldComposition.cpp
+++ b/Engine/Source/Runtime/Engine/Private/WorldComposition.cpp
@@ -53,12 +53,10 @@ void UWorldComposition::PostDuplicate(bool bDuplicateForPIE)
 {
 	Super::PostDuplicate(bDuplicateForPIE);
 
-#if WITH_EDITOR
 	if (bDuplicateForPIE)
 	{
 		FixupForPIE(GetOutermost()->PIEInstanceID);	
 	}
-#endif // WITH_EDITOR
 }
 
 void UWorldComposition::PostLoad()
@@ -272,9 +270,7 @@ void UWorldComposition::Rescan()
 void UWorldComposition::ReinitializeForPIE()
 {
 	Rescan();
-#if WITH_EDITOR
 	FixupForPIE(GetOutermost()->PIEInstanceID);
-#endif// WITH_EDITOR
 	GetWorld()->StreamingLevels.Empty();
 	GetWorld()->StreamingLevels.Append(TilesStreaming);
 }
@@ -693,7 +689,7 @@ void UWorldComposition::UpdateStreamingState()
 	if (Locations.Num())
 	{
 		CentroidLocation/= Locations.Num();
-		if (PlayWorld->GetWorldSettings()->bEnableWorldOriginRebasing)
+		if (PlayWorld->GetWorldSettings()->bEnableWorldOriginRebasing && !PlayWorld->bOriginOffsetThisFrame)
 		{
 			EvaluateWorldOriginLocation(CentroidLocation);
 		}
diff --git a/Engine/Source/Runtime/Engine/Public/EngineUtils.h b/Engine/Source/Runtime/Engine/Public/EngineUtils.h
--- a/Engine/Source/Runtime/Engine/Public/EngineUtils.h
+++ b/Engine/Source/Runtime/Engine/Public/EngineUtils.h
@@ -395,7 +395,20 @@ private:
 	 */
 	static bool CanIterateLevel(ULevel* Level)
 	{
-		return Level->bIsVisible || Level->bIsAssociatingLevel;
+		const bool bIsLevelVisibleOrAssociating = Level->bIsVisible || Level->bIsAssociatingLevel;
+
+		// Only allow iteration of Level if it's in the currently active level collection of the world, or is a static level.
+		const FLevelCollection* const ActorLevelCollection = Level->GetCachedLevelCollection();
+		const FLevelCollection* const ActiveLevelCollection = Level->OwningWorld ? Level->OwningWorld->GetActiveLevelCollection() : nullptr;
+
+		// If the world's active level collection is null, we can't apply any meaningful filter,
+		// so just allow iteration in this case.
+		const bool bIsCurrentLevelCollectionTicking = !ActiveLevelCollection || (ActorLevelCollection == ActiveLevelCollection);
+		
+		const bool bIsLevelCollectionNullOrStatic = !ActorLevelCollection || ActorLevelCollection->GetType() == ELevelCollectionType::StaticLevels;
+		const bool bShouldIterateLevelCollection = bIsCurrentLevelCollectionTicking || bIsLevelCollectionNullOrStatic;
+
+		return bIsLevelVisibleOrAssociating && bShouldIterateLevelCollection;
 	}
 };
 
diff --git a/Engine/Source/Runtime/Engine/Public/Model.h b/Engine/Source/Runtime/Engine/Public/Model.h
--- a/Engine/Source/Runtime/Engine/Public/Model.h
+++ b/Engine/Source/Runtime/Engine/Public/Model.h
@@ -553,14 +553,12 @@ public:
 	/** Release CPU access version of vertex buffer */
 	void ReleaseVertices();
 
-#if WITH_EDITOR
 	/**
 	* Clears local (non RHI) data associated with MaterialIndexBuffers
 	*/
 	ENGINE_API void ClearLocalMaterialIndexBuffersData();
 
 	void CalculateUniqueVertCount();
-#endif
 
 	friend class UWorld;
 	friend class UBrushComponent;
diff --git a/Engine/Source/Runtime/Engine/Public/Net/DataReplication.h b/Engine/Source/Runtime/Engine/Public/Net/DataReplication.h
--- a/Engine/Source/Runtime/Engine/Public/Net/DataReplication.h
+++ b/Engine/Source/Runtime/Engine/Public/Net/DataReplication.h
@@ -8,6 +8,7 @@
 
 class FRepState;
 class FNetFieldExportGroup;
+class FRepChangelistState;
 
 bool FORCEINLINE IsCustomDeltaProperty( const UProperty* Property )
 {
@@ -34,6 +35,31 @@ struct FReplicatedActorProperty
 	{}
 };
 
+/** 
+ *	FReplicationChangelistMgr manages a list of change lists for a particular replicated object that have occurred since the object started replicating
+ *	Once the history is completely full, the very first changelist will then be merged with the next one (freeing a slot)
+ *		This way we always have the entire history for join in progress players
+ *	This information is then used by all connections, to share the compare work needed to determine what to send each connection
+ *	Connections will send any changelist that is new since the last time the connection checked
+ */
+class ENGINE_API FReplicationChangelistMgr
+{
+public:
+	FReplicationChangelistMgr( UNetDriver* InDriver, UObject* InObject );
+
+	~FReplicationChangelistMgr();
+
+	void Update( const UObject* InObject, const uint32 ReplicationFrame, const int32 LastCompareIndex, const FReplicationFlags& RepFlags );
+
+	FRepChangelistState* GetRepChangelistState() const { return RepChangelistState.Get(); }
+
+private:
+	UNetDriver*							Driver;
+	TSharedPtr< FRepLayout >			RepLayout;
+	TUniquePtr< FRepChangelistState >	RepChangelistState;
+	uint32								LastReplicationFrame;
+};
+
 /** FObjectReplicator
  *   Generic class that replicates properties for an object.
  *   All delta/diffing work is done in this class. 
@@ -98,6 +124,8 @@ public:
 	TSet< FNetworkGUID >							ReferencedGuids;
 	int32											TrackedGuidMemoryBytes;
 
+	TSharedPtr< FReplicationChangelistMgr >			ChangelistMgr;
+
 	struct FRPCCallInfo 
 	{
 		FName	FuncName;
diff --git a/Engine/Source/Runtime/Engine/Public/Net/RepLayout.h b/Engine/Source/Runtime/Engine/Public/Net/RepLayout.h
--- a/Engine/Source/Runtime/Engine/Public/Net/RepLayout.h
+++ b/Engine/Source/Runtime/Engine/Public/Net/RepLayout.h
@@ -18,8 +18,6 @@ class FRepChangedParent
 public:
 	FRepChangedParent() : Active( 1 ), OldActive( 1 ), IsConditional( 0 ) {}
 
-	TArray< uint16 >	Changed;
-
 	uint32				Active			: 1;
 	uint32				OldActive		: 1;
 	uint32				IsConditional	: 1;
@@ -33,11 +31,7 @@ class FRepChangedPropertyTracker : public IRepChangedPropertyTracker
 {
 public:
 	FRepChangedPropertyTracker( const bool InbIsReplay, const bool InbIsClientReplayRecording )
-		: LastReplicationGroupFrame( 0 )
-		, LastReplicationFrame( 0 )
-		, ActiveStatusChanged( false )
-		, UnconditionalPropChanged( false )
-		, bIsReplay( InbIsReplay )
+		: bIsReplay( InbIsReplay )
 		, bIsClientReplayRecording( InbIsClientReplayRecording )
 		, ExternalDataNumBits( 0 ) {}
 	virtual ~FRepChangedPropertyTracker() { }
@@ -50,11 +44,6 @@ public:
 
 		Parent.Active = (bIsActive || bIsClientReplayRecording) ? 1 : 0;
 
-		if ( Parent.Active != Parent.OldActive )
-		{
-			ActiveStatusChanged++;		// So we know to rebuild the conditional list with/without this property
-		}
-
 		Parent.OldActive = Parent.Active;
 	}
 
@@ -70,14 +59,8 @@ public:
 		return bIsReplay;
 	}
 
-	uint32						LastReplicationGroupFrame;		// Most recent frame number of the last FRepState that was compared this frame
-	uint32						LastReplicationFrame;
-	class FRepState *			LastRepState;
-
 	TArray< FRepChangedParent >	Parents;
 
-	uint32						ActiveStatusChanged;
-	bool						UnconditionalPropChanged;
 	bool						bIsReplay;							// True when recording/playing replays
 	bool						bIsClientReplayRecording;			// True when recording client replays
 
@@ -124,6 +107,31 @@ public:
 	int32						CmdIndex;
 };
 
+/** FRepChangelistState
+*  Stores changelist history (that are used to know what properties have changed) for objects
+*/
+class FRepChangelistState
+{
+public:
+	FRepChangelistState() :
+		HistoryStart( 0 ),
+		HistoryEnd( 0 ),
+		CompareIndex( 0 )
+	{ }
+
+	TSharedPtr< FRepLayout >						RepLayout;
+
+	static const int32 MAX_CHANGE_HISTORY = 64;
+
+	FRepChangedHistory								ChangeHistory[MAX_CHANGE_HISTORY];
+	int32											HistoryStart;
+	int32											HistoryEnd;
+	int32											CompareIndex;
+
+	// Properties will be copied in here so memory needs aligned to largest type
+	TArray< uint8, TAlignedHeapAllocator<16> >		StaticBuffer;
+};
+
 /** FRepState
  *  Stores state used by the FRepLayout manager
 */
@@ -133,45 +141,47 @@ public:
 	FRepState() : 
 		HistoryStart( 0 ), 
 		HistoryEnd( 0 ),
-		LastReplicationFrame( 0 ),
 		NumNaks( 0 ),
 		OpenAckedCalled( false ),
 		AwakeFromDormancy( false ),
-		ActiveStatusChanged( 0 )
+		LastChangelistIndex( 0 ),
+		LastCompareIndex( 0 )
 	{ }
 
 	~FRepState();
 
 	// Properties will be copied in here so memory needs aligned to largest type
-	TArray< uint8, TAlignedHeapAllocator<16> >				StaticBuffer;
+	TArray< uint8, TAlignedHeapAllocator<16> >	StaticBuffer;
 
-	FGuidReferencesMap			GuidReferencesMap;
+	FGuidReferencesMap				GuidReferencesMap;
 
-	TSharedPtr< FRepLayout >	RepLayout;
+	TSharedPtr< FRepLayout >		RepLayout;
 	
-	TArray< UProperty * >		RepNotifies;
+	TArray< UProperty * >			RepNotifies;
 
 	TSharedPtr<FRepChangedPropertyTracker> RepChangedPropertyTracker;
 
 	static const int32 MAX_CHANGE_HISTORY = 32;
 
-	FRepChangedHistory			ChangeHistory[MAX_CHANGE_HISTORY];
-	int32						HistoryStart;
-	int32						HistoryEnd;
+	FRepChangedHistory				ChangeHistory[MAX_CHANGE_HISTORY];
+	int32							HistoryStart;
+	int32							HistoryEnd;
 
-	uint32						LastReplicationFrame;
-	int32						NumNaks;
+	int32							NumNaks;
 
 	TArray< FRepChangedHistory >	PreOpenAckHistory;
 
 	bool							OpenAckedCalled;
 	bool							AwakeFromDormancy;
 
-	TArray< uint16 >				ConditionalLifetime;		// Properties the need to be checked conditionally (based on net initial, role, etc)
 	FReplicationFlags				RepFlags;
-	uint32							ActiveStatusChanged;
 
 	TArray< uint16 >				LifetimeChangelist;			// This is the unique list of properties that have changed since the channel was first opened
+
+	int32							LastChangelistIndex;		// The last change list history item we replicated from FRepChangelistState (if we are caught up to FRepChangelistState::HistoryEnd, there are no new changelists to replicate)
+	int32							LastCompareIndex;			// If == FRepChangelistState::CompareIndex, then there is definitely no new information since the last time we checked
+
+	bool							ConditionMap[COND_Max];
 };
 
 enum ERepLayoutCmdType
@@ -245,22 +255,96 @@ public:
 	uint16		ParentIndex;		// Index into Parents
 	uint32		CompatibleChecksum;	// Used to determine if property is still compatible
 };
+	
+/** FHandleToCmdIndex
+ *  Converts a relative handle to the appropriate index into the Cmds array
+ */
+class FHandleToCmdIndex
+{
+public:
+	FHandleToCmdIndex() : CmdIndex( INDEX_NONE )
+	{
+	}
+
+	FHandleToCmdIndex( const int32 InHandleToCmdIndex ) : CmdIndex( InHandleToCmdIndex )
+	{
+	}
+
+	int32										CmdIndex;
+	TUniquePtr< TArray< FHandleToCmdIndex > >	HandleToCmdIndex;
+
+	// Fix VS2013 compiler issue (VS2013 doesn't synthesize move constructor or move assignment operator version of these)
+	FHandleToCmdIndex( FHandleToCmdIndex&& Other ) : CmdIndex( Other.CmdIndex ), HandleToCmdIndex( MoveTemp( Other.HandleToCmdIndex ) )
+	{
+	}
+
+	FHandleToCmdIndex& operator=( FHandleToCmdIndex&& Other )
+	{
+		if ( this != &Other )
+		{
+			CmdIndex			= Other.CmdIndex;
+			HandleToCmdIndex	= MoveTemp( Other.HandleToCmdIndex );
+		}
+
+		return *this;
+	}
+};
 
-class FRepWriterState
+class FChangelistIterator
 {
 public:
-	FRepWriterState( FNetBitWriter & InWriter, TArray< uint16 > & InChanged, bool bInDoChecksum ) : 
-		Writer( InWriter ), 
-		Changed( InChanged ),
-		CurrentChanged( 0 ),
-		bDoChecksum( bInDoChecksum )
+	FChangelistIterator( const TArray< uint16 >& InChanged, const int32 InChangedIndex ) : Changed( InChanged ), ChangedIndex( InChangedIndex )
 	{
 	}
 
-	FNetBitWriter &		Writer;
-	TArray< uint16 > &	Changed;
-	int32				CurrentChanged;
-	bool				bDoChecksum;
+	const TArray< uint16 >& Changed;
+	int32					ChangedIndex;
+};
+
+/** FRepHandleIterator
+ *  Iterates over a changelist, taking each handle, and mapping to rep layout index, array index, etc
+ */
+class FRepHandleIterator
+{
+public:
+	FRepHandleIterator(
+		FChangelistIterator&				InChangelistIterator,
+		const TArray< FRepLayoutCmd >&		InCmds,
+		const TArray< FHandleToCmdIndex >&	InHandleToCmdIndex,
+		const int32							InElementSize,
+		const int32							InMaxArrayIndex,
+		const int32							InMinCmdIndex,
+		const int32							InMaxCmdIndex
+	) :
+		ChangelistIterator( InChangelistIterator ),
+		Cmds( InCmds ),
+		HandleToCmdIndex( InHandleToCmdIndex ),
+		NumHandlesPerElement( HandleToCmdIndex.Num() ),
+		ArrayElementSize( InElementSize ),
+		MaxArrayIndex( InMaxArrayIndex ),
+		MinCmdIndex( InMinCmdIndex ),
+		MaxCmdIndex( InMaxCmdIndex )
+	{
+	}
+
+	FChangelistIterator&				ChangelistIterator;
+
+	const TArray< FRepLayoutCmd >&		Cmds;
+	const TArray< FHandleToCmdIndex >&	HandleToCmdIndex;
+	const int32							NumHandlesPerElement;
+	const int32							ArrayElementSize;
+	const int32							MaxArrayIndex;
+	const int32							MinCmdIndex;
+	const int32							MaxCmdIndex;
+
+	int32								Handle;
+	int32								CmdIndex;
+	int32								ArrayIndex;
+	int32								ArrayOffset;
+
+	bool	NextHandle();
+	bool	JumpOverArray();
+	int32	PeekNextHandle() const;
 };
 
 /** FRepLayout
@@ -277,6 +361,11 @@ public:
 
 	void OpenAcked( FRepState * RepState ) const;
 
+	void InitShadowData(
+		TArray< uint8, TAlignedHeapAllocator<16> >&	ShadowData,
+		UClass *									InObjectClass,
+		uint8 *										Src ) const;
+
 	void InitRepState( 
 		FRepState *									RepState, 
 		UClass *									InObjectClass, 
@@ -285,20 +374,20 @@ public:
 
 	void InitChangedTracker( FRepChangedPropertyTracker * ChangedTracker ) const;
 
-	bool ReplicateProperties( 
-		FRepState * RESTRICT		RepState, 
-		const uint8* RESTRICT		Data, 
-		UClass *					ObjectClass,
-		UActorChannel *				OwningChannel,
-		FNetBitWriter&				Writer, 
-		const FReplicationFlags &	RepFlags ) const;
-
-	void SendProperties( 
-		FRepState *	RESTRICT		RepState, 
-		const FReplicationFlags &	RepFlags,
-		const uint8* RESTRICT		Data, 
-		UClass *					ObjectClass,
-		UActorChannel *				OwningChannel,
+	bool ReplicateProperties(
+		FRepState* RESTRICT				RepState,
+		FRepChangelistState* RESTRICT	RepChangelistState,
+		const uint8* RESTRICT			Data,
+		UClass *						ObjectClass,
+		UActorChannel *					OwningChannel,
+		FNetBitWriter&					Writer,
+		const FReplicationFlags &		RepFlags ) const;
+
+	void SendProperties(
+		FRepState*	RESTRICT		RepState,
+		FRepChangedPropertyTracker* ChangedTracker,
+		const uint8* RESTRICT		Data,
+		UClass*						ObjectClass,
 		FNetBitWriter&				Writer,
 		TArray< uint16 >&			Changed ) const;
 
@@ -321,9 +410,10 @@ public:
 	void ValidateWithChecksum( const void* RESTRICT Data, FArchive & Ar ) const;
 	uint32 GenerateChecksum( const FRepState* RepState ) const;
 
+	/** Clamp the changelist so that it conforms to the current size of either the array, or arrays within structs/arrays */
 	void PruneChangeList( FRepState * RepState, const void* RESTRICT Data, const TArray< uint16 >& Changed, TArray< uint16 >& PrunedChanged ) const;
 
-	void MergeDirtyList( FRepState * RepState, const void* RESTRICT Data, const TArray< uint16 > & Dirty1, const TArray< uint16 > & Dirty2, TArray< uint16 > & MergedDirty ) const;
+	void MergeChangeList( const uint8* RESTRICT Data, const TArray< uint16 >& Dirty1, const TArray< uint16 >& Dirty2, TArray< uint16 >& MergedDirty ) const;
 
 	bool DiffProperties( TArray<UProperty*>* RepNotifies, void* RESTRICT Destination, const void* RESTRICT Source, const bool bSync ) const;
 
@@ -345,6 +435,7 @@ public:
 
 	void SendProperties_BackwardsCompatible(
 		FRepState* RESTRICT			RepState,
+		FRepChangedPropertyTracker* ChangedTracker,
 		const uint8* RESTRICT		Data,
 		UNetConnection*				Connection,
 		FNetBitWriter&				Writer,
@@ -359,90 +450,61 @@ public:
 		const bool					bEnableRepNotifies,
 		bool&						bOutGuidsChanged ) const;
 
+	bool CompareProperties(
+		FRepChangelistState* RESTRICT	RepState,
+		const uint8* RESTRICT			Data,
+		const FReplicationFlags&		RepFlags ) const;
+
 private:
 	void RebuildConditionalProperties( FRepState * RESTRICT	RepState, const FRepChangedPropertyTracker& ChangedTracker, const FReplicationFlags& RepFlags ) const;
 
-	void LogChangeListMismatches( 
-		const uint8* 						Data, 
-		UActorChannel *						OwningChannel, 
-		const TArray< uint16 > &			PropertyList,
-		FRepState *							RepState1, 
-		FRepState * 						RepState2, 
-		const TArray< FRepChangedParent > & ChangedParents1,
-		const TArray< FRepChangedParent > & ChangedParents2 ) const;
-
-	void SanityCheckShadowStateAgainstChangeList(
-		FRepState * 						RepState, 
-		const uint8* 						Data, 
-		UActorChannel *						OwningChannel, 
-		const TArray< uint16 > &			PropertyList,
-		FRepState *							OtherRepState,
-		const TArray< FRepChangedParent > & OtherChangedParents ) const;
-
 	void UpdateChangelistHistory( FRepState * RepState, UClass * ObjectClass, const uint8* RESTRICT Data, const int32 AckPacketId, TArray< uint16 > * OutMerged ) const;
 
+	void SendProperties_BackwardsCompatible_r(
+		FRepState* RESTRICT					RepState,
+		UPackageMapClient*					PackageMapClient,
+		FNetFieldExportGroup*				NetFieldExportGroup,
+		FRepChangedPropertyTracker*			ChangedTracker,
+		FNetBitWriter&						Writer,
+		const bool							bDoChecksum,
+		FRepHandleIterator&					HandleIterator,
+		const uint8* RESTRICT				SourceData ) const;
+
+	void SendAllProperties_BackwardsCompatible_r(
+		FNetBitWriter&						Writer,
+		const bool							bDoChecksum,
+		UPackageMapClient*					PackageMapClient,
+		FNetFieldExportGroup*				NetFieldExportGroup,
+		const int32							CmdStart,
+		const int32							CmdEnd,
+		const uint8*						SourceData ) const;
+
+	void SendProperties_r(
+		FRepState*	RESTRICT				RepState,
+		FRepChangedPropertyTracker*			ChangedTracker,
+		FNetBitWriter&						Writer,
+		const bool							bDoChecksum,
+		FRepHandleIterator&					HandleIterator,
+		const uint8* RESTRICT				SourceData ) const;
+
 	uint16 CompareProperties_r(
 		const int32				CmdStart,
 		const int32				CmdEnd,
-		const uint8* RESTRICT	CompareData, 
+		const uint8* RESTRICT	CompareData,
 		const uint8* RESTRICT	Data,
-		TArray< uint16 > &		Changed,
-		uint16					Handle ) const;
+		TArray< uint16 >&		Changed,
+		uint16					Handle,
+		const bool				bIsInitial,
+		const bool				bForceFail ) const;
 
-	void CompareProperties_Array_r( 
-		const uint8* RESTRICT	CompareData, 
+	void CompareProperties_Array_r(
+		const uint8* RESTRICT	CompareData,
 		const uint8* RESTRICT	Data,
-		TArray< uint16 > &		Changed,
+		TArray< uint16 >&		Changed,
 		const uint16			CmdIndex,
-		const uint16			Handle ) const;
-
-	bool CompareProperties( 
-		FRepState * RESTRICT				RepState, 
-		const uint8* RESTRICT				CompareData,
-		const uint8* RESTRICT				Data, 
-		TArray< FRepChangedParent > &		OutChangedParents,
-		const TArray< uint16 > &			PropertyList ) const;
-
-	void SendProperties_DynamicArray_r( 
-		FRepState *	RESTRICT		RepState, 
-		const FReplicationFlags &	RepFlags,
-		FRepWriterState &			WriterState,
-		const int32					CmdIndex, 
-		const uint8* RESTRICT		StoredData, 
-		const uint8* RESTRICT		Data, 
-		uint16						Handle ) const;
-
-	uint16 SendProperties_r( 
-		FRepState *	RESTRICT		RepState, 
-		const FReplicationFlags &	RepFlags,
-		FRepWriterState &			WriterState,
-		const int32					CmdStart, 
-		const int32					CmdEnd, 
-		const uint8* RESTRICT		StoredData, 
-		const uint8* RESTRICT		Data, 
-		uint16						Handle ) const;
-
-	void SendProperties_BackwardsCompatible_DynamicArray_r(
-		FRepState *	RESTRICT		RepState,
-		FRepWriterState &			WriterState,
-		FNetBitWriter &				Writer,
-		UPackageMapClient*			PackageMapClient,
-		FNetFieldExportGroup*		NetFieldExportGroup,
-		const int32					CmdIndex,
-		const uint8* RESTRICT		StoredData,
-		const uint8* RESTRICT		Data ) const;
-
-	uint16 SendProperties_BackwardsCompatible_r(
-		FRepState *	RESTRICT		RepState,
-		FRepWriterState &			WriterState,
-		FNetBitWriter &				Writer,
-		UPackageMapClient*			PackageMapClient,
-		FNetFieldExportGroup*		NetFieldExportGroup,
-		const int32					CmdStart,
-		const int32					CmdEnd,
-		const uint8* RESTRICT		StoredData,
-		const uint8* RESTRICT		Data,
-		uint16						Handle ) const;
+		const uint16			Handle,
+		const bool				bIsInitial,
+		const bool				bForceFail ) const;
 
 	TSharedPtr< FNetFieldExportGroup > CreateNetfieldExportGroup() const;
 
@@ -518,17 +580,29 @@ private:
 		void *				Data,
 		bool &				bHasUnmapped ) const;
 
-	void BuildChangeList_r( const int32 CmdStart, const int32 CmdEnd, uint8* Data, const int32 HandleOffset, TArray< uint16 >& Changed ) const;
+	void MergeChangeList_r(
+		FRepHandleIterator&		RepHandleIterator1,
+		FRepHandleIterator&		RepHandleIterator2,
+		const uint8* RESTRICT	SourceData,
+		TArray< uint16 >&		OutChanged ) const;
+
+	void PruneChangeList_r(
+		FRepHandleIterator&		RepHandleIterator,
+		const uint8* RESTRICT	SourceData,
+		TArray< uint16 >&		OutChanged ) const;
+		
+	void BuildChangeList_r( const TArray< FHandleToCmdIndex >& HandleToCmdIndex, const int32 CmdStart, const int32 CmdEnd, uint8* Data, const int32 HandleOffset, TArray< uint16 >& Changed ) const;
+
+	void BuildHandleToCmdIndexTable_r( const int32 CmdStart, const int32 CmdEnd, TArray< FHandleToCmdIndex >& HandleToCmdIndex );
 
-	void ConstructProperties( FRepState * RepState ) const;
-	void InitProperties( FRepState * RepState, uint8* Src ) const;
+	void ConstructProperties( TArray< uint8, TAlignedHeapAllocator<16> >& ShadowData ) const;
+	void InitProperties( TArray< uint8, TAlignedHeapAllocator<16> >& ShadowData, uint8* Src ) const;
 	void DestructProperties( FRepState * RepState ) const;
 
 	TArray< FRepParentCmd >		Parents;
 	TArray< FRepLayoutCmd >		Cmds;
 
-	TArray< uint16 >			UnconditionalLifetime;		// Properties that don't need special checks (net initial, role, etc)
-	TArray< FLifetimeProperty >	ConditionalLifetime;		// Properties the need to be checked conditionally (based on net initial, role, etc)
+	TArray< FHandleToCmdIndex >	BaseHandleToCmdIndex;		// Converts a relative handle to the appropriate index into the Cmds array
 
 	int32						FirstNonCustomParent;
 	int32						RoleIndex;
diff --git a/Engine/Source/Runtime/Engine/Public/TickTaskManagerInterface.h b/Engine/Source/Runtime/Engine/Public/TickTaskManagerInterface.h
--- a/Engine/Source/Runtime/Engine/Public/TickTaskManagerInterface.h
+++ b/Engine/Source/Runtime/Engine/Public/TickTaskManagerInterface.h
@@ -28,7 +28,7 @@ public:
 	 * @param DeltaSeconds - time in seconds since last tick
 	 * @param TickType - type of tick (viewports only, time only, etc)
 	 */
-	virtual void StartFrame(UWorld* InWorld, float DeltaSeconds, ELevelTick TickType) = 0;
+	virtual void StartFrame(UWorld* InWorld, float DeltaSeconds, ELevelTick TickType, const TArray<ULevel*>& LevelsToTick) = 0;
 
 	/**
 	 * Run all of the ticks for a pause frame synchronously on the game thread.
@@ -37,7 +37,7 @@ public:
 	 * @param DeltaSeconds - time in seconds since last tick
 	 * @param TickType - type of tick (viewports only, time only, etc)
 	 */
-	virtual void RunPauseFrame(UWorld* InWorld, float DeltaSeconds, ELevelTick TickType) = 0;
+	virtual void RunPauseFrame(UWorld* InWorld, float DeltaSeconds, ELevelTick TickType, const TArray<ULevel*>& LevelsToTick) = 0;
 
 	/**
 		* Run a tick group, ticking all actors and components
diff --git a/Engine/Source/Runtime/Engine/Public/TimerManager.h b/Engine/Source/Runtime/Engine/Public/TimerManager.h
--- a/Engine/Source/Runtime/Engine/Public/TimerManager.h
+++ b/Engine/Source/Runtime/Engine/Public/TimerManager.h
@@ -8,6 +8,8 @@
 
 #include "Engine/EngineTypes.h"
 
+class UGameInstance;
+
 DECLARE_DELEGATE(FTimerDelegate);
 
 /** Simple interface to wrap a timer delegate that can be either native or dynamic. */
@@ -114,12 +116,16 @@ struct FTimerData
 
 	FTimerHandle TimerHandle;
 
+	/** The level collection that was active when this timer was created. Used to set the correct context before executing the timer's delegate. */
+	ELevelCollectionType LevelCollection;
+
 	FTimerData()
 		: bLoop(false)
 		, bRequiresDelegate(false)
 		, Status(ETimerStatus::Active)
 		, Rate(0)
 		, ExpireTime(0)
+		, LevelCollection(ELevelCollectionType::DynamicSourceLevels)
 	{}
 
 	/** Operator less, used to sort the heap based on time until execution. **/
@@ -393,6 +399,9 @@ public:
 	/** Get the current last assigned handle */
 	static void ValidateHandle(FTimerHandle& InOutHandle);
 
+	/** Used by the UGameInstance constructor to set this manager's owning game instance. */
+	void SetGameInstance(UGameInstance* InGameInstance) { OwningGameInstance = InGameInstance; }
+
 private:
 	void InternalSetTimer( FTimerHandle& InOutHandle, FTimerUnifiedDelegate const& InDelegate, float InRate, bool InbLoop, float InFirstDelay );
 	void InternalSetTimer( FTimerData& NewTimerData, float InRate, bool InbLoop, float InFirstDelay );
@@ -432,5 +441,8 @@ private:
 
 	/** The last handle we assigned from this timer manager */
 	static uint64 LastAssignedHandle;
+
+	/** The game instance that created this timer manager. May be null if this timer manager wasn't created by a game instance. */
+	UGameInstance* OwningGameInstance;
 };
 
diff --git a/Engine/Source/Runtime/GameplayAbilities/Private/GameplayEffect.cpp b/Engine/Source/Runtime/GameplayAbilities/Private/GameplayEffect.cpp
--- a/Engine/Source/Runtime/GameplayAbilities/Private/GameplayEffect.cpp
+++ b/Engine/Source/Runtime/GameplayAbilities/Private/GameplayEffect.cpp
@@ -1527,6 +1527,13 @@ void FActiveGameplayEffect::PostReplicatedAdd(const struct FActiveGameplayEffect
 		return;
 	}
 
+	if (Spec.Def->Modifiers.Num() != Spec.Modifiers.Num())
+	{
+		// This can happen with older replays, where the replicated Spec.Modifiers size changed in the newer Spec.Def
+		ABILITY_LOG(Error, TEXT("FActiveGameplayEffect::PostReplicatedAdd: Spec.Def->Modifiers.Num() != Spec.Modifiers.Num()"));
+		return;
+	}
+
 	bool ShouldInvokeGameplayCueEvents = true;
 	if (PredictionKey.IsLocalClientKey())
 	{
@@ -1578,6 +1585,14 @@ void FActiveGameplayEffect::PostReplicatedChange(const struct FActiveGameplayEff
 	if (Spec.Def == nullptr)
 	{
 		ABILITY_LOG(Error, TEXT("Received ReplicatedGameplayEffect with no UGameplayEffect def."));
+		return;
+	}
+
+	if (Spec.Def->Modifiers.Num() != Spec.Modifiers.Num())
+	{
+		// This can happen with older replays, where the replicated Spec.Modifiers size changed in the newer Spec.Def
+		ABILITY_LOG(Error, TEXT("FActiveGameplayEffect::PostReplicatedChange: Spec.Def->Modifiers.Num() != Spec.Modifiers.Num()"));
+		return;
 	}
 
 	// Handle potential duration refresh
diff --git a/Engine/Source/Runtime/Launch/Private/LaunchEngineLoop.cpp b/Engine/Source/Runtime/Launch/Private/LaunchEngineLoop.cpp
--- a/Engine/Source/Runtime/Launch/Private/LaunchEngineLoop.cpp
+++ b/Engine/Source/Runtime/Launch/Private/LaunchEngineLoop.cpp
@@ -1273,6 +1273,7 @@ int32 FEngineLoop::PreInit( const TCHAR* CmdLine )
 	ApplyCVarSettingsFromIni(TEXT("/Script/Engine.RendererOverrideSettings"), *GEngineIni, ECVF_SetByProjectSetting);
 	ApplyCVarSettingsFromIni(TEXT("/Script/Engine.StreamingSettings"), *GEngineIni, ECVF_SetByProjectSetting);
 	ApplyCVarSettingsFromIni(TEXT("/Script/Engine.GarbageCollectionSettings"), *GEngineIni, ECVF_SetByProjectSetting);
+	ApplyCVarSettingsFromIni(TEXT("/Script/Engine.NetworkSettings"), *GEngineIni, ECVF_SetByProjectSetting);
 
 #if !UE_SERVER
 	if (!IsRunningDedicatedServer())
diff --git a/Engine/Source/Runtime/LevelSequence/Private/LevelSequenceObjectReference.cpp b/Engine/Source/Runtime/LevelSequence/Private/LevelSequenceObjectReference.cpp
--- a/Engine/Source/Runtime/LevelSequence/Private/LevelSequenceObjectReference.cpp
+++ b/Engine/Source/Runtime/LevelSequence/Private/LevelSequenceObjectReference.cpp
@@ -32,12 +32,8 @@ UObject* FLevelSequenceObjectReference::Resolve(UObject* InContext) const
 {
 	if (ObjectId.IsValid() && InContext != nullptr)
 	{
-	#if WITH_EDITOR
 		int32 PIEInstanceID = InContext->GetOutermost()->PIEInstanceID;
 		FUniqueObjectGuid FixedUpId = PIEInstanceID == -1 ? ObjectId : ObjectId.FixupForPIE(PIEInstanceID);
-	#else
-		FUniqueObjectGuid FixedUpId = ObjectId;
-	#endif
 		
 
 		FLazyObjectPtr LazyPtr;
