Copying //UE4/Dev-Editor to //UE4/Dev-Main (Source: //UE4/Dev-Editor @ 3152045)

diff --git a/Engine/Source/Runtime/Engine/Private/BatchedElements.cpp b/Engine/Source/Runtime/Engine/Private/BatchedElements.cpp
--- a/Engine/Source/Runtime/Engine/Private/BatchedElements.cpp
+++ b/Engine/Source/Runtime/Engine/Private/BatchedElements.cpp
@@ -796,9 +796,42 @@ void FBatchedElements::DrawPointElements(FRHICommandList& RHICmdList, const FMat
 	}
 }
 
+FSceneView FBatchedElements::CreateProxySceneView(const FMatrix& ProjectionMatrix, const FIntRect& ViewRect)
+{
+	FSceneViewInitOptions ProxyViewInitOptions;
+	ProxyViewInitOptions.SetViewRectangle(ViewRect);
+	ProxyViewInitOptions.ViewOrigin = FVector::ZeroVector;
+	ProxyViewInitOptions.ViewRotationMatrix = FMatrix::Identity;
+	ProxyViewInitOptions.ProjectionMatrix = ProjectionMatrix;
+
+	return FSceneView(ProxyViewInitOptions);
+}
 
 bool FBatchedElements::Draw(FRHICommandList& RHICmdList, ERHIFeatureLevel::Type FeatureLevel, bool bNeedToSwitchVerticalAxis, const FMatrix& Transform, uint32 ViewportSizeX, uint32 ViewportSizeY, bool bHitTesting, float Gamma, const FSceneView* View, FTexture2DRHIRef DepthTexture, EBlendModeFilter::Type Filter) const
 {
+	if ( View )
+	{
+		// Going to ignore these parameters in favor of just using the values directly from the scene view, so ensure that they're identical.
+		check(Transform == View->ViewProjectionMatrix);
+		check(ViewportSizeX == View->ViewRect.Width());
+		check(ViewportSizeY == View->ViewRect.Height());
+
+		return Draw(RHICmdList, FeatureLevel, bNeedToSwitchVerticalAxis, *View, bHitTesting, Gamma, DepthTexture, Filter);
+	}
+	else
+	{
+		FIntRect ViewRect = FIntRect(0, 0, ViewportSizeX, ViewportSizeY);
+
+		return Draw(RHICmdList, FeatureLevel, bNeedToSwitchVerticalAxis, CreateProxySceneView(Transform, ViewRect), bHitTesting, Gamma, DepthTexture, Filter);
+	}
+}
+
+bool FBatchedElements::Draw(FRHICommandList& RHICmdList, ERHIFeatureLevel::Type FeatureLevel, bool bNeedToSwitchVerticalAxis, const FSceneView& View, bool bHitTesting, float Gamma /* = 1.0f */, FTexture2DRHIRef DepthTexture /* = FTexture2DRHIRef() */, EBlendModeFilter::Type Filter /* = EBlendModeFilter::All */) const
+{
+	const FMatrix& Transform = View.ViewProjectionMatrix;
+	const uint32 ViewportSizeX = View.ViewRect.Width();
+	const uint32 ViewportSizeY = View.ViewRect.Height();
+
 	if (UNLIKELY(!FApp::CanEverRender()))
 	{
 		return false;
@@ -820,7 +853,7 @@ bool FBatchedElements::Draw(FRHICommandList& RHICmdList, ERHIFeatureLevel::Type
 			FBatchedElementParameters* BatchedElementParameters = NULL;
 
 			// Set the appropriate pixel shader parameters & shader state for the non-textured elements.
-			PrepareShaders(RHICmdList, FeatureLevel, SE_BLEND_Opaque, Transform, bNeedToSwitchVerticalAxis, BatchedElementParameters, GWhiteTexture, bHitTesting, Gamma, NULL, View, DepthTexture);
+			PrepareShaders(RHICmdList, FeatureLevel, SE_BLEND_Opaque, Transform, bNeedToSwitchVerticalAxis, BatchedElementParameters, GWhiteTexture, bHitTesting, Gamma, NULL, &View, DepthTexture);
 
 			// Draw the line elements.
 			if( LineVertices.Num() > 0 )
@@ -846,16 +879,13 @@ bool FBatchedElements::Draw(FRHICommandList& RHICmdList, ERHIFeatureLevel::Type
 
 			if ( ThickLines.Num() > 0 )
 			{
-				PrepareShaders(RHICmdList, FeatureLevel, SE_BLEND_Translucent, Transform, bNeedToSwitchVerticalAxis, BatchedElementParameters, GWhiteTexture, bHitTesting, Gamma, NULL, View, DepthTexture);
+				PrepareShaders(RHICmdList, FeatureLevel, SE_BLEND_Translucent, Transform, bNeedToSwitchVerticalAxis, BatchedElementParameters, GWhiteTexture, bHitTesting, Gamma, NULL, &View, DepthTexture);
 				float OrthoZoomFactor = 1.0f;
 
-				if( View )
+				const bool bIsPerspective = View.ViewMatrices.ProjMatrix.M[3][3] < 1.0f ? true : false;
+				if (!bIsPerspective)
 				{
-					const bool bIsPerspective = View->ViewMatrices.ProjMatrix.M[3][3] < 1.0f ? true : false;
-					if( !bIsPerspective )
-					{
-						OrthoZoomFactor = 1.0f / View->ViewMatrices.ProjMatrix.M[0][0];
-					}
+					OrthoZoomFactor = 1.0f / View.ViewMatrices.ProjMatrix.M[0][0];
 				}
 
 				int32 LineIndex = 0;
@@ -1037,7 +1067,7 @@ bool FBatchedElements::Draw(FRHICommandList& RHICmdList, ERHIFeatureLevel::Type
 						//New batch, draw previous and clear
 						const int32 VertexCount = SpriteList.Num();
 						const int32 PrimCount = VertexCount / 3;
-						PrepareShaders(RHICmdList, FeatureLevel, CurrentBlendMode, Transform, bNeedToSwitchVerticalAxis, BatchedElementParameters, CurrentTexture, bHitTesting, Gamma, NULL, View, DepthTexture);
+						PrepareShaders(RHICmdList, FeatureLevel, CurrentBlendMode, Transform, bNeedToSwitchVerticalAxis, BatchedElementParameters, CurrentTexture, bHitTesting, Gamma, NULL, &View, DepthTexture);
 						DrawPrimitiveUP(RHICmdList, PT_TriangleList, PrimCount, SpriteList.GetData(), sizeof(FSimpleElementVertex));
 
 						SpriteList.Empty(6);
@@ -1077,7 +1107,7 @@ bool FBatchedElements::Draw(FRHICommandList& RHICmdList, ERHIFeatureLevel::Type
 					//Draw last batch
 					const int32 VertexCount = SpriteList.Num();
 					const int32 PrimCount = VertexCount / 3;
-					PrepareShaders(RHICmdList, FeatureLevel, CurrentBlendMode, Transform, bNeedToSwitchVerticalAxis, BatchedElementParameters, CurrentTexture, bHitTesting, Gamma, NULL, View, DepthTexture);
+					PrepareShaders(RHICmdList, FeatureLevel, CurrentBlendMode, Transform, bNeedToSwitchVerticalAxis, BatchedElementParameters, CurrentTexture, bHitTesting, Gamma, NULL, &View, DepthTexture);
 					DrawPrimitiveUP(RHICmdList, PT_TriangleList, PrimCount, SpriteList.GetData(), sizeof(FSimpleElementVertex));
 				}
 			}
@@ -1095,7 +1125,7 @@ bool FBatchedElements::Draw(FRHICommandList& RHICmdList, ERHIFeatureLevel::Type
 				if (Filter & MeshFilter)
 				{
 					// Set the appropriate pixel shader for the mesh.
-					PrepareShaders(RHICmdList, FeatureLevel, MeshElement.BlendMode, Transform, bNeedToSwitchVerticalAxis, MeshElement.BatchedElementParameters, MeshElement.Texture, bHitTesting, Gamma, &MeshElement.GlowInfo);
+					PrepareShaders(RHICmdList, FeatureLevel, MeshElement.BlendMode, Transform, bNeedToSwitchVerticalAxis, MeshElement.BatchedElementParameters, MeshElement.Texture, bHitTesting, Gamma, &MeshElement.GlowInfo, &View);
 
 					// Draw the mesh.
 					DrawIndexedPrimitiveUP(
diff --git a/Engine/Source/Runtime/Engine/Private/Camera/CameraAnim.cpp b/Engine/Source/Runtime/Engine/Private/Camera/CameraAnim.cpp
--- a/Engine/Source/Runtime/Engine/Private/Camera/CameraAnim.cpp
+++ b/Engine/Source/Runtime/Engine/Private/Camera/CameraAnim.cpp
@@ -145,11 +145,11 @@ void UCameraAnim::CalcLocalAABB()
 	}
 }
 
-SIZE_T UCameraAnim::GetResourceSize(EResourceSizeMode::Type Mode)
+void UCameraAnim::GetResourceSizeEx(FResourceSizeEx& CumulativeResourceSize)
 {
-	int32 ResourceSize = Super::GetResourceSize(Mode);
+	Super::GetResourceSizeEx(CumulativeResourceSize);
 
-	if (Mode == EResourceSizeMode::Inclusive && CameraInterpGroup)
+	if (CumulativeResourceSize.GetResourceSizeMode() == EResourceSizeMode::Inclusive && CameraInterpGroup)
 	{
 		// find move track
 		UInterpTrackMove *MoveTrack = NULL;
@@ -160,9 +160,8 @@ SIZE_T UCameraAnim::GetResourceSize(EResourceSizeMode::Type Mode)
 			if (MoveTrack)
 			{
 				FArchiveCountMem CountBytesSize(MoveTrack);
-				ResourceSize += CountBytesSize.GetNum();
+				CumulativeResourceSize.AddDedicatedSystemMemoryBytes(CountBytesSize.GetNum());
 			}
 		}
 	}
-	return ResourceSize;
 }
diff --git a/Engine/Source/Runtime/Engine/Private/Components/BrushComponent.cpp b/Engine/Source/Runtime/Engine/Private/Components/BrushComponent.cpp
--- a/Engine/Source/Runtime/Engine/Private/Components/BrushComponent.cpp
+++ b/Engine/Source/Runtime/Engine/Private/Components/BrushComponent.cpp
@@ -563,18 +563,15 @@ void UBrushComponent::PostLoad()
 #endif
 }
 
-
-SIZE_T UBrushComponent::GetResourceSize(EResourceSizeMode::Type Mode)
+void UBrushComponent::GetResourceSizeEx(FResourceSizeEx& CumulativeResourceSize)
 {
-	SIZE_T ResSize = Super::GetResourceSize(Mode);
+	Super::GetResourceSizeEx(CumulativeResourceSize);
 
 	// Count the bodysetup we own as well for 'inclusive' stats
-	if((Mode == EResourceSizeMode::Inclusive) && (BrushBodySetup != NULL))
+	if((CumulativeResourceSize.GetResourceSizeMode() == EResourceSizeMode::Inclusive) && (BrushBodySetup != NULL))
 	{
-		ResSize += BrushBodySetup->GetResourceSize(Mode);
+		BrushBodySetup->GetResourceSizeEx(CumulativeResourceSize);
 	}
-
-	return ResSize;
 }
 
 uint8 UBrushComponent::GetStaticDepthPriorityGroup() const
diff --git a/Engine/Source/Runtime/Engine/Private/Components/MeshComponent.cpp b/Engine/Source/Runtime/Engine/Private/Components/MeshComponent.cpp
--- a/Engine/Source/Runtime/Engine/Private/Components/MeshComponent.cpp
+++ b/Engine/Source/Runtime/Engine/Private/Components/MeshComponent.cpp
@@ -103,6 +103,30 @@ int32 UMeshComponent::GetNumOverrideMaterials() const
 	return OverrideMaterials.Num();
 }
 
+#if WITH_EDITOR
+void UMeshComponent::CleanUpOverrideMaterials()
+{
+	//We have to remove material override Ids that are bigger then the material list
+	if (GetNumOverrideMaterials() > GetNumMaterials())
+	{
+		//Remove the override material id that are superior to the static mesh materials number
+		int32 RemoveCount = GetNumOverrideMaterials() - GetNumMaterials();
+		OverrideMaterials.RemoveAt(GetNumMaterials(), RemoveCount);
+	}
+	//Remove override at the end of the array until there is a valid material
+	for (int32 i = GetNumOverrideMaterials() - 1; i >= 0; --i)
+	{
+		UMaterialInterface *OverrideMaterial = OverrideMaterials[i];
+		if (OverrideMaterial == nullptr)
+		{
+			OverrideMaterials.RemoveAt(i);
+			continue;
+		}
+		break;
+	}
+}
+#endif
+
 int32 UMeshComponent::GetNumMaterials() const
 {
 	return 0;
diff --git a/Engine/Source/Runtime/Engine/Private/Components/ModelComponent.cpp b/Engine/Source/Runtime/Engine/Private/Components/ModelComponent.cpp
--- a/Engine/Source/Runtime/Engine/Private/Components/ModelComponent.cpp
+++ b/Engine/Source/Runtime/Engine/Private/Components/ModelComponent.cpp
@@ -259,17 +259,15 @@ void UModelComponent::PostEditUndo()
 }
 #endif // WITH_EDITOR
 
-SIZE_T UModelComponent::GetResourceSize(EResourceSizeMode::Type Mode)
+void UModelComponent::GetResourceSizeEx(FResourceSizeEx& CumulativeResourceSize)
 {
-	SIZE_T ResSize = Super::GetResourceSize(Mode);
+	Super::GetResourceSizeEx(CumulativeResourceSize);
 
 	// Count the bodysetup we own as well for 'inclusive' stats
-	if((Mode == EResourceSizeMode::Inclusive) && (ModelBodySetup != NULL))
+	if((CumulativeResourceSize.GetResourceSizeMode() == EResourceSizeMode::Inclusive) && (ModelBodySetup != NULL))
 	{
-		ResSize += ModelBodySetup->GetResourceSize(Mode);
+		ModelBodySetup->GetResourceSizeEx(CumulativeResourceSize);
 	}
-
-	return ResSize;
 }
 
 
diff --git a/Engine/Source/Runtime/Engine/Private/Components/PrimitiveComponent.cpp b/Engine/Source/Runtime/Engine/Private/Components/PrimitiveComponent.cpp
--- a/Engine/Source/Runtime/Engine/Private/Components/PrimitiveComponent.cpp
+++ b/Engine/Source/Runtime/Engine/Private/Components/PrimitiveComponent.cpp
@@ -2917,16 +2917,14 @@ void UPrimitiveComponent::DispatchOnInputTouchEnd(const ETouchIndex::Type Finger
 	}
 }
 
-SIZE_T UPrimitiveComponent::GetResourceSize( EResourceSizeMode::Type Mode )
+void UPrimitiveComponent::GetResourceSizeEx(FResourceSizeEx& CumulativeResourceSize)
 {
-	SIZE_T ResSize = Super::GetResourceSize(Mode);
+	Super::GetResourceSizeEx(CumulativeResourceSize);
 
 	if (BodyInstance.IsValidBodyInstance())
 	{
-		ResSize = BodyInstance.GetBodyInstanceResourceSize(Mode);
+		BodyInstance.GetBodyInstanceResourceSizeEx(CumulativeResourceSize);
 	}
-
-	return ResSize;
 }
 
 void UPrimitiveComponent::SetRenderCustomDepth(bool bValue)
diff --git a/Engine/Source/Runtime/Engine/Private/Components/SkeletalMeshComponent.cpp b/Engine/Source/Runtime/Engine/Private/Components/SkeletalMeshComponent.cpp
--- a/Engine/Source/Runtime/Engine/Private/Components/SkeletalMeshComponent.cpp
+++ b/Engine/Source/Runtime/Engine/Private/Components/SkeletalMeshComponent.cpp
@@ -2056,19 +2056,17 @@ FVector USkeletalMeshComponent::GetClosestCollidingRigidBodyLocation(const FVect
 	return Best;
 }
 
-SIZE_T USkeletalMeshComponent::GetResourceSize( EResourceSizeMode::Type Mode )
+void USkeletalMeshComponent::GetResourceSizeEx(FResourceSizeEx& CumulativeResourceSize)
 {
-	SIZE_T ResSize = Super::GetResourceSize(Mode);
+	Super::GetResourceSizeEx(CumulativeResourceSize);
 
 	for (int32 i=0; i < Bodies.Num(); ++i)
 	{
 		if (Bodies[i] != nullptr && Bodies[i]->IsValidBodyInstance())
 		{
-			ResSize += Bodies[i]->GetBodyInstanceResourceSize(Mode);
+			Bodies[i]->GetBodyInstanceResourceSizeEx(CumulativeResourceSize);
 		}
 	}
-
-	return ResSize;
 }
 
 void USkeletalMeshComponent::SetAnimationMode(EAnimationMode::Type InAnimationMode)
diff --git a/Engine/Source/Runtime/Engine/Private/Components/StaticMeshComponent.cpp b/Engine/Source/Runtime/Engine/Private/Components/StaticMeshComponent.cpp
--- a/Engine/Source/Runtime/Engine/Private/Components/StaticMeshComponent.cpp
+++ b/Engine/Source/Runtime/Engine/Private/Components/StaticMeshComponent.cpp
@@ -2084,17 +2084,17 @@ void UStaticMeshComponent::ApplyComponentInstanceData(FStaticMeshComponentInstan
 						if (LightMap->IsValid(0))
 						{
 							UTexture2D* Tex = LightMap->GetTexture(0);
-							Tex->GetResourceSize(EResourceSizeMode::Exclusive);
+							Tex->GetResourceSizeBytes(EResourceSizeMode::Exclusive);
 						}
 						if (LightMap->IsValid(1))
 						{
 							UTexture2D* Tex = LightMap->GetTexture(1);
-							Tex->GetResourceSize(EResourceSizeMode::Exclusive);
+							Tex->GetResourceSizeBytes(EResourceSizeMode::Exclusive);
 						}
 						if (LightMap->GetSkyOcclusionTexture())
 						{
 							UTexture2D* Tex = LightMap->GetSkyOcclusionTexture();
-							Tex->GetResourceSize(EResourceSizeMode::Exclusive);
+							Tex->GetResourceSizeBytes(EResourceSizeMode::Exclusive);
 						}
 					}
 				}
diff --git a/Engine/Source/Runtime/Engine/Private/Components/TextRenderComponent.cpp b/Engine/Source/Runtime/Engine/Private/Components/TextRenderComponent.cpp
--- a/Engine/Source/Runtime/Engine/Private/Components/TextRenderComponent.cpp
+++ b/Engine/Source/Runtime/Engine/Private/Components/TextRenderComponent.cpp
@@ -387,7 +387,7 @@ public:
 			}
 		}
 
-		TArray<const UMaterialInstanceDynamic*> MIDs;
+		TArray<UMaterialInstanceDynamic*> MIDs;
 	};
 
 	static void Initialize()
@@ -423,10 +423,10 @@ public:
 
 	virtual void AddReferencedObjects(FReferenceCollector& Collector) override
 	{
-		for (const auto& MIDDataPair : CachedMIDs)
+		for (auto& MIDDataPair : CachedMIDs)
 		{
-			const TSharedPtr<const FMIDData>& MIDData = MIDDataPair.Value;
-			for (const UMaterialInstanceDynamic* MID : MIDData->MIDs)
+			TSharedPtr<FMIDData>& MIDData = MIDDataPair.Value;
+			for (UMaterialInstanceDynamic*& MID : MIDData->MIDs)
 			{
 				Collector.AddReferencedObject(MID);
 			}
@@ -511,7 +511,7 @@ private:
 
 	FDelegateHandle TickerHandle;
 
-	TMap<FKey, TSharedPtr<const FMIDData>> CachedMIDs;
+	TMap<FKey, TSharedPtr<FMIDData>> CachedMIDs;
 
 	TSet<FKey> MIDsPendingPurge;
 
diff --git a/Engine/Source/Runtime/Engine/Private/DataTable.cpp b/Engine/Source/Runtime/Engine/Private/DataTable.cpp
--- a/Engine/Source/Runtime/Engine/Private/DataTable.cpp
+++ b/Engine/Source/Runtime/Engine/Private/DataTable.cpp
@@ -103,17 +103,10 @@ void UDataTable::OnPostDataImported(TArray<FString>& OutCollectedImportProblems)
 {
 	if (RowStruct && RowStruct->IsChildOf(FTableRowBase::StaticStruct()))
 	{
-		static const FString ContextString(TEXT("UDataTable::OnPostDataImport"));
-		
-		TArray<FTableRowBase*> TableRowBaseRows;
-		GetAllRows(ContextString, TableRowBaseRows);
-
-		for (FTableRowBase* CurRow : TableRowBaseRows)
+		for (const TPair<FName, uint8*>& TableRowPair : RowMap)
 		{
-			if (CurRow)
-			{
-				CurRow->OnPostDataImport(OutCollectedImportProblems);
-			}
+			FTableRowBase* CurRow = reinterpret_cast<FTableRowBase*>(TableRowPair.Value);
+			CurRow->OnPostDataImport(this, TableRowPair.Key, OutCollectedImportProblems);
 		}
 	}
 }
@@ -178,17 +171,15 @@ void UDataTable::AddReferencedObjects(UObject* InThis, FReferenceCollector& Coll
 	Super::AddReferencedObjects( This, Collector );
 }
 
-SIZE_T UDataTable::GetResourceSize(EResourceSizeMode::Type Mode)
+void UDataTable::GetResourceSizeEx(FResourceSizeEx& CumulativeResourceSize)
 {
-	SIZE_T ResourceSize = Super::GetResourceSize(Mode);
+	Super::GetResourceSizeEx(CumulativeResourceSize);
 
-	ResourceSize += RowMap.GetAllocatedSize();
+	CumulativeResourceSize.AddDedicatedSystemMemoryBytes(RowMap.GetAllocatedSize());
 	if (RowStruct)
 	{
-		ResourceSize += RowMap.Num() * RowStruct->PropertiesSize;
+		CumulativeResourceSize.AddDedicatedSystemMemoryBytes(RowMap.Num() * RowStruct->PropertiesSize);
 	}
-	
-	return ResourceSize;
 }
 
 void UDataTable::FinishDestroy()
diff --git a/Engine/Source/Runtime/Engine/Private/EngineUtils.cpp b/Engine/Source/Runtime/Engine/Private/EngineUtils.cpp
--- a/Engine/Source/Runtime/Engine/Private/EngineUtils.cpp
+++ b/Engine/Source/Runtime/Engine/Private/EngineUtils.cpp
@@ -114,7 +114,7 @@ bool FContentComparisonHelper::CompareClasses(const FString& InBaseClassName, co
 								int32 NewIndex = AssetList->AddZeroed();
 								FContentComparisonAssetInfo& Info = (*AssetList)[NewIndex];
 								Info.AssetName = Object->GetFullName();
-								Info.ResourceSize = Object->GetResourceSize(EResourceSizeMode::Inclusive);
+								Info.ResourceSize = Object->GetResourceSizeBytes(EResourceSizeMode::Inclusive);
 							}
 						}
 					}
diff --git a/Engine/Source/Runtime/Engine/Private/Font.cpp b/Engine/Source/Runtime/Engine/Private/Font.cpp
--- a/Engine/Source/Runtime/Engine/Private/Font.cpp
+++ b/Engine/Source/Runtime/Engine/Private/Font.cpp
@@ -363,9 +363,9 @@ void UFont::GetStringHeightAndWidth( const TCHAR *Text, int32& Height, int32& Wi
 	Height = FMath::CeilToInt( TotalHeight );
 }
 
-SIZE_T UFont::GetResourceSize(EResourceSizeMode::Type Mode)
+void UFont::GetResourceSizeEx(FResourceSizeEx& CumulativeResourceSize)
 {
-	int32 ResourceSize = 0;
+	Super::GetResourceSizeEx(CumulativeResourceSize);
 
 	switch(FontCacheType)
 	{
@@ -375,7 +375,7 @@ SIZE_T UFont::GetResourceSize(EResourceSizeMode::Type Mode)
 			{
 				if (Texture)
 				{
-					ResourceSize += Texture->GetResourceSize(Mode);
+					Texture->GetResourceSizeEx(CumulativeResourceSize);
 				}
 			}
 		}
@@ -383,25 +383,23 @@ SIZE_T UFont::GetResourceSize(EResourceSizeMode::Type Mode)
 
 	case EFontCacheType::Runtime:
 		{
-			auto GetTypefaceResourceSize = [](const FTypeface& Typeface) -> int32
+			auto GetTypefaceResourceSize = [&CumulativeResourceSize](const FTypeface& Typeface)
 			{
-				int32 TypefaceResourceSize = 0;
 				for (const FTypefaceEntry& TypefaceEntry : Typeface.Fonts)
 				{
 					if (TypefaceEntry.Font.BulkDataPtr)
 					{
 						// We use GetBulkDataSizeOnDisk since that will be the resident size once the bulk data has been decompressed
-						TypefaceResourceSize += TypefaceEntry.Font.BulkDataPtr->GetBulkDataSizeOnDisk();
+						CumulativeResourceSize.AddDedicatedSystemMemoryBytes(TypefaceEntry.Font.BulkDataPtr->GetBulkDataSizeOnDisk());
 					}
 				}
-				return TypefaceResourceSize;
 			};
 
 			// Sum the contained font data sizes
-			ResourceSize += GetTypefaceResourceSize(CompositeFont.DefaultTypeface);
+			GetTypefaceResourceSize(CompositeFont.DefaultTypeface);
 			for (const FCompositeSubFont& SubTypeface : CompositeFont.SubTypefaces)
 			{
-				ResourceSize += GetTypefaceResourceSize(SubTypeface.Typeface);
+				GetTypefaceResourceSize(SubTypeface.Typeface);
 			}
 		}
 		break;
@@ -409,8 +407,6 @@ SIZE_T UFont::GetResourceSize(EResourceSizeMode::Type Mode)
 	default:
 		break;
 	}
-
-	return ResourceSize;
 }
 
 void UFont::ForceLoadFontData()
diff --git a/Engine/Source/Runtime/Engine/Private/GameViewportClient.cpp b/Engine/Source/Runtime/Engine/Private/GameViewportClient.cpp
--- a/Engine/Source/Runtime/Engine/Private/GameViewportClient.cpp
+++ b/Engine/Source/Runtime/Engine/Private/GameViewportClient.cpp
@@ -379,7 +379,7 @@ bool UGameViewportClient::InputKey(FViewport* InViewport, int32 ControllerId, FK
 	// Give debugger commands a chance to process key binding
 	if (GameViewportInputKeyDelegate.IsBound())
 	{
-		if ( GameViewportInputKeyDelegate.Execute(Key, FSlateApplication::Get().GetModifierKeys()) )
+		if ( GameViewportInputKeyDelegate.Execute(Key, FSlateApplication::Get().GetModifierKeys(), EventType) )
 		{
 			return true;
 		}
diff --git a/Engine/Source/Runtime/Engine/Private/InstancedStaticMesh.cpp b/Engine/Source/Runtime/Engine/Private/InstancedStaticMesh.cpp
--- a/Engine/Source/Runtime/Engine/Private/InstancedStaticMesh.cpp
+++ b/Engine/Source/Runtime/Engine/Private/InstancedStaticMesh.cpp
@@ -1902,34 +1902,32 @@ void UInstancedStaticMeshComponent::GetNavigationPerInstanceTransforms(const FBo
 	}
 }
 
-SIZE_T UInstancedStaticMeshComponent::GetResourceSize( EResourceSizeMode::Type Mode )
+void UInstancedStaticMeshComponent::GetResourceSizeEx(FResourceSizeEx& CumulativeResourceSize)
 {
-	SIZE_T ResSize = Super::GetResourceSize(Mode);
+	Super::GetResourceSizeEx(CumulativeResourceSize);
 
 	// proxy stuff
-	ResSize += ProxySize; 
+	CumulativeResourceSize.AddDedicatedSystemMemoryBytes(ProxySize); 
 
 	// component stuff
-	ResSize += InstanceBodies.GetAllocatedSize();
+	CumulativeResourceSize.AddDedicatedSystemMemoryBytes(InstanceBodies.GetAllocatedSize());
 	for (int32 i=0; i < InstanceBodies.Num(); ++i)
 	{
 		if (InstanceBodies[i] != NULL && InstanceBodies[i]->IsValidBodyInstance())
 		{
-			ResSize += InstanceBodies[i]->GetBodyInstanceResourceSize(Mode);
+			InstanceBodies[i]->GetBodyInstanceResourceSizeEx(CumulativeResourceSize);
 		}
 	}
-	ResSize += InstanceReorderTable.GetAllocatedSize();
-	ResSize += RemovedInstances.GetAllocatedSize();
+	CumulativeResourceSize.AddDedicatedSystemMemoryBytes(InstanceReorderTable.GetAllocatedSize());
+	CumulativeResourceSize.AddDedicatedSystemMemoryBytes(RemovedInstances.GetAllocatedSize());
 
 #if WITH_EDITOR
-	ResSize += SelectedInstances.GetAllocatedSize();
+	CumulativeResourceSize.AddDedicatedSystemMemoryBytes(SelectedInstances.GetAllocatedSize());
 #endif
 
 #if WITH_PHYSX
-	ResSize += Aggregates.GetAllocatedSize();
+	CumulativeResourceSize.AddDedicatedSystemMemoryBytes(Aggregates.GetAllocatedSize());
 #endif
-
-	return ResSize;
 }
 
 void UInstancedStaticMeshComponent::BeginDestroy()
diff --git a/Engine/Source/Runtime/Engine/Private/Materials/Material.cpp b/Engine/Source/Runtime/Engine/Private/Materials/Material.cpp
--- a/Engine/Source/Runtime/Engine/Private/Materials/Material.cpp
+++ b/Engine/Source/Runtime/Engine/Private/Materials/Material.cpp
@@ -3383,27 +3383,26 @@ void UMaterial::FinishDestroy()
 	Super::FinishDestroy();
 }
 
-
-SIZE_T UMaterial::GetResourceSize(EResourceSizeMode::Type Mode)
+void UMaterial::GetResourceSizeEx(FResourceSizeEx& CumulativeResourceSize)
 {
-	int32 ResourceSize = Super::GetResourceSize(Mode);
+	Super::GetResourceSizeEx(CumulativeResourceSize);
 
 	for (int32 InstanceIndex = 0; InstanceIndex < 3; ++InstanceIndex)
 	{
 		if (DefaultMaterialInstances[InstanceIndex])
 		{
-			ResourceSize += sizeof(FDefaultMaterialInstance);
+			CumulativeResourceSize.AddDedicatedSystemMemoryBytes(sizeof(FDefaultMaterialInstance));
 		}
 	}
 
-	if (Mode == EResourceSizeMode::Inclusive)
+	if (CumulativeResourceSize.GetResourceSizeMode() == EResourceSizeMode::Inclusive)
 	{
 		for (int32 QualityLevelIndex = 0; QualityLevelIndex < EMaterialQualityLevel::Num; QualityLevelIndex++)
 		{
 			for (int32 FeatureLevelIndex = 0; FeatureLevelIndex < ERHIFeatureLevel::Num; FeatureLevelIndex++)
 			{
 				FMaterialResource* CurrentResource = MaterialResources[QualityLevelIndex][FeatureLevelIndex];
-				ResourceSize += CurrentResource->GetResourceSizeInclusive();
+				CurrentResource->GetResourceSizeEx(CumulativeResourceSize);
 			}
 		}
 
@@ -3418,13 +3417,11 @@ SIZE_T UMaterial::GetResourceSize(EResourceSizeMode::Type Mode)
 				if ( !bTextureAlreadyConsidered )
 				{
 					TheReferencedTextures.Add( Texture );
-					ResourceSize += Texture->GetResourceSize(Mode);
+					Texture->GetResourceSizeEx(CumulativeResourceSize);
 				}
 			}
 		}
 	}
-
-	return ResourceSize;
 }
 
 void UMaterial::AddReferencedObjects(UObject* InThis, FReferenceCollector& Collector)
diff --git a/Engine/Source/Runtime/Engine/Private/Materials/MaterialInstance.cpp b/Engine/Source/Runtime/Engine/Private/Materials/MaterialInstance.cpp
--- a/Engine/Source/Runtime/Engine/Private/Materials/MaterialInstance.cpp
+++ b/Engine/Source/Runtime/Engine/Private/Materials/MaterialInstance.cpp
@@ -2570,20 +2570,20 @@ bool UMaterialInstance::GetTexturesInPropertyChain(EMaterialProperty InProperty,
 }
 #endif // WITH_EDITOR
 
-SIZE_T UMaterialInstance::GetResourceSize(EResourceSizeMode::Type Mode)
+void UMaterialInstance::GetResourceSizeEx(FResourceSizeEx& CumulativeResourceSize)
 {
-	SIZE_T ResourceSize = Super::GetResourceSize(Mode);
+	Super::GetResourceSizeEx(CumulativeResourceSize);
 
 	if (bHasStaticPermutationResource)
 	{
-		if (Mode == EResourceSizeMode::Inclusive)
+		if (CumulativeResourceSize.GetResourceSizeMode() == EResourceSizeMode::Inclusive)
 		{
 			for (int32 QualityLevelIndex = 0; QualityLevelIndex < EMaterialQualityLevel::Num; QualityLevelIndex++)
 			{
 				for (int32 FeatureLevelIndex = 0; FeatureLevelIndex < ERHIFeatureLevel::Num; FeatureLevelIndex++)
 				{
 					FMaterialResource* CurrentResource = StaticPermutationMaterialResources[QualityLevelIndex][FeatureLevelIndex];
-					ResourceSize += CurrentResource->GetResourceSizeInclusive();
+					CurrentResource->GetResourceSizeEx(CumulativeResourceSize);
 				}
 			}
 		}
@@ -2593,15 +2593,13 @@ SIZE_T UMaterialInstance::GetResourceSize(EResourceSizeMode::Type Mode)
 	{
 		if (Resources[ResourceIndex])
 		{
-			ResourceSize += sizeof(FMaterialInstanceResource);
-			ResourceSize += ScalarParameterValues.Num() * sizeof(FMaterialInstanceResource::TNamedParameter<float>);
-			ResourceSize += VectorParameterValues.Num() * sizeof(FMaterialInstanceResource::TNamedParameter<FLinearColor>);
-			ResourceSize += TextureParameterValues.Num() * sizeof(FMaterialInstanceResource::TNamedParameter<const UTexture*>);
-			ResourceSize += FontParameterValues.Num() * sizeof(FMaterialInstanceResource::TNamedParameter<const UTexture*>);
+			CumulativeResourceSize.AddDedicatedSystemMemoryBytes(sizeof(FMaterialInstanceResource));
+			CumulativeResourceSize.AddDedicatedSystemMemoryBytes(ScalarParameterValues.Num() * sizeof(FMaterialInstanceResource::TNamedParameter<float>));
+			CumulativeResourceSize.AddDedicatedSystemMemoryBytes(VectorParameterValues.Num() * sizeof(FMaterialInstanceResource::TNamedParameter<FLinearColor>));
+			CumulativeResourceSize.AddDedicatedSystemMemoryBytes(TextureParameterValues.Num() * sizeof(FMaterialInstanceResource::TNamedParameter<const UTexture*>));
+			CumulativeResourceSize.AddDedicatedSystemMemoryBytes(FontParameterValues.Num() * sizeof(FMaterialInstanceResource::TNamedParameter<const UTexture*>));
 		}
 	}
-
-	return ResourceSize;
 }
 
 FPostProcessMaterialNode* IteratePostProcessMaterialNodes(const FFinalPostProcessSettings& Dest, const UMaterial* Material, FBlendableEntry*& Iterator)
diff --git a/Engine/Source/Runtime/Engine/Private/Materials/MaterialShared.cpp b/Engine/Source/Runtime/Engine/Private/Materials/MaterialShared.cpp
--- a/Engine/Source/Runtime/Engine/Private/Materials/MaterialShared.cpp
+++ b/Engine/Source/Runtime/Engine/Private/Materials/MaterialShared.cpp
@@ -1232,51 +1232,58 @@ void FMaterialResource::GetRepresentativeShaderTypesAndDescriptions(TMap<FName,
 
 SIZE_T FMaterialResource::GetResourceSizeInclusive()
 {
-	SIZE_T ResourceSize = 0;
-	TSet<const FMaterialShaderMap*> UniqueShaderMaps;
-	TMap<FShaderId, FShader*> UniqueShaders;
-	TArray<FShaderPipeline*> ShaderPipelines;
-	TSet<FShaderResourceId> UniqueShaderResourceIds;
-
-	ResourceSize += sizeof(FMaterialResource);
-	UniqueShaderMaps.Add(GetGameThreadShaderMap());
+	FResourceSizeEx ResSize = FResourceSizeEx(EResourceSizeMode::Inclusive);
+	GetResourceSizeEx(ResSize);
+	return ResSize.GetTotalMemoryBytes();
+}
 
-	for (TSet<const FMaterialShaderMap*>::TConstIterator It(UniqueShaderMaps); It; ++It)
+void FMaterialResource::GetResourceSizeEx(FResourceSizeEx& CumulativeResourceSize)
+{
+	if (CumulativeResourceSize.GetResourceSizeMode() == EResourceSizeMode::Inclusive)
 	{
-		const FMaterialShaderMap* MaterialShaderMap = *It;
-		if (MaterialShaderMap)
+		TSet<const FMaterialShaderMap*> UniqueShaderMaps;
+		TMap<FShaderId, FShader*> UniqueShaders;
+		TArray<FShaderPipeline*> ShaderPipelines;
+		TSet<FShaderResourceId> UniqueShaderResourceIds;
+
+		CumulativeResourceSize.AddDedicatedSystemMemoryBytes(sizeof(FMaterialResource));
+		UniqueShaderMaps.Add(GetGameThreadShaderMap());
+
+		for (TSet<const FMaterialShaderMap*>::TConstIterator It(UniqueShaderMaps); It; ++It)
 		{
-			ResourceSize += MaterialShaderMap->GetSizeBytes();
-			MaterialShaderMap->GetShaderList(UniqueShaders);
-			MaterialShaderMap->GetShaderPipelineList(ShaderPipelines);
+			const FMaterialShaderMap* MaterialShaderMap = *It;
+			if (MaterialShaderMap)
+			{
+				CumulativeResourceSize.AddDedicatedSystemMemoryBytes(MaterialShaderMap->GetSizeBytes());
+				MaterialShaderMap->GetShaderList(UniqueShaders);
+				MaterialShaderMap->GetShaderPipelineList(ShaderPipelines);
+			}
 		}
-	}
 
-	for (auto& KeyValue : UniqueShaders)
-	{
-		auto* Shader = KeyValue.Value;
-		if (Shader)
+		for (auto& KeyValue : UniqueShaders)
 		{
-			ResourceSize += AddShaderSize(Shader, UniqueShaderResourceIds);
+			auto* Shader = KeyValue.Value;
+			if (Shader)
+			{
+				CumulativeResourceSize.AddDedicatedSystemMemoryBytes(AddShaderSize(Shader, UniqueShaderResourceIds));
+			}
 		}
-	}
 
-	for (FShaderPipeline* Pipeline : ShaderPipelines)
-	{
-		if (Pipeline)
+		for (FShaderPipeline* Pipeline : ShaderPipelines)
 		{
-			for (FShader* Shader : Pipeline->GetShaders())
+			if (Pipeline)
 			{
-				if (Shader)
+				for (FShader* Shader : Pipeline->GetShaders())
 				{
-					ResourceSize += AddShaderSize(Shader, UniqueShaderResourceIds);
+					if (Shader)
+					{
+						CumulativeResourceSize.AddDedicatedSystemMemoryBytes(AddShaderSize(Shader, UniqueShaderResourceIds));
+					}
 				}
+				CumulativeResourceSize.AddDedicatedSystemMemoryBytes(Pipeline->GetSizeBytes());
 			}
-			ResourceSize += Pipeline->GetSizeBytes();
 		}
 	}
-
-	return ResourceSize;
 }
 
 /**
@@ -2534,14 +2541,18 @@ FMaterialUpdateContext::~FMaterialUpdateContext()
 	}
 
 	double EndTime = FPlatformTime::Seconds();
-	UE_LOG(LogMaterial, Log,
-		TEXT("%.2f seconds spent updating %d materials, %d interfaces, %d instances, %d with static permutations."),
-		(float)(EndTime - StartTime),
-		UpdatedMaterials.Num(),
-		UpdatedMaterialInterfaces.Num(),
-		InstancesToUpdate.Num(),
-		NumInstancesWithStaticPermutations
-		);
+
+	if (UpdatedMaterials.Num() > 0)
+	{
+		UE_LOG(LogMaterial, Log,
+			   TEXT("%.2f seconds spent updating %d materials, %d interfaces, %d instances, %d with static permutations."),
+			   (float)(EndTime - StartTime),
+			   UpdatedMaterials.Num(),
+			   UpdatedMaterialInterfaces.Num(),
+			   InstancesToUpdate.Num(),
+			   NumInstancesWithStaticPermutations
+			);
+	}
 }
 
 const TMap<EMaterialProperty, int32> CreatePropertyToIOIndexMap()
diff --git a/Engine/Source/Runtime/Engine/Private/Model.cpp b/Engine/Source/Runtime/Engine/Private/Model.cpp
--- a/Engine/Source/Runtime/Engine/Private/Model.cpp
+++ b/Engine/Source/Runtime/Engine/Private/Model.cpp
@@ -497,19 +497,17 @@ bool UModel::IsReadyForFinishDestroy()
 	return ReleaseResourcesFence.IsFenceComplete() && Super::IsReadyForFinishDestroy();
 }
 
-SIZE_T UModel::GetResourceSize(EResourceSizeMode::Type Mode)
+void UModel::GetResourceSizeEx(FResourceSizeEx& CumulativeResourceSize)
 {
-	int32 ResourceSize = 0;
+	Super::GetResourceSizeEx(CumulativeResourceSize);
 	
 	// I'm adding extra stuff that haven't been covered by Serialize 
 	// I don't have to include VertexFactories (based on Sam Z)
 	for(TMap<UMaterialInterface*,TScopedPointer<FRawIndexBuffer16or32> >::TConstIterator IndexBufferIt(MaterialIndexBuffers);IndexBufferIt;++IndexBufferIt)
 	{
 		const TScopedPointer<FRawIndexBuffer16or32> &IndexBuffer = IndexBufferIt.Value();
-		ResourceSize += IndexBuffer->Indices.Num() * sizeof(uint32);
+		CumulativeResourceSize.AddUnknownMemoryBytes(IndexBuffer->Indices.Num() * sizeof(uint32));
 	}
-	
-	return ResourceSize;
 }
 
 #if WITH_EDITOR
diff --git a/Engine/Source/Runtime/Engine/Private/MorphMesh.cpp b/Engine/Source/Runtime/Engine/Private/MorphMesh.cpp
--- a/Engine/Source/Runtime/Engine/Private/MorphMesh.cpp
+++ b/Engine/Source/Runtime/Engine/Private/MorphMesh.cpp
@@ -26,20 +26,31 @@ void UMorphTarget::Serialize( FArchive& Ar )
 }
 
 
-SIZE_T UMorphTarget::GetResourceSize(EResourceSizeMode::Type Mode)
+void UMorphTarget::GetResourceSizeEx(FResourceSizeEx& CumulativeResourceSize)
 {
-	SIZE_T ResourceSize = 0;
+	Super::GetResourceSizeEx(CumulativeResourceSize);
 
 	for (const auto& LODModel : MorphLODModels)
 	{
-		ResourceSize += LODModel.GetResourceSize();
+		LODModel.GetResourceSizeEx(CumulativeResourceSize);
 	}
-	return ResourceSize;
 }
 
 
 //////////////////////////////////////////////////////////////////////
 SIZE_T FMorphTargetLODModel::GetResourceSize() const
 {
-	return Vertices.GetAllocatedSize() + sizeof(int32);
-}
\ No newline at end of file
+	return GetResourceSizeBytes();
+}
+
+void FMorphTargetLODModel::GetResourceSizeEx(FResourceSizeEx& CumulativeResourceSize) const
+{
+	CumulativeResourceSize.AddUnknownMemoryBytes(Vertices.GetAllocatedSize() + sizeof(int32));
+}
+
+SIZE_T FMorphTargetLODModel::GetResourceSizeBytes() const
+{
+	FResourceSizeEx ResSize;
+	GetResourceSizeEx(ResSize);
+	return ResSize.GetTotalMemoryBytes();
+}
diff --git a/Engine/Source/Runtime/Engine/Private/Particles/ParticleBeam2EmitterInstance.cpp b/Engine/Source/Runtime/Engine/Private/Particles/ParticleBeam2EmitterInstance.cpp
--- a/Engine/Source/Runtime/Engine/Private/Particles/ParticleBeam2EmitterInstance.cpp
+++ b/Engine/Source/Runtime/Engine/Private/Particles/ParticleBeam2EmitterInstance.cpp
@@ -1251,27 +1251,24 @@ void FParticleBeam2EmitterInstance::GetAllocatedSize(int32& OutNum, int32& OutMa
  * @param	Mode	Specifies which resource size should be displayed. ( see EResourceSizeMode::Type )
  * @return	Size of resource as to be displayed to artists/ LDs in the Editor.
  */
-SIZE_T FParticleBeam2EmitterInstance::GetResourceSize(EResourceSizeMode::Type Mode)
+void FParticleBeam2EmitterInstance::GetResourceSizeEx(FResourceSizeEx& CumulativeResourceSize)
 {
-	int32 ResSize = 0;
-	
-	if (Mode == EResourceSizeMode::Inclusive || (Component && Component->SceneProxy))
+	if (CumulativeResourceSize.GetResourceSizeMode() == EResourceSizeMode::Inclusive || (Component && Component->SceneProxy))
 	{
 		int32 MaxActiveParticleDataSize = (ParticleData != NULL) ? (MaxActiveParticles * ParticleStride) : 0;
 		int32 MaxActiveParticleIndexSize = 0;
 		// Take dynamic data into account as well
-		ResSize = sizeof(FDynamicBeam2EmitterData);
-		ResSize += MaxActiveParticleDataSize;		// Copy of the particle data on the render thread
+		CumulativeResourceSize.AddUnknownMemoryBytes(sizeof(FDynamicBeam2EmitterData));
+		CumulativeResourceSize.AddUnknownMemoryBytes(MaxActiveParticleDataSize);		// Copy of the particle data on the render thread
 		if (DynamicParameterDataOffset == 0)
 		{
-			ResSize += MaxActiveParticles * sizeof(FParticleBeamTrailVertex);	// The vertex data array
+			CumulativeResourceSize.AddUnknownMemoryBytes(MaxActiveParticles * sizeof(FParticleBeamTrailVertex));	// The vertex data array
 		}
 		else
 		{
-			ResSize += MaxActiveParticles * sizeof(FParticleBeamTrailVertexDynamicParameter);
+			CumulativeResourceSize.AddUnknownMemoryBytes(MaxActiveParticles * sizeof(FParticleBeamTrailVertexDynamicParameter));
 		}
 	}
-	return ResSize;
 }
 
 /**
diff --git a/Engine/Source/Runtime/Engine/Private/Particles/ParticleComponents.cpp b/Engine/Source/Runtime/Engine/Private/Particles/ParticleComponents.cpp
--- a/Engine/Source/Runtime/Engine/Private/Particles/ParticleComponents.cpp
+++ b/Engine/Source/Runtime/Engine/Private/Particles/ParticleComponents.cpp
@@ -3418,20 +3418,20 @@ void UParticleSystemComponent::FinishDestroy()
 }
 
 
-SIZE_T UParticleSystemComponent::GetResourceSize(EResourceSizeMode::Type Mode)
+void UParticleSystemComponent::GetResourceSizeEx(FResourceSizeEx& CumulativeResourceSize)
 {
 	ForceAsyncWorkCompletion(ENSURE_AND_STALL);
-	int32 ResSize = Super::GetResourceSize(Mode);
+
+	Super::GetResourceSizeEx(CumulativeResourceSize);
 	for (int32 EmitterIdx = 0; EmitterIdx < EmitterInstances.Num(); EmitterIdx++)
 	{
 		FParticleEmitterInstance* EmitterInstance = EmitterInstances[EmitterIdx];
 		if (EmitterInstance != NULL)
 		{
 			// If the data manager has the PSys, force it to report, regardless of a PSysComp scene info being present...
-			ResSize += EmitterInstance->GetResourceSize(Mode);
+			EmitterInstance->GetResourceSizeEx(CumulativeResourceSize);
 		}
 	}
-	return ResSize;
 }
 
 
diff --git a/Engine/Source/Runtime/Engine/Private/Particles/ParticleEmitterInstances.cpp b/Engine/Source/Runtime/Engine/Private/Particles/ParticleEmitterInstances.cpp
--- a/Engine/Source/Runtime/Engine/Private/Particles/ParticleEmitterInstances.cpp
+++ b/Engine/Source/Runtime/Engine/Private/Particles/ParticleEmitterInstances.cpp
@@ -2811,26 +2811,24 @@ void FParticleSpriteEmitterInstance::GetAllocatedSize(int32& OutNum, int32& OutM
  * @param	Mode	Specifies which resource size should be displayed. ( see EResourceSizeMode::Type )
  * @return	Size of resource as to be displayed to artists/ LDs in the Editor.
  */
-SIZE_T FParticleSpriteEmitterInstance::GetResourceSize(EResourceSizeMode::Type Mode)
+void FParticleSpriteEmitterInstance::GetResourceSizeEx(FResourceSizeEx& CumulativeResourceSize)
 {
-	int32 ResSize = 0;
-	if (Mode == EResourceSizeMode::Inclusive || (Component && Component->SceneProxy))
+	if (CumulativeResourceSize.GetResourceSizeMode() == EResourceSizeMode::Inclusive || (Component && Component->SceneProxy))
 	{
 		int32 MaxActiveParticleDataSize = (ParticleData != NULL) ? (MaxActiveParticles * ParticleStride) : 0;
 		int32 MaxActiveParticleIndexSize = (ParticleIndices != NULL) ? (MaxActiveParticles * sizeof(uint16)) : 0;
 		// Take dynamic data into account as well
-		ResSize = sizeof(FDynamicSpriteEmitterData);
-		ResSize += MaxActiveParticleDataSize;								// Copy of the particle data on the render thread
-		ResSize += MaxActiveParticleIndexSize;								// Copy of the particle indices on the render thread
-		ResSize += MaxActiveParticles * sizeof(FParticleSpriteVertex);		// The vertex data array
+		CumulativeResourceSize.AddUnknownMemoryBytes(sizeof(FDynamicSpriteEmitterData));
+		CumulativeResourceSize.AddUnknownMemoryBytes(MaxActiveParticleDataSize);								// Copy of the particle data on the render thread
+		CumulativeResourceSize.AddUnknownMemoryBytes(MaxActiveParticleIndexSize);								// Copy of the particle indices on the render thread
+		CumulativeResourceSize.AddUnknownMemoryBytes(MaxActiveParticles * sizeof(FParticleSpriteVertex));		// The vertex data array
 
 		// Account for dynamic parameter data
 		if (DynamicParameterDataOffset > 0)
 		{
-			ResSize += MaxActiveParticles * sizeof(FParticleVertexDynamicParameter);
+			CumulativeResourceSize.AddUnknownMemoryBytes(MaxActiveParticles * sizeof(FParticleVertexDynamicParameter));
 		}
 	}
-	return ResSize;
 }
 
 /**
@@ -3442,19 +3440,17 @@ void FParticleMeshEmitterInstance::GetAllocatedSize(int32& OutNum, int32& OutMax
  * @param	Mode	Specifies which resource size should be displayed. ( see EResourceSizeMode::Type )
  * @return  Size of resource as to be displayed to artists/ LDs in the Editor.
  */
-SIZE_T FParticleMeshEmitterInstance::GetResourceSize(EResourceSizeMode::Type Mode)
+void FParticleMeshEmitterInstance::GetResourceSizeEx(FResourceSizeEx& CumulativeResourceSize)
 {
-	int32 ResSize = 0;
-	if (Mode == EResourceSizeMode::Inclusive || (Component && Component->SceneProxy))
+	if (CumulativeResourceSize.GetResourceSizeMode() == EResourceSizeMode::Inclusive || (Component && Component->SceneProxy))
 	{
 		int32 MaxActiveParticleDataSize = (ParticleData != NULL) ? (MaxActiveParticles * ParticleStride) : 0;
 		int32 MaxActiveParticleIndexSize = (ParticleIndices != NULL) ? (MaxActiveParticles * sizeof(uint16)) : 0;
 		// Take dynamic data into account as well
-		ResSize = sizeof(FDynamicMeshEmitterData);
-		ResSize += MaxActiveParticleDataSize;								// Copy of the particle data on the render thread
-		ResSize += MaxActiveParticleIndexSize;								// Copy of the particle indices on the render thread
+		CumulativeResourceSize.AddUnknownMemoryBytes(sizeof(FDynamicMeshEmitterData));
+		CumulativeResourceSize.AddUnknownMemoryBytes(MaxActiveParticleDataSize);								// Copy of the particle data on the render thread
+		CumulativeResourceSize.AddUnknownMemoryBytes(MaxActiveParticleIndexSize);								// Copy of the particle indices on the render thread
 	}
-	return ResSize;
 }
 
 /**
diff --git a/Engine/Source/Runtime/Engine/Private/Particles/ParticleTrail2EmitterInstance.cpp b/Engine/Source/Runtime/Engine/Private/Particles/ParticleTrail2EmitterInstance.cpp
--- a/Engine/Source/Runtime/Engine/Private/Particles/ParticleTrail2EmitterInstance.cpp
+++ b/Engine/Source/Runtime/Engine/Private/Particles/ParticleTrail2EmitterInstance.cpp
@@ -2674,27 +2674,25 @@ void FParticleRibbonEmitterInstance::GetAllocatedSize(int32& OutNum, int32& OutM
  * @param	Mode	Specifies which resource size should be displayed. ( see EResourceSizeMode::Type )
  * @return  Size of resource as to be displayed to artists/ LDs in the Editor.
  */
-SIZE_T FParticleRibbonEmitterInstance::GetResourceSize(EResourceSizeMode::Type Mode)
+void FParticleRibbonEmitterInstance::GetResourceSizeEx(FResourceSizeEx& CumulativeResourceSize)
 {
-	int32 ResSize = 0;
-	if (Mode == EResourceSizeMode::Inclusive || (Component && Component->SceneProxy))
+	if (CumulativeResourceSize.GetResourceSizeMode() == EResourceSizeMode::Inclusive || (Component && Component->SceneProxy))
 	{
 		int32 MaxActiveParticleDataSize = (ParticleData != NULL) ? (MaxActiveParticles * ParticleStride) : 0;
 		int32 MaxActiveParticleIndexSize = (ParticleIndices != NULL) ? (MaxActiveParticles * sizeof(uint16)) : 0;
 		// Take dynamic data into account as well
-		ResSize = sizeof(FParticleRibbonEmitterInstance);
-		ResSize += MaxActiveParticleDataSize;								// Copy of the particle data on the render thread
-		ResSize += MaxActiveParticleIndexSize;								// Copy of the particle indices on the render thread
+		CumulativeResourceSize.AddUnknownMemoryBytes(sizeof(FParticleRibbonEmitterInstance));
+		CumulativeResourceSize.AddUnknownMemoryBytes(MaxActiveParticleDataSize);								// Copy of the particle data on the render thread
+		CumulativeResourceSize.AddUnknownMemoryBytes(MaxActiveParticleIndexSize);								// Copy of the particle indices on the render thread
 		if (DynamicParameterDataOffset == 0)
 		{
-			ResSize += MaxActiveParticles * sizeof(FParticleBeamTrailVertex);	// The vertex data array
+			CumulativeResourceSize.AddUnknownMemoryBytes(MaxActiveParticles * sizeof(FParticleBeamTrailVertex));	// The vertex data array
 		}
 		else
 		{
-			ResSize += MaxActiveParticles * sizeof(FParticleBeamTrailVertexDynamicParameter);
+			CumulativeResourceSize.AddUnknownMemoryBytes(MaxActiveParticles * sizeof(FParticleBeamTrailVertexDynamicParameter));
 		}
 	}
-	return ResSize;
 }
 
 /**
@@ -3973,27 +3971,25 @@ void FParticleAnimTrailEmitterInstance::GetAllocatedSize(int32& OutNum, int32& O
  * @param	Mode	Specifies which resource size should be displayed. ( see EResourceSizeMode::Type )
  * @return  Size of resource as to be displayed to artists/ LDs in the Editor.
  */
-SIZE_T FParticleAnimTrailEmitterInstance::GetResourceSize(EResourceSizeMode::Type Mode)
+void FParticleAnimTrailEmitterInstance::GetResourceSizeEx(FResourceSizeEx& CumulativeResourceSize)
 {
-	int32 ResSize = 0;
-	if (Mode == EResourceSizeMode::Inclusive || (Component && Component->SceneProxy))
+	if (CumulativeResourceSize.GetResourceSizeMode() == EResourceSizeMode::Inclusive || (Component && Component->SceneProxy))
 	{
 		int32 MaxActiveParticleDataSize = (ParticleData != NULL) ? (MaxActiveParticles * ParticleStride) : 0;
 		int32 MaxActiveParticleIndexSize = (ParticleIndices != NULL) ? (MaxActiveParticles * sizeof(uint16)) : 0;
 		// Take dynamic data into account as well
-		ResSize = sizeof(FParticleAnimTrailEmitterInstance);
-		ResSize += MaxActiveParticleDataSize;								// Copy of the particle data on the render thread
-		ResSize += MaxActiveParticleIndexSize;								// Copy of the particle indices on the render thread
+		CumulativeResourceSize.AddUnknownMemoryBytes(sizeof(FParticleAnimTrailEmitterInstance));
+		CumulativeResourceSize.AddUnknownMemoryBytes(MaxActiveParticleDataSize);								// Copy of the particle data on the render thread
+		CumulativeResourceSize.AddUnknownMemoryBytes(MaxActiveParticleIndexSize);								// Copy of the particle indices on the render thread
 		if (DynamicParameterDataOffset == 0)
 		{
-			ResSize += MaxActiveParticles * sizeof(FParticleBeamTrailVertex);	// The vertex data array
+			CumulativeResourceSize.AddUnknownMemoryBytes(MaxActiveParticles * sizeof(FParticleBeamTrailVertex));	// The vertex data array
 		}
 		else
 		{
-			ResSize += MaxActiveParticles * sizeof(FParticleBeamTrailVertexDynamicParameter);
+			CumulativeResourceSize.AddUnknownMemoryBytes(MaxActiveParticles * sizeof(FParticleBeamTrailVertexDynamicParameter));
 		}
 	}
-	return ResSize;
 }
 
 void FParticleAnimTrailEmitterInstance::BeginTrail()
diff --git a/Engine/Source/Runtime/Engine/Private/PhysicsEngine/BodyInstance.cpp b/Engine/Source/Runtime/Engine/Private/PhysicsEngine/BodyInstance.cpp
--- a/Engine/Source/Runtime/Engine/Private/PhysicsEngine/BodyInstance.cpp
+++ b/Engine/Source/Runtime/Engine/Private/PhysicsEngine/BodyInstance.cpp
@@ -4657,22 +4657,25 @@ void FBodyInstance::LoadProfileData(bool bVerifyProfile)
 
 SIZE_T FBodyInstance::GetBodyInstanceResourceSize(EResourceSizeMode::Type Mode) const
 {
-	SIZE_T ResSize = 0;
+	FResourceSizeEx ResSize = FResourceSizeEx(Mode);
+	GetBodyInstanceResourceSizeEx(ResSize);
+	return ResSize.GetTotalMemoryBytes();
+}
 
+void FBodyInstance::GetBodyInstanceResourceSizeEx(FResourceSizeEx& CumulativeResourceSize) const
+{
 #if WITH_PHYSX
 	// Get size of PhysX data, skipping contents of 'shared data'
 	if (RigidActorSync != NULL)
 	{
-		ResSize += GetPhysxObjectSize(RigidActorSync, FPhysxSharedData::Get().GetCollection());
+		CumulativeResourceSize.AddDedicatedSystemMemoryBytes(GetPhysxObjectSize(RigidActorSync, FPhysxSharedData::Get().GetCollection()));
 	}
 
 	if (RigidActorAsync != NULL)
 	{
-		ResSize += GetPhysxObjectSize(RigidActorAsync, FPhysxSharedData::Get().GetCollection());
+		CumulativeResourceSize.AddDedicatedSystemMemoryBytes(GetPhysxObjectSize(RigidActorAsync, FPhysxSharedData::Get().GetCollection()));
 	}
 #endif
-
-	return ResSize;
 }
 
 void FBodyInstance::FixupData(class UObject* Loader)
diff --git a/Engine/Source/Runtime/Engine/Private/PhysicsEngine/BodySetup.cpp b/Engine/Source/Runtime/Engine/Private/PhysicsEngine/BodySetup.cpp
--- a/Engine/Source/Runtime/Engine/Private/PhysicsEngine/BodySetup.cpp
+++ b/Engine/Source/Runtime/Engine/Private/PhysicsEngine/BodySetup.cpp
@@ -51,19 +51,29 @@ static TAutoConsoleVariable<float> CVarMaxContactOffset(
 	ECVF_Default);
 
 
-SIZE_T FBodySetupUVInfo::GetResourceSize()
+SIZE_T FBodySetupUVInfo::GetResourceSize() const
 {
-	SIZE_T Size = 0;
-	Size += IndexBuffer.GetAllocatedSize();
-	Size += VertPositions.GetAllocatedSize();
+	return GetResourceSizeBytes();
+}
+
+void FBodySetupUVInfo::GetResourceSizeEx(FResourceSizeEx& CumulativeResourceSize) const
+{
+	CumulativeResourceSize.AddDedicatedSystemMemoryBytes(IndexBuffer.GetAllocatedSize());
+	CumulativeResourceSize.AddDedicatedSystemMemoryBytes(VertPositions.GetAllocatedSize());
 
 	for (int32 ChannelIdx = 0; ChannelIdx < VertUVs.Num(); ChannelIdx++)
 	{
-		Size += VertUVs[ChannelIdx].GetAllocatedSize();
+		CumulativeResourceSize.AddDedicatedSystemMemoryBytes(VertUVs[ChannelIdx].GetAllocatedSize());
 	}
 
-	Size += VertUVs.GetAllocatedSize();
-	return Size;
+	CumulativeResourceSize.AddDedicatedSystemMemoryBytes(VertUVs.GetAllocatedSize());
+}
+
+SIZE_T FBodySetupUVInfo::GetResourceSizeBytes() const
+{
+	FResourceSizeEx ResSize;
+	GetResourceSizeEx(ResSize);
+	return ResSize.GetTotalMemoryBytes();
 }
 
 
@@ -1188,15 +1198,15 @@ void UBodySetup::CopyBodySetupProperty(const UBodySetup* Other)
 
 #endif // WITH_EDITOR
 
-SIZE_T UBodySetup::GetResourceSize( EResourceSizeMode::Type Mode )
+void UBodySetup::GetResourceSizeEx(FResourceSizeEx& CumulativeResourceSize)
 {
-	SIZE_T ResourceSize = Super::GetResourceSize(Mode);
+	Super::GetResourceSizeEx(CumulativeResourceSize);
 
 #if WITH_PHYSX
 	// Count PhysX trimesh mem usage
 	for(PxTriangleMesh* TriMesh : TriMeshes)
 	{
-		ResourceSize += GetPhysxObjectSize(TriMesh, NULL);
+		CumulativeResourceSize.AddDedicatedSystemMemoryBytes(GetPhysxObjectSize(TriMesh, NULL));
 	}
 
 	// Count PhysX convex mem usage
@@ -1206,12 +1216,12 @@ SIZE_T UBodySetup::GetResourceSize( EResourceSizeMode::Type Mode )
 
 		if(ConvexElem.ConvexMesh != NULL)
 		{
-			ResourceSize += GetPhysxObjectSize(ConvexElem.ConvexMesh, NULL);
+			CumulativeResourceSize.AddDedicatedSystemMemoryBytes(GetPhysxObjectSize(ConvexElem.ConvexMesh, NULL));
 		}
 
 		if(ConvexElem.ConvexMeshNegX != NULL)
 		{
-			ResourceSize += GetPhysxObjectSize(ConvexElem.ConvexMeshNegX, NULL);
+			CumulativeResourceSize.AddDedicatedSystemMemoryBytes(GetPhysxObjectSize(ConvexElem.ConvexMeshNegX, NULL));
 		}
 	}
 
@@ -1220,13 +1230,11 @@ SIZE_T UBodySetup::GetResourceSize( EResourceSizeMode::Type Mode )
 	if (CookedFormatData.Contains(FPlatformProperties::GetPhysicsFormat()))
 	{
 		const FByteBulkData& FmtData = CookedFormatData.GetFormat(FPlatformProperties::GetPhysicsFormat());
-		ResourceSize += FmtData.GetElementSize() * FmtData.GetElementCount();
+		CumulativeResourceSize.AddDedicatedSystemMemoryBytes(FmtData.GetElementSize() * FmtData.GetElementCount());
 	}
 	
 	// Count any UV info
-	ResourceSize += UVInfo.GetResourceSize();
-
-	return ResourceSize;
+	UVInfo.GetResourceSizeEx(CumulativeResourceSize);
 }
 
 void FKAggregateGeom::Serialize( const FArchive& Ar )
diff --git a/Engine/Source/Runtime/Engine/Private/PhysicsEngine/PhysicsAsset.cpp b/Engine/Source/Runtime/Engine/Private/PhysicsEngine/PhysicsAsset.cpp
--- a/Engine/Source/Runtime/Engine/Private/PhysicsEngine/PhysicsAsset.cpp
+++ b/Engine/Source/Runtime/Engine/Private/PhysicsEngine/PhysicsAsset.cpp
@@ -693,20 +693,19 @@ void UPhysicsAsset::BodyFindConstraints(int32 BodyIndex, TArray<int32>& Constrai
 	}
 }
 
-SIZE_T UPhysicsAsset::GetResourceSize(EResourceSizeMode::Type Mode)
+void UPhysicsAsset::GetResourceSizeEx(FResourceSizeEx& CumulativeResourceSize)
 {
-	SIZE_T ResourceSize = Super::GetResourceSize(Mode);
+	Super::GetResourceSizeEx(CumulativeResourceSize);
 
 	for (const auto& SingleBody : SkeletalBodySetups)
 	{
-		ResourceSize += SingleBody->GetResourceSize(Mode);
+		SingleBody->GetResourceSizeEx(CumulativeResourceSize);
 	}
 
-	ResourceSize += BodySetupIndexMap.GetAllocatedSize();
-	ResourceSize += CollisionDisableTable.GetAllocatedSize();
+	CumulativeResourceSize.AddDedicatedSystemMemoryBytes(BodySetupIndexMap.GetAllocatedSize());
+	CumulativeResourceSize.AddDedicatedSystemMemoryBytes(CollisionDisableTable.GetAllocatedSize());
 
 	// @todo implement inclusive mode
-	return ResourceSize;
 }
 
 #undef LOCTEXT_NAMESPACE
\ No newline at end of file
diff --git a/Engine/Source/Runtime/Engine/Private/RenderingSettings.cpp b/Engine/Source/Runtime/Engine/Private/RenderingSettings.cpp
--- a/Engine/Source/Runtime/Engine/Private/RenderingSettings.cpp
+++ b/Engine/Source/Runtime/Engine/Private/RenderingSettings.cpp
@@ -23,24 +23,7 @@ URendererSettings::URendererSettings(const FObjectInitializer& ObjectInitializer
 void URendererSettings::PostInitProperties()
 {
 	Super::PostInitProperties();
-
-	if ( IsTemplate() )
-	{
-		// If we have UI scale curve data in the renderer settings config, move it to the user interface settings.
-		if ( FRichCurve* OldCurve = UIScaleCurve_DEPRECATED.GetRichCurve() )
-		{
-			if ( OldCurve->GetNumKeys() != 0 )
-			{
-				UUserInterfaceSettings* UISettings = GetMutableDefault<UUserInterfaceSettings>(UUserInterfaceSettings::StaticClass());
-				UISettings->UIScaleRule = UIScaleRule_DEPRECATED;
-				UISettings->UIScaleCurve = UIScaleCurve_DEPRECATED;
-
-				// Remove all old keys so that we don't attempt to load from renderer settings again.
-				OldCurve->Reset();
-			}
-		}
-	}
-
+	
 	SanatizeReflectionCaptureResolution();
 
 #if WITH_EDITOR
@@ -62,14 +45,6 @@ void URendererSettings::PostEditChangeProperty(FPropertyChangedEvent& PropertyCh
 	{
 		ExportValuesToConsoleVariables(PropertyChangedEvent.Property);
 
-		// If the renderer settings are updated and saved, we need to update the user interface settings too, to prevent data loss
-		// from old curve data for UI settings.
-		if ( IsTemplate() && PropertyChangedEvent.ChangeType != EPropertyChangeType::Interactive )
-		{
-			UUserInterfaceSettings* UISettings = GetMutableDefault<UUserInterfaceSettings>(UUserInterfaceSettings::StaticClass());
-			UISettings->UpdateDefaultConfigFile();
-		}
-
 		if (PropertyChangedEvent.Property->GetFName() == GET_MEMBER_NAME_CHECKED(URendererSettings, ReflectionCaptureResolution) && 
 			PropertyChangedEvent.ChangeType != EPropertyChangeType::Interactive)
 		{
diff --git a/Engine/Source/Runtime/Engine/Private/SceneManagement.cpp b/Engine/Source/Runtime/Engine/Private/SceneManagement.cpp
--- a/Engine/Source/Runtime/Engine/Private/SceneManagement.cpp
+++ b/Engine/Source/Runtime/Engine/Private/SceneManagement.cpp
@@ -164,12 +164,9 @@ void FSimpleElementCollector::DrawBatchedElements(FRHICommandList& RHICmdList, c
 		RHICmdList,
 		InView.GetFeatureLevel(),
 		bNeedToSwitchVerticalAxis,
-		InView.ViewProjectionMatrix,
-		InView.ViewRect.Width(),
-		InView.ViewRect.Height(),
+		InView,
 		InView.Family->EngineShowFlags.HitProxies,
 		1.0f,
-		&InView,
 		DepthTexture,
 		Filter
 		);
diff --git a/Engine/Source/Runtime/Engine/Private/Selection.cpp b/Engine/Source/Runtime/Engine/Private/Selection.cpp
--- a/Engine/Source/Runtime/Engine/Private/Selection.cpp
+++ b/Engine/Source/Runtime/Engine/Private/Selection.cpp
@@ -216,7 +216,11 @@ void USelection::PostEditUndo()
 {
 	Super::PostEditUndo();
 
-	// Sync up the annotation (for fast lookup) with the master array
+
+	// The set of selected objects may have changed, so make sure our annotations exactly match the list, otherwise
+	// UObject::IsSelected() could return a result that was different from the list of objects returned by GetSelectedObjects()
+	GSelectedAnnotation.ClearAll();
+
 	for(TWeakObjectPtr<UObject>& ObjectPtr : SelectedObjects)
 	{
 		UObject* Object = ObjectPtr.Get(true);
diff --git a/Engine/Source/Runtime/Engine/Private/SkeletalMesh.cpp b/Engine/Source/Runtime/Engine/Private/SkeletalMesh.cpp
--- a/Engine/Source/Runtime/Engine/Private/SkeletalMesh.cpp
+++ b/Engine/Source/Runtime/Engine/Private/SkeletalMesh.cpp
@@ -161,7 +161,11 @@ static bool SaveApexClothingAssetToBlob(const apex::ClothingAsset *InAsset, TArr
 ////////////////////////////////////////////////
 SIZE_T FClothingAssetData::GetResourceSize() const
 {
-	SIZE_T ResourceSize = 0;
+	return GetResourceSizeBytes();
+}
+
+void FClothingAssetData::GetResourceSizeEx(FResourceSizeEx& CumulativeResourceSize) const
+{
 #if WITH_APEX_CLOTHING
 	if (ApexClothingAsset)
 	{
@@ -173,20 +177,25 @@ SIZE_T FClothingAssetData::GetResourceSize() const
 				apex::RenderMeshAssetStats AssetStats;
 				RenderAsset->getStats(AssetStats);
 
-				ResourceSize += AssetStats.totalBytes;
+				CumulativeResourceSize.AddUnknownMemoryBytes(AssetStats.totalBytes);
 			}
 		}
 	}
 
-	ResourceSize += ClothCollisionVolumes.GetAllocatedSize();
-	ResourceSize += ClothCollisionConvexPlaneIndices.GetAllocatedSize();
-	ResourceSize += ClothCollisionVolumePlanes.GetAllocatedSize();
-	ResourceSize += ClothBoneSpheres.GetAllocatedSize();
-	ResourceSize += BoneSphereConnections.GetAllocatedSize();
-	ResourceSize += ClothVisualizationInfos.GetAllocatedSize();
+	CumulativeResourceSize.AddUnknownMemoryBytes(ClothCollisionVolumes.GetAllocatedSize());
+	CumulativeResourceSize.AddUnknownMemoryBytes(ClothCollisionConvexPlaneIndices.GetAllocatedSize());
+	CumulativeResourceSize.AddUnknownMemoryBytes(ClothCollisionVolumePlanes.GetAllocatedSize());
+	CumulativeResourceSize.AddUnknownMemoryBytes(ClothBoneSpheres.GetAllocatedSize());
+	CumulativeResourceSize.AddUnknownMemoryBytes(BoneSphereConnections.GetAllocatedSize());
+	CumulativeResourceSize.AddUnknownMemoryBytes(ClothVisualizationInfos.GetAllocatedSize());
 #endif // #if WITH_APEX_CLOTHING
+}
 
-	return ResourceSize;
+SIZE_T FClothingAssetData::GetResourceSizeBytes() const
+{
+	FResourceSizeEx ResSize;
+	GetResourceSizeEx(ResSize);
+	return ResSize.GetTotalMemoryBytes();
 }
 
 /*-----------------------------------------------------------------------------
@@ -1970,18 +1979,21 @@ void FStaticLODModel::ReleaseCPUResources()
 
 SIZE_T FStaticLODModel::GetResourceSize() const
 {
-	SIZE_T ResourceSize = 0;
+	return GetResourceSizeBytes();
+}
 
-	ResourceSize += Sections.GetAllocatedSize();
-	ResourceSize += ActiveBoneIndices.GetAllocatedSize();  
-	ResourceSize += RequiredBones.GetAllocatedSize();
+void FStaticLODModel::GetResourceSizeEx(FResourceSizeEx& CumulativeResourceSize) const
+{
+	CumulativeResourceSize.AddUnknownMemoryBytes(Sections.GetAllocatedSize());
+	CumulativeResourceSize.AddUnknownMemoryBytes(ActiveBoneIndices.GetAllocatedSize()); 
+	CumulativeResourceSize.AddUnknownMemoryBytes(RequiredBones.GetAllocatedSize());
 
 	if(MultiSizeIndexContainer.IsIndexBufferValid())
 	{
 		const FRawStaticIndexBuffer16or32Interface* IndexBuffer = MultiSizeIndexContainer.GetIndexBuffer();
 		if (IndexBuffer)
 		{
-			ResourceSize += IndexBuffer->GetResourceDataSize(); 
+			CumulativeResourceSize.AddUnknownMemoryBytes(IndexBuffer->GetResourceDataSize()); 
 		}
 	}
 
@@ -1990,22 +2002,27 @@ SIZE_T FStaticLODModel::GetResourceSize() const
 		const FRawStaticIndexBuffer16or32Interface* AdjacentIndexBuffer = AdjacencyMultiSizeIndexContainer.GetIndexBuffer();
 		if(AdjacentIndexBuffer)
 		{
-			ResourceSize += AdjacentIndexBuffer->GetResourceDataSize();
+			CumulativeResourceSize.AddUnknownMemoryBytes(AdjacentIndexBuffer->GetResourceDataSize());
 		}
 	}
 
-	ResourceSize += VertexBufferGPUSkin.GetVertexDataSize();
-	ResourceSize += ColorVertexBuffer.GetVertexDataSize();
-	ResourceSize += APEXClothVertexBuffer.GetVertexDataSize();
+	CumulativeResourceSize.AddUnknownMemoryBytes(VertexBufferGPUSkin.GetVertexDataSize());
+	CumulativeResourceSize.AddUnknownMemoryBytes(ColorVertexBuffer.GetVertexDataSize());
+	CumulativeResourceSize.AddUnknownMemoryBytes(APEXClothVertexBuffer.GetVertexDataSize());
 
-	ResourceSize += RawPointIndices.GetBulkDataSize();
-	ResourceSize += LegacyRawPointIndices.GetBulkDataSize();
-	ResourceSize += MeshToImportVertexMap.GetAllocatedSize();
+	CumulativeResourceSize.AddUnknownMemoryBytes(RawPointIndices.GetBulkDataSize());
+	CumulativeResourceSize.AddUnknownMemoryBytes(LegacyRawPointIndices.GetBulkDataSize());
+	CumulativeResourceSize.AddUnknownMemoryBytes(MeshToImportVertexMap.GetAllocatedSize());
 
 	// I suppose we add everything we could
-	ResourceSize += sizeof(int32);
+	CumulativeResourceSize.AddUnknownMemoryBytes(sizeof(int32));
+}
 
-	return ResourceSize;
+SIZE_T FStaticLODModel::GetResourceSizeBytes() const
+{
+	FResourceSizeEx ResSize;
+	GetResourceSizeEx(ResSize);
+	return ResSize.GetTotalMemoryBytes();
 }
 
 void FStaticLODModel::RebuildIndexBuffer(FMultiSizeIndexContainerData* IndexBufferData, FMultiSizeIndexContainerData* AdjacencyIndexBufferData)
@@ -2351,15 +2368,23 @@ bool FSkeletalMeshResource::RequiresCPUSkinning(ERHIFeatureLevel::Type FeatureLe
 
 SIZE_T FSkeletalMeshResource::GetResourceSize()
 {
-	SIZE_T ResourceSize = 0;
+	return GetResourceSizeBytes();
+}
+
+void FSkeletalMeshResource::GetResourceSizeEx(FResourceSizeEx& CumulativeResourceSize)
+{
 	for(int32 LODIndex = 0; LODIndex < LODModels.Num(); ++LODIndex)
 	{
 		const FStaticLODModel& Model = LODModels[LODIndex];
-
-		ResourceSize += Model.GetResourceSize();
+		Model.GetResourceSizeEx(CumulativeResourceSize);
 	}
+}
 
-	return ResourceSize;
+SIZE_T FSkeletalMeshResource::GetResourceSizeBytes()
+{
+	FResourceSizeEx ResSize;
+	GetResourceSizeEx(ResSize);
+	return ResSize.GetTotalMemoryBytes();
 }
 
 /*-----------------------------------------------------------------------------
@@ -2621,24 +2646,25 @@ float USkeletalMesh::GetStreamingTextureFactor( int32 RequestedUVIndex )
 }
 
 
-SIZE_T USkeletalMesh::GetResourceSize(EResourceSizeMode::Type Mode)
+void USkeletalMesh::GetResourceSizeEx(FResourceSizeEx& CumulativeResourceSize)
 {
-	SIZE_T ResourceSize = 0;
+	Super::GetResourceSizeEx(CumulativeResourceSize);
+
 	if (ImportedResource.IsValid())
 	{
-		ResourceSize += ImportedResource->GetResourceSize();
+		ImportedResource->GetResourceSizeEx(CumulativeResourceSize);
 	}
 
-	if (Mode == EResourceSizeMode::Inclusive)
+	if (CumulativeResourceSize.GetResourceSizeMode() == EResourceSizeMode::Inclusive)
 	{
 		for (const auto& MorphTarget : MorphTargets)
 		{
-			ResourceSize += MorphTarget->GetResourceSize(Mode);
+			MorphTarget->GetResourceSizeEx(CumulativeResourceSize);
 		}
 
 		for (const auto& ClothingAsset : ClothingAssets)
 		{
-			ResourceSize += ClothingAsset.GetResourceSize();
+			ClothingAsset.GetResourceSizeEx(CumulativeResourceSize);
 		}
 
 		TSet<UMaterialInterface*> UniqueMaterials;
@@ -2649,29 +2675,27 @@ SIZE_T USkeletalMesh::GetResourceSize(EResourceSizeMode::Type Mode)
 			UniqueMaterials.Add(Material, &bAlreadyCounted);
 			if(!bAlreadyCounted && Material)
 			{
-				ResourceSize += Material->GetResourceSize(Mode);
+				Material->GetResourceSizeEx(CumulativeResourceSize);
 			}
 		}
 
 #if WITH_EDITORONLY_DATA
-		ResourceSize += RetargetBasePose.GetAllocatedSize();
+		CumulativeResourceSize.AddDedicatedSystemMemoryBytes(RetargetBasePose.GetAllocatedSize());
 #endif
 
-		ResourceSize += RefBasesInvMatrix.GetAllocatedSize();
-		ResourceSize += RefSkeleton.GetDataSize();
+		CumulativeResourceSize.AddDedicatedSystemMemoryBytes(RefBasesInvMatrix.GetAllocatedSize());
+		CumulativeResourceSize.AddDedicatedSystemMemoryBytes(RefSkeleton.GetDataSize());
 
 		if (BodySetup)
 		{
-			ResourceSize += BodySetup->GetResourceSize(Mode);
+			BodySetup->GetResourceSizeEx(CumulativeResourceSize);
 		}
 
 		if (PhysicsAsset)
 		{
-			ResourceSize += PhysicsAsset->GetResourceSize(Mode);
+			PhysicsAsset->GetResourceSizeEx(CumulativeResourceSize);
 		}
 	}
-
-	return ResourceSize;
 }
 
 /**
@@ -4738,6 +4762,20 @@ FString ASkeletalMeshActor::GetDetailedInfoInternal() const
 	return Result;  
 }
 
+#if WITH_EDITOR
+void ASkeletalMeshActor::PostEditChangeChainProperty(FPropertyChangedChainEvent& PropertyChangedEvent)
+{
+	Super::PostEditChangeChainProperty(PropertyChangedEvent);
+
+	if (PropertyChangedEvent.Property != nullptr)
+	{
+		if (PropertyChangedEvent.Property->GetFName() == GET_MEMBER_NAME_STRING_CHECKED(ASkeletalMeshActor, SkeletalMeshComponent) && SkeletalMeshComponent->SkeletalMesh != nullptr)
+		{
+			SkeletalMeshComponent->CleanUpOverrideMaterials();
+		}
+	}
+}
+#endif
 
 void ASkeletalMeshActor::PostInitializeComponents()
 {
@@ -5769,17 +5807,15 @@ void USkeletalMeshComponent::Serialize(FArchive& Ar)
 	PRAGMA_ENABLE_DEPRECATION_WARNINGS
 }
 
-
-SIZE_T USkinnedMeshComponent::GetResourceSize(EResourceSizeMode::Type Mode)
+void USkinnedMeshComponent::GetResourceSizeEx(FResourceSizeEx& CumulativeResourceSize)
 {
-	int32 ResourceSize = 0;
+	Super::GetResourceSizeEx(CumulativeResourceSize);
+
 	// Get Mesh Object's memory
 	if(MeshObject)
 	{
-		ResourceSize += MeshObject->GetResourceSize();
+		MeshObject->GetResourceSizeEx(CumulativeResourceSize);
 	}
-
-	return ResourceSize;
 }
 
 FPrimitiveSceneProxy* USkinnedMeshComponent::CreateSceneProxy()
diff --git a/Engine/Source/Runtime/Engine/Private/SkeletalRenderCPUSkin.h b/Engine/Source/Runtime/Engine/Private/SkeletalRenderCPUSkin.h
--- a/Engine/Source/Runtime/Engine/Private/SkeletalRenderCPUSkin.h
+++ b/Engine/Source/Runtime/Engine/Private/SkeletalRenderCPUSkin.h
@@ -123,16 +123,27 @@ public:
 	/**
 	* Returns the size of memory allocated by render data
 	*/
-	virtual SIZE_T GetResourceSize()
+	DEPRECATED(4.14, "GetResourceSize is deprecated. Please use GetResourceSizeEx or GetResourceSizeBytes instead.")
+	SIZE_T GetResourceSize()
+	{
+		return GetResourceSizeBytes();
+	}
+
+	virtual void GetResourceSizeEx(FResourceSizeEx& CumulativeResourceSize)
  	{
-		SIZE_T ResourceSize = sizeof(*this);
+		CumulativeResourceSize.AddDedicatedSystemMemoryBytes(sizeof(*this));
  		
- 		ResourceSize += ReferenceToLocal.GetAllocatedSize();
- 		ResourceSize += ActiveMorphTargets.GetAllocatedSize();
- 		
-		return ResourceSize;
+		CumulativeResourceSize.AddDedicatedSystemMemoryBytes(ReferenceToLocal.GetAllocatedSize());
+		CumulativeResourceSize.AddDedicatedSystemMemoryBytes(ActiveMorphTargets.GetAllocatedSize());
  	}
 
+	SIZE_T GetResourceSizeBytes()
+	{
+		FResourceSizeEx ResSize;
+		GetResourceSizeEx(ResSize);
+		return ResSize.GetTotalMemoryBytes();
+	}
+
 	/** Update Simulated Positions & Normals from APEX Clothing actor */
 	bool UpdateClothSimulationData(USkinnedMeshComponent* InMeshComponent);
 };
@@ -178,27 +189,25 @@ public:
 		return ( DynamicData!=NULL ); 
 	}
 
-	virtual SIZE_T GetResourceSize() override
+	virtual void GetResourceSizeEx(FResourceSizeEx& CumulativeResourceSize) override
 	{
-		SIZE_T ResourceSize=sizeof(*this);
+		CumulativeResourceSize.AddDedicatedSystemMemoryBytes(sizeof(*this));
 
 		if(DynamicData)
 		{
-			ResourceSize += DynamicData->GetResourceSize();
+			DynamicData->GetResourceSizeEx(CumulativeResourceSize);
 		}
 
-		ResourceSize += LODs.GetAllocatedSize(); 
+		CumulativeResourceSize.AddDedicatedSystemMemoryBytes(LODs.GetAllocatedSize()); 
 
 		// include extra data from LOD
 		for (int32 I=0; I<LODs.Num(); ++I)
 		{
-			ResourceSize += LODs[I].GetResourceSize();
+			LODs[I].GetResourceSizeEx(CumulativeResourceSize);
 		}
 
-		ResourceSize += CachedFinalVertices.GetAllocatedSize();
-		ResourceSize += BonesOfInterest.GetAllocatedSize();
-
-		return ResourceSize;
+		CumulativeResourceSize.AddDedicatedSystemMemoryBytes(CachedFinalVertices.GetAllocatedSize());
+		CumulativeResourceSize.AddDedicatedSystemMemoryBytes(BonesOfInterest.GetAllocatedSize());
 	}
 
 	virtual void DrawVertexElements(FPrimitiveDrawInterface* PDI, const FTransform& ToWorldSpace, bool bDrawNormals, bool bDrawTangents, bool bDrawBinormals) const override;
@@ -240,9 +249,22 @@ private:
 		/**
 	 	 * Get Resource Size : return the size of Resource this allocates
 	 	 */
+		DEPRECATED(4.14, "GetResourceSize is deprecated. Please use GetResourceSizeEx or GetResourceSizeBytes instead.")
 		SIZE_T GetResourceSize()
 		{
-			return VertexBuffer.GetResourceSize();
+			return GetResourceSizeBytes();
+		}
+
+		void GetResourceSizeEx(FResourceSizeEx& CumulativeResourceSize)
+		{
+			CumulativeResourceSize.AddUnknownMemoryBytes(VertexBuffer.GetResourceSize());
+		}
+
+		SIZE_T GetResourceSizeBytes()
+		{
+			FResourceSizeEx ResSize;
+			GetResourceSizeEx(ResSize);
+			return ResSize.GetTotalMemoryBytes();
 		}
 	};
 
diff --git a/Engine/Source/Runtime/Engine/Private/SkeletalRenderGPUSkin.h b/Engine/Source/Runtime/Engine/Private/SkeletalRenderGPUSkin.h
--- a/Engine/Source/Runtime/Engine/Private/SkeletalRenderGPUSkin.h
+++ b/Engine/Source/Runtime/Engine/Private/SkeletalRenderGPUSkin.h
@@ -90,14 +90,25 @@ public:
 	/**
 	* Returns the size of memory allocated by render data
 	*/
-	virtual SIZE_T GetResourceSize()
+	DEPRECATED(4.14, "GetResourceSize is deprecated. Please use GetResourceSizeEx or GetResourceSizeBytes instead.")
+	SIZE_T GetResourceSize()
 	{
-		SIZE_T ResourceSize = sizeof(*this);
+		return GetResourceSizeBytes();
+	}
+
+	virtual void GetResourceSizeEx(FResourceSizeEx& CumulativeResourceSize)
+	{
+		CumulativeResourceSize.AddDedicatedSystemMemoryBytes(sizeof(*this));
 		
-		ResourceSize += ReferenceToLocal.GetAllocatedSize();
-		ResourceSize += ActiveMorphTargets.GetAllocatedSize();
+		CumulativeResourceSize.AddUnknownMemoryBytes(ReferenceToLocal.GetAllocatedSize());
+		CumulativeResourceSize.AddUnknownMemoryBytes(ActiveMorphTargets.GetAllocatedSize());
+	}
 
-		return ResourceSize;
+	SIZE_T GetResourceSizeBytes()
+	{
+		FResourceSizeEx ResSize;
+		GetResourceSizeEx(ResSize);
+		return ResSize.GetTotalMemoryBytes();
 	}
 
 	/** Update Simulated Positions & Normals from APEX Clothing actor */
@@ -258,26 +269,24 @@ public:
 		return ( DynamicData!=NULL ); 
 	}
 
-	virtual SIZE_T GetResourceSize() override
+	virtual void GetResourceSizeEx(FResourceSizeEx& CumulativeResourceSize) override
 	{
-		SIZE_T ResourceSize = sizeof(*this);
+		CumulativeResourceSize.AddDedicatedSystemMemoryBytes(sizeof(*this));
 		
 		if(DynamicData)
 		{
-			ResourceSize += DynamicData->GetResourceSize();
+			DynamicData->GetResourceSizeEx(CumulativeResourceSize);
 		}
 		
-		ResourceSize += LODs.GetAllocatedSize(); 
+		CumulativeResourceSize.AddUnknownMemoryBytes(LODs.GetAllocatedSize()); 
 
 		// include extra data from LOD
 		for (int32 I=0; I<LODs.Num(); ++I)
 		{
-			ResourceSize += LODs[I].GetResourceSize();
+			LODs[I].GetResourceSizeEx(CumulativeResourceSize);
 		}
 
-		ResourceSize += MorphWeightsVertexBuffer.VertexBufferRHI ? MorphWeightsVertexBuffer.VertexBufferRHI->GetSize() : 0;
-
-		return ResourceSize;
+		CumulativeResourceSize.AddUnknownMemoryBytes(MorphWeightsVertexBuffer.VertexBufferRHI ? MorphWeightsVertexBuffer.VertexBufferRHI->GetSize() : 0);
 	}
 	//~ End FSkeletalMeshObject Interface
 
@@ -423,13 +432,23 @@ private:
 		/**
 		 * @return memory in bytes of size of the resources for this LOD
 		 */
+		DEPRECATED(4.14, "GetResourceSize is deprecated. Please use GetResourceSizeEx or GetResourceSizeBytes instead.")
 		SIZE_T GetResourceSize()
 		{
-			SIZE_T Size = MorphVertexBuffer.GetResourceSize();
+			return GetResourceSizeBytes();
+		}
 
-			Size += GPUSkinVertexFactories.GetResourceSize();
+		void GetResourceSizeEx(FResourceSizeEx& CumulativeResourceSize)
+		{
+			CumulativeResourceSize.AddUnknownMemoryBytes(MorphVertexBuffer.GetResourceSize());
+			CumulativeResourceSize.AddUnknownMemoryBytes(GPUSkinVertexFactories.GetResourceSize());
+		}
 
-			return Size;
+		SIZE_T GetResourceSizeBytes()
+		{
+			FResourceSizeEx ResSize;
+			GetResourceSizeEx(ResSize);
+			return ResSize.GetTotalMemoryBytes();
 		}
 
 		FSkeletalMeshResource* SkelMeshResource;
diff --git a/Engine/Source/Runtime/Engine/Private/Slate/SceneViewport.cpp b/Engine/Source/Runtime/Engine/Private/Slate/SceneViewport.cpp
--- a/Engine/Source/Runtime/Engine/Private/Slate/SceneViewport.cpp
+++ b/Engine/Source/Runtime/Engine/Private/Slate/SceneViewport.cpp
@@ -44,6 +44,12 @@ FSceneViewport::FSceneViewport( FViewportClient* InViewportClient, TSharedPtr<SV
 	{
 		bShouldCaptureMouseOnActivate = InViewportClient->CaptureMouseOnLaunch();
 	}
+
+	if(FSlateApplication::IsInitialized())
+	{
+		FSlateApplication::Get().GetRenderer()->OnPreResizeWindowBackBuffer().AddRaw(this, &FSceneViewport::OnPreResizeWindowBackbuffer);
+		FSlateApplication::Get().GetRenderer()->OnPostResizeWindowBackBuffer().AddRaw(this, &FSceneViewport::OnPostResizeWindowBackbuffer);
+	}
 }
 
 FSceneViewport::~FSceneViewport()
@@ -51,6 +57,13 @@ FSceneViewport::~FSceneViewport()
 	Destroy();
 	// Wait for resources to be deleted
 	FlushRenderingCommands();
+
+	if(FSlateApplication::IsInitialized())
+	{
+		FSlateApplication::Get().GetRenderer()->OnPreResizeWindowBackBuffer().RemoveAll(this);
+		FSlateApplication::Get().GetRenderer()->OnPostResizeWindowBackBuffer().RemoveAll(this);
+	}
+
 }
 
 bool FSceneViewport::HasMouseCapture() const
@@ -523,8 +536,18 @@ FReply FSceneViewport::OnMouseButtonUp( const FGeometry& InGeometry, const FPoin
 			CurrentReplyState = FReply::Unhandled(); 
 		}
 		bCursorVisible = ViewportClient->GetCursor(this, GetMouseX(), GetMouseY()) != EMouseCursor::None;
+
+		bool bShouldMouseBeVisible = false;
+
+		UWorld* World = ViewportClient->GetWorld();
+		if (World && World->IsGameWorld() && World->GetGameInstance())
+		{
+			APlayerController* PC = World->GetGameInstance()->GetFirstLocalPlayerController();
+			bShouldMouseBeVisible = PC && PC->ShouldShowMouseCursor();
+		}
+
 		bReleaseMouse = 
-			bCursorVisible || 
+			(bCursorVisible && !bShouldMouseBeVisible) || 
 			ViewportClient->CaptureMouseOnClick() == EMouseCaptureMode::CaptureDuringMouseDown ||
 			( ViewportClient->CaptureMouseOnClick() == EMouseCaptureMode::CaptureDuringRightMouseDown && InMouseEvent.GetEffectingButton() == EKeys::RightMouseButton );
 	}
@@ -1545,6 +1568,33 @@ void FSceneViewport::WindowRenderTargetUpdate(FSlateRenderer* Renderer, SWindow*
 	}
 }
 
+void FSceneViewport::OnPreResizeWindowBackbuffer(void* Backbuffer)
+{
+	check(IsInGameThread());
+	FViewportRHIRef TestReference = *(FViewportRHIRef*)Backbuffer;
+	// Backbuffer we are rendering to is being released.  We must free our resource
+	if(ViewportRHI == TestReference)
+	{
+		ViewportRHI.SafeRelease();
+	}
+}
+
+void FSceneViewport::OnPostResizeWindowBackbuffer(void* Backbuffer)
+{
+	check(IsInGameThread());
+
+	if(!UseSeparateRenderTarget() && !IsValidRef(ViewportRHI))
+	{
+		TSharedPtr<FSlateRenderer> Renderer = FSlateApplication::Get().GetRenderer();
+		FWidgetPath WidgetPath;
+		void* ViewportResource = Renderer->GetViewportResource(*FSlateApplication::Get().FindWidgetWindow(ViewportWidget.Pin().ToSharedRef(), WidgetPath));
+		if(ViewportResource)
+		{
+			ViewportRHI = *((FViewportRHIRef*)ViewportResource);
+		}
+	}
+}
+
 void FSceneViewport::InitDynamicRHI()
 {
 	if(bRequiresHitProxyStorage)
diff --git a/Engine/Source/Runtime/Engine/Private/SoundCue.cpp b/Engine/Source/Runtime/Engine/Private/SoundCue.cpp
--- a/Engine/Source/Runtime/Engine/Private/SoundCue.cpp
+++ b/Engine/Source/Runtime/Engine/Private/SoundCue.cpp
@@ -276,17 +276,17 @@ FString USoundCue::GetDesc()
 	return Description;
 }
 
-SIZE_T USoundCue::GetResourceSize(EResourceSizeMode::Type Mode)
+void USoundCue::GetResourceSizeEx(FResourceSizeEx& CumulativeResourceSize)
 {
-	if( Mode == EResourceSizeMode::Exclusive )
+	Super::GetResourceSizeEx(CumulativeResourceSize);
+
+	if( CumulativeResourceSize.GetResourceSizeMode() == EResourceSizeMode::Exclusive )
 	{
-		return( 0 );
+		// Nothing to add
 	}
 	else
 	{
 		// Sum up the size of referenced waves
-		int32 ResourceSize = 0;
-
 		TArray<USoundNodeWavePlayer*> WavePlayers;
 		RecursiveFindNode<USoundNodeWavePlayer>( FirstNode, WavePlayers );
 
@@ -295,11 +295,9 @@ SIZE_T USoundCue::GetResourceSize(EResourceSizeMode::Type Mode)
 			USoundWave* SoundWave = WavePlayers[WaveIndex]->GetSoundWave();
 			if (SoundWave)
 			{
-				ResourceSize += SoundWave->GetResourceSize(Mode);
+				SoundWave->GetResourceSizeEx(CumulativeResourceSize);
 			}
 		}
-
-		return( ResourceSize );
 	}
 }
 
diff --git a/Engine/Source/Runtime/Engine/Private/SoundWave.cpp b/Engine/Source/Runtime/Engine/Private/SoundWave.cpp
--- a/Engine/Source/Runtime/Engine/Private/SoundWave.cpp
+++ b/Engine/Source/Runtime/Engine/Private/SoundWave.cpp
@@ -81,15 +81,15 @@ USoundWave::USoundWave(const FObjectInitializer& ObjectInitializer)
 	ResourceState = ESoundWaveResourceState::NeedsFree;
 }
 
-SIZE_T USoundWave::GetResourceSize(EResourceSizeMode::Type Mode)
+void USoundWave::GetResourceSizeEx(FResourceSizeEx& CumulativeResourceSize)
 {
+	Super::GetResourceSizeEx(CumulativeResourceSize);
+
 	if (!GEngine)
 	{
-		return 0;
+		return;
 	}
 
-	SIZE_T CalculatedResourceSize = 0;
-
 	if (FAudioDevice* LocalAudioDevice = GEngine->GetMainAudioDevice())
 	{
 		if (LocalAudioDevice->HasCompressedAudioInfoClass(this) && DecompressionType == DTYPE_Native)
@@ -102,23 +102,21 @@ SIZE_T USoundWave::GetResourceSize(EResourceSizeMode::Type Mode)
 			{
 				ensureMsgf(ResourceSize == 0, TEXT("ResourceSize for DTYPE_Native USoundWave '%s' was not 0 (%d)."), *GetName(), ResourceSize);
 			}
-			CalculatedResourceSize = RawPCMDataSize;
+			CumulativeResourceSize.AddDedicatedSystemMemoryBytes(RawPCMDataSize);
 		}
 		else 
 		{
 			if (DecompressionType == DTYPE_RealTime && CachedRealtimeFirstBuffer)
 			{
-				CalculatedResourceSize = MONO_PCM_BUFFER_SIZE * NumChannels;
+				CumulativeResourceSize.AddDedicatedSystemMemoryBytes(MONO_PCM_BUFFER_SIZE * NumChannels);
 			}
 			
 			if ((!FPlatformProperties::SupportsAudioStreaming() || !IsStreaming()))
 			{
-				CalculatedResourceSize += GetCompressedDataSize(LocalAudioDevice->GetRuntimeFormat(this));
+				CumulativeResourceSize.AddDedicatedSystemMemoryBytes(GetCompressedDataSize(LocalAudioDevice->GetRuntimeFormat(this)));
 			}
 		}
 	}
-
-	return CalculatedResourceSize;
 }
 
 int32 USoundWave::GetResourceSizeForFormat(FName Format)
diff --git a/Engine/Source/Runtime/Engine/Private/SoundWaveProcedural.cpp b/Engine/Source/Runtime/Engine/Private/SoundWaveProcedural.cpp
--- a/Engine/Source/Runtime/Engine/Private/SoundWaveProcedural.cpp
+++ b/Engine/Source/Runtime/Engine/Private/SoundWaveProcedural.cpp
@@ -96,9 +96,9 @@ int32 USoundWaveProcedural::GetAvailableAudioByteCount()
 	return AvailableByteCount.GetValue();
 }
 
-SIZE_T USoundWaveProcedural::GetResourceSize(EResourceSizeMode::Type Mode)
+void USoundWaveProcedural::GetResourceSizeEx(FResourceSizeEx& CumulativeResourceSize)
 {
-	return 0;
+	Super::GetResourceSizeEx(CumulativeResourceSize);
 }
 
 int32 USoundWaveProcedural::GetResourceSizeForFormat(FName Format)
diff --git a/Engine/Source/Runtime/Engine/Private/StaticMesh.cpp b/Engine/Source/Runtime/Engine/Private/StaticMesh.cpp
--- a/Engine/Source/Runtime/Engine/Private/StaticMesh.cpp
+++ b/Engine/Source/Runtime/Engine/Private/StaticMesh.cpp
@@ -582,7 +582,7 @@ void FStaticMeshLODResources::InitResources(UStaticMesh* Parent)
 	if (DistanceFieldData)
 	{
 		DistanceFieldData->VolumeTexture.Initialize();
-		INC_DWORD_STAT_BY( STAT_StaticMeshDistanceFieldMemory, DistanceFieldData->GetResourceSize() );
+		INC_DWORD_STAT_BY( STAT_StaticMeshDistanceFieldMemory, DistanceFieldData->GetResourceSizeBytes() );
 	}
 
 	ENQUEUE_UNIQUE_RENDER_COMMAND_ONEPARAMETER(
@@ -648,7 +648,7 @@ void FStaticMeshLODResources::ReleaseResources()
 
 	if (DistanceFieldData)
 	{
-		DEC_DWORD_STAT_BY( STAT_StaticMeshDistanceFieldMemory, DistanceFieldData->GetResourceSize() );
+		DEC_DWORD_STAT_BY( STAT_StaticMeshDistanceFieldMemory, DistanceFieldData->GetResourceSizeBytes() );
 		DistanceFieldData->VolumeTexture.Release();
 	}
 }
@@ -1394,7 +1394,7 @@ void UStaticMesh::InitResources()
 		UpdateMemoryStats,
 		UStaticMesh*, This, this,
 		{
- 			const uint32 StaticMeshResourceSize = This->GetResourceSize( EResourceSizeMode::Exclusive );
+ 			const uint32 StaticMeshResourceSize = This->GetResourceSizeBytes( EResourceSizeMode::Exclusive );
  			INC_DWORD_STAT_BY( STAT_StaticMeshTotalMemory, StaticMeshResourceSize );
  			INC_DWORD_STAT_BY( STAT_StaticMeshTotalMemory2, StaticMeshResourceSize );
 		} );
@@ -1406,14 +1406,15 @@ void UStaticMesh::InitResources()
  *
  * @return size of resource as to be displayed to artists/ LDs in the Editor.
  */
-SIZE_T UStaticMesh::GetResourceSize(EResourceSizeMode::Type Mode)
+void UStaticMesh::GetResourceSizeEx(FResourceSizeEx& CumulativeResourceSize)
 {
-	SIZE_T ResourceSize = 0;
+	Super::GetResourceSizeEx(CumulativeResourceSize);
+
 	if (RenderData)
 	{
-		ResourceSize += RenderData->GetResourceSize();
+		RenderData->GetResourceSizeEx(CumulativeResourceSize);
 	}
-	if (Mode == EResourceSizeMode::Inclusive)
+	if (CumulativeResourceSize.GetResourceSizeMode() == EResourceSizeMode::Inclusive)
 	{
 		TSet<UMaterialInterface*> UniqueMaterials;
 		for (int32 MaterialIndex = 0; MaterialIndex < StaticMaterials.Num(); ++MaterialIndex)
@@ -1423,27 +1424,31 @@ SIZE_T UStaticMesh::GetResourceSize(EResourceSizeMode::Type Mode)
 			UniqueMaterials.Add(StaticMaterial.MaterialInterface,&bAlreadyCounted);
 			if (!bAlreadyCounted && StaticMaterial.MaterialInterface)
 			{
-				ResourceSize += StaticMaterial.MaterialInterface->GetResourceSize(Mode);
+				StaticMaterial.MaterialInterface->GetResourceSizeEx(CumulativeResourceSize);
 			}
 		}
 
 		if(BodySetup)
 		{
-			ResourceSize += BodySetup->GetResourceSize(Mode);
+			BodySetup->GetResourceSizeEx(CumulativeResourceSize);
 		}
 	}
-	return ResourceSize;
 }
 
 SIZE_T FStaticMeshRenderData::GetResourceSize() const
 {
-	SIZE_T ResourceSize = sizeof(*this);
+	return GetResourceSizeBytes();
+}
+
+void FStaticMeshRenderData::GetResourceSizeEx(FResourceSizeEx& CumulativeResourceSize) const
+{
+	CumulativeResourceSize.AddDedicatedSystemMemoryBytes(sizeof(*this));
 
 	// Count dynamic arrays.
-	ResourceSize += LODResources.GetAllocatedSize();
+	CumulativeResourceSize.AddUnknownMemoryBytes(LODResources.GetAllocatedSize());
 #if WITH_EDITORONLY_DATA
-	ResourceSize += DerivedDataKey.GetAllocatedSize();
-	ResourceSize += WedgeMap.GetAllocatedSize();
+	CumulativeResourceSize.AddDedicatedSystemMemoryBytes(DerivedDataKey.GetAllocatedSize());
+	CumulativeResourceSize.AddDedicatedSystemMemoryBytes(WedgeMap.GetAllocatedSize());
 #endif // #if WITH_EDITORONLY_DATA
 
 	for(int32 LODIndex = 0;LODIndex < LODResources.Num();LODIndex++)
@@ -1457,12 +1462,12 @@ SIZE_T FStaticMeshRenderData::GetResourceSize() const
 			+ LODRenderData.WireframeIndexBuffer.GetAllocatedSize()
 			+ (RHISupportsTessellation(GShaderPlatformForFeatureLevel[GMaxRHIFeatureLevel]) ? LODRenderData.AdjacencyIndexBuffer.GetAllocatedSize() : 0);
 
-		ResourceSize += VBSize + IBSize;
-		ResourceSize += LODRenderData.Sections.GetAllocatedSize();
+		CumulativeResourceSize.AddUnknownMemoryBytes(VBSize + IBSize);
+		CumulativeResourceSize.AddUnknownMemoryBytes(LODRenderData.Sections.GetAllocatedSize());
 
 		if (LODRenderData.DistanceFieldData)
 		{
-			ResourceSize += LODRenderData.DistanceFieldData->GetResourceSize();
+			LODRenderData.DistanceFieldData->GetResourceSizeEx(CumulativeResourceSize);
 		}
 	}
 
@@ -1470,11 +1475,16 @@ SIZE_T FStaticMeshRenderData::GetResourceSize() const
 	// If render data for multiple platforms is loaded, count it all.
 	if (NextCachedRenderData)
 	{
-		ResourceSize += NextCachedRenderData->GetResourceSize();
+		NextCachedRenderData->GetResourceSizeEx(CumulativeResourceSize);
 	}
 #endif // #if WITH_EDITORONLY_DATA
+}
 
-	return ResourceSize;
+SIZE_T FStaticMeshRenderData::GetResourceSizeBytes() const
+{
+	FResourceSizeEx ResSize;
+	GetResourceSizeEx(ResSize);
+	return ResSize.GetTotalMemoryBytes();
 }
 
 int32 UStaticMesh::GetNumVertices(int32 LODIndex) const
@@ -1669,7 +1679,7 @@ bool UStaticMesh::GetStreamingTextureFactor(float& OutTexelFactor, FBoxSphereBou
 void UStaticMesh::ReleaseResources()
 {
 #if STATS
-	uint32 StaticMeshResourceSize = GetResourceSize(EResourceSizeMode::Exclusive);
+	uint32 StaticMeshResourceSize = GetResourceSizeBytes(EResourceSizeMode::Exclusive);
 	DEC_DWORD_STAT_BY( STAT_StaticMeshTotalMemory, StaticMeshResourceSize );
 	DEC_DWORD_STAT_BY( STAT_StaticMeshTotalMemory2, StaticMeshResourceSize );
 #endif
diff --git a/Engine/Source/Runtime/Engine/Private/StaticMeshActor.cpp b/Engine/Source/Runtime/Engine/Private/StaticMeshActor.cpp
--- a/Engine/Source/Runtime/Engine/Private/StaticMeshActor.cpp
+++ b/Engine/Source/Runtime/Engine/Private/StaticMeshActor.cpp
@@ -95,11 +95,18 @@ void AStaticMeshActor::PostEditChangeChainProperty(FPropertyChangedChainEvent& P
 {
 	Super::PostEditChangeChainProperty(PropertyChangedEvent);
 
-	// If we want to replicate movement, set RemoteRole to match
-	if(PropertyChangedEvent.Property != NULL && PropertyChangedEvent.Property->GetFName() == FName(TEXT("bStaticMeshReplicateMovement")))
+	if (PropertyChangedEvent.Property != nullptr)
 	{
-		bReplicateMovement = bStaticMeshReplicateMovement;
-		SetReplicates(bReplicateMovement);
+		// If we want to replicate movement, set RemoteRole to match
+		if (PropertyChangedEvent.Property->GetFName() == FName(TEXT("bStaticMeshReplicateMovement")))
+		{
+			bReplicateMovement = bStaticMeshReplicateMovement;
+			SetReplicates(bReplicateMovement);
+		}
+		else if (PropertyChangedEvent.Property->GetFName() == GET_MEMBER_NAME_STRING_CHECKED(AStaticMeshActor, StaticMeshComponent) && StaticMeshComponent->StaticMesh != nullptr)
+		{
+			StaticMeshComponent->CleanUpOverrideMaterials();
+		}
 	}
 
 	// Get 'deepest' property name we changed.
diff --git a/Engine/Source/Runtime/Engine/Private/StaticMeshBuild.cpp b/Engine/Source/Runtime/Engine/Private/StaticMeshBuild.cpp
--- a/Engine/Source/Runtime/Engine/Private/StaticMeshBuild.cpp
+++ b/Engine/Source/Runtime/Engine/Private/StaticMeshBuild.cpp
@@ -92,8 +92,13 @@ void UStaticMesh::Build(bool bSilent, TArray<FText>* OutErrors)
 	// Free existing render data and recache.
 	CacheDerivedData();
 
-	// Reinitialize the static mesh's resources.
-	InitResources();
+	// Note: meshes can be built during automated importing.  We should not create resources in that case
+	// as they will never be released when this object is deleted
+	if(FApp::CanEverRender())
+	{
+		// Reinitialize the static mesh's resources.
+		InitResources();
+	}
 
 	// Ensure we have a bodysetup.
 	CreateBodySetup();
diff --git a/Engine/Source/Runtime/Engine/Private/Texture2D.cpp b/Engine/Source/Runtime/Engine/Private/Texture2D.cpp
--- a/Engine/Source/Runtime/Engine/Private/Texture2D.cpp
+++ b/Engine/Source/Runtime/Engine/Private/Texture2D.cpp
@@ -900,22 +900,21 @@ FTextureResource* UTexture2D::CreateResource()
 }
 
 
-SIZE_T UTexture2D::GetResourceSize(EResourceSizeMode::Type Mode)
+void UTexture2D::GetResourceSizeEx(FResourceSizeEx& CumulativeResourceSize)
 {
-	if (Mode == EResourceSizeMode::Exclusive)
+	Super::GetResourceSizeEx(CumulativeResourceSize);
+
+	if (CumulativeResourceSize.GetResourceSizeMode() == EResourceSizeMode::Exclusive)
 	{
-		return CalcTextureMemorySize(ResidentMips);
+		CumulativeResourceSize.AddUnknownMemoryBytes(CalcTextureMemorySize(ResidentMips));
 	}
 	else
 	{
-		SIZE_T ResourceSize = 0;
 		if (PlatformData)
 		{
-			ResourceSize += CalcTextureSize(GetSizeX(), GetSizeY(), GetPixelFormat(), GetNumMips());
+			CumulativeResourceSize.AddUnknownMemoryBytes(CalcTextureSize(GetSizeX(), GetSizeY(), GetPixelFormat(), GetNumMips()));
 		}
-		return ResourceSize;
 	}
-	return 0;
 }
 
 
diff --git a/Engine/Source/Runtime/Engine/Private/TextureCube.cpp b/Engine/Source/Runtime/Engine/Private/TextureCube.cpp
--- a/Engine/Source/Runtime/Engine/Private/TextureCube.cpp
+++ b/Engine/Source/Runtime/Engine/Private/TextureCube.cpp
@@ -346,9 +346,10 @@ FTextureResource* UTextureCube::CreateResource()
 	return NewResource;
 }
 
-SIZE_T UTextureCube::GetResourceSize(EResourceSizeMode::Type Mode)
+void UTextureCube::GetResourceSizeEx(FResourceSizeEx& CumulativeResourceSize)
 {
-	return CalcTextureMemorySizeEnum(TMC_ResidentMips);
+	Super::GetResourceSizeEx(CumulativeResourceSize);
+	CumulativeResourceSize.AddUnknownMemoryBytes(CalcTextureMemorySizeEnum(TMC_ResidentMips));
 }
 
 #if WITH_EDITOR
diff --git a/Engine/Source/Runtime/Engine/Private/TextureRenderTarget2D.cpp b/Engine/Source/Runtime/Engine/Private/TextureRenderTarget2D.cpp
--- a/Engine/Source/Runtime/Engine/Private/TextureRenderTarget2D.cpp
+++ b/Engine/Source/Runtime/Engine/Private/TextureRenderTarget2D.cpp
@@ -54,8 +54,10 @@ EMaterialValueType UTextureRenderTarget2D::GetMaterialType()
 	return MCT_Texture2D;
 }
 
-SIZE_T UTextureRenderTarget2D::GetResourceSize(EResourceSizeMode::Type Mode)
+void UTextureRenderTarget2D::GetResourceSizeEx(FResourceSizeEx& CumulativeResourceSize)
 {
+	Super::GetResourceSizeEx(CumulativeResourceSize);
+
 	// Calculate size based on format.
 	EPixelFormat Format = GetFormat();
 	int32 BlockSizeX	= GPixelFormats[Format].BlockSizeX;
@@ -65,7 +67,7 @@ SIZE_T UTextureRenderTarget2D::GetResourceSize(EResourceSizeMode::Type Mode)
 	int32 NumBlocksY	= (SizeY + BlockSizeY - 1) / BlockSizeY;
 	int32 NumBytes	= NumBlocksX * NumBlocksY * BlockBytes;
 
-	return NumBytes;
+	CumulativeResourceSize.AddUnknownMemoryBytes(NumBytes);
 }
 
 void UTextureRenderTarget2D::InitCustomFormat( uint32 InSizeX, uint32 InSizeY, EPixelFormat InOverrideFormat, bool bInForceLinearGamma )
diff --git a/Engine/Source/Runtime/Engine/Private/TextureRenderTargetCube.cpp b/Engine/Source/Runtime/Engine/Private/TextureRenderTargetCube.cpp
--- a/Engine/Source/Runtime/Engine/Private/TextureRenderTargetCube.cpp
+++ b/Engine/Source/Runtime/Engine/Private/TextureRenderTargetCube.cpp
@@ -51,8 +51,10 @@ void UTextureRenderTargetCube::InitAutoFormat(uint32 InSizeX)
 }
 
 
-SIZE_T UTextureRenderTargetCube::GetResourceSize(EResourceSizeMode::Type Mode)
+void UTextureRenderTargetCube::GetResourceSizeEx(FResourceSizeEx& CumulativeResourceSize)
 {
+	Super::GetResourceSizeEx(CumulativeResourceSize);
+
 	// Calculate size based on format.
 	EPixelFormat Format = GetFormat();
 	int32 BlockSizeX	= GPixelFormats[Format].BlockSizeX;
@@ -62,7 +64,7 @@ SIZE_T UTextureRenderTargetCube::GetResourceSize(EResourceSizeMode::Type Mode)
 	int32 NumBlocksY	= (SizeX + BlockSizeY - 1) / BlockSizeY;
 	int32 NumBytes	= NumBlocksX * NumBlocksY * BlockBytes * 6;
 
-	return NumBytes;
+	CumulativeResourceSize.AddUnknownMemoryBytes(NumBytes);
 }
 
 
diff --git a/Engine/Source/Runtime/Engine/Private/UnrealClient.cpp b/Engine/Source/Runtime/Engine/Private/UnrealClient.cpp
--- a/Engine/Source/Runtime/Engine/Private/UnrealClient.cpp
+++ b/Engine/Source/Runtime/Engine/Private/UnrealClient.cpp
@@ -926,7 +926,7 @@ void FViewport::HighResScreenshot()
 	bTakeHighResScreenShot = false;
 
 	// Notification of a successful screenshot
-	if (GIsEditor || !IsFullscreen())
+	if ((GIsEditor || !IsFullscreen()) && !GIsAutomationTesting )
 	{
 		auto Message = NSLOCTEXT("UnrealClient", "HighResScreenshotSavedAs", "High resolution screenshot saved as");
 		FNotificationInfo Info(Message);
