Copying //UE4/Dev-Core to //UE4/Dev-Main (Source: //UE4/Dev-Core @ 3151653)

diff --git a/Engine/Source/Runtime/Engine/Classes/Components/SkyLightComponent.h b/Engine/Source/Runtime/Engine/Classes/Components/SkyLightComponent.h
--- a/Engine/Source/Runtime/Engine/Classes/Components/SkyLightComponent.h
+++ b/Engine/Source/Runtime/Engine/Classes/Components/SkyLightComponent.h
@@ -255,6 +255,7 @@ protected:
 	 */
 	static TArray<USkyLightComponent*> SkyCapturesToUpdate;
 	static TArray<USkyLightComponent*> SkyCapturesToUpdateBlendDestinations;
+	static FCriticalSection SkyCapturesToUpdateLock;
 
 	//~ Begin UActorComponent Interface
 	virtual void CreateRenderState_Concurrent() override;
diff --git a/Engine/Source/Runtime/Engine/Classes/Engine/BlueprintGeneratedClass.h b/Engine/Source/Runtime/Engine/Classes/Engine/BlueprintGeneratedClass.h
--- a/Engine/Source/Runtime/Engine/Classes/Engine/BlueprintGeneratedClass.h
+++ b/Engine/Source/Runtime/Engine/Classes/Engine/BlueprintGeneratedClass.h
@@ -649,6 +649,7 @@ public:
 	virtual void Serialize(FArchive& Ar) override;
 	virtual void PostLoad() override;
 	virtual void PostInitProperties() override;
+	virtual void GetPreloadDependencies(TArray<UObject*>& OutDeps) override;
 	// End UObject interface
 	
 	// UClass interface
diff --git a/Engine/Source/Runtime/Engine/Classes/Engine/CoreSettings.h b/Engine/Source/Runtime/Engine/Classes/Engine/CoreSettings.h
--- a/Engine/Source/Runtime/Engine/Classes/Engine/CoreSettings.h
+++ b/Engine/Source/Runtime/Engine/Classes/Engine/CoreSettings.h
@@ -85,6 +85,11 @@ protected:
 		ToolTip = "Batching granularity used to register actor components during level streaming."))
 	int32 LevelStreamingComponentsRegistrationGranularity;
 
+	UPROPERTY(config, EditAnywhere, Category = PackageStreaming, meta = (
+		ConsoleVariable = "s.EventDrivenLoaderEnabled", DisplayName = "Event Driven Loader Enabled",
+		ToolTip = "Enables the event driven loader in cooked builds."))
+	uint32 EventDrivenLoaderEnabled : 1;
+
 	//~ Begin UObject Interface
 	virtual void PostInitProperties() override;
 
diff --git a/Engine/Source/Runtime/Engine/Classes/Engine/DataTable.h b/Engine/Source/Runtime/Engine/Classes/Engine/DataTable.h
--- a/Engine/Source/Runtime/Engine/Classes/Engine/DataTable.h
+++ b/Engine/Source/Runtime/Engine/Classes/Engine/DataTable.h
@@ -58,6 +58,7 @@ class UDataTable
 	ENGINE_API virtual void FinishDestroy() override;
 	ENGINE_API virtual void Serialize(FArchive& Ar) override;
 	ENGINE_API static void AddReferencedObjects(UObject* InThis, FReferenceCollector& Collector);
+	ENGINE_API virtual void GetPreloadDependencies(TArray<UObject*>& OutDeps) override;
 	ENGINE_API virtual void GetResourceSizeEx(FResourceSizeEx& CumulativeResourceSize) override;
 #if WITH_EDITORONLY_DATA
 	ENGINE_API FName GetRowStructName() const;
diff --git a/Engine/Source/Runtime/Engine/Classes/Engine/Engine.h b/Engine/Source/Runtime/Engine/Classes/Engine/Engine.h
--- a/Engine/Source/Runtime/Engine/Classes/Engine/Engine.h
+++ b/Engine/Source/Runtime/Engine/Classes/Engine/Engine.h
@@ -1685,7 +1685,12 @@ public:
 	bool HandleViewnamesCommand( const TCHAR* Cmd, FOutputDevice& Ar );
 	bool HandleFreezeStreamingCommand( const TCHAR* Cmd, FOutputDevice& Ar, UWorld* InWorld );		// Smedis
 	bool HandleFreezeAllCommand( const TCHAR* Cmd, FOutputDevice& Ar, UWorld* InWorld );			// Smedis
+
+#if !USE_NEW_ASYNC_IO
 	bool HandleFlushIOManagerCommand( const TCHAR* Cmd, FOutputDevice& Ar );						// Smedis
+	bool HandleListPreCacheMapPackagesCommand(const TCHAR* Cmd, FOutputDevice& Ar);
+#endif
+
 	bool HandleToggleRenderingThreadCommand( const TCHAR* Cmd, FOutputDevice& Ar );	
 	bool HandleToggleAsyncComputeCommand( const TCHAR* Cmd, FOutputDevice& Ar );
 	bool HandleRecompileShadersCommand( const TCHAR* Cmd, FOutputDevice& Ar );
@@ -1710,7 +1715,6 @@ public:
 	bool HandleMemReportDeferredCommand( const TCHAR* Cmd, FOutputDevice& Ar, UWorld* InWorld );
 	bool HandleParticleMeshUsageCommand( const TCHAR* Cmd, FOutputDevice& Ar );
 	bool HandleDumpParticleCountsCommand( const TCHAR* Cmd, FOutputDevice& Ar );
-	bool HandleListPreCacheMapPackagesCommand( const TCHAR* Cmd, FOutputDevice& Ar );
 	bool HandleListLoadedPackagesCommand( const TCHAR* Cmd, FOutputDevice& Ar );
 	bool HandleMemCommand( const TCHAR* Cmd, FOutputDevice& Ar );
 	bool HandleDebugCommand( const TCHAR* Cmd, FOutputDevice& Ar );
diff --git a/Engine/Source/Runtime/Engine/Classes/Engine/World.h b/Engine/Source/Runtime/Engine/Classes/Engine/World.h
--- a/Engine/Source/Runtime/Engine/Classes/Engine/World.h
+++ b/Engine/Source/Runtime/Engine/Classes/Engine/World.h
@@ -661,7 +661,7 @@ private:
  */
 
 UCLASS(customConstructor, config=Engine)
-class ENGINE_API UWorld : public UObject, public FNetworkNotify
+class ENGINE_API UWorld final : public UObject, public FNetworkNotify
 {
 	GENERATED_UCLASS_BODY()
 
diff --git a/Engine/Source/Runtime/Engine/Private/AudioStreaming.cpp b/Engine/Source/Runtime/Engine/Private/AudioStreaming.cpp
--- a/Engine/Source/Runtime/Engine/Private/AudioStreaming.cpp
+++ b/Engine/Source/Runtime/Engine/Private/AudioStreaming.cpp
@@ -10,6 +10,9 @@ AudioStreaming.cpp: Implementation of audio streaming classes.
 #include "DerivedDataCacheInterface.h"
 #include "Sound/AudioSettings.h"
 
+#if USE_NEW_ASYNC_IO
+#include "AsyncFileHandle.h"
+#endif
 /*------------------------------------------------------------------------------
 	Streaming chunks from the derived data cache.
 ------------------------------------------------------------------------------*/
@@ -66,26 +69,84 @@ void FAsyncStreamDerivedChunkWorker::DoWork()
 
 FStreamingWaveData::FStreamingWaveData()
 	: SoundWave(NULL)
+#if USE_NEW_ASYNC_IO
+	, IORequestHandle(nullptr)
+#endif
 {
+#if USE_NEW_ASYNC_IO
+	AsyncFileCallBack =
+		[this](bool bWasCancelled, IAsyncReadRequest* Req)
+	{
+		uint8* Mem = Req->GetReadResults();
+		if (Mem)
+		{
+			bool bFound = false;
+			for (FLoadedAudioChunk& LoadedChunk : LoadedChunks)
+			{
+				if (LoadedChunk.IORequest == Req)
+				{
+					check(!LoadedChunk.Data);
+					LoadedChunk.Data = Mem;
+					DEC_MEMORY_STAT_BY(STAT_AsyncFileMemory, LoadedChunk.DataSize);
+					INC_DWORD_STAT_BY(STAT_AudioMemorySize, LoadedChunk.DataSize);
+					INC_DWORD_STAT_BY(STAT_AudioMemory, LoadedChunk.DataSize);
+					bFound = true;
+					break;
+				}
+			}
+			check(bFound);
+		}
+		// else was canceled
+		PendingChunkChangeRequestStatus.Decrement();
+	};
+#endif
 }
 
 FStreamingWaveData::~FStreamingWaveData()
 {
 	// Make sure there are no pending requests in flight.
-	while (UpdateStreamingStatus() == true)
+#if USE_NEW_ASYNC_IO
+	{
+		for (int32 Pass = 0; Pass < 3; Pass++)
+		{
+			BlockTillAllRequestsFinished();
+			if (!UpdateStreamingStatus())
+			{
+				break;
+			}
+			check(Pass < 2); // we should be done after two passes. Pass 0 will start anything we need and pass 1 will complete those requests
+		}
+	}
+#else
 	{
-		// Give up timeslice.
-		FPlatformProcess::Sleep(0);
+		QUICK_SCOPE_CYCLE_COUNTER(FStreamingWaveData_Destructor_Spin);
+		while (UpdateStreamingStatus() == true)
+		{
+			// Give up timeslice.
+			check(!USE_NEW_ASYNC_IO); // should be done because we block on it above
+			FPlatformProcess::SleepNoStats(0.001f);
+		}
 	}
+#endif
 
-	for (auto& LoadedChunk : LoadedChunks)
+	for (FLoadedAudioChunk& LoadedChunk : LoadedChunks)
 	{
 		FreeLoadedChunk(LoadedChunk);
 	}
+#if USE_NEW_ASYNC_IO
+	if (IORequestHandle)
+	{
+		delete IORequestHandle;
+		IORequestHandle = nullptr;
+	}
+#endif
 }
 
 void FStreamingWaveData::Initialize(USoundWave* InSoundWave)
 {
+#if USE_NEW_ASYNC_IO
+	check(!IORequestHandle);
+#endif
 	SoundWave = InSoundWave;
 
 	// Always get the first chunk of data so we can play immediately
@@ -132,11 +193,25 @@ bool FStreamingWaveData::UpdateStreamingStatus()
 #if WITH_EDITORONLY_DATA
 			bFailedRequests = FinishDDCRequests();
 #endif //WITH_EDITORONLY_DATA
+
+#if USE_NEW_ASYNC_IO
+			// could maybe iterate over the things we know are done, but I couldn't tell if that was IndicesToLoad or not.
+			for (FLoadedAudioChunk& LoadedChunk : LoadedChunks)
+			{
+				if (LoadedChunk.IORequest && LoadedChunk.IORequest->PollCompletion())
+				{
+					LoadedChunk.IORequest->WaitCompletion();
+					delete LoadedChunk.IORequest;
+					LoadedChunk.IORequest = nullptr;
+				}
+			}
+#endif
+
 			PendingChunkChangeRequestStatus.Decrement();
 			bHasPendingRequestInFlight = false;
 			LoadedChunkIndices = CurrentRequest.RequiredIndices;
 		}
-		else if (RequestStatus == AudioState_ReadyFor_Requests)
+		else if (RequestStatus == AudioState_ReadyFor_Requests) // odd that this is an else, probably we should start requests right now
 		{
 			BeginPendingRequests(IndicesToLoad, IndicesToFree);
 		}
@@ -153,8 +228,10 @@ void FStreamingWaveData::UpdateChunkRequests(FWaveRequest& InWaveRequest)
 
 	CurrentRequest = InWaveRequest;
 
+#if !USE_NEW_ASYNC_IO
 	// Clear last batch of request indices
 	IORequestIndices.Empty();
+#endif
 }
 
 bool FStreamingWaveData::HasPendingRequests(TArray<uint32>& IndicesToLoad, TArray<uint32>& IndicesToFree) const
@@ -197,6 +274,20 @@ void FStreamingWaveData::BeginPendingRequests(const TArray<uint32>& IndicesToLoa
 	}
 
 	// Mark Chunks for removal in case they can be reused
+#if USE_NEW_ASYNC_IO
+	for (auto Index : IndicesToFree)
+	{
+		for (int32 ChunkIndex = 0; ChunkIndex < LoadedChunks.Num(); ++ChunkIndex)
+		{
+			if (LoadedChunks[ChunkIndex].Index == Index)
+			{
+				FreeLoadedChunk(LoadedChunks[ChunkIndex]);
+				LoadedChunks.RemoveAt(ChunkIndex);
+				break;
+			}
+		}
+	}
+#else
 	TArray<uint32> FreeChunkIndices;
 	for (auto Index : IndicesToFree)
 	{
@@ -204,11 +295,13 @@ void FStreamingWaveData::BeginPendingRequests(const TArray<uint32>& IndicesToLoa
 		{
 			if (LoadedChunks[ChunkIndex].Index == Index)
 			{
-				FreeChunkIndices.AddUnique(ChunkIndex);
+				check(!FreeChunkIndices.Contains(ChunkIndex));
+				FreeChunkIndices.Add(ChunkIndex);
 				break;
 			}
 		}
 	}
+#endif
 
 	if (IndicesToLoad.Num() > 0)
 	{
@@ -222,6 +315,7 @@ void FStreamingWaveData::BeginPendingRequests(const TArray<uint32>& IndicesToLoa
 
 			FLoadedAudioChunk* ChunkStorage = NULL;
 
+#if !USE_NEW_ASYNC_IO
 			for (auto FreeIndex : FreeChunkIndices)
 			{
 				if (LoadedChunks[FreeIndex].MemorySize >= ChunkSize)
@@ -233,6 +327,7 @@ void FStreamingWaveData::BeginPendingRequests(const TArray<uint32>& IndicesToLoa
 					break;
 				}
 			}
+#endif
 			if (ChunkStorage == NULL)
 			{
 				ChunkStorage = AddNewLoadedChunk(ChunkSize);
@@ -263,6 +358,29 @@ void FStreamingWaveData::BeginPendingRequests(const TArray<uint32>& IndicesToLoa
 #endif // #if WITH_EDITORONLY_DATA
 			{
 				check(Chunk.BulkData.GetFilename().Len());
+#if USE_NEW_ASYNC_IO
+				if (Chunk.BulkData.IsStoredCompressedOnDisk())
+				{
+					check(!"USE_NEW_ASYNC_IO does not support compression at the package level.");
+				}
+				else
+				{
+					check(!ChunkStorage->IORequest);
+					if (!IORequestHandle)
+					{
+						IORequestHandle = FPlatformFileManager::Get().GetPlatformFile().OpenAsyncRead(*Chunk.BulkData.GetFilename());
+						check(IORequestHandle); // this generally cannot fail because it is async
+					}
+					check(Chunk.BulkData.GetBulkDataSize() == ChunkStorage->DataSize);
+					ChunkStorage->IORequest = IORequestHandle->ReadRequest(Chunk.BulkData.GetBulkDataOffsetInFile(), ChunkStorage->DataSize, AsyncIOPriority, &AsyncFileCallBack);
+					if (!ChunkStorage->IORequest)
+					{
+						// we failed for some reason; file not found I guess.
+						PendingChunkChangeRequestStatus.Decrement();
+					}
+
+				}
+#else
 				if (Chunk.BulkData.IsStoredCompressedOnDisk())
 				{
 					IORequestIndices.AddUnique(FIOSystem::Get().LoadCompressedData(
@@ -291,6 +409,7 @@ void FStreamingWaveData::BeginPendingRequests(const TArray<uint32>& IndicesToLoa
 						);
 				}
 				check(IORequestIndices[IORequestIndices.Num() - 1]);
+#endif
 			}
 		}
 
@@ -303,6 +422,7 @@ void FStreamingWaveData::BeginPendingRequests(const TArray<uint32>& IndicesToLoa
 		PendingChunkChangeRequestStatus.Set(AudioState_ReadyFor_Finalization);
 	}
 
+#if !USE_NEW_ASYNC_IO
 	// Ensure indices are in order so we can step through backwards
 	FreeChunkIndices.Sort();
 
@@ -311,7 +431,47 @@ void FStreamingWaveData::BeginPendingRequests(const TArray<uint32>& IndicesToLoa
 		FreeLoadedChunk(LoadedChunks[FreeChunkIndices[FreeIndex]]);
 		LoadedChunks.RemoveAt(FreeChunkIndices[FreeIndex]);
 	}
+#endif
+}
+
+#if USE_NEW_ASYNC_IO
+
+bool FStreamingWaveData::BlockTillAllRequestsFinished(float TimeLimit)
+{
+	QUICK_SCOPE_CYCLE_COUNTER(FStreamingWaveData_BlockTillAllRequestsFinished);
+	if (TimeLimit == 0.0f)
+	{
+		for (FLoadedAudioChunk& LoadedChunk : LoadedChunks)
+		{
+			if (LoadedChunk.IORequest)
+			{
+				LoadedChunk.IORequest->WaitCompletion();
+				delete LoadedChunk.IORequest;
+				LoadedChunk.IORequest = nullptr;
+			}
+		}
+	}
+	else
+	{
+		double EndTime = FPlatformTime::Seconds() + TimeLimit;
+		for (FLoadedAudioChunk& LoadedChunk : LoadedChunks)
+		{
+			if (LoadedChunk.IORequest)
+			{
+				float ThisTimeLimit = EndTime - FPlatformTime::Seconds();
+				if (ThisTimeLimit < .001f || // one ms is the granularity of the platform event system
+					!LoadedChunk.IORequest->WaitCompletion(ThisTimeLimit))
+				{
+					return false;
+				}
+				delete LoadedChunk.IORequest;
+				LoadedChunk.IORequest = nullptr;
+			}
+		}
+	}
+	return true;
 }
+#endif
 
 #if WITH_EDITORONLY_DATA
 bool FStreamingWaveData::FinishDDCRequests()
@@ -334,31 +494,49 @@ bool FStreamingWaveData::FinishDDCRequests()
 FLoadedAudioChunk* FStreamingWaveData::AddNewLoadedChunk(int32 ChunkSize)
 {
 	int32 NewIndex = LoadedChunks.Num();
-	LoadedChunks.AddZeroed();
+	LoadedChunks.AddDefaulted();
+#if !USE_NEW_ASYNC_IO
+	LoadedChunks[NewIndex].MemorySize = ChunkSize;
 	LoadedChunks[NewIndex].Data = static_cast<uint8*>(FMemory::Malloc(ChunkSize));
-	LoadedChunks[NewIndex].DataSize = LoadedChunks[NewIndex].MemorySize = ChunkSize;
-
 	INC_DWORD_STAT_BY(STAT_AudioMemorySize, ChunkSize);
 	INC_DWORD_STAT_BY(STAT_AudioMemory, ChunkSize);
+#endif
+	LoadedChunks[NewIndex].DataSize = ChunkSize;
 
 	return &LoadedChunks[NewIndex];
 }
 
 void FStreamingWaveData::FreeLoadedChunk(FLoadedAudioChunk& LoadedChunk)
 {
+#if USE_NEW_ASYNC_IO
+	if (LoadedChunk.IORequest)
+	{
+		LoadedChunk.IORequest->Cancel();
+		LoadedChunk.IORequest->WaitCompletion();
+		delete LoadedChunk.IORequest;
+		LoadedChunk.IORequest = nullptr;
+	}
+#endif
 	if (LoadedChunk.Data != NULL)
 	{
 		FMemory::Free(LoadedChunk.Data);
 
 		// Stat housekeeping
+#if USE_NEW_ASYNC_IO
+		DEC_DWORD_STAT_BY(STAT_AudioMemorySize, LoadedChunk.DataSize);
+		DEC_DWORD_STAT_BY(STAT_AudioMemory, LoadedChunk.DataSize);
+#else
 		DEC_DWORD_STAT_BY(STAT_AudioMemorySize, LoadedChunk.MemorySize);
 		DEC_DWORD_STAT_BY(STAT_AudioMemory, LoadedChunk.MemorySize);
+#endif
 
-		LoadedChunk.Data = NULL;
-		LoadedChunk.DataSize = 0;
-		LoadedChunk.MemorySize = 0;
-		LoadedChunk.Index = 0;
 	}
+	LoadedChunk.Data = NULL;
+	LoadedChunk.DataSize = 0;
+#if !USE_NEW_ASYNC_IO
+	LoadedChunk.MemorySize = 0;
+#endif
+	LoadedChunk.Index = 0;
 }
 
 ////////////////////////////
@@ -377,7 +555,7 @@ void FAudioStreamingManager::UpdateResourceStreaming(float DeltaTime, bool bProc
 {
 	for (auto& WavePair : StreamingSoundWaves)
 	{
-		WavePair.Value.UpdateStreamingStatus();
+		WavePair.Value->UpdateStreamingStatus();
 	}
 
 	for (auto Source : StreamingSoundSources)
@@ -386,10 +564,11 @@ void FAudioStreamingManager::UpdateResourceStreaming(float DeltaTime, bool bProc
 		USoundWave* Wave = WaveInstance ? WaveInstance->WaveData : nullptr;
 		if (Wave)
 		{
-			FStreamingWaveData* WaveData = StreamingSoundWaves.Find(Wave);
+			FStreamingWaveData** WaveDataPtr = StreamingSoundWaves.Find(Wave);
 	
-			if (WaveData && WaveData->PendingChunkChangeRequestStatus.GetValue() == AudioState_ReadyFor_Requests)
+			if (WaveDataPtr && (*WaveDataPtr)->PendingChunkChangeRequestStatus.GetValue() == AudioState_ReadyFor_Requests)
 			{
+				FStreamingWaveData* WaveData = *WaveDataPtr;
 				// Request the chunk the source is using and the one after that
 				FWaveRequest& WaveRequest = GetWaveRequest(Wave);
 				int32 SourceChunk = Source->GetBuffer()->GetCurrentChunkIndex();
@@ -407,23 +586,48 @@ void FAudioStreamingManager::UpdateResourceStreaming(float DeltaTime, bool bProc
 			}
 		}
 	}
-
-	for (auto WaveRequest : WaveRequests)
+	for (auto Iter = WaveRequests.CreateIterator(); Iter; ++Iter)
 	{
-		USoundWave* Wave = WaveRequest.Key;
-		FStreamingWaveData* WaveData = StreamingSoundWaves.Find(Wave);
+		USoundWave* Wave = Iter.Key();
+		FStreamingWaveData* WaveData = StreamingSoundWaves.FindRef(Wave);
 
 		if (WaveData && WaveData->PendingChunkChangeRequestStatus.GetValue() == AudioState_ReadyFor_Requests)
 		{
-			WaveData->UpdateChunkRequests(WaveRequest.Value);
+			WaveData->UpdateChunkRequests(Iter.Value());
 			WaveData->UpdateStreamingStatus();
-			WaveRequests.Remove(Wave);
+			Iter.RemoveCurrent();
 		}
 	}
 }
 
-int32 FAudioStreamingManager::BlockTillAllRequestsFinished(float TimeLimit /*= 0.0f*/, bool bLogResults /*= false*/)
+int32 FAudioStreamingManager::BlockTillAllRequestsFinished(float TimeLimit, bool)
 {
+#if USE_NEW_ASYNC_IO
+	{
+		QUICK_SCOPE_CYCLE_COUNTER(FAudioStreamingManager_BlockTillAllRequestsFinished);
+		if (TimeLimit == 0.0f)
+		{
+			for (auto& WavePair : StreamingSoundWaves)
+			{
+				WavePair.Value->BlockTillAllRequestsFinished();
+			}
+		}
+		else
+		{
+			double EndTime = FPlatformTime::Seconds() + TimeLimit;
+			for (auto& WavePair : StreamingSoundWaves)
+			{
+				float ThisTimeLimit = EndTime - FPlatformTime::Seconds();
+				if (ThisTimeLimit < .001f || // one ms is the granularity of the platform event system
+					!WavePair.Value->BlockTillAllRequestsFinished(ThisTimeLimit))
+				{
+					return 1; // we don't report the actual number, just 1 for any number of outstanding requests
+				}
+			}
+		}
+	}
+
+#endif
 	// Not sure yet whether this will work the same as textures - aside from just before destroying
 	return 0;
 }
@@ -451,31 +655,32 @@ void FAudioStreamingManager::RemoveLevel(class ULevel* Level)
 void FAudioStreamingManager::AddStreamingSoundWave(USoundWave* SoundWave)
 {
 	if (FPlatformProperties::SupportsAudioStreaming() && SoundWave->IsStreaming()
-	&& StreamingSoundWaves.Find(SoundWave) == NULL)
+	&& StreamingSoundWaves.FindRef(SoundWave) == NULL)
 	{
-		FStreamingWaveData& WaveData = StreamingSoundWaves.Add(SoundWave);
+		FStreamingWaveData& WaveData = *StreamingSoundWaves.Add(SoundWave, new FStreamingWaveData);
 		WaveData.Initialize(SoundWave);
 	}
 }
 
 void FAudioStreamingManager::RemoveStreamingSoundWave(USoundWave* SoundWave)
 {
-	FStreamingWaveData* WaveData = StreamingSoundWaves.Find(SoundWave);
+	FStreamingWaveData* WaveData = StreamingSoundWaves.FindRef(SoundWave);
 	if (WaveData)
 	{
 		StreamingSoundWaves.Remove(SoundWave);
+		delete WaveData;
 	}
 	WaveRequests.Remove(SoundWave);
 }
 
 bool FAudioStreamingManager::IsManagedStreamingSoundWave(const USoundWave* SoundWave) const
 {
-	return StreamingSoundWaves.Find(SoundWave) != NULL;
+	return StreamingSoundWaves.FindRef(SoundWave) != NULL;
 }
 
 bool FAudioStreamingManager::IsStreamingInProgress(const USoundWave* SoundWave)
 {
-	FStreamingWaveData* WaveData = StreamingSoundWaves.Find(SoundWave);
+	FStreamingWaveData* WaveData = StreamingSoundWaves.FindRef(SoundWave);
 	if (WaveData)
 	{
 		return WaveData->UpdateStreamingStatus();
@@ -568,7 +773,7 @@ bool FAudioStreamingManager::IsManagedStreamingSoundSource(const FSoundSource* S
 
 const uint8* FAudioStreamingManager::GetLoadedChunk(const USoundWave* SoundWave, uint32 ChunkIndex) const
 {
-	const FStreamingWaveData* WaveData = StreamingSoundWaves.Find(SoundWave);
+	const FStreamingWaveData* WaveData = StreamingSoundWaves.FindRef(SoundWave);
 	if (WaveData)
 	{
 		if (WaveData->LoadedChunkIndices.Contains(ChunkIndex))
diff --git a/Engine/Source/Runtime/Engine/Private/Blueprint.cpp b/Engine/Source/Runtime/Engine/Private/Blueprint.cpp
--- a/Engine/Source/Runtime/Engine/Private/Blueprint.cpp
+++ b/Engine/Source/Runtime/Engine/Private/Blueprint.cpp
@@ -330,6 +330,18 @@ void UBlueprint::PreSave(const class ITargetPlatform* TargetPlatform)
 }
 #endif // WITH_EDITORONLY_DATA
 
+void UBlueprint::GetPreloadDependencies(TArray<UObject*>& OutDeps)
+{
+	Super::GetPreloadDependencies(OutDeps);
+	for (UClass* ClassIt = ParentClass; (ClassIt != NULL) && !(ClassIt->HasAnyClassFlags(CLASS_Native)); ClassIt = ClassIt->GetSuperClass())
+	{
+		if (ClassIt->ClassGeneratedBy)
+		{
+			OutDeps.Add(ClassIt->ClassGeneratedBy);
+		}
+	}
+}
+
 void UBlueprint::Serialize(FArchive& Ar)
 {
 	Super::Serialize(Ar);
diff --git a/Engine/Source/Runtime/Engine/Private/BlueprintGeneratedClass.cpp b/Engine/Source/Runtime/Engine/Private/BlueprintGeneratedClass.cpp
--- a/Engine/Source/Runtime/Engine/Private/BlueprintGeneratedClass.cpp
+++ b/Engine/Source/Runtime/Engine/Private/BlueprintGeneratedClass.cpp
@@ -965,6 +965,31 @@ void UBlueprintGeneratedClass::DestroyPersistentUberGraphFrame(UObject* Obj, boo
 	}
 }
 
+void UBlueprintGeneratedClass::GetPreloadDependencies(TArray<UObject*>& OutDeps)
+{
+	Super::GetPreloadDependencies(OutDeps);
+	for (UField* Field = Children; Field; Field = Field->Next)
+	{
+		OutDeps.Add(Field);
+	}
+	OutDeps.Add(GetSuperClass());
+	OutDeps.Add(GetSuperClass()->GetDefaultObject());
+
+	OutDeps.Add(UberGraphFunction);
+	OutDeps.Add(GetInheritableComponentHandler());
+	UObject *CDO = GetDefaultObject();
+	if (CDO)
+	{
+		TArray<UObject*> Subobjects;
+		GetObjectsWithOuter(CDO, Subobjects);
+		for (UObject* SubObj : Subobjects)
+		{
+			OutDeps.Add(SubObj->GetClass());
+			OutDeps.Add(SubObj->GetArchetype());
+		}
+	}
+}
+
 void UBlueprintGeneratedClass::Link(FArchive& Ar, bool bRelinkExistingProperties)
 {
 	// Ensure that function netflags equate to any super function in a parent BP prior to linking; it may have been changed by the user
diff --git a/Engine/Source/Runtime/Engine/Private/Components/PrimitiveComponent.cpp b/Engine/Source/Runtime/Engine/Private/Components/PrimitiveComponent.cpp
--- a/Engine/Source/Runtime/Engine/Private/Components/PrimitiveComponent.cpp
+++ b/Engine/Source/Runtime/Engine/Private/Components/PrimitiveComponent.cpp
@@ -236,7 +236,7 @@ void UPrimitiveComponent::InvalidateLightingCacheDetailed(bool bInvalidateBuildE
 
 bool UPrimitiveComponent::IsEditorOnly() const
 {
-	return ((AlwaysLoadOnClient == false) && (AlwaysLoadOnServer == false)) || Super::IsEditorOnly();
+	return Super::IsEditorOnly() || ((AlwaysLoadOnClient == false) && (AlwaysLoadOnServer == false));
 }
 
 bool UPrimitiveComponent::HasStaticLighting() const
@@ -973,7 +973,7 @@ void UPrimitiveComponent::FinishDestroy()
 
 bool UPrimitiveComponent::NeedsLoadForClient() const
 {
-	if(!IsVisible() && !IsCollisionEnabled() && !AlwaysLoadOnClient)
+	if (!IsVisible() && !IsCollisionEnabled() && !AlwaysLoadOnClient)
 	{
 		return 0;
 	}
diff --git a/Engine/Source/Runtime/Engine/Private/Components/ReflectionCaptureComponent.cpp b/Engine/Source/Runtime/Engine/Private/Components/ReflectionCaptureComponent.cpp
--- a/Engine/Source/Runtime/Engine/Private/Components/ReflectionCaptureComponent.cpp
+++ b/Engine/Source/Runtime/Engine/Private/Components/ReflectionCaptureComponent.cpp
@@ -799,6 +799,7 @@ private:
 
 TArray<UReflectionCaptureComponent*> UReflectionCaptureComponent::ReflectionCapturesToUpdate;
 TArray<UReflectionCaptureComponent*> UReflectionCaptureComponent::ReflectionCapturesToUpdateForLoad;
+FCriticalSection UReflectionCaptureComponent::ReflectionCapturesToUpdateForLoadLock;
 
 UReflectionCaptureComponent::UReflectionCaptureComponent(const FObjectInitializer& ObjectInitializer)
 	: Super(ObjectInitializer)
@@ -857,6 +858,7 @@ void UReflectionCaptureComponent::PostInitProperties()
 
 	if (!HasAnyFlags(RF_ClassDefaultObject | RF_ArchetypeObject))
 	{
+		FScopeLock Lock(&ReflectionCapturesToUpdateForLoadLock);
 		ReflectionCapturesToUpdateForLoad.AddUnique(this);
 		bCaptureDirty = true; 
 	}
@@ -1253,6 +1255,7 @@ void UReflectionCaptureComponent::BeginDestroy()
 	// Deregister the component from the update queue
 	if (bCaptureDirty)
 	{
+		FScopeLock Lock(&ReflectionCapturesToUpdateForLoadLock);
 		ReflectionCapturesToUpdate.Remove(this);
 		ReflectionCapturesToUpdateForLoad.Remove(this);
 	}
@@ -1438,15 +1441,19 @@ void UReflectionCaptureComponent::UpdateReflectionCaptureContents(UWorld* WorldT
 
 		TArray<UReflectionCaptureComponent*> WorldCapturesToUpdateForLoad;
 
-		for (int32 CaptureIndex = ReflectionCapturesToUpdateForLoad.Num() - 1; CaptureIndex >= 0; CaptureIndex--)
+		if (ReflectionCapturesToUpdateForLoad.Num() > 0)
 		{
-			UReflectionCaptureComponent* CaptureComponent = ReflectionCapturesToUpdateForLoad[CaptureIndex];
-
-			if (!CaptureComponent->GetOwner() || WorldToUpdate->ContainsActor(CaptureComponent->GetOwner()))
+			FScopeLock Lock(&ReflectionCapturesToUpdateForLoadLock);
+			for (int32 CaptureIndex = ReflectionCapturesToUpdateForLoad.Num() - 1; CaptureIndex >= 0; CaptureIndex--)
 			{
-				WorldCombinedCaptures.Add(CaptureComponent);
-				WorldCapturesToUpdateForLoad.Add(CaptureComponent);
-				ReflectionCapturesToUpdateForLoad.RemoveAt(CaptureIndex);
+				UReflectionCaptureComponent* CaptureComponent = ReflectionCapturesToUpdateForLoad[CaptureIndex];
+
+				if (!CaptureComponent->GetOwner() || WorldToUpdate->ContainsActor(CaptureComponent->GetOwner()))
+				{
+					WorldCombinedCaptures.Add(CaptureComponent);
+					WorldCapturesToUpdateForLoad.Add(CaptureComponent);
+					ReflectionCapturesToUpdateForLoad.RemoveAt(CaptureIndex);
+				}
 			}
 		}
 
diff --git a/Engine/Source/Runtime/Engine/Private/Components/ShapeComponent.cpp b/Engine/Source/Runtime/Engine/Private/Components/ShapeComponent.cpp
--- a/Engine/Source/Runtime/Engine/Private/Components/ShapeComponent.cpp
+++ b/Engine/Source/Runtime/Engine/Private/Components/ShapeComponent.cpp
@@ -108,7 +108,7 @@ void UShapeComponent::CreateShapeBodySetupIfNeeded()
 {
 	if (ShapeBodySetup == nullptr || ShapeBodySetup->IsPendingKill())
 	{
-		ShapeBodySetup = NewObject<UBodySetup>(this);
+		ShapeBodySetup = NewObject<UBodySetup>(this, NAME_None, RF_Transient);
 		if (GUObjectArray.IsDisregardForGC(this))
 		{
 			ShapeBodySetup->AddToRoot();
diff --git a/Engine/Source/Runtime/Engine/Private/Components/SkyLightComponent.cpp b/Engine/Source/Runtime/Engine/Private/Components/SkyLightComponent.cpp
--- a/Engine/Source/Runtime/Engine/Private/Components/SkyLightComponent.cpp
+++ b/Engine/Source/Runtime/Engine/Private/Components/SkyLightComponent.cpp
@@ -178,6 +178,8 @@ void USkyLightComponent::SetCaptureIsDirty()
 { 
 	if (bVisible && bAffectsWorld)
 	{
+		FScopeLock Lock(&SkyCapturesToUpdateLock);
+
 		SkyCapturesToUpdate.AddUnique(this);
 
 		// Mark saved values as invalid, in case a sky recapture is requested in a construction script between a save / restore of sky capture state
@@ -205,6 +207,7 @@ void USkyLightComponent::SetBlendDestinationCaptureIsDirty()
 
 TArray<USkyLightComponent*> USkyLightComponent::SkyCapturesToUpdate;
 TArray<USkyLightComponent*> USkyLightComponent::SkyCapturesToUpdateBlendDestinations;
+FCriticalSection USkyLightComponent::SkyCapturesToUpdateLock;
 
 void USkyLightComponent::CreateRenderState_Concurrent()
 {
@@ -242,6 +245,7 @@ void USkyLightComponent::PostInitProperties()
 	{
 		// Enqueue an update by default, so that newly placed components will get an update
 		// PostLoad will undo this for components loaded from disk
+		FScopeLock Lock(&SkyCapturesToUpdateLock);
 		SkyCapturesToUpdate.AddUnique(this);
 	}
 
@@ -257,6 +261,7 @@ void USkyLightComponent::PostLoad()
 	// All components are queued for update on creation by default, remove if needed
 	if (!bVisible || HasAnyFlags(RF_ClassDefaultObject | RF_ArchetypeObject))
 	{
+		FScopeLock Lock(&SkyCapturesToUpdateLock);
 		SkyCapturesToUpdate.Remove(this);
 	}
 }
@@ -407,7 +412,11 @@ void USkyLightComponent::CheckForErrors()
 void USkyLightComponent::BeginDestroy()
 {
 	// Deregister the component from the update queue
-	SkyCapturesToUpdate.Remove(this);
+	{
+		FScopeLock Lock(&SkyCapturesToUpdateLock); 
+		SkyCapturesToUpdate.Remove(this);
+	}
+	
 	SkyCapturesToUpdateBlendDestinations.Remove(this);
 
 	// Release reference
@@ -475,6 +484,7 @@ void USkyLightComponent::ApplyComponentInstanceData(FPrecomputedSkyLightInstance
 	if (ProcessedSkyTexture && bSavedConstructionScriptValuesValid)
 	{
 		// We have valid capture state, remove the queued update
+		FScopeLock Lock(&SkyCapturesToUpdateLock);
 		SkyCapturesToUpdate.Remove(this);
 	}
 
@@ -545,8 +555,16 @@ void USkyLightComponent::UpdateSkyCaptureContents(UWorld* WorldToUpdate)
 	if (WorldToUpdate->Scene)
 	{
 		QUICK_SCOPE_CYCLE_COUNTER(STAT_SkylightCaptures);
-		UpdateSkyCaptureContentsArray(WorldToUpdate, SkyCapturesToUpdate, true);
-		UpdateSkyCaptureContentsArray(WorldToUpdate, SkyCapturesToUpdateBlendDestinations, false);
+		if (SkyCapturesToUpdate.Num() > 0)
+		{
+			FScopeLock Lock(&SkyCapturesToUpdateLock);
+			UpdateSkyCaptureContentsArray(WorldToUpdate, SkyCapturesToUpdate, true);
+		}
+		
+		if (SkyCapturesToUpdateBlendDestinations.Num() > 0)
+		{
+			UpdateSkyCaptureContentsArray(WorldToUpdate, SkyCapturesToUpdateBlendDestinations, false);
+		}
 	}
 }
 
diff --git a/Engine/Source/Runtime/Engine/Private/CoreSettings.cpp b/Engine/Source/Runtime/Engine/Private/CoreSettings.cpp
--- a/Engine/Source/Runtime/Engine/Private/CoreSettings.cpp
+++ b/Engine/Source/Runtime/Engine/Private/CoreSettings.cpp
@@ -70,7 +70,8 @@ UStreamingSettings::UStreamingSettings()
 	AsyncLoadingUseFullTimeLimit = true;
 	PriorityAsyncLoadingExtraTime = 20.0f;
 	LevelStreamingActorsUpdateTimeLimit = 5.0f;
-	LevelStreamingComponentsRegistrationGranularity = 10;	
+	LevelStreamingComponentsRegistrationGranularity = 10;
+	EventDrivenLoaderEnabled = false;
 }
 
 void UStreamingSettings::PostInitProperties()
diff --git a/Engine/Source/Runtime/Engine/Private/DataTable.cpp b/Engine/Source/Runtime/Engine/Private/DataTable.cpp
--- a/Engine/Source/Runtime/Engine/Private/DataTable.cpp
+++ b/Engine/Source/Runtime/Engine/Private/DataTable.cpp
@@ -99,6 +99,13 @@ void UDataTable::SaveStructData(FArchive& Ar)
 	}
 }
 
+
+void UDataTable::GetPreloadDependencies(TArray<UObject*>& OutDeps)
+{
+	Super::GetPreloadDependencies(OutDeps);
+	OutDeps.Add(RowStruct);
+}
+
 void UDataTable::OnPostDataImported(TArray<FString>& OutCollectedImportProblems)
 {
 	if (RowStruct && RowStruct->IsChildOf(FTableRowBase::StaticStruct()))
diff --git a/Engine/Source/Runtime/Engine/Private/InheritableComponentHandler.cpp b/Engine/Source/Runtime/Engine/Private/InheritableComponentHandler.cpp
--- a/Engine/Source/Runtime/Engine/Private/InheritableComponentHandler.cpp
+++ b/Engine/Source/Runtime/Engine/Private/InheritableComponentHandler.cpp
@@ -33,6 +33,16 @@ void UInheritableComponentHandler::PostLoad()
 	}
 }
 
+void UInheritableComponentHandler::GetPreloadDependencies(TArray<UObject*>& OutDeps)
+{
+	Super::GetPreloadDependencies(OutDeps);
+	for (auto Record : Records)
+	{
+		OutDeps.Add(Record.ComponentTemplate);
+	}
+}
+
+
 #if WITH_EDITOR
 UActorComponent* UInheritableComponentHandler::CreateOverridenComponentTemplate(FComponentKey Key)
 {
diff --git a/Engine/Source/Runtime/Engine/Private/LevelStreaming.cpp b/Engine/Source/Runtime/Engine/Private/LevelStreaming.cpp
--- a/Engine/Source/Runtime/Engine/Private/LevelStreaming.cpp
+++ b/Engine/Source/Runtime/Engine/Private/LevelStreaming.cpp
@@ -471,19 +471,6 @@ bool ULevelStreaming::RequestLevel(UWorld* PersistentWorld, bool bAllowLevelLoad
 		const FName DesiredPackageNameToLoad = bIsGameWorld ? GetLODPackageNameToLoad() : PackageNameToLoad;
 		const FString PackageNameToLoadFrom = DesiredPackageNameToLoad != NAME_None ? DesiredPackageNameToLoad.ToString() : DesiredPackageName.ToString();
 
-		if (GUseSeekFreeLoading)
-		{
-			// Only load localized package if it exists as async package loading doesn't handle errors gracefully.
-			FString LocalizedPackageName = PackageNameToLoadFrom + LOCALIZED_SEEKFREE_SUFFIX;
-			FString LocalizedFileName;
-			if (FPackageName::DoesPackageExist(LocalizedPackageName, NULL, &LocalizedFileName))
-			{
-				// Load localized part of level first in case it exists. We don't need to worry about GC or completion 
-				// callback as we always kick off another async IO for the level below.
-				LoadPackageAsync(*(GetWorldAssetPackageName() + LOCALIZED_SEEKFREE_SUFFIX), nullptr, *LocalizedPackageName, FLoadPackageAsyncDelegate(), PackageFlags, PIEInstanceID);
-			}
-		}
-
 		if (FPackageName::DoesPackageExist(PackageNameToLoadFrom, NULL, NULL))
 		{
 			bHasLoadRequestPending = true;
diff --git a/Engine/Source/Runtime/Engine/Private/Materials/Material.cpp b/Engine/Source/Runtime/Engine/Private/Materials/Material.cpp
--- a/Engine/Source/Runtime/Engine/Private/Materials/Material.cpp
+++ b/Engine/Source/Runtime/Engine/Private/Materials/Material.cpp
@@ -389,6 +389,12 @@ void UMaterialInterface::PostLoadDefaultMaterials()
 			UMaterial* Material = GDefaultMaterials[Domain];
 			check(Material);
 			Material->ConditionalPostLoad();
+			// Sometimes the above will get called before the material has been fully serialized
+			// in this case its NeedPostLoad flag will not be cleared.
+			if (Material->HasAnyFlags(RF_NeedPostLoad))
+			{
+				bPostLoaded = false;
+			}
 		}
 	}
 }
@@ -2466,23 +2472,26 @@ void UMaterial::GetQualityLevelNodeUsage(TArray<bool, TInlineAllocator<EMaterial
 
 void UMaterial::UpdateResourceAllocations()
 {
-	for (int32 FeatureLevelIndex = 0; FeatureLevelIndex < ERHIFeatureLevel::Num; FeatureLevelIndex++)
+	if (FApp::CanEverRender())
 	{
-		TArray<bool, TInlineAllocator<EMaterialQualityLevel::Num> > QualityLevelsUsed;
-		EShaderPlatform ShaderPlatform = GShaderPlatformForFeatureLevel[FeatureLevelIndex];
-		GetQualityLevelUsage(QualityLevelsUsed, ShaderPlatform);
-		for (int32 QualityLevelIndex = 0; QualityLevelIndex < EMaterialQualityLevel::Num; QualityLevelIndex++)
+		for (int32 FeatureLevelIndex = 0; FeatureLevelIndex < ERHIFeatureLevel::Num; FeatureLevelIndex++)
 		{
-			FMaterialResource*& CurrentResource = MaterialResources[QualityLevelIndex][FeatureLevelIndex];
-
-			if (!CurrentResource)
+			TArray<bool, TInlineAllocator<EMaterialQualityLevel::Num> > QualityLevelsUsed;
+			EShaderPlatform ShaderPlatform = GShaderPlatformForFeatureLevel[FeatureLevelIndex];
+			GetQualityLevelUsage(QualityLevelsUsed, ShaderPlatform);
+			for (int32 QualityLevelIndex = 0; QualityLevelIndex < EMaterialQualityLevel::Num; QualityLevelIndex++)
 			{
-				CurrentResource = AllocateResource();
-			}
+				FMaterialResource*& CurrentResource = MaterialResources[QualityLevelIndex][FeatureLevelIndex];
+
+				if (!CurrentResource)
+				{
+					CurrentResource = AllocateResource();
+				}
 
-			const bool bHasQualityLevelUsage = QualityLevelsUsed[QualityLevelIndex];
-			// Setup transient FMaterialResource properties that are needed to use this resource for rendering or compilation
-			CurrentResource->SetMaterial(this, (EMaterialQualityLevel::Type)QualityLevelIndex, bHasQualityLevelUsage, (ERHIFeatureLevel::Type)FeatureLevelIndex);
+				const bool bHasQualityLevelUsage = QualityLevelsUsed[QualityLevelIndex];
+				// Setup transient FMaterialResource properties that are needed to use this resource for rendering or compilation
+				CurrentResource->SetMaterial(this, (EMaterialQualityLevel::Type)QualityLevelIndex, bHasQualityLevelUsage, (ERHIFeatureLevel::Type)FeatureLevelIndex);
+			}
 		}
 	}
 }
@@ -3433,7 +3442,10 @@ void UMaterial::AddReferencedObjects(UObject* InThis, FReferenceCollector& Colle
 		for (int32 FeatureLevelIndex = 0; FeatureLevelIndex < ERHIFeatureLevel::Num; FeatureLevelIndex++)
 		{
 			FMaterialResource* CurrentResource = This->MaterialResources[QualityLevelIndex][FeatureLevelIndex];
-			CurrentResource->AddReferencedObjects(Collector);
+			if (CurrentResource)
+			{
+				CurrentResource->AddReferencedObjects(Collector);
+			}
 		}
 	}
 #if WITH_EDITORONLY_DATA
diff --git a/Engine/Source/Runtime/Engine/Private/Materials/MaterialInstance.cpp b/Engine/Source/Runtime/Engine/Private/Materials/MaterialInstance.cpp
--- a/Engine/Source/Runtime/Engine/Private/Materials/MaterialInstance.cpp
+++ b/Engine/Source/Runtime/Engine/Private/Materials/MaterialInstance.cpp
@@ -1984,6 +1984,9 @@ void UMaterialInstance::PostLoad()
 	// Ensure that the instance's parent is PostLoaded before the instance.
 	if(Parent)
 	{
+#if !WITH_EDITORONLY_DATA
+		check(!Parent->HasAnyFlags(RF_NeedLoad));
+#endif
 		Parent->ConditionalPostLoad();
 	}
 
diff --git a/Engine/Source/Runtime/Engine/Private/ObjectLibrary.cpp b/Engine/Source/Runtime/Engine/Private/ObjectLibrary.cpp
--- a/Engine/Source/Runtime/Engine/Private/ObjectLibrary.cpp
+++ b/Engine/Source/Runtime/Engine/Private/ObjectLibrary.cpp
@@ -3,6 +3,7 @@
 #include "EnginePrivate.h"
 #include "AssetRegistryModule.h"
 #include "Engine/ObjectLibrary.h"
+#include "Engine/StreamableManager.h"
 
 UObjectLibrary::UObjectLibrary(const FObjectInitializer& ObjectInitializer)
 	: Super(ObjectInitializer)
@@ -448,6 +449,7 @@ int32 UObjectLibrary::LoadBlueprintAssetDataFromPaths(const TArray<FString>& Pat
 int32 UObjectLibrary::LoadAssetsFromAssetData()
 {
 	int32 Count = 0;
+	bool bPreloadObjects = !WITH_EDITOR;
 
 	if (bIsFullyLoaded)
 	{
@@ -457,12 +459,40 @@ int32 UObjectLibrary::LoadAssetsFromAssetData()
 
 	bIsFullyLoaded = true;
 
+	// Preload the packages with an async call, faster in cooked builds
+	if (bPreloadObjects)
+	{
+		TArray<FStringAssetReference> AssetsToStream;
+
+		for (int32 AssetIdx = 0; AssetIdx < AssetDataList.Num(); AssetIdx++)
+		{
+			FAssetData& Data = AssetDataList[AssetIdx];
+			AssetsToStream.AddUnique(Data.PackageName.ToString());
+		}
+
+		if (AssetsToStream.Num())
+		{
+			FStreamableManager Streamable;
+			bool bLoadFinished = false;
+
+			Streamable.RequestAsyncLoad(AssetsToStream,
+				FStreamableDelegate::CreateLambda([&bLoadFinished]()
+			{
+				bLoadFinished = true;
+			}));
+
+			FlushAsyncLoading();
+			check(bLoadFinished);
+		}
+
+	}
+
 	for(int32 AssetIdx=0; AssetIdx<AssetDataList.Num(); AssetIdx++)
 	{
 		FAssetData& Data = AssetDataList[AssetIdx];
 
 		UObject *LoadedObject = NULL;
-			
+		
 		if (!bHasBlueprintClasses)
 		{
 			LoadedObject = Data.GetAsset();
diff --git a/Engine/Source/Runtime/Engine/Private/PhysicsEngine/PhysXLibs.cpp b/Engine/Source/Runtime/Engine/Private/PhysicsEngine/PhysXLibs.cpp
--- a/Engine/Source/Runtime/Engine/Private/PhysicsEngine/PhysXLibs.cpp
+++ b/Engine/Source/Runtime/Engine/Private/PhysicsEngine/PhysXLibs.cpp
@@ -152,22 +152,34 @@ void LoadPhysXModules()
 	FString APEXSuffix(TEXT(".dylib"));
 #endif
 
-	PxFoundationHandle = LoadPhysicsLibrary(PhysXBinariesRoot + "libPxFoundation" + PhysXSuffix);
-	PhysX3CommonHandle = LoadPhysicsLibrary(PhysXBinariesRoot + "libPhysX3Common" + PhysXSuffix);
-	PxPvdSDKHandle = LoadPhysicsLibrary(PhysXBinariesRoot + "libPxPvdSDK" + PhysXSuffix);
-	PhysX3Handle = LoadPhysicsLibrary(PhysXBinariesRoot + "libPhysX3" + PhysXSuffix);
+	const FString PxFoundationLibName = FString::Printf(TEXT("%slibPxFoundation%s"), *PhysXBinariesRoot, *PhysXSuffix);
+	PxFoundationHandle = LoadPhysicsLibrary(PxFoundationLibName);
+
+	const FString PhysX3CommonLibName = FString::Printf(TEXT("%slibPhysX3Common%s"), *PhysXBinariesRoot, *PhysXSuffix);
+	PhysX3CommonHandle = LoadPhysicsLibrary(PhysX3CommonLibName);
+
+	const FString PxPvdSDKLibName = FString::Printf(TEXT("%slibPxPvdSDK%s"), *PhysXBinariesRoot, *PhysXSuffix);
+	PxPvdSDKHandle = LoadPhysicsLibrary(PxPvdSDKLibName);
+
+	const FString PhysX3LibName = FString::Printf(TEXT("%slibPhysX3%s"), *PhysXBinariesRoot, *PhysXSuffix);
+	PhysX3Handle = LoadPhysicsLibrary(PhysX3LibName);
 
 	#if WITH_PHYSICS_COOKING || WITH_RUNTIME_PHYSICS_COOKING
-		PhysX3CookingHandle = LoadPhysicsLibrary(PhysXBinariesRoot + "libPhysX3Cooking" + PhysXSuffix);
+		const FString PhysX3CookinLibName = FString::Printf(TEXT("%slibPhysX3Cooking%s"), *PhysXBinariesRoot, *PhysXSuffix);
+		PhysX3CookingHandle = LoadPhysicsLibrary(PhysX3CookinLibName);
 	#endif
 	#if WITH_APEX
-		APEXFrameworkHandle = LoadPhysicsLibrary(PhysXBinariesRoot + "libAPEXFramework" + APEXSuffix);
-		APEX_DestructibleHandle = LoadPhysicsLibrary(PhysXBinariesRoot + "libAPEX_Destructible" + APEXSuffix);
+		const FString APEXFrameworkLibName = FString::Printf(TEXT("%slibAPEXFramework%s"), *PhysXBinariesRoot, *APEXSuffix);
+		APEXFrameworkHandle = LoadPhysicsLibrary(APEXFrameworkLibName);
+		const FString APEX_DestructibleHandleLibName = FString::Printf(TEXT("%slibAPEX_Destructible%s"), *PhysXBinariesRoot, *APEXSuffix);
+		APEX_DestructibleHandle = LoadPhysicsLibrary(APEX_DestructibleHandleLibName);
 		#if WITH_APEX_LEGACY
-			APEX_LegacyHandle = LoadPhysicsLibrary(PhysXBinariesRoot + "libAPEX_Legacy" + APEXSuffix);
+			const FString APEX_LegacyHandleLibName = FString::Printf(TEXT("%slibAPEX_Legacy%s"), *PhysXBinariesRoot, *APEXSuffix);
+			APEX_LegacyHandle = LoadPhysicsLibrary(APEX_LegacyHandleLibName);
 		#endif //WITH_APEX_LEGACY
 		#if WITH_APEX_CLOTHING
-			APEX_ClothingHandle = LoadPhysicsLibrary(PhysXBinariesRoot + "libAPEX_Clothing" + APEXSuffix);
+			const FString APEX_ClothingHandleLibName = FString::Printf(TEXT("%slibAPEX_Clothing%s"), *PhysXBinariesRoot, *APEXSuffix);
+			APEX_ClothingHandle = LoadPhysicsLibrary(APEX_ClothingHandleLibName);
 		#endif //WITH_APEX_CLOTHING
 	#endif	//WITH_APEX
 #endif	//PLATFORM_WINDOWS
diff --git a/Engine/Source/Runtime/Engine/Private/SoundWave.cpp b/Engine/Source/Runtime/Engine/Private/SoundWave.cpp
--- a/Engine/Source/Runtime/Engine/Private/SoundWave.cpp
+++ b/Engine/Source/Runtime/Engine/Private/SoundWave.cpp
@@ -37,6 +37,7 @@ void FStreamedAudioChunk::Serialize(FArchive& Ar, UObject* Owner, int32 ChunkInd
 	bool bCooked = Ar.IsCooking();
 	Ar << bCooked;
 
+	BulkData.SetBulkDataFlags(BULKDATA_Force_NOT_InlinePayload);
 	BulkData.Serialize(Ar, Owner, ChunkIndex);
 	Ar << DataSize;
 
diff --git a/Engine/Source/Runtime/Engine/Private/Streaming/StreamingManagerTexture.cpp b/Engine/Source/Runtime/Engine/Private/Streaming/StreamingManagerTexture.cpp
--- a/Engine/Source/Runtime/Engine/Private/Streaming/StreamingManagerTexture.cpp
+++ b/Engine/Source/Runtime/Engine/Private/Streaming/StreamingManagerTexture.cpp
@@ -226,7 +226,7 @@ bool FStreamingManagerTexture::StreamOutTextureData( int64 RequiredMemorySize )
 			int32 NumMips = Texture->GetNumMips();
 			int32 MipTailBaseIndex = Texture->GetMipTailBaseIndex();
 			int32 NumRequiredResidentMips = (MipTailBaseIndex >= 0) ? FMath::Max<int32>(NumMips - MipTailBaseIndex, 0) : 0;
-			NumRequiredResidentMips = FMath::Max<int32>(NumRequiredResidentMips, UTexture2D::GetMinTextureResidentMipCount());
+			NumRequiredResidentMips = FMath::Max<int32>(NumRequiredResidentMips, StreamingTexture.NumNonStreamingMips);
 
 			// Only consider streamable textures that have enough miplevels, and who are currently ready for streaming.
 			if ( IsStreamingTexture(Texture) && Texture->IsReadyForStreaming() && Texture->ResidentMips > NumRequiredResidentMips  )
diff --git a/Engine/Source/Runtime/Engine/Private/Texture2D.cpp b/Engine/Source/Runtime/Engine/Private/Texture2D.cpp
--- a/Engine/Source/Runtime/Engine/Private/Texture2D.cpp
+++ b/Engine/Source/Runtime/Engine/Private/Texture2D.cpp
@@ -19,6 +19,10 @@
 #include "ContentStreaming.h"
 #include "Streaming/StreamingManagerTexture.h"
 
+#if USE_NEW_ASYNC_IO
+#include "AsyncFileHandle.h"
+#endif
+
 UTexture2D::UTexture2D(const FObjectInitializer& ObjectInitializer)
 	: Super(ObjectInitializer)
 {
@@ -1107,7 +1111,11 @@ void UTexture2D::RefreshSamplerStates()
 FTexture2DResource::FTexture2DResource( UTexture2D* InOwner, int32 InitialMipCount )
 :	Owner( InOwner )
 ,	ResourceMem( InOwner->ResourceMem )
-,	AsyncCreateTextureTask(NULL)
+,	AsyncCreateTextureTask(nullptr)
+#if USE_NEW_ASYNC_IO
+,	IORequestHandle(nullptr)
+,	IORequestOffsetOffset(0)
+#endif
 ,	IORequestCount( 0 )
 ,	bUsingAsyncCreation(false)
 ,	bPrioritizedIORequest(false)
@@ -1159,6 +1167,10 @@ FTexture2DResource::FTexture2DResource( UTexture2D* InOwner, int32 InitialMipCou
  */
 FTexture2DResource::~FTexture2DResource()
 {
+#if USE_NEW_ASYNC_IO
+	check(!IORequests.Num()); // if these aren't gone by now, that seems problematic
+	verify(BlockTillAllRequestsFinished());
+#endif
 	// free resource memory that was preallocated
 	// The deletion needs to happen in the rendering thread.
 	ENQUEUE_UNIQUE_RENDER_COMMAND_ONEPARAMETER(
@@ -1735,6 +1747,47 @@ void FTexture2DResource::UpdateMipCount()
 	}
 }
 
+#if !defined(USE_NEW_ASYNC_IO) || !defined(SPLIT_COOKED_FILES) || !defined(USE_EVENT_DRIVEN_ASYNC_LOAD)
+#error "USE_NEW_ASYNC_IO and SPLIT_COOKED_FILES and USE_EVENT_DRIVEN_ASYNC_LOAD must be defined."
+#endif
+
+#if USE_NEW_ASYNC_IO
+void FTexture2DResource::AsyncPrep(const FByteBulkData& BulkData)
+{
+	FString Filename = BulkData.GetFilename();
+	check(Filename.Len());
+	if (!IORequestHandle)
+	{
+#if SPLIT_COOKED_FILES && USE_EVENT_DRIVEN_ASYNC_LOAD
+		IORequestOffsetOffset = 0;
+		if (Filename.EndsWith(TEXT(".uasset")) || Filename.EndsWith(TEXT(".umap")))
+		{
+			IORequestOffsetOffset = -IFileManager::Get().FileSize(*Filename);
+			check(IORequestOffsetOffset < 0);
+			Filename = FPaths::GetBaseFilename(Filename, false) + TEXT(".uexp");
+			UE_LOG(LogTexture, Error, TEXT("Streaming from the .uexp file '%s' this MUST be in a ubulk instead for best performance."), *Filename);
+		}
+#endif
+		IORequestFilename = MoveTemp(Filename);
+		IORequestHandle = FPlatformFileManager::Get().GetPlatformFile().OpenAsyncRead(*IORequestFilename);
+		check(IORequestHandle); // this generally cannot fail because it is async
+
+		if (!AsyncFileCallBack)
+		{
+			AsyncFileCallBack = [this](bool bWasCancelled, IAsyncReadRequest* Req)
+			{
+				check(Owner && Owner->IsValidLowLevel());
+				Owner->PendingMipChangeRequestStatus.Decrement();
+			};
+		}
+	}
+	else
+	{
+		check(FPaths::GetBaseFilename(IORequestFilename, false) == FPaths::GetBaseFilename(Filename, false)); // all of the streaming mips must be stored in the same file (could relax that, but this code makes this assumption)
+	}
+}
+#endif
+
 /**
  * Called from the rendering thread to start async I/O to load in new mips.
  */
@@ -1747,6 +1800,10 @@ void FTexture2DResource::LoadMipData()
 
 	if ( (Texture2DRHI->GetFlags() & TexCreate_Virtual) == TexCreate_Virtual )
 	{
+
+#if USE_NEW_ASYNC_IO
+		verify(BlockTillAllRequestsFinished());
+#endif
 		IORequestCount = 0;
 		if ( !Owner->bHasCancelationPending )
 		{
@@ -1777,7 +1834,19 @@ void FTexture2DResource::LoadMipData()
 				Owner->PendingMipChangeRequestStatus.Increment();
 
 				EAsyncIOPriority AsyncIOPriority = bPrioritizedIORequest ? AIOP_BelowNormal : AIOP_Low;
+#if USE_NEW_ASYNC_IO
+				if (MipMap.BulkData.IsStoredCompressedOnDisk())
+				{
+					check(!"USE_NEW_ASYNC_IO does not support compression at the package level.");
+				}
+				else
+				{
+					AsyncPrep(MipMap.BulkData);
+					IORequests.Add(IORequestHandle->ReadRequest(MipMap.BulkData.GetBulkDataOffsetInFile() + IORequestOffsetOffset, MipMap.BulkData.GetBulkDataSize(), AsyncIOPriority, &AsyncFileCallBack, (uint8*)TheMipData));
+					IORequestCount++; // no clue what this actually means, but it is important
+				}
 
+#else
 				// Load and decompress async.
 				{
 					check(MipMap.BulkData.GetFilename().Len());
@@ -1808,6 +1877,7 @@ void FTexture2DResource::LoadMipData()
 					}
 					check(IORequestIndices[MipIndex]);
 				}
+#endif
 
 				// For consistency with other code paths, track the pointer to the locked buffer.
 				MipData[ActualMipIndex] = TheMipData;
@@ -1828,7 +1898,11 @@ void FTexture2DResource::LoadMipData()
 		return;
 	}
 
+#if USE_NEW_ASYNC_IO
+	verify(BlockTillAllRequestsFinished());
+#endif
 	IORequestCount = 0;
+
 	if ( (bUsingAsyncCreation || IsValidRef(IntermediateTextureRHI)) && !Owner->bHasCancelationPending )
 	{
 		STAT( IntermediateTextureSize = Owner->CalcTextureMemorySize( Owner->RequestedMips ) );
@@ -1899,6 +1973,18 @@ void FTexture2DResource::LoadMipData()
 					QUICK_SCOPE_CYCLE_COUNTER(STAT_FTexture2DResource_LoadMipData_Malloc);
 					MipData[ActualMipIndex] = FMemory::Malloc(MipSize);
 				}
+#if USE_NEW_ASYNC_IO
+				if (MipMap.BulkData.IsStoredCompressedOnDisk())
+				{
+					check(!"USE_NEW_ASYNC_IO does not support compression at the package level.");
+				}
+				else
+				{
+					AsyncPrep(MipMap.BulkData);
+					IORequests.Add(IORequestHandle->ReadRequest(MipMap.BulkData.GetBulkDataOffsetInFile() + IORequestOffsetOffset, MipMap.BulkData.GetBulkDataSize(), AsyncIOPriority, &AsyncFileCallBack, (uint8*)MipData[ActualMipIndex]));
+					IORequestCount++; // no clue what this actually means, but it is important
+				}
+#else
 				check(MipMap.BulkData.GetFilename().Len());
 				if( MipMap.BulkData.IsStoredCompressedOnDisk() )
 				{
@@ -1926,6 +2012,7 @@ void FTexture2DResource::LoadMipData()
 						);
 				}
 				check(IORequestIndices[MipIndex]);
+#endif
 			}
 		}
 
@@ -2080,19 +2167,72 @@ void FTexture2DResource::UploadMipData()
 	}
 }
 
+#if USE_NEW_ASYNC_IO
+bool FTexture2DResource::BlockTillAllRequestsFinished(float TimeLimit)
+{
+	QUICK_SCOPE_CYCLE_COUNTER(FTexture2DResource_BlockTillAllRequestsFinished);
+	if (TimeLimit == 0.0f)
+	{
+		double StartTime = FPlatformTime::Seconds();
+		for (IAsyncReadRequest* IORequest : IORequests)
+		{
+			if (IORequest)
+			{
+				IORequest->WaitCompletion();
+				delete IORequest;
+			}
+		}
+		float ThisTime = FPlatformTime::Seconds() - StartTime;
+		if (ThisTime > .01f)
+		{
+			UE_LOG(LogTexture, Log, TEXT("Hitch (%6.1fms) in waiting for async IO in FTexture2DResource::BlockTillAllRequestsFinished; this can happen, especially on desktop platforms. If it happens a lot, or on consoles, that would be indicative of a bug in the texture streamer."), ThisTime * 1000.0f);
+		}
+	}
+	else
+	{
+		double EndTime = FPlatformTime::Seconds() + TimeLimit;
+		for (IAsyncReadRequest*& IORequest : IORequests)
+		{
+			if (IORequest)
+			{
+				float ThisTimeLimit = EndTime - FPlatformTime::Seconds();
+				if (ThisTimeLimit < .001f || // one ms is the granularity of the platform event system
+					!IORequest->WaitCompletion(ThisTimeLimit))
+				{
+					IORequests.Remove(nullptr);
+					return false;
+				}
+				delete IORequest;
+				IORequest = nullptr;
+			}
+		}
+	}
+	IORequests.Empty();
+	return true;
+}
+#endif
+
 /**
 * Helper function for cleaning up bulk data files after streaming
-* @todo: make it smarter, close only when we know we won't be streaming anymore or at least for a while
-* @todo: What if each mip is in a different bulk data file? might need to loop over all
 */
-inline void HintDoneWithStreamedTextureFiles(const UTexture2D* InTexture)
+void FTexture2DResource::HintDoneWithStreamedTextureFiles()
 {
+#if USE_NEW_ASYNC_IO
+	verify(BlockTillAllRequestsFinished());
+	if (IORequestHandle)
+	{
+		delete IORequestHandle;
+		IORequestHandle = nullptr;
+	}
+	check(!IORequestHandle);
+#else
 	if (FPlatformProperties::RequiresCookedData())
 	{
-		const TIndirectArray<FTexture2DMipMap>& OwnerMips = InTexture->GetPlatformMips();
+		const TIndirectArray<FTexture2DMipMap>& OwnerMips = Owner->GetPlatformMips();
 		const FTexture2DMipMap& MipMap = OwnerMips[0];
 		FIOSystem::Get().HintDoneWithFile(MipMap.BulkData.GetFilename());
 	}
+#endif
 }
 
 /**
@@ -2133,7 +2273,7 @@ void FTexture2DResource::FinalizeMipCount()
 		// We're done.
 		Owner->PendingMipChangeRequestStatus.Decrement();
 
-		HintDoneWithStreamedTextureFiles(Owner);
+		HintDoneWithStreamedTextureFiles();
 
 		return;
 	}
@@ -2184,7 +2324,7 @@ void FTexture2DResource::FinalizeMipCount()
 		}
 		IntermediateTextureRHI.SafeRelease();
 
-		HintDoneWithStreamedTextureFiles(Owner);
+		HintDoneWithStreamedTextureFiles();
 	}
 	else
 	{
@@ -2214,8 +2354,22 @@ void FTexture2DResource::CancelUpdate()
 	// We only have anything worth cancellation if there are outstanding I/O requests.
 	if( IORequestCount )
 	{
+#if USE_NEW_ASYNC_IO
+		// Cancel requests. Some may already have been fulfilled or be in the process of being fulfilled.
+		for (IAsyncReadRequest* Item : IORequests)
+		{
+			Item->Cancel();
+		}
+		for (IAsyncReadRequest* Item : IORequests)
+		{
+			Item->WaitCompletion();
+			delete Item;
+		}
+		IORequests.Empty();
+#else
 		// Cancel requests. This only cancels pending requests and not ones currently being fulfilled.
 		FIOSystem::Get().CancelRequests( IORequestIndices, IORequestCount );
+#endif
 	}
 
 	if ( !bUsingAsyncCreation && IsValidRef(IntermediateTextureRHI) )
diff --git a/Engine/Source/Runtime/Engine/Private/TextureDerivedData.cpp b/Engine/Source/Runtime/Engine/Private/TextureDerivedData.cpp
--- a/Engine/Source/Runtime/Engine/Private/TextureDerivedData.cpp
+++ b/Engine/Source/Runtime/Engine/Private/TextureDerivedData.cpp
@@ -1216,7 +1216,7 @@ int32 FTexturePlatformData::GetNumNonStreamingMips() const
 		for (const FTexture2DMipMap& Mip : Mips)
 		{
 			uint32 BulkDataFlags = Mip.BulkData.GetBulkDataFlags();
-			if (BulkDataFlags & BULKDATA_PayloadInSeperateFile || BulkDataFlags & BULKDATA_PayloadAtEndOfFile)
+			if ((BulkDataFlags & BULKDATA_PayloadInSeperateFile) || (BulkDataFlags & BULKDATA_PayloadAtEndOfFile))
 			{
 				--NumNonStreamingMips;
 			}
@@ -1334,6 +1334,10 @@ static void SerializePlatformData(
 			MinMipToInline = FMath::Max(0, NumMips - PlatformData->GetNumNonStreamingMips());
 		}
 
+		for (int32 MipIndex = 0; MipIndex < NumMips && MipIndex < MinMipToInline; ++MipIndex)
+		{
+			PlatformData->Mips[MipIndex + FirstMipToSerialize].BulkData.SetBulkDataFlags(BULKDATA_Force_NOT_InlinePayload);
+		}
 		for (int32 MipIndex = MinMipToInline; MipIndex < NumMips; ++MipIndex)
 		{
 			PlatformData->Mips[MipIndex + FirstMipToSerialize].BulkData.SetBulkDataFlags(BULKDATA_ForceInlinePayload | BULKDATA_SingleUse);
diff --git a/Engine/Source/Runtime/Engine/Private/URL.cpp b/Engine/Source/Runtime/Engine/Private/URL.cpp
--- a/Engine/Source/Runtime/Engine/Private/URL.cpp
+++ b/Engine/Source/Runtime/Engine/Private/URL.cpp
@@ -74,7 +74,14 @@ FURL::FURL( const TCHAR* LocalFilename )
 	// strip off any extension from map name
 	if (LocalFilename)
 	{
-		Map = FPaths::GetBaseFilename(LocalFilename);
+		if (FPackageName::IsValidLongPackageName(LocalFilename))
+		{
+			Map = LocalFilename;
+		}
+		else
+		{
+			Map = FPaths::GetBaseFilename(LocalFilename);
+		}
 	}
 	else
 	{
diff --git a/Engine/Source/Runtime/Engine/Private/UnrealEngine.cpp b/Engine/Source/Runtime/Engine/Private/UnrealEngine.cpp
--- a/Engine/Source/Runtime/Engine/Private/UnrealEngine.cpp
+++ b/Engine/Source/Runtime/Engine/Private/UnrealEngine.cpp
@@ -141,6 +141,7 @@
 #endif	// UE_BUILD_SHIPPING
 
 #include "GeneralProjectSettings.h"
+#include "LoadTimeTracker.h"
 
 DEFINE_LOG_CATEGORY(LogEngine);
 IMPLEMENT_MODULE( FEngineModule, Engine );
@@ -2725,10 +2726,18 @@ bool UEngine::Exec( UWorld* InWorld, const TCHAR* Cmd, FOutputDevice& Ar )
 	{
 		return HandleFreezeAllCommand( Cmd, Ar, InWorld );
 	}
+#if !USE_NEW_ASYNC_IO
 	else if( FParse::Command(&Cmd, TEXT("FLUSHIOMANAGER")) )
 	{
 		return HandleFlushIOManagerCommand( Cmd, Ar );
 	}
+	// This will list out the packages which are in the precache list and have not been "loaded" out.  (e.g. could be just there taking up memory!)
+	else if (FParse::Command(&Cmd, TEXT("ListPrecacheMapPackages")))
+	{
+		return HandleListPreCacheMapPackagesCommand(Cmd, Ar);
+	}
+#endif
+
 	else if( FParse::Command(&Cmd,TEXT("ToggleRenderingThread")) )
 	{
 		return HandleToggleRenderingThreadCommand( Cmd, Ar );
@@ -2834,11 +2843,6 @@ bool UEngine::Exec( UWorld* InWorld, const TCHAR* Cmd, FOutputDevice& Ar )
 	{
 		return HandleDumpParticleCountsCommand( Cmd, Ar );
 	}
-	// This will list out the packages which are in the precache list and have not been "loaded" out.  (e.g. could be just there taking up memory!)
-	else if( FParse::Command( &Cmd, TEXT("ListPrecacheMapPackages") ) )
-	{		
-		return HandleListPreCacheMapPackagesCommand( Cmd, Ar );
-	}
 	// we can't always do an obj linkers, as cooked games have their linkers tossed out.  So we need to look at the actual packages which are loaded
 	else if( FParse::Command( &Cmd, TEXT("ListLoadedPackages") ) )
 	{
@@ -3366,11 +3370,31 @@ bool UEngine::HandleFreezeAllCommand( const TCHAR* Cmd, FOutputDevice& Ar, UWorl
 	return true;
 }
 
+#if !USE_NEW_ASYNC_IO
 bool UEngine::HandleFlushIOManagerCommand( const TCHAR* Cmd, FOutputDevice& Ar )
 {
 	FIOSystem::Get().BlockTillAllRequestsFinishedAndFlushHandles();
 	return true;
 }
+bool UEngine::HandleListPreCacheMapPackagesCommand(const TCHAR* Cmd, FOutputDevice& Ar)
+{
+	TArray<FString> Packages;
+	FLinkerLoad::GetListOfPackagesInPackagePrecacheMap(Packages);
+
+	Packages.Sort();
+
+	Ar.Logf(TEXT("Total Number Of Packages In PrecacheMap: %i "), Packages.Num());
+
+	for (int32 i = 0; i < Packages.Num(); ++i)
+	{
+		Ar.Logf(TEXT("%i %s"), i, *Packages[i]);
+	}
+	Ar.Logf(TEXT("Total Number Of Packages In PrecacheMap: %i "), Packages.Num());
+
+	return true;
+}
+
+#endif
 
 bool UEngine::HandleFreezeRenderingCommand( const TCHAR* Cmd, FOutputDevice& Ar, UWorld* InWorld  )
 {
@@ -4496,23 +4520,6 @@ bool UEngine::HandleDumpParticleCountsCommand( const TCHAR* Cmd, FOutputDevice&
 	return true;
 }
 
-bool UEngine::HandleListPreCacheMapPackagesCommand( const TCHAR* Cmd, FOutputDevice& Ar )
-{
-	TArray<FString> Packages;
-	FLinkerLoad::GetListOfPackagesInPackagePrecacheMap( Packages );
-
-	Packages.Sort();
-
-	Ar.Logf( TEXT( "Total Number Of Packages In PrecacheMap: %i " ), Packages.Num() );
-
-	for( int32 i = 0; i < Packages.Num(); ++i )
-	{
-		Ar.Logf( TEXT( "%i %s" ), i, *Packages[i] );
-	}
-	Ar.Logf( TEXT( "Total Number Of Packages In PrecacheMap: %i " ), Packages.Num() );
-
-	return true;
-}
 
 bool UEngine::HandleListLoadedPackagesCommand( const TCHAR* Cmd, FOutputDevice& Ar )
 {
@@ -9599,6 +9606,8 @@ bool UEngine::LoadMap( FWorldContext& WorldContext, FURL URL, class UPendingNetG
 	MALLOC_PROFILER( FMallocProfiler::SnapshotMemoryLoadMapStart( URL.Map ) );
 	Error = TEXT("");
 
+	FLoadTimeTracker::Get().ResetRawLoadTimes();
+
 	FPlatformMisc::PreLoadMap(URL.Map, WorldContext.LastURL.Map, nullptr);
 	
 	// make sure level streaming isn't frozen
@@ -9892,11 +9901,6 @@ bool UEngine::LoadMap( FWorldContext& WorldContext, FURL URL, class UPendingNetG
 
 		FScopeCycleCounterUObject MapScope(WorldPackage);
 
-		if (FPlatformProperties::RequiresCookedData() && GUseSeekFreeLoading && !WorldPackage->HasAnyPackageFlags(PKG_DisallowLazyLoading))
-		{
-			UE_LOG(LogLoad, Fatal, TEXT("Map '%s' has not been cooked correctly! Most likely stale version on the XDK."), *WorldPackage->GetName());
-		}
-
 		if (WorldContext.WorldType == EWorldType::PIE)
 		{
 			// If we are a PIE world and the world we just found is already initialized, then we're probably reloading the editor world and we
@@ -10100,6 +10104,7 @@ bool UEngine::LoadMap( FWorldContext& WorldContext, FURL URL, class UPendingNetG
 	double StopTime = FPlatformTime::Seconds();
 
 	UE_LOG(LogLoad, Log, TEXT("Took %f seconds to LoadMap(%s)"), StopTime - StartTime, *URL.Map);
+	FLoadTimeTracker::Get().DumpRawLoadTimes();
 
 	// Successfully started local level.
 	return true;
@@ -10818,19 +10823,7 @@ bool UEngine::PrepareMapChange(FWorldContext &Context, const TArray<FName>& Leve
 		for( int32 LevelIndex=0; LevelIndex < Context.LevelsToLoadForPendingMapChange.Num(); LevelIndex++ )
 		{
 			const FName LevelName = Context.LevelsToLoadForPendingMapChange[LevelIndex];
-			if( GUseSeekFreeLoading )
-			{
-				// Only load localized package if it exists as async package loading doesn't handle errors gracefully.
-				FString LocalizedPackageName = LevelName.ToString() + LOCALIZED_SEEKFREE_SUFFIX;
-				FString LocalizedFileName;
-				if( FPackageName::DoesPackageExist( LocalizedPackageName, NULL, &LocalizedFileName ) )
-				{
-					// Load localized part of level first in case it exists. We don't need to worry about GC or completion 
-					// callback as we always kick off another async IO for the level below.
-					LoadPackageAsync(LocalizedPackageName);
-				}
-			}
-			
+		
 			STAT_ADD_CUSTOMMESSAGE_NAME( STAT_NamedMarker, *(FString( TEXT( "PrepareMapChange - " ) + LevelName.ToString() )) );
 			LoadPackageAsync(LevelName.ToString(),
 				FLoadPackageAsyncDelegate::CreateStatic(&AsyncMapChangeLevelLoadCompletionCallback, Context.ContextHandle)
diff --git a/Engine/Source/Runtime/Engine/Public/AudioStreaming.h b/Engine/Source/Runtime/Engine/Public/AudioStreaming.h
--- a/Engine/Source/Runtime/Engine/Public/AudioStreaming.h
+++ b/Engine/Source/Runtime/Engine/Public/AudioStreaming.h
@@ -83,9 +83,25 @@ struct FWaveRequest
 struct FLoadedAudioChunk
 {
 	uint8*	Data;
+#if USE_NEW_ASYNC_IO
+	class IAsyncReadRequest* IORequest;
+#else
+	int32	MemorySize; 
+#endif
 	int32	DataSize;
-	int32	MemorySize;
 	uint32	Index;
+
+	FLoadedAudioChunk()
+		: Data(nullptr)
+#if USE_NEW_ASYNC_IO
+		, IORequest(nullptr)
+#else
+		, MemorySize(0)
+#endif
+		, DataSize(0)
+		, Index(0)
+	{
+	}
 };
 
 /**
@@ -133,6 +149,16 @@ struct FStreamingWaveData final
 	 */
 	void BeginPendingRequests(const TArray<uint32>& IndicesToLoad, const TArray<uint32>& IndicesToFree);
 
+#if USE_NEW_ASYNC_IO
+	/**
+	* Blocks till all pending requests are fulfilled.
+	*
+	* @param TimeLimit		Optional time limit for processing, in seconds. Specifying 0 means infinite time limit.
+	* @return				Return true if there are no requests left in flight, false if the time limit was reached before they were finished.
+	*/
+	bool BlockTillAllRequestsFinished(float TimeLimit = 0.0f);
+#endif
+
 #if WITH_EDITORONLY_DATA
 	/**
 	 * Finishes any Derived Data Cache requests that may be in progress
@@ -160,8 +186,13 @@ public:
 	/* Contains pointers to Chunks of audio data that have been streamed in */
 	TArray<FLoadedAudioChunk> LoadedChunks;
 
+#if USE_NEW_ASYNC_IO
+	class IAsyncReadFileHandle* IORequestHandle;
+	FAsyncFileCallBack AsyncFileCallBack;
+#else
 	/** Potentially outstanding audio chunk I/O requests */
 	TArray<uint64>	IORequestIndices;
+#endif
 
 	/** Indices of chunks that are currently loaded */
 	TArray<uint32>	LoadedChunkIndices;
@@ -217,7 +248,7 @@ protected:
 	FWaveRequest& GetWaveRequest(USoundWave* SoundWave);
 
 	/** Sound Waves being managed. */
-	TMap<USoundWave*, FStreamingWaveData> StreamingSoundWaves;
+	TMap<USoundWave*, FStreamingWaveData*> StreamingSoundWaves;
 
 	/** Sound Sources being managed. */
 	TArray<FSoundSource*>	StreamingSoundSources;
diff --git a/Engine/Source/Runtime/Engine/Public/EngineUtils.h b/Engine/Source/Runtime/Engine/Public/EngineUtils.h
--- a/Engine/Source/Runtime/Engine/Public/EngineUtils.h
+++ b/Engine/Source/Runtime/Engine/Public/EngineUtils.h
@@ -734,7 +734,7 @@ namespace EngineUtils
 }
 
 /** Helper class for serializing flags describing which data have been stripped (if any). */
-class FStripDataFlags
+class ENGINE_API FStripDataFlags
 {
 	/** Serialized engine strip flags (up to 8 flags). */
 	uint8 GlobalStripFlags;
diff --git a/Engine/Source/Runtime/Engine/Public/TextureResource.h b/Engine/Source/Runtime/Engine/Public/TextureResource.h
--- a/Engine/Source/Runtime/Engine/Public/TextureResource.h
+++ b/Engine/Source/Runtime/Engine/Public/TextureResource.h
@@ -315,9 +315,21 @@ private:
 	
 	/** Local copy/ cache of mip data between creation and first call to InitRHI.							*/
 	void*				MipData[MAX_TEXTURE_MIP_COUNT];
+
+#if USE_NEW_ASYNC_IO
+	/**  Async handle */
+	class IAsyncReadFileHandle* IORequestHandle;
+	FString IORequestFilename;
+	int64 IORequestOffsetOffset;
+	FAsyncFileCallBack AsyncFileCallBack;
+
+	/** Potentially outstanding texture I/O requests.														*/
+	TArray<class IAsyncReadRequest*> IORequests;
+#else
 	/** Potentially outstanding texture I/O requests.														*/
 	uint64				IORequestIndices[MAX_TEXTURE_MIP_COUNT];
-	/** Number of file I/O requests for current request														*/
+#endif
+	/** Number of file I/O requests for current request. The use of this is crazy confusing.				*/
 	int32					IORequestCount;
 
 	FThreadSafeCounter AsyncReallocateCounter;
@@ -387,6 +399,21 @@ private:
 	// releases and recreates sampler state objects.
 	// used when updating mip map bias offset
 	void RefreshSamplerStates();
+
+	/**
+	* Helper function for cleaning up bulk data files after streaming
+	*/
+	void HintDoneWithStreamedTextureFiles();
+
+#if USE_NEW_ASYNC_IO
+	/**
+	* Block until all requests are done.
+	* A time limit of zero means infinite
+	*/
+	bool BlockTillAllRequestsFinished(float TimeLimit = 0.0f);
+	void AsyncPrep(const FByteBulkData& BulkData);
+#endif
+
 };
 
 /** A dynamic 2D texture resource. */
diff --git a/Engine/Source/Runtime/Landscape/Private/LandscapeCollision.cpp b/Engine/Source/Runtime/Landscape/Private/LandscapeCollision.cpp
--- a/Engine/Source/Runtime/Landscape/Private/LandscapeCollision.cpp
+++ b/Engine/Source/Runtime/Landscape/Private/LandscapeCollision.cpp
@@ -129,9 +129,13 @@ static FString GetHFDDCKeyString(const FName& Format, bool bDefMaterial, const F
 
 ECollisionEnabled::Type ULandscapeHeightfieldCollisionComponent::GetCollisionEnabled() const
 {
-	ALandscapeProxy* Proxy = GetLandscapeProxy();
+	if (!HasAnyFlags(RF_ClassDefaultObject))
+	{
+		ALandscapeProxy* Proxy = GetLandscapeProxy();
 
-	return Proxy->BodyInstance.GetCollisionEnabled();
+		return Proxy->BodyInstance.GetCollisionEnabled();
+	}
+	return ECollisionEnabled::QueryAndPhysics;
 }
 
 ECollisionResponse ULandscapeHeightfieldCollisionComponent::GetCollisionResponseToChannel(ECollisionChannel Channel) const
diff --git a/Engine/Source/Runtime/Landscape/Private/LandscapeGrass.cpp b/Engine/Source/Runtime/Landscape/Private/LandscapeGrass.cpp
--- a/Engine/Source/Runtime/Landscape/Private/LandscapeGrass.cpp
+++ b/Engine/Source/Runtime/Landscape/Private/LandscapeGrass.cpp
@@ -1173,7 +1173,7 @@ FArchive& operator<<(FArchive& Ar, FLandscapeComponentGrassData& Data)
 	// Each weight data array, being 1 byte will be serialized in bulk.
 	Ar << Data.WeightData;
 
-	if (Ar.IsLoading() && !GIsEditor && CVarGrassDiscardDataOnLoad.GetValueOnGameThread())
+	if (Ar.IsLoading() && !GIsEditor && CVarGrassDiscardDataOnLoad.GetValueOnAnyThread())
 	{
 		//Data = FLandscapeComponentGrassData();
 		Data.WeightData.Empty();
diff --git a/Engine/Source/Runtime/Launch/Private/LaunchEngineLoop.cpp b/Engine/Source/Runtime/Launch/Private/LaunchEngineLoop.cpp
--- a/Engine/Source/Runtime/Launch/Private/LaunchEngineLoop.cpp
+++ b/Engine/Source/Runtime/Launch/Private/LaunchEngineLoop.cpp
@@ -1190,7 +1190,11 @@ int32 FEngineLoop::PreInit( const TCHAR* CmdLine )
 	FScopeCycleCounter CycleCount_AfterStats( GET_STATID( STAT_FEngineLoop_PreInit_AfterStats ) );
 
 	// Load Core modules required for everything else to work (needs to be loaded before InitializeRenderingCVarsCaching)
-	LoadCoreModules();
+	if (!LoadCoreModules())
+	{
+		UE_LOG(LogInit, Error, TEXT("Failed to load Core modules."));
+		return 1;
+	}
 
 #if WITH_ENGINE
 	extern ENGINE_API void InitializeRenderingCVarsCaching();
@@ -1412,7 +1416,9 @@ int32 FEngineLoop::PreInit( const TCHAR* CmdLine )
 		InitializeStdOutDevice();
 	}
 
+#if !USE_NEW_ASYNC_IO
 	FIOSystem::Get(); // force it to be created if it isn't already
+#endif
 
 	// allow the platform to start up any features it may need
 	IPlatformFeaturesModule::Get();
@@ -2025,11 +2031,14 @@ int32 FEngineLoop::PreInit( const TCHAR* CmdLine )
 }
 
 
-void FEngineLoop::LoadCoreModules()
+bool FEngineLoop::LoadCoreModules()
 {
 	// Always attempt to load CoreUObject. It requires additional pre-init which is called from its module's StartupModule method.
 #if WITH_COREUOBJECT
-	FModuleManager::Get().LoadModule(TEXT("CoreUObject"));
+	bool bResult = FModuleManager::Get().LoadModule(TEXT("CoreUObject")).IsValid();
+	return bResult;
+#else
+	return true;
 #endif
 }
 
@@ -2512,7 +2521,9 @@ void FEngineLoop::Exit()
 
 	FTaskGraphInterface::Shutdown();
 	IStreamingManager::Shutdown();
+#if !USE_NEW_ASYNC_IO
 	FIOSystem::Shutdown();
+#endif
 }
 
 
diff --git a/Engine/Source/Runtime/Launch/Public/LaunchEngineLoop.h b/Engine/Source/Runtime/Launch/Public/LaunchEngineLoop.h
--- a/Engine/Source/Runtime/Launch/Public/LaunchEngineLoop.h
+++ b/Engine/Source/Runtime/Launch/Public/LaunchEngineLoop.h
@@ -46,7 +46,7 @@ public:
 	void LoadPreInitModules();
 
 	/** Load core modules. */
-	void LoadCoreModules();
+	bool LoadCoreModules();
 
 #if WITH_ENGINE
 	
