Link function inputs into the caller graph before re-entering after compiling an input.  Fixes cases where chained calls to the same function will overwrite input connections on previous calls in the chain.

diff --git a/Engine/Shaders/ShaderVersion.usf b/Engine/Shaders/ShaderVersion.usf
--- a/Engine/Shaders/ShaderVersion.usf
+++ b/Engine/Shaders/ShaderVersion.usf
@@ -2,4 +2,5 @@
 // Each time the console command is executed it generates a new GUID. As this file is included
 // in common.usf (which should be included in any shader) it allows to invalidate the shader DDC.
 // 
-// GUID = FCe205F2167C4A0EB4DDC870A2Bb74E7
+// GUID = FCe205F2167C4A0EB5DDC870A2Bb74E7
+ 
\ No newline at end of file
diff --git a/Engine/Source/Runtime/Engine/Classes/Materials/MaterialExpressionFunctionInput.h b/Engine/Source/Runtime/Engine/Classes/Materials/MaterialExpressionFunctionInput.h
--- a/Engine/Source/Runtime/Engine/Classes/Materials/MaterialExpressionFunctionInput.h
+++ b/Engine/Source/Runtime/Engine/Classes/Materials/MaterialExpressionFunctionInput.h
@@ -25,7 +25,7 @@ class UMaterialExpressionFunctionInput : public UMaterialExpression
 {
 	GENERATED_UCLASS_BODY()
 
-	/** Used for previewing when editing the function, also temporarily used to connect to the outside material when compiling that material. */
+	/** Used for previewing when editing the function, or when bUsePreviewValueAsDefault is enabled. */
 	UPROPERTY(meta=(RequiredInput = "false"))
 	FExpressionInput Preview;
 
diff --git a/Engine/Source/Runtime/Engine/Classes/Materials/MaterialFunction.h b/Engine/Source/Runtime/Engine/Classes/Materials/MaterialFunction.h
--- a/Engine/Source/Runtime/Engine/Classes/Materials/MaterialFunction.h
+++ b/Engine/Source/Runtime/Engine/Classes/Materials/MaterialFunction.h
@@ -78,7 +78,12 @@ public:
 	/** Get the inputs and outputs that this function exposes, for a function call expression to use. */
 	void GetInputsAndOutputs(TArray<struct FFunctionExpressionInput>& OutInputs, TArray<struct FFunctionExpressionOutput>& OutOutputs) const;
 
-	int32 Compile(class FMaterialCompiler* Compiler, const struct FFunctionExpressionOutput& Output, const TArray<struct FFunctionExpressionInput>& Inputs);
+	int32 Compile(class FMaterialCompiler* Compiler, const struct FFunctionExpressionOutput& Output);
+
+	/** Called during compilation before entering the function. */
+	void LinkIntoCaller(const TArray<FFunctionExpressionInput>& CallerInputs);
+
+	void UnlinkFromCaller();
 
 	/** @return true if this function is dependent on the passed in function, directly or indirectly. */
 	ENGINE_API bool IsDependent(UMaterialFunction* OtherFunction);
diff --git a/Engine/Source/Runtime/Engine/Private/Materials/MaterialExpressions.cpp b/Engine/Source/Runtime/Engine/Private/Materials/MaterialExpressions.cpp
--- a/Engine/Source/Runtime/Engine/Private/Materials/MaterialExpressions.cpp
+++ b/Engine/Source/Runtime/Engine/Private/Materials/MaterialExpressions.cpp
@@ -7826,7 +7826,25 @@ static int32 FindOutputIndexByName(const FString& Name, const TArray<FFunctionEx
 }
 
 #if WITH_EDITOR
-int32 UMaterialFunction::Compile(class FMaterialCompiler* Compiler, const FFunctionExpressionOutput& Output, const TArray<FFunctionExpressionInput>& Inputs)
+int32 UMaterialFunction::Compile(class FMaterialCompiler* Compiler, const FFunctionExpressionOutput& Output)
+{
+	int32 ReturnValue = INDEX_NONE;
+	if (Output.ExpressionOutput->A.Expression)
+	{
+		// Compile the given function output
+		ReturnValue = Output.ExpressionOutput->A.Compile(Compiler);
+	}
+	else
+	{
+		ReturnValue = Compiler->Errorf(TEXT("Missing function output connection '%s'"), *Output.ExpressionOutput->OutputName);
+	}
+
+	return ReturnValue;
+}
+
+#endif // WITH_EDITOR
+
+void UMaterialFunction::LinkIntoCaller(const TArray<FFunctionExpressionInput>& CallerInputs)
 {
 	// Go through all the function's input expressions and hook their inputs up to the corresponding expression in the material being compiled.
 	for (int32 ExpressionIndex = 0; ExpressionIndex < FunctionExpressions.Num(); ExpressionIndex++)
@@ -7842,7 +7860,7 @@ int32 UMaterialFunction::Compile(class FMaterialCompiler* Compiler, const FFunct
 			InputExpression->EffectivePreviewDuringCompile = InputExpression->Preview;
 
 			// Get the FExpressionInput which stores information about who this input node should be linked to in order to compile
-			const FExpressionInput* MatchingInput = FindInputByExpression(InputExpression, Inputs);
+			const FExpressionInput* MatchingInput = FindInputByExpression(InputExpression, CallerInputs);
 
 			if (MatchingInput 
 				// Only change the connection if the input has a valid connection,
@@ -7860,18 +7878,10 @@ int32 UMaterialFunction::Compile(class FMaterialCompiler* Compiler, const FFunct
 			}
 		}
 	}
+}
 
-	int32 ReturnValue = INDEX_NONE;
-	if (Output.ExpressionOutput->A.Expression)
-	{
-		// Compile the given function output
-		ReturnValue = Output.ExpressionOutput->A.Compile(Compiler);
-	}
-	else
-	{
-		ReturnValue = Compiler->Errorf(TEXT("Missing function output connection '%s'"), *Output.ExpressionOutput->OutputName);
-	}
-
+void UMaterialFunction::UnlinkFromCaller()
+{
 	for (int32 ExpressionIndex = 0; ExpressionIndex < FunctionExpressions.Num(); ExpressionIndex++)
 	{
 		UMaterialExpression* CurrentExpression = FunctionExpressions[ExpressionIndex];
@@ -7881,12 +7891,11 @@ int32 UMaterialFunction::Compile(class FMaterialCompiler* Compiler, const FFunct
 		{
 			// Restore the default value
 			InputExpression->bCompilingFunctionPreview = true;
+			// Clear the reference to make stale accesses obvious
+			InputExpression->EffectivePreviewDuringCompile.Expression = NULL;
 		}
 	}
-
-	return ReturnValue;
 }
-#endif // WITH_EDITOR
 
 bool UMaterialFunction::IsDependent(UMaterialFunction* OtherFunction)
 {
@@ -8083,16 +8092,22 @@ int32 UMaterialExpressionMaterialFunctionCall::Compile(class FMaterialCompiler*
 		return Compiler->Errorf(TEXT("Invalid function output"));
 	}
 
+	// Link the function's inputs into the caller graph before entering
+	MaterialFunction->LinkIntoCaller(FunctionInputs);
+
 	// Tell the compiler that we are entering a function
 	Compiler->PushFunction(FMaterialFunctionCompileState(this));
 
 	// Compile the requested output
-	const int32 ReturnValue = MaterialFunction->Compile(Compiler, FunctionOutputs[OutputIndex], FunctionInputs);
+	const int32 ReturnValue = MaterialFunction->Compile(Compiler, FunctionOutputs[OutputIndex]);
 
 	// Tell the compiler that we are leaving a function
 	const FMaterialFunctionCompileState CompileState = Compiler->PopFunction();
 	check(CompileState.ExpressionStack.Num() == 0);
 
+	// Restore the function since we are leaving it
+	MaterialFunction->UnlinkFromCaller();
+
 	return ReturnValue;
 }
 
@@ -8703,9 +8718,17 @@ int32 UMaterialExpressionFunctionInput::Compile(class FMaterialCompiler* Compile
 			// Tell the compiler that we are leaving the function
 			const FMaterialFunctionCompileState FunctionState = Compiler->PopFunction();
 
+			// Backup EffectivePreviewDuringCompile which will be modified by UnlinkFromCaller and LinkIntoCaller of any potential chained function calls to the same function
+			FExpressionInput LocalPreviewDuringCompile = EffectivePreviewDuringCompile;
+
+			// Restore the function since we are leaving it
+			FunctionState.FunctionCall->MaterialFunction->UnlinkFromCaller();
+
 			// Compile the function input
-			// Warning: EffectivePreviewDuringCompile will change during this call if the same function is called again
-			ExpressionResult = EffectivePreviewDuringCompile.Compile(Compiler);
+			ExpressionResult = LocalPreviewDuringCompile.Compile(Compiler);
+
+			// Link the function's inputs into the caller graph before entering
+			FunctionState.FunctionCall->MaterialFunction->LinkIntoCaller(FunctionState.FunctionCall->FunctionInputs);
 
 			// Tell the compiler that we are re-entering the function
 			Compiler->PushFunction(FunctionState);
