Get appropriate input expression when translating texture object function inputs.

diff --git a/Engine/Shaders/ShaderVersion.usf b/Engine/Shaders/ShaderVersion.usf
--- a/Engine/Shaders/ShaderVersion.usf
+++ b/Engine/Shaders/ShaderVersion.usf
@@ -2,5 +2,5 @@
 // Each time the console command is executed it generates a new GUID. As this file is included
 // in common.usf (which should be included in any shader) it allows to invalidate the shader DDC.
 // 
-// GUID = FCe205F2167C4A0EB5DDC870A2Bb74E7
+// GUID = 763D97AE101C4AF88A7222EF58F2A920
  
\ No newline at end of file
diff --git a/Engine/Source/Runtime/Engine/Classes/Materials/MaterialExpressionFunctionInput.h b/Engine/Source/Runtime/Engine/Classes/Materials/MaterialExpressionFunctionInput.h
--- a/Engine/Source/Runtime/Engine/Classes/Materials/MaterialExpressionFunctionInput.h
+++ b/Engine/Source/Runtime/Engine/Classes/Materials/MaterialExpressionFunctionInput.h
@@ -71,6 +71,12 @@ class UMaterialExpressionFunctionInput : public UMaterialExpression
 	/** The Preview input to use during compilation from another material, when bCompilingFunctionPreview is false. */
 	FExpressionInput EffectivePreviewDuringCompile;
 
+	/** Returns the appropriate preview expression when compiling a function or material preview. */
+	UMaterialExpression* GetEffectivePreviewExpression()
+	{
+		return bCompilingFunctionPreview ? Preview.Expression : EffectivePreviewDuringCompile.Expression;
+	}
+
 	//~ Begin UObject Interface.
 	virtual void PostLoad() override;
 	virtual void PostDuplicate(bool bDuplicateForPIE) override;
diff --git a/Engine/Source/Runtime/Engine/Private/Materials/MaterialExpressions.cpp b/Engine/Source/Runtime/Engine/Private/Materials/MaterialExpressions.cpp
--- a/Engine/Source/Runtime/Engine/Private/Materials/MaterialExpressions.cpp
+++ b/Engine/Source/Runtime/Engine/Private/Materials/MaterialExpressions.cpp
@@ -1372,16 +1372,24 @@ int32 UMaterialExpressionTextureSample::Compile(class FMaterialCompiler* Compile
 		{
 			UMaterialExpression* InputExpression = TextureObject.Expression;
 			UMaterialExpressionFunctionInput* FunctionInput = Cast<UMaterialExpressionFunctionInput>(InputExpression);
-			if (FunctionInput && FunctionInput->Preview.Expression)
-			{
+			if (FunctionInput)
+			{	
 				UMaterialExpressionFunctionInput* NestedFunctionInput = FunctionInput;
 
 				// Walk the input chain to find the last node in the chain
-				while (NestedFunctionInput->Preview.Expression && NestedFunctionInput->Preview.Expression->IsA(UMaterialExpressionFunctionInput::StaticClass()))
+				while (true)
 				{
-					NestedFunctionInput = CastChecked<UMaterialExpressionFunctionInput>(NestedFunctionInput->Preview.Expression);
+					UMaterialExpression* PreviewExpression = NestedFunctionInput->GetEffectivePreviewExpression();
+					if (PreviewExpression && PreviewExpression->IsA(UMaterialExpressionFunctionInput::StaticClass()))
+					{
+						NestedFunctionInput = CastChecked<UMaterialExpressionFunctionInput>(PreviewExpression);
+					}
+					else
+					{
+						break;
+					}
 				}
-				InputExpression = NestedFunctionInput->Preview.Expression;
+				InputExpression = NestedFunctionInput->GetEffectivePreviewExpression();
 			}
 
 			UMaterialExpressionTextureObject* TextureObjectExpression = Cast<UMaterialExpressionTextureObject>(InputExpression);
