Development compile error fix.

diff --git a/Engine/Source/Runtime/Engine/Classes/Materials/MaterialExpressionFunctionInput.h b/Engine/Source/Runtime/Engine/Classes/Materials/MaterialExpressionFunctionInput.h
--- a/Engine/Source/Runtime/Engine/Classes/Materials/MaterialExpressionFunctionInput.h
+++ b/Engine/Source/Runtime/Engine/Classes/Materials/MaterialExpressionFunctionInput.h
@@ -71,11 +71,13 @@ class UMaterialExpressionFunctionInput : public UMaterialExpression
 	/** The Preview input to use during compilation from another material, when bCompilingFunctionPreview is false. */
 	FExpressionInput EffectivePreviewDuringCompile;
 
+#if WITH_EDITOR
 	/** Returns the appropriate preview expression when compiling a function or material preview. */
 	UMaterialExpression* GetEffectivePreviewExpression()
 	{
 		return bCompilingFunctionPreview ? Preview.Expression : EffectivePreviewDuringCompile.Expression;
 	}
+#endif
 
 	//~ Begin UObject Interface.
 	virtual void PostLoad() override;
