Legacy MapBuildDataRegistry objects are no longer public so they don't get shown in the content browser, prevents users trying to move them.

diff --git a/Engine/Source/Runtime/Engine/Private/Level.cpp b/Engine/Source/Runtime/Engine/Private/Level.cpp
--- a/Engine/Source/Runtime/Engine/Private/Level.cpp
+++ b/Engine/Source/Runtime/Engine/Private/Level.cpp
@@ -1702,6 +1702,7 @@ UMapBuildDataRegistry* ULevel::CreateMapBuildDataRegistry() const
 	// PKG_ContainsMapData required so FEditorFileUtils::GetDirtyContentPackages can treat this as a map package
 	BuiltDataPackage->SetPackageFlags(PKG_ContainsMapData);
 	FName ShortPackageName = FPackageName::GetShortFName(*PackageName);
+	// Top level UObjects have to have both RF_Standalone and RF_Public to be saved into packages
 	UMapBuildDataRegistry* NewMapBuildData = NewObject<UMapBuildDataRegistry>(BuiltDataPackage, ShortPackageName, RF_Standalone | RF_Public);
 	return NewMapBuildData;
 }
@@ -1710,7 +1711,8 @@ UMapBuildDataRegistry* ULevel::GetOrCreateMapBuildData()
 {
 	if (!MapBuildData 
 		// If MapBuildData is in the level package we need to create a new one, see CreateRegistryForLegacyMap
-		|| MapBuildData->GetOutermost() == GetOutermost())
+		|| MapBuildData->GetOutermost() == GetOutermost()
+		|| !MapBuildData->HasAllFlags(RF_Public | RF_Standalone))
 	{
 		if (MapBuildData)
 		{
diff --git a/Engine/Source/Runtime/Engine/Private/MapBuildData.cpp b/Engine/Source/Runtime/Engine/Private/MapBuildData.cpp
--- a/Engine/Source/Runtime/Engine/Private/MapBuildData.cpp
+++ b/Engine/Source/Runtime/Engine/Private/MapBuildData.cpp
@@ -70,7 +70,7 @@ UMapBuildDataRegistry* CreateRegistryForLegacyMap(ULevel* Level)
 	static FName RegistryName(TEXT("MapBuildDataRegistry"));
 	// Create a new registry for legacy map build data, but put it in the level's package.  
 	// This avoids creating a new package during cooking which the cooker won't know about.
-	Level->MapBuildData = NewObject<UMapBuildDataRegistry>(Level->GetOutermost(), RegistryName, RF_Standalone | RF_Public);
+	Level->MapBuildData = NewObject<UMapBuildDataRegistry>(Level->GetOutermost(), RegistryName, RF_NoFlags);
 	return Level->MapBuildData;
 }
 
