[Integrate] Legacy lightmaps are renamed with the world that uses them.  Fixes 'graph linked to external object' when renaming a map in the content browser.

diff --git a/Engine/Source/Runtime/Engine/Private/World.cpp b/Engine/Source/Runtime/Engine/Private/World.cpp
--- a/Engine/Source/Runtime/Engine/Private/World.cpp
+++ b/Engine/Source/Runtime/Engine/Private/World.cpp
@@ -437,8 +437,33 @@ bool UWorld::Rename(const TCHAR* InName, UObject* NewOuter, ERenameFlags Flags)
 {
 	check(PersistentLevel);
 
+	// Rename LightMaps and ShadowMaps to the new location. Keep the old name, since they are not named after the world.
+	TArray<UTexture2D*> LightMapsAndShadowMaps;
+	GetLightMapsAndShadowMaps(PersistentLevel, LightMapsAndShadowMaps);
+
 	UPackage* OldPackage = GetOutermost();
 
+	for (auto* Tex : LightMapsAndShadowMaps)
+	{
+		if ( Tex && Tex->GetOutermost() == OldPackage)
+		{
+			if (!Tex->Rename(*Tex->GetName(), NewOuter, Flags))
+			{
+				return false;
+			}
+		}
+	}
+
+	if (PersistentLevel
+		&& PersistentLevel->MapBuildData 
+		&& PersistentLevel->MapBuildData->GetOutermost() == OldPackage)
+	{
+		if (!PersistentLevel->MapBuildData->Rename(*PersistentLevel->MapBuildData->GetName(), NewOuter, Flags))
+		{
+			return false;
+		}
+	}
+
 	bool bShouldFail = false;
 	FWorldDelegates::OnPreWorldRename.Broadcast(this, InName, NewOuter, Flags, bShouldFail);
 
