Fixed movable lights getting assigned a shadowmap channel

diff --git a/Engine/Source/Runtime/Engine/Private/Components/LightComponent.cpp b/Engine/Source/Runtime/Engine/Private/Components/LightComponent.cpp
--- a/Engine/Source/Runtime/Engine/Private/Components/LightComponent.cpp
+++ b/Engine/Source/Runtime/Engine/Private/Components/LightComponent.cpp
@@ -209,7 +209,7 @@ FLightSceneProxy::FLightSceneProxy(const ULightComponent* InLightComponent)
 {
 	const FLightComponentMapBuildData* MapBuildData = InLightComponent->GetLightComponentMapBuildData();
 	
-	if (MapBuildData)
+	if (MapBuildData && bStaticShadowing && !bStaticLighting)
 	{
 		ShadowMapChannel = MapBuildData->ShadowMapChannel;
 	}
@@ -987,7 +987,7 @@ void ULightComponent::PropagateLightingScenarioChange()
 
 bool ULightComponent::IsPrecomputedLightingValid() const
 {
-	return GetLightComponentMapBuildData() != NULL;
+	return GetLightComponentMapBuildData() != NULL && HasStaticShadowing();
 }
 
 int32 ULightComponent::GetNumMaterials() const
