Fix some issues with manipulating sublevels in the editor.
Duplicates CL 3160045 from Release-4.14.

diff --git a/Engine/Source/Runtime/Engine/Classes/Engine/World.h b/Engine/Source/Runtime/Engine/Classes/Engine/World.h
--- a/Engine/Source/Runtime/Engine/Classes/Engine/World.h
+++ b/Engine/Source/Runtime/Engine/Classes/Engine/World.h
@@ -792,7 +792,7 @@ private:
 	TArray<class ULevel*>						Levels;
 
 	/** Array of level collections currently in this world. */
-	UPROPERTY(Transient)
+	UPROPERTY(Transient, NonTransactional)
 	TArray<FLevelCollection>					LevelCollections;
 
 	/** Pointer to the level collection that's currently ticking. */
diff --git a/Engine/Source/Runtime/Engine/Private/World.cpp b/Engine/Source/Runtime/Engine/Private/World.cpp
--- a/Engine/Source/Runtime/Engine/Private/World.cpp
+++ b/Engine/Source/Runtime/Engine/Private/World.cpp
@@ -142,8 +142,11 @@ FLevelCollection::~FLevelCollection()
 {
 	for (ULevel* Level : Levels)
 	{
-		check(Level->GetCachedLevelCollection() == this);
-		Level->SetCachedLevelCollection(nullptr);
+		if (Level)
+		{
+			check(Level->GetCachedLevelCollection() == this);
+			Level->SetCachedLevelCollection(nullptr);
+		}
 	}
 }
 
@@ -158,8 +161,11 @@ FLevelCollection::FLevelCollection(FLevelCollection&& Other)
 {
 	for (ULevel* Level : Levels)
 	{
-		check(Level->GetCachedLevelCollection() == &Other);
-		Level->SetCachedLevelCollection(this);
+		if (Level)
+		{
+			check(Level->GetCachedLevelCollection() == &Other);
+			Level->SetCachedLevelCollection(this);
+		}
 	}
 }
 
@@ -177,8 +183,11 @@ FLevelCollection& FLevelCollection::operator=(FLevelCollection&& Other)
 
 		for (ULevel* Level : Levels)
 		{
-			check(Level->GetCachedLevelCollection() == &Other);
-			Level->SetCachedLevelCollection(this);
+			if (Level)
+			{
+				check(Level->GetCachedLevelCollection() == &Other);
+				Level->SetCachedLevelCollection(this);
+			}
 		}
 	}
 
@@ -6006,7 +6015,10 @@ void UWorld::SetActiveLevelCollection(const FLevelCollection* InCollection)
 	ActiveLevelCollection = InCollection;
 
 	PersistentLevel = InCollection->GetPersistentLevel();
-	SetCurrentLevel(InCollection->GetPersistentLevel());
+	if (IsGameWorld())
+	{
+		SetCurrentLevel(InCollection->GetPersistentLevel());
+	}
 	GameState = InCollection->GetGameState();
 	NetDriver = InCollection->GetNetDriver();
 	DemoNetDriver = InCollection->GetDemoNetDriver();
