Fix assert when adding a new sublevel.
Duplicates CL 3160518 from Release-4.14.

diff --git a/Engine/Source/Runtime/Engine/Private/LevelStreaming.cpp b/Engine/Source/Runtime/Engine/Private/LevelStreaming.cpp
--- a/Engine/Source/Runtime/Engine/Private/LevelStreaming.cpp
+++ b/Engine/Source/Runtime/Engine/Private/LevelStreaming.cpp
@@ -317,6 +317,12 @@ void ULevelStreaming::SetLoadedLevel(class ULevel* Level)
 
 	FLevelCollection& LC = GetWorld()->FindOrAddCollectionByType(CollectionType);
 	LC.RemoveLevel(PendingUnloadLevel);
+
+	// Remove the loaded level from its current collection, if any.
+	if (LoadedLevel && LoadedLevel->GetCachedLevelCollection())
+	{
+		LoadedLevel->GetCachedLevelCollection()->RemoveLevel(LoadedLevel);
+	}
 	LC.AddLevel(LoadedLevel);
 }
 
