Integrate CL 3159903 to //UE4/Main/   (resolve blocker for Min Oh)

diff --git a/Engine/Source/Runtime/Engine/Private/Components/ActorComponent.cpp b/Engine/Source/Runtime/Engine/Private/Components/ActorComponent.cpp
--- a/Engine/Source/Runtime/Engine/Private/Components/ActorComponent.cpp
+++ b/Engine/Source/Runtime/Engine/Private/Components/ActorComponent.cpp
@@ -667,7 +667,7 @@ void UActorComponent::OnRegister()
 	if (bAutoActivate)
 	{
 		AActor* Owner = GetOwner();
-		if (Owner == nullptr || Owner->IsActorInitialized())
+		if (!WorldPrivate->IsGameWorld() || Owner == nullptr || Owner->IsActorInitialized())
 		{
 			Activate(true);
 		}
