Merging CL 3162934 to //UE4/Main
UE-37044 Fix crash when loading map that has null actors in the actor list

diff --git a/Engine/Source/Runtime/Engine/Private/MapBuildData.cpp b/Engine/Source/Runtime/Engine/Private/MapBuildData.cpp
--- a/Engine/Source/Runtime/Engine/Private/MapBuildData.cpp
+++ b/Engine/Source/Runtime/Engine/Private/MapBuildData.cpp
@@ -96,6 +96,11 @@ void ULevel::HandleLegacyMapBuildData()
 
 		for (int32 ActorIndex = 0; ActorIndex < Actors.Num(); ActorIndex++)
 		{
+			if (!Actors[ActorIndex])
+			{
+				continue;
+			}
+
 			TInlineComponentArray<UActorComponent*> Components;
 			Actors[ActorIndex]->GetComponents(Components);
 
