Fix Lixux compile

diff --git a/Engine/Source/Runtime/Engine/Private/Materials/MaterialInterface.cpp b/Engine/Source/Runtime/Engine/Private/Materials/MaterialInterface.cpp
--- a/Engine/Source/Runtime/Engine/Private/Materials/MaterialInterface.cpp
+++ b/Engine/Source/Runtime/Engine/Private/Materials/MaterialInterface.cpp
@@ -438,7 +438,7 @@ namespace Algo
 bool FMaterialTextureInfo::IsValid(bool bCheckTextureIndex) const
 { 
 #if WITH_EDITORONLY_DATA
-	if (bCheckTextureIndex && TextureIndex < 0 || TextureIndex >= TEXSTREAM_MAX_NUM_TEXTURES_PER_MATERIAL)
+	if (bCheckTextureIndex && (TextureIndex < 0 || TextureIndex >= TEXSTREAM_MAX_NUM_TEXTURES_PER_MATERIAL))
 	{
 		return false;
 	}
