Don't clear level collections in UWorld::CleanupWorld unless bCleanupResources is true.
Duplicates CL 3175695 from Release-4.14.

diff --git a/Engine/Source/Runtime/Engine/Private/World.cpp b/Engine/Source/Runtime/Engine/Private/World.cpp
--- a/Engine/Source/Runtime/Engine/Private/World.cpp
+++ b/Engine/Source/Runtime/Engine/Private/World.cpp
@@ -3428,7 +3428,10 @@ void UWorld::CleanupWorld(bool bSessionEnded, bool bCleanupResources, UWorld* Ne
 		}
 	}
 
-	LevelCollections.Empty();
+	if (bCleanupResources)
+	{
+		LevelCollections.Empty();
+	}
 }
 
 UGameViewportClient* UWorld::GetGameViewport() const
