Switched shader compilation warnings to display messages.
#test Cook shootergame

diff --git a/Engine/Source/Runtime/Engine/Private/Materials/MaterialShader.cpp b/Engine/Source/Runtime/Engine/Private/Materials/MaterialShader.cpp
--- a/Engine/Source/Runtime/Engine/Private/Materials/MaterialShader.cpp
+++ b/Engine/Source/Runtime/Engine/Private/Materials/MaterialShader.cpp
@@ -1206,7 +1206,7 @@ void FMaterialShaderMap::Compile(
 			NewCorrespondingMaterials.Add(Material);
 			ShaderMapsBeingCompiled.Add(this, NewCorrespondingMaterials);
 #if DEBUG_INFINITESHADERCOMPILE
-			UE_LOG(LogTemp, Warning, TEXT("Added material ShaderMap 0x%08X%08X with Material 0x%08X%08X to ShaderMapsBeingCompiled"), (int)((int64)(this) >> 32), (int)((int64)(this)), (int)((int64)(Material) >> 32), (int)((int64)(Material)));
+			UE_LOG(LogTemp, Display, TEXT("Added material ShaderMap 0x%08X%08X with Material 0x%08X%08X to ShaderMapsBeingCompiled"), (int)((int64)(this) >> 32), (int)((int64)(this)), (int)((int64)(Material) >> 32), (int)((int64)(Material)));
 #endif  
 			// Setup the material compilation environment.
 			Material->SetupMaterialEnvironment(InPlatform, InMaterialCompilationOutput.UniformExpressionSet, *MaterialEnvironment);
@@ -1650,7 +1650,7 @@ bool FMaterialShaderMap::TryToAddToExistingCompilationTask(FMaterial* Material)
 	{
 		CorrespondingMaterials->AddUnique(Material);
 #if DEBUG_INFINITESHADERCOMPILE
-		UE_LOG(LogTemp, Warning, TEXT("Added shader map 0x%08X%08X from material 0x%08X%08X"), (int)((int64)(this) >> 32), (int)((int64)(this)), (int)((int64)(Material) >> 32), (int)((int64)(Material)));
+		UE_LOG(LogTemp, Display, TEXT("Added shader map 0x%08X%08X from material 0x%08X%08X"), (int)((int64)(this) >> 32), (int)((int64)(this)), (int)((int64)(Material) >> 32), (int)((int64)(Material)));
 #endif
 		return true;
 	}
@@ -2181,7 +2181,7 @@ void FMaterialShaderMap::RemovePendingMaterial(FMaterial* Material)
 #if DEBUG_INFINITESHADERCOMPILE
 		if ( Result )
 		{
-			UE_LOG(LogTemp, Warning, TEXT("Removed shader map 0x%08X%08X from material 0x%08X%08X"), (int)((int64)(It.Key().GetReference()) >> 32), (int)((int64)(It.Key().GetReference())), (int)((int64)(Material) >> 32), (int)((int64)(Material)));
+			UE_LOG(LogTemp, Display, TEXT("Removed shader map 0x%08X%08X from material 0x%08X%08X"), (int)((int64)(It.Key().GetReference()) >> 32), (int)((int64)(It.Key().GetReference())), (int)((int64)(Material) >> 32), (int)((int64)(Material)));
 		}
 #endif
 	}
diff --git a/Engine/Source/Runtime/Engine/Private/Materials/MaterialShared.cpp b/Engine/Source/Runtime/Engine/Private/Materials/MaterialShared.cpp
--- a/Engine/Source/Runtime/Engine/Private/Materials/MaterialShared.cpp
+++ b/Engine/Source/Runtime/Engine/Private/Materials/MaterialShared.cpp
@@ -1571,7 +1571,7 @@ bool FMaterial::CacheShaders(const FMaterialShaderMapId& ShaderMapId, EShaderPla
 	{
 		//FMaterialShaderMap::ShaderMapsBeingCompiled.Find(GameThreadShaderMap);
 #if DEBUG_INFINITESHADERCOMPILE
-		UE_LOG(LogTemp, Log, TEXT("Found exisitng compiling shader for material %s, linking to other GameThreadShaderMap 0x%08X%08X"), *GetFriendlyName(), (int)((int64)(GameThreadShaderMap.GetReference()) >> 32), (int)((int64)(GameThreadShaderMap.GetReference())) );
+		UE_LOG(LogTemp, Display, TEXT("Found exisitng compiling shader for material %s, linking to other GameThreadShaderMap 0x%08X%08X"), *GetFriendlyName(), (int)((int64)(GameThreadShaderMap.GetReference()) >> 32), (int)((int64)(GameThreadShaderMap.GetReference())) );
 #endif
 		OutstandingCompileShaderMapIds.AddUnique(GameThreadShaderMap->GetCompilingId());
 		// Reset the shader map so the default material will be used until the compile finishes.
@@ -1692,7 +1692,7 @@ bool FMaterial::BeginCompileShaderMap(
 		else
 		{
 #if DEBUG_INFINITESHADERCOMPILE
-			UE_LOG(LogTemp, Log, TEXT("Kicking of shader compilation for %s, GameThreadShaderMap 0x%08X%08X"), *GetFriendlyName(), (int)((int64)(NewShaderMap.GetReference()) >> 32), (int)((int64)(NewShaderMap.GetReference())));
+			UE_LOG(LogTemp, Display, TEXT("Kicking of shader compilation for %s, GameThreadShaderMap 0x%08X%08X"), *GetFriendlyName(), (int)((int64)(NewShaderMap.GetReference()) >> 32), (int)((int64)(NewShaderMap.GetReference())));
 #endif
 			OutstandingCompileShaderMapIds.AddUnique( NewShaderMap->GetCompilingId() );
 			// Async compile, use NULL so that rendering will fall back to the default material.
diff --git a/Engine/Source/Runtime/Engine/Private/ShaderCompiler/ShaderCompiler.cpp b/Engine/Source/Runtime/Engine/Private/ShaderCompiler/ShaderCompiler.cpp
--- a/Engine/Source/Runtime/Engine/Private/ShaderCompiler/ShaderCompiler.cpp
+++ b/Engine/Source/Runtime/Engine/Private/ShaderCompiler/ShaderCompiler.cpp
@@ -1724,16 +1724,16 @@ void FShaderCompilingManager::ProcessCompiledShaderMaps(
 				LocalShaderMapReferences.Add(ShaderMap);
 				FMaterialShaderMap::ShaderMapsBeingCompiled.Remove(ShaderMap);
 #if DEBUG_INFINITESHADERCOMPILE
-				UE_LOG(LogTemp, Warning, TEXT("Finished compile of shader map 0x%08X%08X"), (int)((int64)(ShaderMap.GetReference()) >> 32), (int)((int64)(ShaderMap.GetReference())));
+				UE_LOG(LogTemp, Display, TEXT("Finished compile of shader map 0x%08X%08X"), (int)((int64)(ShaderMap.GetReference()) >> 32), (int)((int64)(ShaderMap.GetReference())));
 #endif
 				for (int32 MaterialIndex = 0; MaterialIndex < MaterialsArray.Num(); MaterialIndex++)
 				{
 					FMaterial* Material = MaterialsArray[MaterialIndex];
 					FMaterialShaderMap* CompletedShaderMap = ShaderMap;
 #if DEBUG_INFINITESHADERCOMPILE
-					UE_LOG(LogTemp, Warning, TEXT("Shader map %s complete, GameThreadShaderMap 0x%08X%08X, marking material %s as finished"), *ShaderMap->GetFriendlyName(), (int)((int64)(ShaderMap.GetReference()) >> 32), (int)((int64)(ShaderMap.GetReference())), *Material->GetFriendlyName());
+					UE_LOG(LogTemp, Display, TEXT("Shader map %s complete, GameThreadShaderMap 0x%08X%08X, marking material %s as finished"), *ShaderMap->GetFriendlyName(), (int)((int64)(ShaderMap.GetReference()) >> 32), (int)((int64)(ShaderMap.GetReference())), *Material->GetFriendlyName());
 
-					UE_LOG(LogTemp, Warning, TEXT("Marking material as finished 0x%08X%08X"), (int)((int64)(Material) >> 32), (int)((int64)(Material)));
+					UE_LOG(LogTemp, Display, TEXT("Marking material as finished 0x%08X%08X"), (int)((int64)(Material) >> 32), (int)((int64)(Material)));
 #endif
 					Material->RemoveOutstandingCompileId(ShaderMap->CompilingId);
 
