Copying //UE4/Dev-Framework to //UE4/Dev-Main (Source: //UE4/Dev-Framework @ 3182037)

diff --git a/Engine/Source/Runtime/AssetRegistry/Private/PackageReader.cpp b/Engine/Source/Runtime/AssetRegistry/Private/PackageReader.cpp
--- a/Engine/Source/Runtime/AssetRegistry/Private/PackageReader.cpp
+++ b/Engine/Source/Runtime/AssetRegistry/Private/PackageReader.cpp
@@ -352,6 +352,7 @@ bool FPackageReader::ReadDependencyData (FPackageDependencyData& OutDependencyDa
 	SerializeNameMap();
 	SerializeImportMap(OutDependencyData.ImportMap);
 	SerializeStringAssetReferencesMap(OutDependencyData.StringAssetReferencesMap);
+	SerializeSearchableNamesMap(OutDependencyData);
 
 	return true;
 }
@@ -401,7 +402,7 @@ void FPackageReader::SerializeExportMap(TArray<FObjectExport>& OutExportMap)
 
 void FPackageReader::SerializeStringAssetReferencesMap(TArray<FString>& OutStringAssetReferencesMap)
 {
-	if (UE4Ver() >= VER_UE4_ADD_STRING_ASSET_REFERENCES_MAP && PackageFileSummary.StringAssetReferencesCount > 0)
+	if (UE4Ver() >= VER_UE4_ADD_STRING_ASSET_REFERENCES_MAP && PackageFileSummary.StringAssetReferencesOffset > 0 && PackageFileSummary.StringAssetReferencesCount > 0)
 	{
 		Seek(PackageFileSummary.StringAssetReferencesOffset);
 
@@ -442,6 +443,16 @@ void FPackageReader::SerializeStringAssetReferencesMap(TArray<FString>& OutStrin
 	}
 }
 
+void FPackageReader::SerializeSearchableNamesMap(FPackageDependencyData& OutDependencyData)
+{
+	if (UE4Ver() >= VER_UE4_ADDED_SEARCHABLE_NAMES && PackageFileSummary.SearchableNamesOffset > 0)
+	{
+		Seek(PackageFileSummary.SearchableNamesOffset);
+
+		OutDependencyData.SerializeSearchableNamesMap(*this);
+	}
+}
+
 void FPackageReader::Serialize( void* V, int64 Length )
 {
 	check(Loader);
diff --git a/Engine/Source/Runtime/AssetRegistry/Private/PackageReader.h b/Engine/Source/Runtime/AssetRegistry/Private/PackageReader.h
--- a/Engine/Source/Runtime/AssetRegistry/Private/PackageReader.h
+++ b/Engine/Source/Runtime/AssetRegistry/Private/PackageReader.h
@@ -40,6 +40,7 @@ public:
 	void SerializeImportMap(TArray<FObjectImport>& OutImportMap);
 	void SerializeExportMap(TArray<FObjectExport>& OutExportMap);
 	void SerializeStringAssetReferencesMap(TArray<FString>& OutStringAssetReferencesMap);
+	void SerializeSearchableNamesMap(FPackageDependencyData& OutDependencyData);
 
 	/** Returns flags the asset package was saved with */
 	uint32 GetPackageFlags() const;
diff --git a/Engine/Source/Runtime/AssetRegistry/Public/AssetData.h b/Engine/Source/Runtime/AssetRegistry/Public/AssetData.h
--- a/Engine/Source/Runtime/AssetRegistry/Public/AssetData.h
+++ b/Engine/Source/Runtime/AssetRegistry/Public/AssetData.h
@@ -463,3 +463,105 @@ inline FName FAssetData::GetTagValueRef<FName>(const FName InTagName) const
 	GetTagValueNameImpl(InTagName, TmpValue);
 	return TmpValue;
 }
+
+/** A structure defining a thing that can be reference bt something else in the asset registry */
+struct FAssetIdentifier
+{
+	/** The name of the package that is depended on, this is always set */
+	FName PackageName;
+	/** Specific object within a package. If empty, assumed to be the default asset */
+	FName ObjectName;
+	/** Name of specific value being referenced, if ObjectName specifies a type such as a UStruct */
+	FName ValueName;
+
+	/** Can be implicitly constructed from just the package name */
+	FAssetIdentifier(FName InPackageName, FName InObjectName = NAME_None, FName InValueName = NAME_None)
+		: PackageName(InPackageName), ObjectName(InObjectName), ValueName(InValueName)
+	{}
+
+	FAssetIdentifier(UObject* SourceObject, FName InValueName)
+	{
+		if (SourceObject)
+		{
+			UPackage* Package = SourceObject->GetOutermost();
+			PackageName = Package->GetFName();
+			ObjectName = SourceObject->GetFName();
+			ValueName = InValueName;
+		}
+	}
+
+	FAssetIdentifier()
+		: PackageName(NAME_None), ObjectName(NAME_None), ValueName(NAME_None)
+	{}
+
+	/** Returns true if this represents a specific value */
+	bool IsValue() const
+	{
+		return ValueName != NAME_None;
+	}
+
+	/** Returns true if this represents a specific value */
+	bool IsObject() const
+	{
+		return ObjectName != NAME_None && !IsValue();
+	}
+
+	/** Returns string version of this identifier in Package.Object::Name format */
+	FString ToString() const
+	{
+		FString Result = PackageName.ToString();
+		if (ObjectName != NAME_None)
+		{
+			Result += TEXT(".");
+			Result += ObjectName.ToString();
+		}
+		if (ValueName != NAME_None)
+		{
+			Result += TEXT("::");
+			Result += ValueName.ToString();
+		}
+		return Result;
+	}
+
+	/** Converts from Package.Object::Name format */
+	static FAssetIdentifier FromString(const FString& String)
+	{
+		// To right of :: is value
+		FString PackageString;
+		FString ObjectString;
+		FString ValueString;
+
+		// Try to split value out
+		if (!String.Split(TEXT("::"), &PackageString, &ValueString))
+		{
+			PackageString = String;
+		}
+
+		// Try to split on first . , if it fails PackageString will stay the same
+		FString(PackageString).Split(TEXT("."), &PackageString, &ObjectString);
+
+		return FAssetIdentifier(*PackageString, *ObjectString, *ValueString);
+	}
+
+	friend inline bool operator==(const FAssetIdentifier& A, const FAssetIdentifier& B)
+	{
+		return A.PackageName == B.PackageName && A.ObjectName == B.ObjectName && A.ValueName == B.ValueName;
+	}
+
+	friend inline uint32 GetTypeHash(const FAssetIdentifier& Key)
+	{
+		uint32 Hash = 0;
+
+		// Most of the time only packagename is set
+		if (Key.ObjectName.IsNone() && Key.ValueName.IsNone())
+		{
+			return GetTypeHash(Key.PackageName);
+		}
+
+		Hash = HashCombine(Hash, GetTypeHash(Key.PackageName));
+		Hash = HashCombine(Hash, GetTypeHash(Key.ObjectName));
+		Hash = HashCombine(Hash, GetTypeHash(Key.ValueName));
+		return Hash;
+	}
+};
+
diff --git a/Engine/Source/Runtime/AssetRegistry/Public/IAssetRegistry.h b/Engine/Source/Runtime/AssetRegistry/Public/IAssetRegistry.h
--- a/Engine/Source/Runtime/AssetRegistry/Public/IAssetRegistry.h
+++ b/Engine/Source/Runtime/AssetRegistry/Public/IAssetRegistry.h
@@ -95,23 +95,41 @@ public:
 	 */
 	virtual bool GetAllAssets(TArray<FAssetData>& OutAssetData, bool bIncludeOnlyOnDiskAssets = false) const = 0;
 
+	/**
+	 * Gets a list of packages and searchable names that are referenced by the supplied package or name. (On disk references ONLY)
+	 *
+	 * @param AssetIdentifier	the name of the package/name for which to gather dependencies
+	 * @param OutDependencies	a list of things that are referenced by AssetIdentifier
+	 * @param InDependencyType	which kinds of dependency to include in the output list
+	 */
+	virtual bool GetDependencies(const FAssetIdentifier& AssetIdentifier, TArray<FAssetIdentifier>& OutDependencies, EAssetRegistryDependencyType::Type InDependencyType = EAssetRegistryDependencyType::All) const = 0;
+
 	/**
 	 * Gets a list of paths to objects that are referenced by the supplied package. (On disk references ONLY)
 	 *
 	 * @param PackageName		the name of the package for which to gather dependencies
-	 * @param OutDependencies	a list of paths to objects that are referenced by the package whose path is PackageName
+	 * @param OutDependencies	a list of packages that are referenced by the package whose path is PackageName
 	 * @param InDependencyType	which kinds of dependency to include in the output list
-	 * @param bResolveIniStringReferences Tells if the method should also resolve INI references.
 	 */
-	virtual bool GetDependencies(FName PackageName, TArray<FName>& OutDependencies, EAssetRegistryDependencyType::Type InDependencyType = EAssetRegistryDependencyType::All, bool bResolveIniStringReferences = false) const = 0;
+	virtual bool GetDependencies(FName PackageName, TArray<FName>& OutDependencies, EAssetRegistryDependencyType::Type InDependencyType = EAssetRegistryDependencyType::Packages) const = 0;
 
 	/**
-	 * Gets a list of paths to objects that reference the supplied package. (On disk references ONLY)
+	 * Gets a list of packages and searchable names that reference the supplied package or name. (On disk references ONLY)
+	 *
+	 * @param AssetIdentifier	the name of the package/name for which to gather dependencies
+	 * @param OutReferencers	a list of things that reference AssetIdentifier
+	 * @param InReferenceType	which kinds of reference to include in the output list
+	 */
+	virtual bool GetReferencers(const FAssetIdentifier& AssetIdentifier, TArray<FAssetIdentifier>& OutReferencers, EAssetRegistryDependencyType::Type InReferenceType = EAssetRegistryDependencyType::All) const = 0;
+
+	/**
+	 * Gets a list of packages that reference the supplied package. (On disk references ONLY)
 	 *
 	 * @param PackageName		the name of the package for which to gather dependencies
-	 * @param OutReferencers	a list of paths to objects that reference the package whose path is PackageName
+	 * @param OutReferencers	a list of packages that reference the package whose path is PackageName
+	 * @param InReferenceType	which kinds of reference to include in the output list
 	 */
-	virtual bool GetReferencers(FName PackageName, TArray<FName>& OutReferencers, EAssetRegistryDependencyType::Type InReferenceType = EAssetRegistryDependencyType::All) const = 0;
+	virtual bool GetReferencers(FName PackageName, TArray<FName>& OutReferencers, EAssetRegistryDependencyType::Type InReferenceType = EAssetRegistryDependencyType::Packages) const = 0;
 
 	/** Returns true if the specified ClassName's ancestors could be found. If so, OutAncestorClassNames is a list of all its ancestors */
 	virtual bool GetAncestorClassNames(FName ClassName, TArray<FName>& OutAncestorClassNames) const = 0;
@@ -239,6 +257,13 @@ public:
 	DECLARE_EVENT_OneParam( IAssetRegistry, FFileLoadProgressUpdatedEvent, const FFileLoadProgressUpdateData& /*ProgressUpdateData*/ );
 	virtual FFileLoadProgressUpdatedEvent& OnFileLoadProgressUpdated() = 0;
 
+	/** Register callback for when someone tries to edit a searchable name */
+	DECLARE_DELEGATE_RetVal_OneParam(bool, FAssetEditSearchableNameDelegate, const FAssetIdentifier&);
+	virtual FAssetEditSearchableNameDelegate& OnEditSearchableName(FName PackageName, FName ObjectName) = 0;
+
+	/** Tries to edit a searchablename, returns true if any of the callbacks handled it */
+	virtual bool EditSearchableName(const FAssetIdentifier& SearchableName) = 0;
+
 	/** Returns true if the asset registry is currently loading files and does not yet know about all assets */
 	virtual bool IsLoadingAssets() const = 0;
 
diff --git a/Engine/Source/Runtime/Core/Private/Math/UnrealMath.cpp b/Engine/Source/Runtime/Core/Private/Math/UnrealMath.cpp
--- a/Engine/Source/Runtime/Core/Private/Math/UnrealMath.cpp
+++ b/Engine/Source/Runtime/Core/Private/Math/UnrealMath.cpp
@@ -1482,218 +1482,164 @@ float FMath::PointDistToSegmentSquared(const FVector &Point, const FVector &Star
 	return (Point - ClosestPoint).SizeSquared();
 }
 
-void FMath::SegmentDistToSegmentSafe(FVector A1, FVector B1, FVector A2, FVector B2, FVector& OutP1, FVector& OutP2)
+struct SegmentDistToSegment_Solver
 {
-	// Segments
-	const	FVector	S1		= B1 - A1;
-	const	FVector	S2		= B2 - A2;
-	const	FVector	S3		= A1 - A2;
-
+	SegmentDistToSegment_Solver(const FVector& InA1, const FVector& InB1, const FVector& InA2, const FVector& InB2):
+		bLinesAreNearlyParallel(false),
+		A1(InA1),
+		A2(InA2),
+		S1(InB1 - InA1),
+		S2(InB2 - InA2),
+		S3(InA1 - InA2)
+	{
+	}
 
-	const	FVector	S1_norm		= S1.GetSafeNormal();
-	const	FVector	S2_norm		= S2.GetSafeNormal();
+	bool bLinesAreNearlyParallel;
 
-	const	float	Dot11	= S1 | S1;	// always >= 0
-	const	float	Dot22	= S2 | S2;	// always >= 0
-	const	float	Dot12	= S1 | S2;
-	const	float	Dot13	= S1 | S3;
-	const	float	Dot23	= S2 | S3;
+	const FVector& A1;
+	const FVector& A2;
 
-	const	float	Dot11_norm	= S1_norm | S1_norm;	// always >= 0
-	const	float	Dot22_norm	= S2_norm | S2_norm;	// always >= 0
-	const	float	Dot12_norm	= S1_norm | S2_norm;
+	const FVector S1;
+	const FVector S2;
+	const FVector S3;
 
-	// Numerator
-	float	N1, N2;
+	void Solve(FVector& OutP1, FVector& OutP2)
+	{
+		const float Dot11 = S1 | S1;
+		const float Dot12 = S1 | S2;
+		const float Dot13 = S1 | S3;
+		const float Dot22 = S2 | S2;
+		const float Dot23 = S2 | S3;
 
-	// Denominator
-	const	float	D		= Dot11*Dot22 - Dot12*Dot12;	// always >= 0
-	const	float	D_norm	= Dot11_norm*Dot22_norm - Dot12_norm*Dot12_norm;	// always >= 0
+		const float D = Dot11*Dot22 - Dot12*Dot12;
 
-	float	D1		= D;		// T1 = N1 / D1, default D1 = D >= 0
-	float	D2		= D;		// T2 = N2 / D2, default D2 = D >= 0
+		float D1 = D;
+		float D2 = D;
 
-	// compute the line parameters of the two closest points
-	if( D < KINDA_SMALL_NUMBER || D_norm < KINDA_SMALL_NUMBER ) 
-	{ 
-		// the lines are almost parallel
-		N1 = 0.f;	// force using point A on segment S1
-		D1 = 1.f;	// to prevent possible division by 0 later
-		N2 = Dot23;
-		D2 = Dot22;
-	}
-	else 
-	{                
-		// get the closest points on the infinite lines
-		N1 = (Dot12*Dot23 - Dot22*Dot13);
-		N2 = (Dot11*Dot23 - Dot12*Dot13);
+		float N1;
+		float N2;
 
-		if( N1 < 0.f ) 
+		if (bLinesAreNearlyParallel || D < KINDA_SMALL_NUMBER)
 		{
-			// t1 < 0.f => the s==0 edge is visible
-			N1 = 0.f;
+			// the lines are almost parallel
+			N1 = 0.f;	// force using point A on segment S1
+			D1 = 1.f;	// to prevent possible division by 0 later
 			N2 = Dot23;
 			D2 = Dot22;
 		}
-		else if( N1 > D1 ) 
+		else
 		{
-			// t1 > 1 => the t1==1 edge is visible
-			N1 = D1;
-			N2 = Dot23 + Dot12;
-			D2 = Dot22;
-		}
-	}
+			// get the closest points on the infinite lines
+			N1 = (Dot12*Dot23 - Dot22*Dot13);
+			N2 = (Dot11*Dot23 - Dot12*Dot13);
 
-	if( N2 < 0.f ) 
-	{           
-		// t2 < 0 => the t2==0 edge is visible
-		N2 = 0.f;
-
-		// recompute t1 for this edge
-		if( -Dot13 < 0.f )
-		{
-			N1 = 0.f;
-		}
-		else if( -Dot13 > Dot11 )
-		{
-			N1 = D1;
-		}
-		else 
-		{
-			N1 = -Dot13;
-			D1 = Dot11;
+			if (N1 < 0.f)
+			{
+				// t1 < 0.f => the s==0 edge is visible
+				N1 = 0.f;
+				N2 = Dot23;
+				D2 = Dot22;
+			}
+			else if (N1 > D1)
+			{
+				// t1 > 1 => the t1==1 edge is visible
+				N1 = D1;
+				N2 = Dot23 + Dot12;
+				D2 = Dot22;
+			}
 		}
-	}
-	else if( N2 > D2 ) 
-	{      
-		// t2 > 1 => the t2=1 edge is visible
-		N2 = D2;
 
-		// recompute t1 for this edge
-		if( (-Dot13 + Dot12) < 0.f )
+		if (N2 < 0.f)
 		{
-			N1 = 0.f;
-		}
-		else if( (-Dot13 + Dot12) > Dot11 )
-		{
-			N1 = D1;
+			// t2 < 0 => the t2==0 edge is visible
+			N2 = 0.f;
+
+			// recompute t1 for this edge
+			if (-Dot13 < 0.f)
+			{
+				N1 = 0.f;
+			}
+			else if (-Dot13 > Dot11)
+			{
+				N1 = D1;
+			}
+			else
+			{
+				N1 = -Dot13;
+				D1 = Dot11;
+			}
 		}
-		else 
+		else if (N2 > D2)
 		{
-			N1 = (-Dot13 + Dot12);
-			D1 = Dot11;
+			// t2 > 1 => the t2=1 edge is visible
+			N2 = D2;
+
+			// recompute t1 for this edge
+			if ((-Dot13 + Dot12) < 0.f)
+			{
+				N1 = 0.f;
+			}
+			else if ((-Dot13 + Dot12) > Dot11)
+			{
+				N1 = D1;
+			}
+			else
+			{
+				N1 = (-Dot13 + Dot12);
+				D1 = Dot11;
+			}
 		}
-	}
 
-	// finally do the division to get the points' location
-	const float T1 = (FMath::Abs(N1) < KINDA_SMALL_NUMBER ? 0.f : N1 / D1);
-	const float T2 = (FMath::Abs(N2) < KINDA_SMALL_NUMBER ? 0.f : N2 / D2);
+		// finally do the division to get the points' location
+		const float T1 = (FMath::Abs(N1) < KINDA_SMALL_NUMBER ? 0.f : N1 / D1);
+		const float T2 = (FMath::Abs(N2) < KINDA_SMALL_NUMBER ? 0.f : N2 / D2);
 
-	// return the closest points
-	OutP1 = A1 + T1 * S1;
-	OutP2 = A2 + T2 * S2;
-}
+		// return the closest points
+		OutP1 = A1 + T1 * S1;
+		OutP2 = A2 + T2 * S2;
+	}
+};
 
-void FMath::SegmentDistToSegment(FVector A1, FVector B1, FVector A2, FVector B2, FVector& OutP1, FVector& OutP2)
+void FMath::SegmentDistToSegmentSafe(FVector A1, FVector B1, FVector A2, FVector B2, FVector& OutP1, FVector& OutP2)
 {
-	// Segments
-	const	FVector	S1		= B1 - A1;
-	const	FVector	S2		= B2 - A2;
-	const	FVector	S3		= A1 - A2;
-
-	const	float	Dot11	= S1 | S1;	// always >= 0
-	const	float	Dot22	= S2 | S2;	// always >= 0
-	const	float	Dot12	= S1 | S2;
-	const	float	Dot13	= S1 | S3;
-	const	float	Dot23	= S2 | S3;
-
-	// Numerator
-			float	N1, N2;
-
-	// Denominator
-	const	float	D		= Dot11*Dot22 - Dot12*Dot12;	// always >= 0
-			float	D1		= D;		// T1 = N1 / D1, default D1 = D >= 0
-			float	D2		= D;		// T2 = N2 / D2, default D2 = D >= 0
-
-	// compute the line parameters of the two closest points
-	if( D < KINDA_SMALL_NUMBER ) 
-	{ 
-		// the lines are almost parallel
-		N1 = 0.f;	// force using point A on segment S1
-		D1 = 1.f;	// to prevent possible division by 0 later
-		N2 = Dot23;
-		D2 = Dot22;
-	}
-	else 
-	{                
-		// get the closest points on the infinite lines
-		N1 = (Dot12*Dot23 - Dot22*Dot13);
-		N2 = (Dot11*Dot23 - Dot12*Dot13);
-
-		if( N1 < 0.f ) 
-		{
-			// t1 < 0.f => the s==0 edge is visible
-			N1 = 0.f;
-			N2 = Dot23;
-			D2 = Dot22;
-		}
-		else if( N1 > D1 ) 
-		{
-			// t1 > 1 => the t1==1 edge is visible
-			N1 = D1;
-			N2 = Dot23 + Dot12;
-			D2 = Dot22;
-		}
-	}
+	SegmentDistToSegment_Solver Solver(A1, B1, A2, B2);
 
-	if( N2 < 0.f ) 
-	{           
-		// t2 < 0 => the t2==0 edge is visible
-		N2 = 0.f;
+	const FVector S1_norm = Solver.S1.GetSafeNormal();
+	const FVector S2_norm = Solver.S2.GetSafeNormal();
 
-		// recompute t1 for this edge
-		if( -Dot13 < 0.f )
-		{
-			N1 = 0.f;
-		}
-		else if( -Dot13 > Dot11 )
-		{
-			N1 = D1;
-		}
-		else 
-		{
-			N1 = -Dot13;
-			D1 = Dot11;
-		}
-	}
-	else if( N2 > D2 ) 
-	{      
-		// t2 > 1 => the t2=1 edge is visible
-		N2 = D2;
+	const bool bS1IsPoint = S1_norm.IsZero();
+	const bool bS2IsPoint = S2_norm.IsZero();
 
-		// recompute t1 for this edge
-		if( (-Dot13 + Dot12) < 0.f )
-		{
-			N1 = 0.f;
-		}
-		else if( (-Dot13 + Dot12) > Dot11 )
-		{
-			N1 = D1;
-		}
-		else 
-		{
-			N1 = (-Dot13 + Dot12);
-			D1 = Dot11;
-		}
+	if (bS1IsPoint && bS2IsPoint)
+	{
+		OutP1 = A1;
+		OutP2 = A2;
 	}
+	else if (bS2IsPoint)
+	{
+		OutP1 = ClosestPointOnSegment(A2, A1, B1);
+		OutP2 = A2;
+	}
+	else if (bS1IsPoint)
+	{
+		OutP1 = A1;
+		OutP2 = ClosestPointOnSegment(A1, A2, B2);
+	}
+	else
+	{
+		const	float	Dot11_norm = S1_norm | S1_norm;	// always >= 0
+		const	float	Dot22_norm = S2_norm | S2_norm;	// always >= 0
+		const	float	Dot12_norm = S1_norm | S2_norm;
+		const	float	D_norm	= Dot11_norm*Dot22_norm - Dot12_norm*Dot12_norm;	// always >= 0
 
-	// finally do the division to get the points' location
-	const float T1 = (FMath::Abs(N1) < KINDA_SMALL_NUMBER ? 0.f : N1 / D1);
-	const float T2 = (FMath::Abs(N2) < KINDA_SMALL_NUMBER ? 0.f : N2 / D2);
+		Solver.bLinesAreNearlyParallel = D_norm < KINDA_SMALL_NUMBER;
+		Solver.Solve(OutP1, OutP2);
+	}
+}
 
-	// return the closest points
-	OutP1 = A1 + T1 * S1;
-	OutP2 = A2 + T2 * S2;
+void FMath::SegmentDistToSegment(FVector A1, FVector B1, FVector A2, FVector B2, FVector& OutP1, FVector& OutP2)
+{
+	SegmentDistToSegment_Solver(A1, B1, A2, B2).Solve(OutP1, OutP2);
 }
 
 float FMath::GetTForSegmentPlaneIntersect(const FVector& StartPoint, const FVector& EndPoint, const FPlane& Plane)
diff --git a/Engine/Source/Runtime/Core/Public/HAL/Platform.h b/Engine/Source/Runtime/Core/Public/HAL/Platform.h
--- a/Engine/Source/Runtime/Core/Public/HAL/Platform.h
+++ b/Engine/Source/Runtime/Core/Public/HAL/Platform.h
@@ -586,34 +586,55 @@ int32 tchar_main(int32 ArgC, TCHAR* ArgV[])
 // Transfer the platform types to global types
 //------------------------------------------------------------------
 
-// Unsigned base types.
-typedef FPlatformTypes::uint8		uint8;		///< An 8-bit unsigned integer.
-typedef FPlatformTypes::uint16		uint16;		///< A 16-bit unsigned integer.
-typedef FPlatformTypes::uint32		uint32;		///< A 32-bit unsigned integer.
-typedef FPlatformTypes::uint64		uint64;		///< A 64-bit unsigned integer.
-
-// Signed base types.
-typedef	FPlatformTypes::int8		int8;		///< An 8-bit signed integer.
-typedef FPlatformTypes::int16		int16;		///< A 16-bit signed integer.
-typedef FPlatformTypes::int32		int32;		///< A 32-bit signed integer.
-typedef FPlatformTypes::int64		int64;		///< A 64-bit signed integer.
-
-// Character types.
-typedef FPlatformTypes::ANSICHAR	ANSICHAR;	///< An ANSI character. Normally a signed type.
-typedef FPlatformTypes::WIDECHAR	WIDECHAR;	///< A wide character. Normally a signed type.
-typedef FPlatformTypes::TCHAR		TCHAR;		///< Either ANSICHAR or WIDECHAR, depending on whether the platform supports wide characters or the requirements of the licensee.
-typedef FPlatformTypes::CHAR8		UTF8CHAR;	///< An 8-bit character containing a UTF8 (Unicode, 8-bit, variable-width) code unit.
-typedef FPlatformTypes::CHAR16		UCS2CHAR;	///< A 16-bit character containing a UCS2 (Unicode, 16-bit, fixed-width) code unit, used for compatibility with 'Windows TCHAR' across multiple platforms.
-typedef FPlatformTypes::CHAR16		UTF16CHAR;	///< A 16-bit character containing a UTF16 (Unicode, 16-bit, variable-width) code unit.
-typedef FPlatformTypes::CHAR32		UTF32CHAR;	///< A 32-bit character containing a UTF32 (Unicode, 32-bit, fixed-width) code unit.
-
-typedef FPlatformTypes::UPTRINT UPTRINT;		///< An unsigned integer the same size as a pointer
-typedef FPlatformTypes::PTRINT PTRINT;			///< A signed integer the same size as a pointer
-typedef FPlatformTypes::SIZE_T SIZE_T;			///< An unsigned integer the same size as a pointer, the same as UPTRINT
-typedef FPlatformTypes::SSIZE_T SSIZE_T;		///< An integer the same size as a pointer, the same as PTRINT
-
-typedef FPlatformTypes::TYPE_OF_NULL	TYPE_OF_NULL;		///< The type of the NULL constant.
-typedef FPlatformTypes::TYPE_OF_NULLPTR	TYPE_OF_NULLPTR;	///< The type of the C++ nullptr keyword.
+//~ Unsigned base types.
+/// An 8-bit unsigned integer.
+typedef FPlatformTypes::uint8		uint8;
+/// A 16-bit unsigned integer.
+typedef FPlatformTypes::uint16		uint16;
+/// A 32-bit unsigned integer.
+typedef FPlatformTypes::uint32		uint32;
+/// A 64-bit unsigned integer.
+typedef FPlatformTypes::uint64		uint64;
+
+//~ Signed base types.
+/// An 8-bit signed integer.
+typedef	FPlatformTypes::int8		int8;
+/// A 16-bit signed integer.
+typedef FPlatformTypes::int16		int16;
+/// A 32-bit signed integer.
+typedef FPlatformTypes::int32		int32;
+/// A 64-bit signed integer.
+typedef FPlatformTypes::int64		int64;
+
+//~ Character types.
+/// An ANSI character. Normally a signed type.
+typedef FPlatformTypes::ANSICHAR	ANSICHAR;
+/// A wide character. Normally a signed type.
+typedef FPlatformTypes::WIDECHAR	WIDECHAR;
+/// Either ANSICHAR or WIDECHAR, depending on whether the platform supports wide characters or the requirements of the licensee.
+typedef FPlatformTypes::TCHAR		TCHAR;
+/// An 8-bit character containing a UTF8 (Unicode, 8-bit, variable-width) code unit.
+typedef FPlatformTypes::CHAR8		UTF8CHAR;
+/// A 16-bit character containing a UCS2 (Unicode, 16-bit, fixed-width) code unit, used for compatibility with 'Windows TCHAR' across multiple platforms.
+typedef FPlatformTypes::CHAR16		UCS2CHAR;
+/// A 16-bit character containing a UTF16 (Unicode, 16-bit, variable-width) code unit.
+typedef FPlatformTypes::CHAR16		UTF16CHAR;
+/// A 32-bit character containing a UTF32 (Unicode, 32-bit, fixed-width) code unit.
+typedef FPlatformTypes::CHAR32		UTF32CHAR;
+
+/// An unsigned integer the same size as a pointer
+typedef FPlatformTypes::UPTRINT UPTRINT;
+/// A signed integer the same size as a pointer
+typedef FPlatformTypes::PTRINT PTRINT;
+/// An unsigned integer the same size as a pointer, the same as UPTRINT
+typedef FPlatformTypes::SIZE_T SIZE_T;
+/// An integer the same size as a pointer, the same as PTRINT
+typedef FPlatformTypes::SSIZE_T SSIZE_T;
+
+/// The type of the NULL constant.
+typedef FPlatformTypes::TYPE_OF_NULL	TYPE_OF_NULL;
+/// The type of the C++ nullptr keyword.
+typedef FPlatformTypes::TYPE_OF_NULLPTR	TYPE_OF_NULLPTR;
 
 //------------------------------------------------------------------
 // Test the global types
diff --git a/Engine/Source/Runtime/Core/Public/Math/UnrealMathUtility.h b/Engine/Source/Runtime/Core/Public/Math/UnrealMathUtility.h
--- a/Engine/Source/Runtime/Core/Public/Math/UnrealMathUtility.h
+++ b/Engine/Source/Runtime/Core/Public/Math/UnrealMathUtility.h
@@ -1103,6 +1103,9 @@ struct FMath : public FPlatformMath
 
 	/** 
 	 * Find closest points between 2 segments.
+	 *
+	 * If either segment may have a length of 0, use SegmentDistToSegmentSafe instance.
+	 *
 	 * @param	(A1, B1)	defines the first segment.
 	 * @param	(A2, B2)	defines the second segment.
 	 * @param	OutP1		Closest point on segment 1 to segment 2.
@@ -1112,6 +1115,10 @@ struct FMath : public FPlatformMath
 
 	/** 
 	 * Find closest points between 2 segments.
+	 *
+	 * This is the safe version, and will check both segments' lengths.
+	 * Use this if either (or both) of the segments lengths may be 0.
+	 *
 	 * @param	(A1, B1)	defines the first segment.
 	 * @param	(A2, B2)	defines the second segment.
 	 * @param	OutP1		Closest point on segment 1 to segment 2.
diff --git a/Engine/Source/Runtime/Core/Public/Misc/AutomationTest.h b/Engine/Source/Runtime/Core/Public/Misc/AutomationTest.h
--- a/Engine/Source/Runtime/Core/Public/Misc/AutomationTest.h
+++ b/Engine/Source/Runtime/Core/Public/Misc/AutomationTest.h
@@ -1128,6 +1128,9 @@ public:
 		}
 	}
 
+	/** Macro version of above, uses the passed in expression as the description as well */
+	#define TestTrueExpr(Expression) TestTrue(TEXT(#Expression), Expression)
+
 	/**
 	 * Logs an error if the given shared pointer is not valid.
 	 *
diff --git a/Engine/Source/Runtime/Core/Public/Modules/ModuleInterface.h b/Engine/Source/Runtime/Core/Public/Modules/ModuleInterface.h
--- a/Engine/Source/Runtime/Core/Public/Modules/ModuleInterface.h
+++ b/Engine/Source/Runtime/Core/Public/Modules/ModuleInterface.h
@@ -24,6 +24,9 @@ public:
 
 	/**
 	 * Called right after the module DLL has been loaded and the module object has been created
+	 * Load dependent modules here, and they will be guaranteed to be available during ShutdownModule. ie:
+	 *
+	 * FModuleManager::Get().LoadModuleChecked(TEXT("HTTP"));
 	 */
 	virtual void StartupModule()
 	{
@@ -45,6 +48,9 @@ public:
 
 	/**
 	 * Called before the module is unloaded, right before the module object is destroyed.
+	 * During normal shutdown, this is called in reverse order that modules finish StartupModule().
+	 * This means that, as long as a module references dependent modules in it's StartupModule(), it
+	 * can safely reference those dependencies in ShutdownModule() as well.
 	 */
 	virtual void ShutdownModule()
 	{
diff --git a/Engine/Source/Runtime/Core/Public/Modules/ModuleManager.h b/Engine/Source/Runtime/Core/Public/Modules/ModuleManager.h
--- a/Engine/Source/Runtime/Core/Public/Modules/ModuleManager.h
+++ b/Engine/Source/Runtime/Core/Public/Modules/ModuleManager.h
@@ -198,6 +198,8 @@ public:
 
 	/**
 	 * Unloads a specific module
+	 * NOTE: You can manually unload a module before the normal shutdown occurs with this, but be careful as you may be unloading another
+	 * module's dependency too early!
 	 *
 	 * @param InModuleName The name of the module to unload.  Should not include path, extension or platform/configuration info.  This is just the "module name" part of the module file name.
 	 * @param bIsShutdown Is this unload module call occurring at shutdown (default = false).
@@ -337,7 +339,11 @@ public:
 	int32 GetModuleCount( ) const;
 
 	/**
-	 * Unloads modules during the shutdown process.
+	 * Unloads modules during the shutdown process. Modules are unloaded in reverse order to when their StartupModule() FINISHES.
+	 * The practical implication of this is that if module A depends on another module B, and A loads B during A's StartupModule, 
+	 * that B will actually get Unloaded after A during shutdown. This allows A's ShutdownModule() call to still reference module B.
+	 * You can manually unload a module yourself which will change this ordering, but be careful as you may be unloading another 
+	 * module's dependency!
 	 *
 	 * This method is Usually called at various points while exiting an application.
 	 */
diff --git a/Engine/Source/Runtime/Core/Public/Serialization/Archive.h b/Engine/Source/Runtime/Core/Public/Serialization/Archive.h
--- a/Engine/Source/Runtime/Core/Public/Serialization/Archive.h
+++ b/Engine/Source/Runtime/Core/Public/Serialization/Archive.h
@@ -654,6 +654,11 @@ public:
 	 */
 	virtual void MarkScriptSerializationEnd(const UObject* Obj) { }
 
+	/**
+	 * Called to register a reference to a specific name value, of type TypeObject (UEnum or UStruct normally). Const so it can be called from PostSerialize
+	 */
+	virtual void MarkSearchableName(const UObject* TypeObject, const FName& ValueName) const { }
+
 	/**
 	* Called to retrieve the archetype from the event driven loader. If this returns null, then call GetArchetype yourself.
 	*/
diff --git a/Engine/Source/Runtime/Core/Public/UObject/FrameworkObjectVersion.h b/Engine/Source/Runtime/Core/Public/UObject/FrameworkObjectVersion.h
--- a/Engine/Source/Runtime/Core/Public/UObject/FrameworkObjectVersion.h
+++ b/Engine/Source/Runtime/Core/Public/UObject/FrameworkObjectVersion.h
@@ -57,12 +57,15 @@ struct CORE_API FFrameworkObjectVersion
 		// major flag - i.e. material types - moves to skeleton and handle in one place
 		MoveCurveTypesToSkeleton,
 		
-	        // Cache destructible overlaps on save
+		// Cache destructible overlaps on save
 		CacheDestructibleOverlaps,
 
-	        // Added serialization of materials applied to geometry cache objects
+		// Added serialization of materials applied to geometry cache objects
 		GeometryCacheMissingMaterials,
 
+		// Switch static & skeletal meshes to calculate LODs based on resolution-independent screen size
+		LODsUseResolutionIndependentScreenSize,
+
 		// -----<new versions can be added above this line>-------------------------------------------------
 		VersionPlusOne,
 		LatestVersion = VersionPlusOne - 1
diff --git a/Engine/Source/Runtime/Core/Public/UObject/ObjectVersion.h b/Engine/Source/Runtime/Core/Public/UObject/ObjectVersion.h
--- a/Engine/Source/Runtime/Core/Public/UObject/ObjectVersion.h
+++ b/Engine/Source/Runtime/Core/Public/UObject/ObjectVersion.h
@@ -614,7 +614,9 @@ enum EUnrealEngineObjectUE4Version
 	VER_UE4_TemplateIndex_IN_COOKED_EXPORTS,
 	// FPropertyTag includes contained type(s) for Set and Map properties:
 	VER_UE4_PROPERTY_TAG_SET_MAP_SUPPORT,
-	
+	// Added SearchableNames to the package summary and asset registry
+	VER_UE4_ADDED_SEARCHABLE_NAMES,
+
 	// -----<new versions can be added before this line>-------------------------------------------------
 	// - this needs to be the last line (see note below)
 	VER_UE4_AUTOMATIC_VERSION_PLUS_ONE,
diff --git a/Engine/Source/Runtime/CoreUObject/Private/Serialization/ArchiveUObject.cpp b/Engine/Source/Runtime/CoreUObject/Private/Serialization/ArchiveUObject.cpp
--- a/Engine/Source/Runtime/CoreUObject/Private/Serialization/ArchiveUObject.cpp
+++ b/Engine/Source/Runtime/CoreUObject/Private/Serialization/ArchiveUObject.cpp
@@ -115,7 +115,7 @@ FArchive& FObjectAndNameAsStringProxyArchive::operator<<(class UObject*& Obj)
 		// look up the object by fully qualified pathname
 		Obj = FindObject<UObject>(nullptr, *LoadedString, false);
 		// If we couldn't find it, and we want to load it, do that
-		if(Obj && bLoadIfFindFails)
+		if(!Obj && bLoadIfFindFails)
 		{
 			Obj = LoadObject<UObject>(nullptr, *LoadedString);
 		}
diff --git a/Engine/Source/Runtime/CoreUObject/Private/UObject/Linker.cpp b/Engine/Source/Runtime/CoreUObject/Private/UObject/Linker.cpp
--- a/Engine/Source/Runtime/CoreUObject/Private/UObject/Linker.cpp
+++ b/Engine/Source/Runtime/CoreUObject/Private/UObject/Linker.cpp
@@ -79,6 +79,27 @@ void FGenerationInfo::Serialize(FArchive& Ar, const struct FPackageFileSummary&
 extern int32 GLinkerAllowDynamicClasses;
 #endif
 
+void FLinkerTables::SerializeSearchableNamesMap(FArchive& Ar)
+{
+#if WITH_EDITOR
+	FArchive::FScopeSetDebugSerializationFlags S(Ar, DSF_IgnoreDiff, true);
+#endif
+
+	if (Ar.IsSaving())
+	{
+		// Sort before saving to keep order consistent
+		SearchableNamesMap.KeySort(TLess<FPackageIndex>());
+
+		for (TPair<FPackageIndex, TArray<FName> >& Pair : SearchableNamesMap)
+		{
+			Pair.Value.Sort();
+		}
+	}
+
+	// Default Map serialize works fine
+	Ar << SearchableNamesMap;
+}
+
 FName FLinker::GetExportClassName( int32 i )
 {
 	if (ExportMap.IsValidIndex(i))
@@ -130,41 +151,17 @@ FLinker::FLinker(ELinkerType::Type InType, UPackage* InRoot, const TCHAR* InFile
 
 void FLinker::Serialize( FArchive& Ar )
 {
+	// This function is only used for counting memory, actual serialization uses a different path
 	if( Ar.IsCountingMemory() )
 	{
 		// Can't use CountBytes as ExportMap is array of structs of arrays.
 		Ar << ImportMap;
 		Ar << ExportMap;
 		Ar << DependsMap;
-
-		if (Ar.IsSaving() || Ar.UE4Ver() >= VER_UE4_ADD_STRING_ASSET_REFERENCES_MAP)
-		{
-			Ar << StringAssetReferencesMap;
-		}
-	}
-
-	if (Ar.IsSaving() || Ar.UE4Ver() >= VER_UE4_SERIALIZE_TEXT_IN_PACKAGES)
-	{
+		Ar << StringAssetReferencesMap;
 		Ar << GatherableTextDataMap;
+		Ar << SearchableNamesMap;
 	}
-
-	// Prevent garbage collecting of linker's names and package.
-	Ar << NameMap << LinkerRoot;
-	{
-		for (auto& E : ExportMap)
-		{
-			Ar << E.ObjectName;
-		}
-	}
-	{
-		for (auto& I : ImportMap)
-		{
-			UObject* LegacyLinkerObject = nullptr;
-			Ar << LegacyLinkerObject;
-			Ar << I.ClassPackage << I.ClassName;
-		}
-	}
-
 }
 
 void FLinker::AddReferencedObjects(FReferenceCollector& Collector)
diff --git a/Engine/Source/Runtime/CoreUObject/Private/UObject/PackageFileSummary.cpp b/Engine/Source/Runtime/CoreUObject/Private/UObject/PackageFileSummary.cpp
--- a/Engine/Source/Runtime/CoreUObject/Private/UObject/PackageFileSummary.cpp
+++ b/Engine/Source/Runtime/CoreUObject/Private/UObject/PackageFileSummary.cpp
@@ -186,6 +186,11 @@ FArchive& operator<<( FArchive& Ar, FPackageFileSummary& Sum )
 			Ar << Sum.StringAssetReferencesCount << Sum.StringAssetReferencesOffset;
 		}
 
+		if (Ar.IsSaving() || Sum.FileVersionUE4 >= VER_UE4_ADDED_SEARCHABLE_NAMES)
+		{
+			Ar << Sum.SearchableNamesOffset;
+		}
+
 		Ar << Sum.ThumbnailTableOffset;
 
 		int32 GenerationCount = Sum.Generations.Num();
diff --git a/Engine/Source/Runtime/CoreUObject/Private/UObject/PropertyBaseObject.cpp b/Engine/Source/Runtime/CoreUObject/Private/UObject/PropertyBaseObject.cpp
--- a/Engine/Source/Runtime/CoreUObject/Private/UObject/PropertyBaseObject.cpp
+++ b/Engine/Source/Runtime/CoreUObject/Private/UObject/PropertyBaseObject.cpp
@@ -336,7 +336,7 @@ UObject* UObjectPropertyBase::FindImportedObject( const UProperty* Property, UOb
 			UObject* ScopedSearchRoot = SearchStart;
 			while (Result == NULL && ScopedSearchRoot != NULL)
 			{
-				Result = StaticFindObject(ObjectClass, ScopedSearchRoot, Text);
+				Result = StaticFindObjectSafe(ObjectClass, ScopedSearchRoot, Text);
 				// don't think it's possible to get a non-subobject here, but it doesn't hurt to check
 				if (Result != NULL && !Result->IsTemplate(RF_ClassDefaultObject))
 				{
@@ -359,7 +359,7 @@ UObject* UObjectPropertyBase::FindImportedObject( const UProperty* Property, UOb
 	UObject* ScopedSearchRoot = OwnerObject;
 	while (Result == NULL && ScopedSearchRoot != NULL)
 	{
-		Result = StaticFindObject(ObjectClass, ScopedSearchRoot, Text);
+		Result = StaticFindObjectSafe(ObjectClass, ScopedSearchRoot, Text);
 		// disallow class default subobjects here while importing defaults
 		// this prevents the use of a subobject name that doesn't exist in the scope of the default object being imported
 		// from grabbing some other subobject with the same name and class in some other arbitrary default object
@@ -374,12 +374,12 @@ UObject* UObjectPropertyBase::FindImportedObject( const UProperty* Property, UOb
 	if (Result == NULL)
 	{
 		// attempt to find a fully qualified object
-		Result = StaticFindObject(ObjectClass, NULL, Text);
+		Result = StaticFindObjectSafe(ObjectClass, NULL, Text);
 
 		if (Result == NULL)
 		{
 			// match any object of the correct class whose path contains the specified path
-			Result = StaticFindObject(ObjectClass, ANY_PACKAGE, Text);
+			Result = StaticFindObjectSafe(ObjectClass, ANY_PACKAGE, Text);
 			// disallow class default subobjects here while importing defaults
 			if (Result != NULL && (PortFlags & PPF_ParsingDefaultProperties) && Result->IsTemplate(RF_ClassDefaultObject))
 			{
@@ -410,7 +410,7 @@ UObject* UObjectPropertyBase::FindImportedObject( const UProperty* Property, UOb
 			}
 		}
 		// If we still can't find it, try to load it. (Only try to load fully qualified names)
-		if(!Result && Dot)
+		if(!Result && Dot && !GIsSavingPackage)
 		{
 #if USE_CIRCULAR_DEPENDENCY_LOAD_DEFERRING
 			FLinkerLoad* Linker = (OwnerObject != nullptr) ? OwnerObject->GetClass()->GetLinker() : nullptr;
diff --git a/Engine/Source/Runtime/CoreUObject/Private/UObject/SavePackage.cpp b/Engine/Source/Runtime/CoreUObject/Private/UObject/SavePackage.cpp
--- a/Engine/Source/Runtime/CoreUObject/Private/UObject/SavePackage.cpp
+++ b/Engine/Source/Runtime/CoreUObject/Private/UObject/SavePackage.cpp
@@ -1028,20 +1028,19 @@ class FArchiveSaveTagImports : public FArchiveUObject
 public:
 	FLinkerSave* Linker;
 	TArray<UObject*> Dependencies;
-	TArray<FString>& StringAssetReferencesMap;
 
-	FArchiveSaveTagImports(FLinkerSave* InLinker, TArray<FString>& InStringAssetReferencesMap)
-		: Linker(InLinker), StringAssetReferencesMap(InStringAssetReferencesMap)
+	FArchiveSaveTagImports(FLinkerSave* InLinker)
+		: Linker(InLinker)
 	{
+		check(Linker);
+
 		ArIsSaving				= true;
 		ArIsPersistent			= true;
 		ArIsObjectReferenceCollector = true;
 		ArShouldSkipBulkData	= true;
-		if ( Linker != NULL )
-		{
-			ArPortFlags = Linker->GetPortFlags();
-			SetCookingTarget(Linker->CookingTarget());
-		}
+
+		ArPortFlags = Linker->GetPortFlags();
+		SetCookingTarget(Linker->CookingTarget());
 	}
 	FArchive& operator<<(UObject*& Obj) override;
 	FArchive& operator<<(FLazyObjectPtr& LazyObjectPtr) override;
@@ -1059,14 +1058,14 @@ public:
 
 			if (GetIniFilenameFromObjectsReference(Path) != nullptr)
 			{
-				StringAssetReferencesMap.AddUnique(Path);
+				Linker->StringAssetReferencesMap.AddUnique(Path);
 			}
 			else
 			{
 				FString NormalizedPath = FPackageName::GetNormalizedObjectPath(Path);
 				if (!NormalizedPath.IsEmpty())
 				{
-					StringAssetReferencesMap.AddUnique(FPackageName::ObjectPathToPackageName(NormalizedPath));
+					Linker->StringAssetReferencesMap.AddUnique(FPackageName::ObjectPathToPackageName(NormalizedPath));
 					Value.SetPath(MoveTemp(NormalizedPath));
 				}
 			}
@@ -1078,6 +1077,24 @@ public:
 		SavePackageState->MarkNameAsReferenced(Name);
 		return *this;
 	}
+	virtual void MarkSearchableName(const UObject* TypeObject, const FName& ValueName) const override
+	{
+		if (!TypeObject)
+		{
+			return;
+		}
+
+		if (!Dependencies.Contains(TypeObject))
+		{
+			// Serialize object to make sure it ends up in import table
+			// This is doing a const cast to avoid backward compatibility issues
+			FArchiveSaveTagImports* MutableArchive = const_cast<FArchiveSaveTagImports*>(this);
+			UObject* TempObject = const_cast<UObject*>(TypeObject);
+			(*MutableArchive) << TempObject;
+		}
+
+		Linker->SearchableNamesObjectMap.FindOrAdd(TypeObject).AddUnique(ValueName);
+	}
 
 	virtual FArchive& operator<< (struct FWeakObjectPtr& Value) override;
 
@@ -1090,7 +1107,6 @@ public:
 	virtual FString GetArchiveName() const override;
 };
 
-
 /**
  * Returns the name of the Archive.  Useful for getting the name of the package a struct or object
  * is in when a loading error occurs.
@@ -3973,7 +3989,6 @@ ESavePackageResult UPackage::Save(UPackage* InOuter, UObject* Base, EObjectFlags
 				}
 
 				// Import objects & names.
-				TArray<FString> StringAssetReferencesMap;
 				{
 					const EObjectMark ObjectMarks = UPackage::GetObjectMarksForTargetPlatform(TargetPlatform, Linker->IsCooking());
 					TArray<UObject*> TagExpObjects;
@@ -3983,7 +3998,7 @@ ESavePackageResult UPackage::Save(UPackage* InOuter, UObject* Base, EObjectFlags
 						UObject* Obj = TagExpObjects[Index];
 						check(Obj->HasAnyMarks(OBJECTMARK_TagExp));
 						// Build list.
-						FArchiveSaveTagImports ImportTagger(Linker, StringAssetReferencesMap);
+						FArchiveSaveTagImports ImportTagger(Linker);
 						ImportTagger.SetPortFlags(ComparisonFlags);
 						ImportTagger.SetFilterEditorOnly(FilterEditorOnly);
 
@@ -4699,6 +4714,19 @@ ESavePackageResult UPackage::Save(UPackage* InOuter, UObject* Base, EObjectFlags
 					}
 				}
 
+				// Convert the searchable names map from UObject to packageindex
+				for (const TPair<const UObject *, TArray<FName>>& SearchableNamePair : Linker->SearchableNamesObjectMap)
+				{
+					const FPackageIndex PackageIndex = Linker->MapObject(SearchableNamePair.Key);
+
+					// This should always be in the imports already
+					if (ensure(!PackageIndex.IsNull()))
+					{
+						Linker->SearchableNamesMap.FindOrAdd(PackageIndex) = SearchableNamePair.Value;
+					}
+				}
+				// Clear temporary map
+				Linker->SearchableNamesObjectMap.Empty();
 
 				SlowTask.EnterProgressFrame();
 
@@ -4770,23 +4798,36 @@ ESavePackageResult UPackage::Save(UPackage* InOuter, UObject* Base, EObjectFlags
 				}
 				SlowTask.EnterProgressFrame();
 
-				// Save string asset reference map
-				Linker->Summary.StringAssetReferencesOffset = Linker->Tell();
-				Linker->Summary.StringAssetReferencesCount = StringAssetReferencesMap.Num();
+				// Only save string asset and searchable name map if saving for editor
+				if (!(Linker->Summary.PackageFlags & PKG_FilterEditorOnly))
 				{
+					// Save string asset reference map
+					Linker->Summary.StringAssetReferencesOffset = Linker->Tell();
+					Linker->Summary.StringAssetReferencesCount = Linker->StringAssetReferencesMap.Num();
+					{
 #if WITH_EDITOR
-					FArchive::FScopeSetDebugSerializationFlags S(*Linker, DSF_IgnoreDiff, true);
+						FArchive::FScopeSetDebugSerializationFlags S(*Linker, DSF_IgnoreDiff, true);
 #endif
-					for (int32 i = 0; i < StringAssetReferencesMap.Num(); i++)
-					{
-
-						*Linker << StringAssetReferencesMap[i];
+						for (FString& StringAssetReference : Linker->StringAssetReferencesMap)
+						{
+							*Linker << StringAssetReference;
+						}
 					}
+
+					// Save searchable names map
+					Linker->Summary.SearchableNamesOffset = Linker->Tell();
+					Linker->SerializeSearchableNamesMap(*Linker);
+				}
+				else
+				{
+					Linker->Summary.StringAssetReferencesCount = 0;
+					Linker->Summary.StringAssetReferencesOffset = 0;
+					Linker->Summary.SearchableNamesOffset = 0;
 				}
+
 				// Save thumbnails
 				UPackage::SaveThumbnails( InOuter, Linker );
 
-				
 				// Save asset registry data so the editor can search for information about assets in this package
 				UPackage::SaveAssetRegistryData( InOuter, Linker );
 
@@ -5843,4 +5884,4 @@ bool UPackage::IsEmptyPackage (UPackage* Package, const UObject* LastReferencer)
 
 	// Invalid package
 	return false;
-}
+}
\ No newline at end of file
diff --git a/Engine/Source/Runtime/CoreUObject/Public/Misc/AssetRegistryInterface.h b/Engine/Source/Runtime/CoreUObject/Public/Misc/AssetRegistryInterface.h
--- a/Engine/Source/Runtime/CoreUObject/Public/Misc/AssetRegistryInterface.h
+++ b/Engine/Source/Runtime/CoreUObject/Public/Misc/AssetRegistryInterface.h
@@ -13,10 +13,14 @@ namespace EAssetRegistryDependencyType
 		Soft = 1,
 
 		// Dependencies which are required for correct usage of the source asset, and must be loaded at the same time
-		Hard = 2
+		Hard = 2,
+
+		// References to specific SearchableNames inside a package
+		SearchableName = 4
 	};
 
-	static const Type All = (Type)(Soft | Hard);
+	static const Type All = (Type)(Soft | Hard | SearchableName);
+	static const Type Packages = (Type)(Soft | Hard);
 };
 
 /**
@@ -42,6 +46,6 @@ public:
 	/**
 	* Lookup dependencies for the given package name and fill OutDependencies with direct dependencies
 	*/
-	virtual void GetDependencies(FName InPackageName, TArray<FName>& OutDependencies, EAssetRegistryDependencyType::Type InDependencyType = EAssetRegistryDependencyType::All) = 0;
+	virtual void GetDependencies(FName InPackageName, TArray<FName>& OutDependencies, EAssetRegistryDependencyType::Type InDependencyType = EAssetRegistryDependencyType::Packages) = 0;
 };
 
diff --git a/Engine/Source/Runtime/CoreUObject/Public/UObject/Linker.h b/Engine/Source/Runtime/CoreUObject/Public/UObject/Linker.h
--- a/Engine/Source/Runtime/CoreUObject/Public/UObject/Linker.h
+++ b/Engine/Source/Runtime/CoreUObject/Public/UObject/Linker.h
@@ -48,6 +48,8 @@ public:
 	TArray<TArray<FPackageIndex> > DependsMap;
 	/** Map that holds info about string asset references from the package. */
 	TArray<FString> StringAssetReferencesMap;
+	/** List of Searchable Names, by object containing them. Not in MultiMap to allow sorting, and sizes are usually small enough where TArray makes sense */
+	TMap<FPackageIndex, TArray<FName> > SearchableNamesMap;
 
 	/**
 	 * Check that this Index is non-null and return an import or export
@@ -159,6 +161,9 @@ public:
 		}
 		return NULL;
 	}
+
+	/** Serializes the searchable name map */
+	COREUOBJECT_API void SerializeSearchableNamesMap(FArchive &Ar);
 };
 
 
@@ -288,7 +293,6 @@ public:
 		return NAME_None;
 	}
 
-
 	FORCEINLINE ELinkerType::Type GetType() const
 	{
 		return LinkerType;
diff --git a/Engine/Source/Runtime/CoreUObject/Public/UObject/LinkerSave.h b/Engine/Source/Runtime/CoreUObject/Public/UObject/LinkerSave.h
--- a/Engine/Source/Runtime/CoreUObject/Public/UObject/LinkerSave.h
+++ b/Engine/Source/Runtime/CoreUObject/Public/UObject/LinkerSave.h
@@ -35,6 +35,9 @@ public:
 	/** Index array - location of the resource for a UObject is stored in the ObjectIndices array using the UObject's Index */
 	TMap<UObject *,FPackageIndex> ObjectIndicesMap;
 
+	/** List of Searchable Names, by object containing them. This gets turned into package indices later */
+	TMap<const UObject *, TArray<FName> > SearchableNamesObjectMap;
+
 	/** Index array - location of the name in the NameMap array for each FName is stored in the NameIndices array using the FName's Index */
 	TMap<FName, int32, FDefaultSetAllocator, TLinkerNameMapKeyFuncs<int32>> NameIndices;
 
diff --git a/Engine/Source/Runtime/CoreUObject/Public/UObject/PackageFileSummary.h b/Engine/Source/Runtime/CoreUObject/Public/UObject/PackageFileSummary.h
--- a/Engine/Source/Runtime/CoreUObject/Public/UObject/PackageFileSummary.h
+++ b/Engine/Source/Runtime/CoreUObject/Public/UObject/PackageFileSummary.h
@@ -123,6 +123,11 @@ public:
 	*/
 	int32		StringAssetReferencesOffset;
 
+	/**
+	* Location into the file on disk for the SearchableNamesMap data
+	*/
+	int32		SearchableNamesOffset;
+
 	/**
 	* Thumbnail table offset
 	*/
diff --git a/Engine/Source/Runtime/Engine/Classes/AI/Navigation/NavLinkComponent.h b/Engine/Source/Runtime/Engine/Classes/AI/Navigation/NavLinkComponent.h
new file mode 100644
--- /dev/null
+++ b/Engine/Source/Runtime/Engine/Classes/AI/Navigation/NavLinkComponent.h
@@ -0,0 +1,26 @@
+// Copyright 1998-2016 Epic Games, Inc. All Rights Reserved.
+
+#pragma once
+
+#include "AI/Navigation/NavLinkDefinition.h"
+#include "AI/Navigation/NavLinkHostInterface.h"
+#include "NavLinkComponent.generated.h"
+
+UCLASS(ClassGroup = (Navigation), meta = (BlueprintSpawnableComponent), hidecategories = (Activation))
+class ENGINE_API UNavLinkComponent : public UPrimitiveComponent, public INavLinkHostInterface
+{
+	GENERATED_UCLASS_BODY()
+
+	UPROPERTY(EditAnywhere, BlueprintReadOnly, Category = Navigation)
+	TArray<FNavigationLink> Links;
+
+	virtual void GetNavigationData(FNavigationRelevantData& Data) const override;
+	virtual bool IsNavigationRelevant() const override;
+
+	virtual bool GetNavigationLinksClasses(TArray<TSubclassOf<class UNavLinkDefinition> >& OutClasses) const override { return false; }
+	virtual bool GetNavigationLinksArray(TArray<FNavigationLink>& OutLink, TArray<FNavigationSegmentLink>& OutSegments) const override;
+
+	virtual FBoxSphereBounds CalcBounds(const FTransform &LocalToWorld) const override;
+	virtual FPrimitiveSceneProxy* CreateSceneProxy() override;
+	virtual bool ShouldRecreateProxyOnUpdateTransform() const override { return true; }
+};
diff --git a/Engine/Source/Runtime/Engine/Classes/Animation/AnimInstance.h b/Engine/Source/Runtime/Engine/Classes/Animation/AnimInstance.h
--- a/Engine/Source/Runtime/Engine/Classes/Animation/AnimInstance.h
+++ b/Engine/Source/Runtime/Engine/Classes/Animation/AnimInstance.h
@@ -456,6 +456,10 @@ public:
 	UFUNCTION(BlueprintCallable, Category = "Animation")
 	virtual APawn* TryGetPawnOwner() const;
 
+	/** Takes a snapshot of the current skeletal mesh component pose. This snapshot can then be retrieved in the animation blueprint for blending. The snapshot is taken at the current LOD, so if for example you took the snapshot at LOD1 and then used it at LOD0 any bones not in LOD1 will use the reference pose */
+	UFUNCTION(BlueprintCallable, Category = "Animation")
+	virtual void SavePoseSnapshot(FName SnapshotName);
+
 	// Are we being evaluated on a worker thread
 	bool IsRunningParallelEvaluation() const;
 
diff --git a/Engine/Source/Runtime/Engine/Classes/Animation/Skeleton.h b/Engine/Source/Runtime/Engine/Classes/Animation/Skeleton.h
--- a/Engine/Source/Runtime/Engine/Classes/Animation/Skeleton.h
+++ b/Engine/Source/Runtime/Engine/Classes/Animation/Skeleton.h
@@ -328,6 +328,7 @@ public:
 	// you'll have to call REfreshCAchedAnimationCurveData to apply
 	ENGINE_API FCurveMetaData* GetCurveMetaData(const FName& CurveName);
 	ENGINE_API const FCurveMetaData* GetCurveMetaData(const FName& CurveName) const;
+	ENGINE_API const FCurveMetaData* GetCurveMetaData(const SmartName::UID_Type CurveUID) const;
 	ENGINE_API FCurveMetaData* GetCurveMetaData(const FSmartName& CurveName);
 	ENGINE_API const FCurveMetaData* GetCurveMetaData(const FSmartName& CurveName) const;
 	// this is called when you know both flags - called by post serialize
diff --git a/Engine/Source/Runtime/Engine/Classes/Animation/SmartName.h b/Engine/Source/Runtime/Engine/Classes/Animation/SmartName.h
--- a/Engine/Source/Runtime/Engine/Classes/Animation/SmartName.h
+++ b/Engine/Source/Runtime/Engine/Classes/Animation/SmartName.h
@@ -150,6 +150,8 @@ struct ENGINE_API FSmartNameMapping
 
 	friend FArchive& operator<<(FArchive& Ar, FSmartNameMapping& Elem);
 
+	/* initialize curve meta data for the container */
+	void InitializeCurveMetaData(class USkeleton* Skeleton);
 private:
 	SmartName::UID_Type NextUid;				// The next SmartName::UID_Type to use 
 	TMap<SmartName::UID_Type, FName> UidMap;	// Mapping of UIDs and names. This is transient and built upon serialize - this is run-time data, and searching from SmartName::UID_Type to FName is important
@@ -179,6 +181,9 @@ protected:
 	FSmartNameMapping* GetContainerInternal(const FName& ContainerName);
 	const FSmartNameMapping* GetContainerInternal(const FName& ContainerName) const;
 
+	/* initialize curve meta data for the container */
+	void InitializeCurveMetaData(class USkeleton* Skeleton);
+
 private:
 	TMap<FName, FSmartNameMapping> NameMappings;	// List of smartname mappings
 };
@@ -190,7 +195,7 @@ struct ENGINE_API FSmartName
 	GENERATED_USTRUCT_BODY();
 
 	// name 
-	UPROPERTY()
+	UPROPERTY(VisibleAnywhere, Category=FSmartName)
 	FName DisplayName;
 
 	// SmartName::UID_Type - for faster access
diff --git a/Engine/Source/Runtime/Engine/Classes/Components/PrimitiveComponent.h b/Engine/Source/Runtime/Engine/Classes/Components/PrimitiveComponent.h
--- a/Engine/Source/Runtime/Engine/Classes/Components/PrimitiveComponent.h
+++ b/Engine/Source/Runtime/Engine/Classes/Components/PrimitiveComponent.h
@@ -677,8 +677,8 @@ public:
 	 *  @param	ObjectQueryParams	List of object types it's looking for. When this enters, we do object query with component shape
 	 *  @return true if OutOverlaps contains any blocking results
 	 */
-	bool ComponentOverlapMulti(TArray<struct FOverlapResult>& OutOverlaps, const class UWorld* InWorld, const FVector& Pos, const FQuat& Rot, ECollisionChannel TestChannel, const struct FComponentQueryParams& Params, const struct FCollisionObjectQueryParams& ObjectQueryParams = FCollisionObjectQueryParams::DefaultObjectQueryParam) const;
-	bool ComponentOverlapMulti(TArray<struct FOverlapResult>& OutOverlaps, const class UWorld* InWorld, const FVector& Pos, const FRotator& Rot, ECollisionChannel TestChannel, const struct FComponentQueryParams& Params, const struct FCollisionObjectQueryParams& ObjectQueryParams = FCollisionObjectQueryParams::DefaultObjectQueryParam) const;
+	bool ComponentOverlapMulti(TArray<struct FOverlapResult>& OutOverlaps, const class UWorld* InWorld, const FVector& Pos, const FQuat& Rot, ECollisionChannel TestChannel, const struct FComponentQueryParams& Params = FComponentQueryParams::DefaultComponentQueryParams, const struct FCollisionObjectQueryParams& ObjectQueryParams = FCollisionObjectQueryParams::DefaultObjectQueryParam) const;
+	bool ComponentOverlapMulti(TArray<struct FOverlapResult>& OutOverlaps, const class UWorld* InWorld, const FVector& Pos, const FRotator& Rot, ECollisionChannel TestChannel, const struct FComponentQueryParams& Params = FComponentQueryParams::DefaultComponentQueryParams, const struct FCollisionObjectQueryParams& ObjectQueryParams = FCollisionObjectQueryParams::DefaultObjectQueryParam) const;
 
 protected:
 
diff --git a/Engine/Source/Runtime/Engine/Classes/Components/WindDirectionalSourceComponent.h b/Engine/Source/Runtime/Engine/Classes/Components/WindDirectionalSourceComponent.h
--- a/Engine/Source/Runtime/Engine/Classes/Components/WindDirectionalSourceComponent.h
+++ b/Engine/Source/Runtime/Engine/Classes/Components/WindDirectionalSourceComponent.h
@@ -5,42 +5,84 @@
 #include "Components/SceneComponent.h"
 #include "WindDirectionalSourceComponent.generated.h"
 
+UENUM()
+enum class EWindSourceType : uint8
+{
+	Directional,
+	Point,
+};
+
 /** Component that provides a directional wind source. Only affects SpeedTree assets. */
 UCLASS(collapsecategories, hidecategories=(Object, Mobility), editinlinenew)
-class ENGINE_API UWindDirectionalSourceComponent : public USceneComponent
+class UWindDirectionalSourceComponent : public USceneComponent
 {
 	GENERATED_UCLASS_BODY()
 
-	UPROPERTY(interp, Category=WindDirectionalSourceComponent)
+	UPROPERTY(Interp, BlueprintReadOnly, Category=WindDirectionalSourceComponent)
 	float Strength;
 
-	UPROPERTY(interp, Category=WindDirectionalSourceComponent)
+	UPROPERTY(Interp, BlueprintReadOnly, Category=WindDirectionalSourceComponent)
 	float Speed;
 
-	UPROPERTY(interp, Category = WindDirectionalSourceComponent)
+	UPROPERTY(Interp, BlueprintReadOnly, Category = WindDirectionalSourceComponent)
 	float MinGustAmount;
 
-	UPROPERTY(interp, Category = WindDirectionalSourceComponent)
+	UPROPERTY(Interp, BlueprintReadOnly, Category = WindDirectionalSourceComponent)
 	float MaxGustAmount;
 
-	UPROPERTY(interp, Category = WindDirectionalSourceComponent, meta = (editcondition = "bSimulatePhysics", ClampMin = "0.1", UIMin = "0.1"))
+	UPROPERTY(Interp, BlueprintReadOnly, Category = WindDirectionalSourceComponent, meta = (editcondition = "bSimulatePhysics", ClampMin = "0.1", UIMin = "0.1"))
 	float Radius;
 
-	UPROPERTY(EditAnywhere, Category = WindDirectionalSourceComponent)
+	UPROPERTY(BlueprintReadOnly, EditAnywhere, Category = WindDirectionalSourceComponent)
 	uint32 bPointWind : 1;
 
 public:
 	class FWindSourceSceneProxy* SceneProxy;
 
+	/**
+	 * Because the actual data used to query wind is stored on the render thread in
+	 * an instance of FWindSourceSceneProxy all of our properties are read only.
+	 * The data can be manipulated with the following functions which will queue 
+	 * a render thread update for this component
+	 */
+
+	/** Sets the strength of the generated wind */
+	UFUNCTION(BlueprintCallable, Category = Wind)
+	void SetStrength(float InNewStrength);
+
+	/** Sets the windspeed of the generated wind */
+	UFUNCTION(BlueprintCallable, Category = Wind)
+	void SetSpeed(float InNewSpeed);
+
+	/** Set minimum deviation for wind gusts */
+	UFUNCTION(BlueprintCallable, Category = Wind)
+	void SetMinimumGustAmount(float InNewMinGust);
+
+	/** Set maximum deviation for wind gusts */
+	UFUNCTION(BlueprintCallable, Category = Wind)
+	void SetMaximumGustAmount(float InNewMaxGust);
+
+	/** Set the effect radius for point wind */
+	UFUNCTION(BlueprintCallable, Category = Wind)
+	void SetRadius(float InNewRadius);
+
+	/** Set the type of wind generator to use */
+	UFUNCTION(BlueprintCallable, Category = Wind)
+	void SetWindType(EWindSourceType InNewType);
+
 protected:
 	//~ Begin UActorComponent Interface.
 	virtual void Activate(bool bReset) override;
 	virtual void Deactivate() override;
 	virtual void CreateRenderState_Concurrent() override;
 	virtual void SendRenderTransform_Concurrent() override;
+	virtual void SendRenderDynamicData_Concurrent() override;
 	virtual void DestroyRenderState_Concurrent() override;
 	//~ End UActorComponent Interface.
 
+private:
+	void UpdateSceneData_Concurrent();
+
 public:
 	/**
 	 * Creates a proxy to represent the primitive to the scene manager in the rendering thread.
diff --git a/Engine/Source/Runtime/Engine/Classes/Engine/CurveTable.h b/Engine/Source/Runtime/Engine/Classes/Engine/CurveTable.h
--- a/Engine/Source/Runtime/Engine/Classes/Engine/CurveTable.h
+++ b/Engine/Source/Runtime/Engine/Classes/Engine/CurveTable.h
@@ -165,8 +165,17 @@ struct ENGINE_API FCurveTableRowHandle
 
 	bool operator==(const FCurveTableRowHandle& Other) const;
 	bool operator!=(const FCurveTableRowHandle& Other) const;
+	void PostSerialize(const FArchive& Ar);
 };
 
+template<>
+struct TStructOpsTypeTraits< FCurveTableRowHandle > : public TStructOpsTypeTraitsBase
+{
+	enum
+	{
+		WithPostSerialize = true,
+	};
+};
 
 /** Macro to call GetCurve with a correct error info. Assumed to be called within a UObject */
 #define GETCURVE_REPORTERROR(Handle) Handle.GetCurve(FString::Printf(TEXT("%s.%s"), *GetPathName(), TEXT(#Handle)))
diff --git a/Engine/Source/Runtime/Engine/Classes/Engine/DataTable.h b/Engine/Source/Runtime/Engine/Classes/Engine/DataTable.h
--- a/Engine/Source/Runtime/Engine/Classes/Engine/DataTable.h
+++ b/Engine/Source/Runtime/Engine/Classes/Engine/DataTable.h
@@ -302,6 +302,16 @@ struct ENGINE_API FDataTableRowHandle
 
 	bool operator==(FDataTableRowHandle const& Other) const;
 	bool operator!=(FDataTableRowHandle const& Other) const;
+	void PostSerialize(const FArchive& Ar);
+};
+
+template<>
+struct TStructOpsTypeTraits< FDataTableRowHandle > : public TStructOpsTypeTraitsBase
+{
+	enum
+	{
+		WithPostSerialize = true,
+	};
 };
 
 /** Handle to a particular row in a table*/
diff --git a/Engine/Source/Runtime/Engine/Classes/Engine/LODActor.h b/Engine/Source/Runtime/Engine/Classes/Engine/LODActor.h
--- a/Engine/Source/Runtime/Engine/Classes/Engine/LODActor.h
+++ b/Engine/Source/Runtime/Engine/Classes/Engine/LODActor.h
@@ -131,6 +131,7 @@ public:
 #if WITH_EDITOR
 	virtual void PreEditChange(UProperty* PropertyThatWillChange) override;
 	virtual void PostEditChangeProperty(FPropertyChangedEvent& PropertyChangedEvent) override;
+	virtual void Serialize(FArchive& Ar) override;
 #endif // WITH_EDITOR	
 	//~ End UObject Interface.		
 protected:
@@ -161,7 +162,11 @@ public:
 	UPROPERTY(EditAnywhere, Category = HierarchicalLODSettings)
 	bool bOverrideTransitionScreenSize;
 
-	/** Override transition screen size value, determines the screen size / draw distance at which the proxy mesh is visible */
+	/** 
+	 * Override transition screen size value, determines the screen size at which the proxy is visible 
+	 * The screen size is based around the projected diameter of the bounding 
+	 * sphere of the model. i.e. 0.5 means half the screen's maximum dimension.
+	 */
 	UPROPERTY(EditAnywhere, Category = HierarchicalLODSettings, meta = (editcondition = "bOverrideTransitionScreenSize"))
 	float TransitionScreenSize;
 
@@ -169,7 +174,7 @@ public:
 	UPROPERTY(EditAnywhere, Category = HierarchicalLODSettings)
 	bool bOverrideScreenSize;
 
-	/** Override screen size value, used when creating the proxy mesh */
+	/** Override screen size value used in mesh reduction, when creating the proxy mesh */
 	UPROPERTY(EditAnywhere, Category = HierarchicalLODSettings, meta = (editcondition = "bOverrideScreenSize"))
 	int32 ScreenSize;
 #endif // WITH_EDITORONLY_DATA
@@ -191,6 +196,12 @@ private:
  	// Have we already tried to register components? (a cache to avoid having to query the owning world when the global HLOD max level setting is changed)
  	uint8 bHasActorTriedToRegisterComponents : 1;
 
+	/**
+	 * If true on post load we need to calculate resolution independent Display Factors from the
+	 * loaded LOD screen sizes.
+	 */
+	uint8 bRequiresLODScreenSizeConversion : 1;
+
 	// Sink for when CVars are changed to check to see if the maximum HLOD level value has changed
 	static FAutoConsoleVariableSink CVarSink;
 }; 
\ No newline at end of file
diff --git a/Engine/Source/Runtime/Engine/Classes/Engine/SkeletalMesh.h b/Engine/Source/Runtime/Engine/Classes/Engine/SkeletalMesh.h
--- a/Engine/Source/Runtime/Engine/Classes/Engine/SkeletalMesh.h
+++ b/Engine/Source/Runtime/Engine/Classes/Engine/SkeletalMesh.h
@@ -236,7 +236,11 @@ struct FSkeletalMeshLODInfo
 {
 	GENERATED_USTRUCT_BODY()
 
-	/**	Indicates when to use this LOD. A smaller number means use this LOD when further away. */
+	/** 
+	 * ScreenSize to display this LOD.
+	 * The screen size is based around the projected diameter of the bounding
+	 * sphere of the model. i.e. 0.5 means half the screen's maximum dimension.
+	 */
 	UPROPERTY(EditAnywhere, Category=SkeletalMeshLODInfo)
 	float ScreenSize;
 
@@ -267,6 +271,10 @@ struct FSkeletalMeshLODInfo
 	UPROPERTY(EditAnywhere, Category = ReductionSettings)
 	TArray<FName> RemovedBones;
 
+	/** The filename of the file tha was used to import this LOD if it was not auto generated. */
+	UPROPERTY(VisibleAnywhere, Category= SkeletalMeshLODInfo, AdvancedDisplay)
+	FString SourceImportFilename;
+
 	FSkeletalMeshLODInfo()
 		: ScreenSize(0)
 		, LODHysteresis(0)
@@ -652,6 +660,12 @@ public:
 	UPROPERTY(EditAnywhere, Category = Physics)
 	uint32 bEnablePerPolyCollision : 1;
 
+	/**
+	 * If true on post load we need to calculate resolution independent Display Factors from the
+	 * loaded LOD screen sizes.
+	 */
+	uint32 bRequiresLODScreenSizeConversion : 1;
+
 	// Physics data for the per poly collision case. In 99% of cases you will not need this and are better off using simple ragdoll collision (physics asset)
 	UPROPERTY(transient)
 	class UBodySetup* BodySetup;
@@ -973,6 +987,9 @@ public:
 	ENGINE_API void BuildPhysicsData();
 	ENGINE_API void AddBoneToReductionSetting(int32 LODIndex, const TArray<FName>& BoneNames);
 	ENGINE_API void AddBoneToReductionSetting(int32 LODIndex, FName BoneName);
+
+	/** Convert legacy screen size (based on fixed resolution) into screen size (diameter in screen units) */
+	void ConvertLegacyLODScreenSize();
 #endif
 	
 
diff --git a/Engine/Source/Runtime/Engine/Classes/Engine/StaticMesh.h b/Engine/Source/Runtime/Engine/Classes/Engine/StaticMesh.h
--- a/Engine/Source/Runtime/Engine/Classes/Engine/StaticMesh.h
+++ b/Engine/Source/Runtime/Engine/Classes/Engine/StaticMesh.h
@@ -154,7 +154,11 @@ struct FStaticMeshSourceModel
 	UPROPERTY()
 	float LODDistance_DEPRECATED;
 
-	/** ScreenSize to display this LOD */
+	/** 
+	 * ScreenSize to display this LOD.
+	 * The screen size is based around the projected diameter of the bounding
+	 * sphere of the model. i.e. 0.5 means half the screen's maximum dimension.
+	 */
 	UPROPERTY(EditAnywhere, Category=ReductionSettings)
 	float ScreenSize;
 
@@ -427,7 +431,11 @@ class UStaticMesh : public UObject, public IInterface_CollisionDataProvider, pub
 	*/
 	bool bRequiresLODDistanceConversion : 1;
 
-
+	/**
+	 * If true on post load we need to calculate resolution independent Display Factors from the
+	 * loaded LOD screen sizes.
+	 */
+	bool bRequiresLODScreenSizeConversion : 1;
 
 #endif // #if WITH_EDITORONLY_DATA
 
@@ -797,6 +805,11 @@ private:
 	 */
 	void ConvertLegacyLODDistance();
 
+	/**
+	 * Converts legacy LOD screen area in the source models to resolution-independent screen size
+	 */
+	void ConvertLegacyLODScreenArea();
+
 	/**
 	 * Fixes up static meshes that were imported with sections that had zero triangles.
 	 */
diff --git a/Engine/Source/Runtime/Engine/Classes/GameFramework/CharacterMovementComponent.h b/Engine/Source/Runtime/Engine/Classes/GameFramework/CharacterMovementComponent.h
--- a/Engine/Source/Runtime/Engine/Classes/GameFramework/CharacterMovementComponent.h
+++ b/Engine/Source/Runtime/Engine/Classes/GameFramework/CharacterMovementComponent.h
@@ -1881,8 +1881,8 @@ public:
 	/** Get prediction data for a server game. Should not be used if not running as a server. Allocates the data on demand and can be overridden to allocate a custom override if desired. */
 	virtual class FNetworkPredictionData_Server* GetPredictionData_Server() const override;
 
-	virtual bool HasPredictionData_Client() const override { return ClientPredictionData != NULL; }
-	virtual bool HasPredictionData_Server() const override { return ServerPredictionData != NULL; }
+	virtual bool HasPredictionData_Client() const override;
+	virtual bool HasPredictionData_Server() const override;
 
 	virtual void ResetPredictionData_Client() override;
 	virtual void ResetPredictionData_Server() override;
diff --git a/Engine/Source/Runtime/Engine/Classes/GameFramework/ForceFeedbackEffect.h b/Engine/Source/Runtime/Engine/Classes/GameFramework/ForceFeedbackEffect.h
--- a/Engine/Source/Runtime/Engine/Classes/GameFramework/ForceFeedbackEffect.h
+++ b/Engine/Source/Runtime/Engine/Classes/GameFramework/ForceFeedbackEffect.h
@@ -33,7 +33,7 @@ struct FForceFeedbackChannelDetails
 };
 
 USTRUCT()
-struct FActiveForceFeedbackEffect
+struct ENGINE_API FActiveForceFeedbackEffect
 {
 	GENERATED_USTRUCT_BODY()
 
diff --git a/Engine/Source/Runtime/Engine/Classes/Kismet/KismetMathLibrary.h b/Engine/Source/Runtime/Engine/Classes/Kismet/KismetMathLibrary.h
--- a/Engine/Source/Runtime/Engine/Classes/Kismet/KismetMathLibrary.h
+++ b/Engine/Source/Runtime/Engine/Classes/Kismet/KismetMathLibrary.h
@@ -994,6 +994,10 @@ class ENGINE_API UKismetMathLibrary : public UBlueprintFunctionLibrary
 	UFUNCTION(BlueprintPure, Category="Math|Rotator", meta=(Keywords="make construct build rotate rotation"))
 	static FRotator RotatorFromAxisAndAngle(FVector Axis, float Angle);
 
+	/** Get an axis and angle from a given rotation */
+	UFUNCTION(BlueprintPure, Category="Math|Rotator", meta=(Keywords="rotate rotation axis angle"))
+	static void RotatorToAxisAndAngle(const FRotator& Rotation, FVector& Axis, float& Angle);
+
 	/**
 	* Clamps an angle to the range of [0, 360].
 	*
diff --git a/Engine/Source/Runtime/Engine/Engine.Build.cs b/Engine/Source/Runtime/Engine/Engine.Build.cs
--- a/Engine/Source/Runtime/Engine/Engine.Build.cs
+++ b/Engine/Source/Runtime/Engine/Engine.Build.cs
@@ -69,6 +69,7 @@ public class Engine : ModuleRules
 				"EngineMessages",
 				"EngineSettings",
 				"SynthBenchmark",
+                "GameplayTags",
                 "AIModule",
 				"DatabaseSupport",
                 "PacketHandler",
@@ -148,6 +149,7 @@ public class Engine : ModuleRules
         }
         }
 
+        CircularlyReferencedDependentModules.Add("GameplayTags");
         CircularlyReferencedDependentModules.Add("AIModule");
 		CircularlyReferencedDependentModules.Add("Landscape");
         CircularlyReferencedDependentModules.Add("UMG");
diff --git a/Engine/Source/Runtime/Engine/Private/Actor.cpp b/Engine/Source/Runtime/Engine/Private/Actor.cpp
--- a/Engine/Source/Runtime/Engine/Private/Actor.cpp
+++ b/Engine/Source/Runtime/Engine/Private/Actor.cpp
@@ -231,12 +231,10 @@ void AActor::ResetOwnedComponents()
 	}
 #endif
 
-	TArray<UObject*> ActorChildren;
 	OwnedComponents.Reset();
 	ReplicatedComponents.Reset();
-	GetObjectsWithOuter(this, ActorChildren, true, RF_NoFlags, EInternalObjectFlags::PendingKill);
 
-	for (UObject* Child : ActorChildren)
+	ForEachObjectWithOuter(this, [this](UObject* Child)
 	{
 		UActorComponent* Component = Cast<UActorComponent>(Child);
 		if (Component && Component->GetOwner() == this)
@@ -248,7 +246,7 @@ void AActor::ResetOwnedComponents()
 				ReplicatedComponents.Add(Component);
 			}
 		}
-	}
+	}, true, RF_NoFlags, EInternalObjectFlags::PendingKill);
 }
 
 void AActor::PostInitProperties()
diff --git a/Engine/Source/Runtime/Engine/Private/Animation/AnimInstance.cpp b/Engine/Source/Runtime/Engine/Private/Animation/AnimInstance.cpp
--- a/Engine/Source/Runtime/Engine/Private/Animation/AnimInstance.cpp
+++ b/Engine/Source/Runtime/Engine/Private/Animation/AnimInstance.cpp
@@ -153,6 +153,15 @@ APawn* UAnimInstance::TryGetPawnOwner() const
 	return NULL;
 }
 
+void UAnimInstance::SavePoseSnapshot(FName SnapshotName)
+{
+	FAnimInstanceProxy& Proxy = GetProxyOnGameThread<FAnimInstanceProxy>();
+	if(USkeletalMeshComponent* SkeletalMeshComponent = GetSkelMeshComponent())
+	{
+		Proxy.SavePoseSnapshot(SkeletalMeshComponent, SnapshotName);	
+	}
+}
+
 USkeletalMeshComponent* UAnimInstance::GetOwningComponent() const
 {
 	return GetSkelMeshComponent();
@@ -1663,7 +1672,7 @@ float UAnimInstance::PlaySlotAnimation(UAnimSequenceBase* Asset, FName SlotNodeN
 	if (!bValidAsset)
 	{
 		// user warning
-		UE_LOG(LogAnimMontage, Warning, TEXT("Invalid Asset. If Montage, use Montage_Play"));
+		UE_LOG(LogAnimMontage, Warning, TEXT("PlaySlotAnimation: Invalid input asset(%s). If Montage, please use Montage_Play"), *GetNameSafe(Asset));
 		return 0.f;
 	}
 
@@ -1724,7 +1733,7 @@ UAnimMontage* UAnimInstance::PlaySlotAnimationAsDynamicMontage(UAnimSequenceBase
 	if (!bValidAsset)
 	{
 		// user warning
-		UE_LOG(LogAnimMontage, Warning, TEXT("Invalid Asset. If Montage, use Montage_Play"));
+		UE_LOG(LogAnimMontage, Warning, TEXT("PlaySlotAnimationAsDynamicMontage: Invalid input asset(%s). If Montage, please use Montage_Play"), *GetNameSafe(Asset));
 		return nullptr;
 	}
 
diff --git a/Engine/Source/Runtime/Engine/Private/Animation/AnimInstanceProxy.cpp b/Engine/Source/Runtime/Engine/Private/Animation/AnimInstanceProxy.cpp
--- a/Engine/Source/Runtime/Engine/Private/Animation/AnimInstanceProxy.cpp
+++ b/Engine/Source/Runtime/Engine/Private/Animation/AnimInstanceProxy.cpp
@@ -237,6 +237,42 @@ void FAnimInstanceProxy::PreUpdate(UAnimInstance* InAnimInstance, float DeltaSec
 	}
 }
 
+void FAnimInstanceProxy::SavePoseSnapshot(USkeletalMeshComponent* InSkeletalMeshComponent, FName SnapshotName)
+{
+	FPoseSnapshot* PoseSnapshot = PoseSnapshots.FindByPredicate([SnapshotName](const FPoseSnapshot& PoseData) { return PoseData.SnapshotName == SnapshotName; });
+	if (PoseSnapshot == nullptr)
+	{
+		PoseSnapshot = new (PoseSnapshots) FPoseSnapshot();
+		PoseSnapshot->SnapshotName = SnapshotName;
+	}
+
+	const TArray<FTransform>& ComponentSpaceTMs = InSkeletalMeshComponent->GetComponentSpaceTransforms();
+	const FReferenceSkeleton& RefSkeleton = InSkeletalMeshComponent->SkeletalMesh->RefSkeleton;
+	const TArray<FTransform>& RefPoseSpaceBaseTMs = RefSkeleton.GetRefBonePose();
+
+	const int32 NumSpaceBases = ComponentSpaceTMs.Num();
+	PoseSnapshot->LocalTransforms.Reset(NumSpaceBases);
+	PoseSnapshot->LocalTransforms.AddUninitialized(NumSpaceBases);
+	PoseSnapshot->LocalTransforms[0] = ComponentSpaceTMs[0];	//Set root bone which is always evaluated.
+
+	int32 CurrentRequiredBone = 1;
+	for(int32 ComponentSpaceIdx = 1; ComponentSpaceIdx < NumSpaceBases; ++ComponentSpaceIdx)
+	{
+		const bool bBoneHasEvaluated = InSkeletalMeshComponent->FillComponentSpaceTransformsRequiredBones.IsValidIndex(CurrentRequiredBone) && ComponentSpaceIdx == InSkeletalMeshComponent->FillComponentSpaceTransformsRequiredBones[CurrentRequiredBone];
+		const int32 ParentIndex = RefSkeleton.GetParentIndex(ComponentSpaceIdx);
+		ensureMsgf(ParentIndex != INDEX_NONE, TEXT("Getting an invalid parent bone for bone %d, but this should not be possible since this is not the root bone!"), ComponentSpaceIdx);
+
+		const FTransform& ParentTransform = ComponentSpaceTMs[ParentIndex];
+		const FTransform& ChildTransform =  ComponentSpaceTMs[ComponentSpaceIdx];
+		PoseSnapshot->LocalTransforms[ComponentSpaceIdx] = bBoneHasEvaluated ? ChildTransform.GetRelativeTransform(ParentTransform) :  RefPoseSpaceBaseTMs[ComponentSpaceIdx];
+
+		if (bBoneHasEvaluated)
+		{
+			CurrentRequiredBone++;
+		}
+	}
+}
+
 void FAnimInstanceProxy::PostUpdate(UAnimInstance* InAnimInstance) const
 {
 #if WITH_EDITORONLY_DATA
@@ -1701,4 +1737,9 @@ void FAnimInstanceProxy::RegisterWatchedPose(const FCSPose<FCompactPose>& Pose,
 }
 #endif
 
+const FAnimInstanceProxy::FPoseSnapshot* FAnimInstanceProxy::GetPoseSnapshot(FName SnapshotName) const
+{
+	return PoseSnapshots.FindByPredicate([SnapshotName](const FPoseSnapshot& PoseData) { return PoseData.SnapshotName == SnapshotName; });
+}
+
 #undef LOCTEXT_NAMESPACE
diff --git a/Engine/Source/Runtime/Engine/Private/Animation/AnimNotifyState_Trail.cpp b/Engine/Source/Runtime/Engine/Private/Animation/AnimNotifyState_Trail.cpp
--- a/Engine/Source/Runtime/Engine/Private/Animation/AnimNotifyState_Trail.cpp
+++ b/Engine/Source/Runtime/Engine/Private/Animation/AnimNotifyState_Trail.cpp
@@ -27,15 +27,13 @@ static void GetCandidateSystems(USkeletalMeshComponent& MeshComp, ParticleSystem
 	else
 	{
 		// No actor owner in some editor windows. Get PSCs spawned by the MeshComp.
-		TArray<UObject*> Children;
-		GetObjectsWithOuter(&MeshComp, Children, false, RF_NoFlags, EInternalObjectFlags::PendingKill);
-		for (UObject* Child : Children)
+		ForEachObjectWithOuter(&MeshComp, [&Components](UObject* Child)
 		{
 			if (UParticleSystemComponent* ChildPSC = Cast<UParticleSystemComponent>(Child))
 			{
 				Components.Add(ChildPSC);
 			}
-		}
+		}, false, RF_NoFlags, EInternalObjectFlags::PendingKill);
 	}
 }
 
diff --git a/Engine/Source/Runtime/Engine/Private/Animation/SkeletalControl.cpp b/Engine/Source/Runtime/Engine/Private/Animation/SkeletalControl.cpp
--- a/Engine/Source/Runtime/Engine/Private/Animation/SkeletalControl.cpp
+++ b/Engine/Source/Runtime/Engine/Private/Animation/SkeletalControl.cpp
@@ -27,6 +27,7 @@ bool FBoneReference::Initialize(const FBoneContainer& RequiredBones)
 	BoneName = *BoneName.ToString().Trim().TrimTrailing();
 	BoneIndex = RequiredBones.GetPoseBoneIndexForBoneName(BoneName);
 
+	bUseSkeletonIndex = false;
 	// If bone name is not found, look into the master skeleton to see if it's found there.
 	// SkeletalMeshes can exclude bones from the master skeleton, and that's OK.
 	// If it's not found in the master skeleton, the bone does not exist at all! so we should report it as a warning.
@@ -51,6 +52,7 @@ bool FBoneReference::Initialize(const USkeleton* Skeleton)
 	{
 		BoneName = *BoneName.ToString().Trim().TrimTrailing();
 		BoneIndex = Skeleton->GetReferenceSkeleton().FindBoneIndex(BoneName);
+		bUseSkeletonIndex = true;
 	}
 	else
 	{
diff --git a/Engine/Source/Runtime/Engine/Private/Animation/Skeleton.cpp b/Engine/Source/Runtime/Engine/Private/Animation/Skeleton.cpp
--- a/Engine/Source/Runtime/Engine/Private/Animation/Skeleton.cpp
+++ b/Engine/Source/Runtime/Engine/Private/Animation/Skeleton.cpp
@@ -88,6 +88,9 @@ void USkeleton::PostLoad()
 
 	const bool bRebuildNameMap = false;
 	ReferenceSkeleton.RebuildRefSkeleton(this, bRebuildNameMap);
+
+	// refresh linked bone indices
+	SmartNames.InitializeCurveMetaData(this);
 }
 
 void USkeleton::PostDuplicate(bool bDuplicateForPIE)
@@ -1429,6 +1432,21 @@ const FCurveMetaData* USkeleton::GetCurveMetaData(const FName& CurveName) const
 	return nullptr;
 }
 
+const FCurveMetaData* USkeleton::GetCurveMetaData(const SmartName::UID_Type CurveUID) const
+{
+	const FSmartNameMapping* Mapping = SmartNames.GetContainerInternal(USkeleton::AnimCurveMappingName);
+	if (ensureAlways(Mapping))
+	{
+		FSmartName SmartName;
+		if (Mapping->FindSmartNameByUID(CurveUID, SmartName))
+		{
+			return Mapping->GetCurveMetaData(SmartName.DisplayName);
+		}
+	}
+
+	return nullptr;
+}
+
 FCurveMetaData* USkeleton::GetCurveMetaData(const FSmartName& CurveName)
 {
 	FSmartNameMapping* Mapping = SmartNames.GetContainerInternal(USkeleton::AnimCurveMappingName);
diff --git a/Engine/Source/Runtime/Engine/Private/Animation/SmartName.cpp b/Engine/Source/Runtime/Engine/Private/Animation/SmartName.cpp
--- a/Engine/Source/Runtime/Engine/Private/Animation/SmartName.cpp
+++ b/Engine/Source/Runtime/Engine/Private/Animation/SmartName.cpp
@@ -296,6 +296,20 @@ bool FSmartNameMapping::FindSmartNameByUID(SmartName::UID_Type UID, FSmartName&
 
 	return false;
 }
+
+/* initialize curve meta data for the container */
+void FSmartNameMapping::InitializeCurveMetaData(class USkeleton* Skeleton)
+{
+	// initialize bone indices for skeleton
+	for (TPair<FName, FCurveMetaData>& Iter : CurveMetaDataMap)
+	{
+		FCurveMetaData& CurveMetaData = Iter.Value;
+		for (int32 LinkedBoneIndex = 0; LinkedBoneIndex < CurveMetaData.LinkedBones.Num(); ++LinkedBoneIndex)
+		{
+			CurveMetaData.LinkedBones[LinkedBoneIndex].Initialize(Skeleton);
+		}
+	}
+}
 ////////////////////////////////////////////////////////////////////////
 //
 // FSmartNameContainer
@@ -329,6 +343,16 @@ const FSmartNameMapping* FSmartNameContainer::GetContainerInternal(const FName&
 	return NameMappings.Find(ContainerName);
 }
 
+/* initialize curve meta data for the container */
+void FSmartNameContainer::InitializeCurveMetaData(class USkeleton* Skeleton)
+{
+	FSmartNameMapping* CurveMappingTable = GetContainerInternal(USkeleton::AnimCurveMappingName);
+	if (CurveMappingTable)
+	{
+		CurveMappingTable->InitializeCurveMetaData(Skeleton);
+	}
+}
+
 ////////////////////////////////////////////////////////////////////////
 //
 // FSmartName
diff --git a/Engine/Source/Runtime/Engine/Private/Audio.cpp b/Engine/Source/Runtime/Engine/Private/Audio.cpp
--- a/Engine/Source/Runtime/Engine/Private/Audio.cpp
+++ b/Engine/Source/Runtime/Engine/Private/Audio.cpp
@@ -41,7 +41,7 @@ DEFINE_STAT(STAT_VorbisDecompressTime);
 DEFINE_STAT(STAT_VorbisPrepareDecompressionTime);
 DEFINE_STAT(STAT_AudioDecompressTime);
 DEFINE_STAT(STAT_AudioPrepareDecompressionTime);
-DEFINE_STAT(STAT_OpusDecompressTime);
+DEFINE_STAT(STAT_AudioStreamedDecompressTime);
 
 DEFINE_STAT(STAT_AudioUpdateEffects);
 DEFINE_STAT(STAT_AudioEvaluateConcurrency);
@@ -780,7 +780,7 @@ FString FWaveInstance::GetName() const
 #endif
 
 // Main Riff-Wave header.
-struct FRiffWaveHeader
+struct FRiffWaveHeaderChunk
 {
 	uint32	rID;			// Contains 'RIFF'
 	uint32	ChunkLen;		// Remaining length of the entire riff chunk (= file).
@@ -795,7 +795,7 @@ struct FRiffChunkOld
 };
 
 // ChunkID: 'fmt ' ("WaveFormatEx" structure )
-struct FFormatChunk
+struct FRiffFormatChunk
 {
 	uint16   wFormatTag;        // Format type: 1 = PCM
 	uint16   nChannels;         // Number of channels (i.e. mono, stereo...).
@@ -832,7 +832,7 @@ struct FSubformatGUID
 // ChunkID: 'fmt ' ("WaveFormatExtensible" structure)
 struct FExtendedFormatChunk
 {
-	FFormatChunk Format;			// Standard WaveFormatEx ('fmt ') chunk, with
+	FRiffFormatChunk Format;			// Standard WaveFormatEx ('fmt ') chunk, with
 									// wFormatTag == WAVE_FORMAT_EXTENSIBLE and cbSize == 22
 	union
 	{
@@ -844,20 +844,6 @@ struct FExtendedFormatChunk
 	FSubformatGUID SubFormat;		// Subformat identifier.
 };
 
-// ChunkID: 'smpl'
-struct FSampleChunk
-{
-	uint32   dwManufacturer;
-	uint32   dwProduct;
-	uint32   dwSamplePeriod;
-	uint32   dwMIDIUnityNote;
-	uint32   dwMIDIPitchFraction;
-	uint32	dwSMPTEFormat;
-	uint32   dwSMPTEOffset;		//
-	uint32   cSampleLoops;		// Number of tSampleLoop structures following this chunk
-	uint32   cbSamplerData;		//
-};
-
 #if PLATFORM_SUPPORTS_PRAGMA_PACK
 #pragma pack(pop)
 #endif
@@ -867,9 +853,9 @@ struct FSampleChunk
 //
 bool FWaveModInfo::ReadWaveInfo( uint8* WaveData, int32 WaveDataSize, FString* ErrorReason )
 {
-	FFormatChunk* FmtChunk;
+	FRiffFormatChunk* FmtChunk;
 	FExtendedFormatChunk* FmtChunkEx = nullptr;
-	FRiffWaveHeader* RiffHdr = ( FRiffWaveHeader* )WaveData;
+	FRiffWaveHeaderChunk* RiffHdr = (FRiffWaveHeaderChunk* )WaveData;
 	WaveDataEnd = WaveData + WaveDataSize;
 
 	if( WaveDataSize == 0 )
@@ -924,7 +910,7 @@ bool FWaveModInfo::ReadWaveInfo( uint8* WaveData, int32 WaveDataSize, FString* E
 		return( false );
 	}
 
-	FmtChunk = ( FFormatChunk* )( ( uint8* )RiffChunk + 8 );
+	FmtChunk = ( FRiffFormatChunk* )( ( uint8* )RiffChunk + 8 );
 #if !PLATFORM_LITTLE_ENDIAN
 	if( !AlreadySwapped )
 	{
diff --git a/Engine/Source/Runtime/Engine/Private/AudioDecompress.cpp b/Engine/Source/Runtime/Engine/Private/AudioDecompress.cpp
--- a/Engine/Source/Runtime/Engine/Private/AudioDecompress.cpp
+++ b/Engine/Source/Runtime/Engine/Private/AudioDecompress.cpp
@@ -7,6 +7,339 @@
 #include "Sound/SoundWave.h"
 #include "IAudioFormat.h"
 
+IStreamedCompressedInfo::IStreamedCompressedInfo()
+	: SrcBufferData(nullptr)
+	, SrcBufferDataSize(0)
+	, SrcBufferOffset(0)
+	, AudioDataOffset(0)
+	, SampleRate(0)
+	, TrueSampleCount(0)
+	, CurrentSampleCount(0)
+	, NumChannels(0)
+	, MaxFrameSizeSamples(0)
+	, SampleStride(0)
+	, LastPCMByteSize(0)
+	, LastPCMOffset(0)
+	, bStoringEndOfFile(false)
+	, StreamingSoundWave(nullptr)
+	, CurrentChunkIndex(0)
+{
+}
+
+uint32 IStreamedCompressedInfo::Read(void *OutBuffer, uint32 DataSize)
+{
+	uint32 BytesToRead = FMath::Min(DataSize, SrcBufferDataSize - SrcBufferOffset);
+	if (BytesToRead > 0)
+	{
+		FMemory::Memcpy(OutBuffer, SrcBufferData + SrcBufferOffset, BytesToRead);
+		SrcBufferOffset += BytesToRead;
+	}
+	return BytesToRead;
+}
+
+bool IStreamedCompressedInfo::ReadCompressedInfo(const uint8* InSrcBufferData, uint32 InSrcBufferDataSize, FSoundQualityInfo* QualityInfo)
+{
+	check(!SrcBufferData);
+
+	SCOPE_CYCLE_COUNTER(STAT_AudioStreamedDecompressTime);
+
+	// Parse the format header, this is done different for each format
+	if (!ParseHeader(InSrcBufferData, InSrcBufferDataSize, QualityInfo))
+	{
+		return false;
+	}
+
+	// After parsing the header, the SrcBufferData should be none-null
+	check(SrcBufferData != nullptr);
+
+	// Sample Stride is 
+	SampleStride = NumChannels * sizeof(int16);
+
+	MaxFrameSizeSamples = GetMaxFrameSizeSamples();
+		
+	LastDecodedPCM.Empty(MaxFrameSizeSamples * SampleStride);
+	LastDecodedPCM.AddUninitialized(MaxFrameSizeSamples * SampleStride);
+
+	return CreateDecoder();
+}
+
+bool IStreamedCompressedInfo::ReadCompressedData(uint8* Destination, bool bLooping, uint32 BufferSize)
+{
+	check(Destination);
+
+	SCOPE_CYCLE_COUNTER(STAT_AudioStreamedDecompressTime);
+
+	// Write out any PCM data that was decoded during the last request
+	uint32 RawPCMOffset = WriteFromDecodedPCM(Destination, BufferSize);
+
+	bool bLooped = false;
+
+	if (bStoringEndOfFile && LastPCMByteSize > 0)
+	{
+		// delayed returning looped because we hadn't read the entire buffer
+		bLooped = true;
+		bStoringEndOfFile = false;
+	}
+
+	while (RawPCMOffset < BufferSize)
+	{
+		uint16 FrameSize = GetFrameSize();
+		int32 DecodedSamples = DecompressToPCMBuffer(FrameSize);
+		if (DecodedSamples < 0)
+		{
+			LastPCMByteSize = 0;
+			ZeroBuffer(Destination + RawPCMOffset, BufferSize - RawPCMOffset);
+			return false;
+		}
+		else
+		{
+			LastPCMByteSize = IncrementCurrentSampleCount(DecodedSamples) * SampleStride;
+			RawPCMOffset += WriteFromDecodedPCM(Destination + RawPCMOffset, BufferSize - RawPCMOffset);
+
+			if (SrcBufferOffset >= SrcBufferDataSize)
+			{
+				// check whether all decoded PCM was written
+				if (LastPCMByteSize == 0)
+				{
+					bLooped = true;
+				}
+				else
+				{
+					bStoringEndOfFile = true;
+				}
+				if (bLooping)
+				{
+					SrcBufferOffset = AudioDataOffset;
+					CurrentSampleCount = 0;
+				}
+				else
+				{
+					RawPCMOffset += ZeroBuffer(Destination + RawPCMOffset, BufferSize - RawPCMOffset);
+				}
+			}
+		}
+	}
+
+	return bLooped;
+}
+
+void IStreamedCompressedInfo::ExpandFile(uint8* DstBuffer, struct FSoundQualityInfo* QualityInfo)
+{
+	check(DstBuffer);
+	check(QualityInfo);
+	check(QualityInfo->SampleDataSize <= SrcBufferDataSize);
+
+	// Ensure we're at the start of the audio data
+	SrcBufferOffset = AudioDataOffset;
+
+	uint32 RawPCMOffset = 0;
+
+	while (RawPCMOffset < QualityInfo->SampleDataSize)
+	{
+		uint16 FrameSize = 0;
+		Read(&FrameSize, sizeof(uint16));
+		int32 DecodedSamples = DecompressToPCMBuffer(FrameSize);
+
+		if (DecodedSamples < 0)
+		{
+			RawPCMOffset += ZeroBuffer(DstBuffer + RawPCMOffset, QualityInfo->SampleDataSize - RawPCMOffset);
+		}
+		else
+		{
+			LastPCMByteSize = IncrementCurrentSampleCount(DecodedSamples) * SampleStride;
+			RawPCMOffset += WriteFromDecodedPCM(DstBuffer + RawPCMOffset, QualityInfo->SampleDataSize - RawPCMOffset);
+		}
+	}
+}
+
+bool IStreamedCompressedInfo::StreamCompressedInfo(USoundWave* Wave, struct FSoundQualityInfo* QualityInfo)
+{
+	StreamingSoundWave = Wave;
+
+	// Get the first chunk of audio data (should always be loaded)
+	CurrentChunkIndex = 0;
+	const uint8* FirstChunk = IStreamingManager::Get().GetAudioStreamingManager().GetLoadedChunk(StreamingSoundWave, CurrentChunkIndex);
+
+	if (FirstChunk)
+	{
+		return ReadCompressedInfo(FirstChunk, Wave->RunningPlatformData->Chunks[0].DataSize, QualityInfo);
+	}
+
+	return false;
+}
+
+bool IStreamedCompressedInfo::StreamCompressedData(uint8* Destination, bool bLooping, uint32 BufferSize)
+{
+	check(Destination);
+
+	SCOPE_CYCLE_COUNTER(STAT_AudioStreamedDecompressTime);
+
+	UE_LOG(LogAudio, Log, TEXT("Streaming compressed data from SoundWave'%s' - Chunk %d, Offset %d"), *StreamingSoundWave->GetName(), CurrentChunkIndex, SrcBufferOffset);
+
+	// Write out any PCM data that was decoded during the last request
+	uint32 RawPCMOffset = WriteFromDecodedPCM(Destination, BufferSize);
+
+	// If next chunk wasn't loaded when last one finished reading, try to get it again now
+	if (SrcBufferData == NULL)
+	{
+		SrcBufferData = IStreamingManager::Get().GetAudioStreamingManager().GetLoadedChunk(StreamingSoundWave, CurrentChunkIndex);
+		if (SrcBufferData)
+		{
+			SrcBufferDataSize = StreamingSoundWave->RunningPlatformData->Chunks[CurrentChunkIndex].DataSize;
+			SrcBufferOffset = CurrentChunkIndex == 0 ? AudioDataOffset : 0;
+		}
+		else
+		{
+			// Still not loaded, zero remainder of current buffer
+			UE_LOG(LogAudio, Warning, TEXT("Unable to read from chunk %d of SoundWave'%s'"), CurrentChunkIndex, *StreamingSoundWave->GetName());
+			ZeroBuffer(Destination + RawPCMOffset, BufferSize - RawPCMOffset);
+			return false;
+		}
+	}
+
+	bool bLooped = false;
+
+	if (bStoringEndOfFile && LastPCMByteSize > 0)
+	{
+		// delayed returning looped because we hadn't read the entire buffer
+		bLooped = true;
+		bStoringEndOfFile = false;
+	}
+
+	while (RawPCMOffset < BufferSize)
+	{
+		// Get the platform-dependent size of the current "frame" of encoded audio (note: frame is used here differently than audio frame/sample)
+		uint16 FrameSize = GetFrameSize();
+
+		// Decompress the next compression frame of audio (many samples) into the PCM buffer
+		int32 DecodedSamples = DecompressToPCMBuffer(FrameSize);
+
+		if (DecodedSamples < 0)
+		{
+			LastPCMByteSize = 0;
+			ZeroBuffer(Destination + RawPCMOffset, BufferSize - RawPCMOffset);
+			return false;
+		}
+		else
+		{
+			LastPCMByteSize = IncrementCurrentSampleCount(DecodedSamples) * SampleStride;
+
+			RawPCMOffset += WriteFromDecodedPCM(Destination + RawPCMOffset, BufferSize - RawPCMOffset);
+
+			// Have we reached the end of buffer
+			if (SrcBufferOffset >= SrcBufferDataSize)
+			{
+				// Special case for the last chunk of audio
+				if (CurrentChunkIndex == StreamingSoundWave->RunningPlatformData->NumChunks - 1)
+				{
+					// check whether all decoded PCM was written
+					if (LastPCMByteSize == 0)
+					{
+						bLooped = true;
+					}
+					else
+					{
+						bStoringEndOfFile = true;
+					}
+
+					if (bLooping)
+					{
+						CurrentChunkIndex = 0;
+						SrcBufferOffset = AudioDataOffset;
+						CurrentSampleCount = 0;
+
+						// Prepare the decoder to begin looping.
+						PrepareToLoop();
+					}
+					else
+					{
+						RawPCMOffset += ZeroBuffer(Destination + RawPCMOffset, BufferSize - RawPCMOffset);
+					}
+				}
+				else
+				{
+					CurrentChunkIndex++;
+					SrcBufferOffset = 0;
+				}
+
+				SrcBufferData = IStreamingManager::Get().GetAudioStreamingManager().GetLoadedChunk(StreamingSoundWave, CurrentChunkIndex);
+				if (SrcBufferData)
+				{
+					UE_LOG(LogAudio, Log, TEXT("Incremented current chunk from SoundWave'%s' - Chunk %d, Offset %d"), *StreamingSoundWave->GetName(), CurrentChunkIndex, SrcBufferOffset);
+					SrcBufferDataSize = StreamingSoundWave->RunningPlatformData->Chunks[CurrentChunkIndex].DataSize;
+				}
+				else
+				{
+					SrcBufferDataSize = 0;
+					RawPCMOffset += ZeroBuffer(Destination + RawPCMOffset, BufferSize - RawPCMOffset);
+				}
+			}
+		}
+	}
+
+	return bLooped;
+}
+
+int32 IStreamedCompressedInfo::DecompressToPCMBuffer(uint16 FrameSize)
+{
+	if (SrcBufferOffset + FrameSize > SrcBufferDataSize)
+	{
+		// if frame size is too large, something has gone wrong
+		return -1;
+	}
+
+	const uint8* SrcPtr = SrcBufferData + SrcBufferOffset;
+	SrcBufferOffset += FrameSize;
+	LastPCMOffset = 0;
+	
+	const int32 SamplesDecoded = Decode(SrcPtr, FrameSize, (int16*)LastDecodedPCM.GetData(), MaxFrameSizeSamples);
+	return SamplesDecoded;
+}
+
+uint32 IStreamedCompressedInfo::IncrementCurrentSampleCount(uint32 NewSamples)
+{
+	if (CurrentSampleCount + NewSamples > TrueSampleCount)
+	{
+		NewSamples = TrueSampleCount - CurrentSampleCount;
+		CurrentSampleCount = TrueSampleCount;
+	}
+	else
+	{
+		CurrentSampleCount += NewSamples;
+	}
+	return NewSamples;
+}
+
+uint32 IStreamedCompressedInfo::WriteFromDecodedPCM(uint8* Destination, uint32 BufferSize)
+{
+	uint32 BytesToCopy = FMath::Min(BufferSize, LastPCMByteSize - LastPCMOffset);
+	if (BytesToCopy > 0)
+	{
+		check(BytesToCopy <= LastDecodedPCM.Num() - LastPCMOffset);
+		FMemory::Memcpy(Destination, LastDecodedPCM.GetData() + LastPCMOffset, BytesToCopy);
+		LastPCMOffset += BytesToCopy;
+		if (LastPCMOffset >= LastPCMByteSize)
+		{
+			LastPCMOffset = 0;
+			LastPCMByteSize = 0;
+		}
+	}
+	return BytesToCopy;
+}
+
+uint32 IStreamedCompressedInfo::ZeroBuffer(uint8* Destination, uint32 BufferSize)
+{
+	check(Destination);
+
+	if (BufferSize > 0)
+	{
+		FMemory::Memzero(Destination, BufferSize);
+		return BufferSize;
+	}
+	return 0;
+}
+
+
 //////////////////////////////////////////////////////////////////////////
 // Copied from IOS - probably want to split and share
 //////////////////////////////////////////////////////////////////////////
diff --git a/Engine/Source/Runtime/Engine/Private/AudioDevice.cpp b/Engine/Source/Runtime/Engine/Private/AudioDevice.cpp
--- a/Engine/Source/Runtime/Engine/Private/AudioDevice.cpp
+++ b/Engine/Source/Runtime/Engine/Private/AudioDevice.cpp
@@ -1661,7 +1661,9 @@ bool FAudioDevice::TryClearingSoundMix(USoundMix* SoundMix, FSoundMixState* Soun
 {
 	if (SoundMix && SoundMixState)
 	{
-		if (SoundMixState->ActiveRefCount == 0 && SoundMixState->PassiveRefCount == 0 && SoundMixState->IsBaseSoundMix == false)
+		// Only manually clear the sound mix if it's no longer referenced and if the duration was not set.
+		// If the duration was set by sound designer, let the sound mix clear itself up automatically.
+		if (SoundMix->Duration < 0.0f && SoundMixState->ActiveRefCount == 0 && SoundMixState->PassiveRefCount == 0 && SoundMixState->IsBaseSoundMix == false)
 		{
 			// do whatever is needed to remove influence of this SoundMix
 			if (SoundMix->FadeOutTime > 0.f)
@@ -2450,15 +2452,6 @@ void FAudioDevice::PopSoundMixModifier(USoundMix* SoundMix, bool bIsPassive)
 			if (bIsPassive && SoundMixState->PassiveRefCount > 0)
 			{
 				SoundMixState->PassiveRefCount--;
-				if (SoundMixState->PassiveRefCount == 0)
-				{
-					// Check whether Fade out time was previously set and reset it to current time
-					if (SoundMixState->FadeOutStartTime >= 0.f && FApp::GetCurrentTime() > SoundMixState->FadeOutStartTime)
-					{
-						SoundMixState->FadeOutStartTime = FApp::GetCurrentTime();
-						SoundMixState->EndTime = SoundMixState->FadeOutStartTime + SoundMix->FadeOutTime;
-					}
-				}
 			}
 			else if (!bIsPassive && SoundMixState->ActiveRefCount > 0)
 			{
@@ -3105,6 +3098,7 @@ void FAudioDevice::SendUpdateResultsToGameThread(const int32 FirstActiveIndex)
 				WaveInstanceInfo.Description = Source ? Source->Describe((RequestedAudioStats & ERequestedAudioStats::LongSoundNames) != 0) : FString(TEXT("No source"));
 				WaveInstanceInfo.ActualVolume = WaveInstance->GetVolume();
 				WaveInstanceInfo.InstanceIndex = InstanceIndex;
+				WaveInstanceInfo.WaveInstanceName = *WaveInstance->GetName();
 				StatSoundInfos[*SoundInfoIndex].WaveInstanceInfos.Add(MoveTemp(WaveInstanceInfo));
 			}
 		}
diff --git a/Engine/Source/Runtime/Engine/Private/AudioDeviceManager.cpp b/Engine/Source/Runtime/Engine/Private/AudioDeviceManager.cpp
--- a/Engine/Source/Runtime/Engine/Private/AudioDeviceManager.cpp
+++ b/Engine/Source/Runtime/Engine/Private/AudioDeviceManager.cpp
@@ -586,6 +586,30 @@ void FAudioDeviceManager::ToggleVisualize3dDebug()
 	bVisualize3dDebug = !bVisualize3dDebug;
 }
 
+void FAudioDeviceManager::ToggleDebugStat(const uint8 StatBitMask)
+{
+#if !UE_BUILD_SHIPPING
+	if (!IsInAudioThread())
+	{
+		FAudioDeviceManager* AudioDeviceManager = this;
+		FAudioThread::RunCommandOnAudioThread([AudioDeviceManager, StatBitMask]()
+		{
+			AudioDeviceManager->ToggleDebugStat(StatBitMask);
+		});
+
+		return;
+	}
+
+	for (FAudioDevice* AudioDevice : Devices)
+	{
+		if (AudioDevice)
+		{
+			AudioDevice->UpdateRequestedStat(StatBitMask);
+		}
+	}
+#endif
+}
+
 void FAudioDeviceManager::SetDebugSoloSoundClass(const TCHAR* SoundClassName)
 {
 	if (!IsInAudioThread())
