Copying //UE4/Dev-Framework to //UE4/Dev-Main (Source: //UE4/Dev-Framework @ 3182037)

diff --git a/Engine/Source/Runtime/Engine/Private/AudioThread.cpp b/Engine/Source/Runtime/Engine/Private/AudioThread.cpp
--- a/Engine/Source/Runtime/Engine/Private/AudioThread.cpp
+++ b/Engine/Source/Runtime/Engine/Private/AudioThread.cpp
@@ -41,7 +41,7 @@ struct FAudioThreadInteractor
 			{
 				UE_LOG(LogAudio, Warning, TEXT("Audio threading is disabled in the editor."));
 			}
-			else
+			else if (!FAudioThread::IsUsingThreadedAudio())
 			{
 				UE_LOG(LogAudio, Warning, TEXT("Cannot manipulate audio thread when disabled by platform or ini."));
 			}
@@ -52,12 +52,10 @@ struct FAudioThreadInteractor
 static FAutoConsoleVariableSink CVarUseAudioThreadSink(FConsoleCommandDelegate::CreateStatic(&FAudioThreadInteractor::UseAudioThreadCVarSinkFunction));
 
 bool FAudioThread::bIsAudioThreadRunning = false;
+bool FAudioThread::bIsAudioThreadSuspended = false;
 bool FAudioThread::bUseThreadedAudio = false;
-uint32 FAudioThread::SuspensionCount = 0;
 uint32 FAudioThread::CachedAudioThreadId = 0;
 FRunnable* FAudioThread::AudioThreadRunnable = nullptr;
-FString FAudioThread::AudioThreadError;
-volatile bool FAudioThread::bIsAudioThreadHealthy = true;
 
 
 /** The audio thread main loop */
@@ -72,25 +70,8 @@ void AudioThreadMain( FEvent* TaskGraphBoundSyncEvent )
 		TaskGraphBoundSyncEvent->Trigger();
 	}
 
-	// set the thread back to real time mode
-	FPlatformProcess::SetRealTimeMode();
-
-#if STATS
-	if (FThreadStats::WillEverCollectData())
-	{
-		FThreadStats::ExplicitFlush(); // flush the stats and set update the scope so we don't flush again until a frame update, this helps prevent fragmentation
-	}
-#endif
-
 	FTaskGraphInterface::Get().ProcessThreadUntilRequestReturn(ENamedThreads::AudioThread);
 	FPlatformMisc::MemoryBarrier();
-	
-#if STATS
-	if (FThreadStats::WillEverCollectData())
-	{
-		FThreadStats::ExplicitFlush(); // Another explicit flush to clean up the ScopeCount established above for any stats lingering since the last frame
-	}
-#endif
 }
 
 FAudioThread::FAudioThread()
@@ -112,11 +93,10 @@ FAudioThread::~FAudioThread()
 
 void FAudioThread::SuspendAudioThread()
 {
-	check(IsInGameThread());
+	check(FPlatformTLS::GetCurrentThreadId() == GGameThreadId);
+	check(!bIsAudioThreadSuspended || CVarSuspendAudioThread.GetValueOnGameThread() != 0);
 	if (IsAudioThreadRunning())
 	{
-		++SuspensionCount;
-
 		// Make GC wait on the audio thread finishing processing
 		FAudioCommandFence AudioFence;
 		AudioFence.BeginFence();
@@ -124,26 +104,23 @@ void FAudioThread::SuspendAudioThread()
 
 		CachedAudioThreadId = GAudioThreadId;
 		GAudioThreadId = 0; // While we are suspended we will pretend we have no audio thread
-		bIsAudioThreadRunning = false;
-
+		bIsAudioThreadSuspended = true;
 		FPlatformMisc::MemoryBarrier();
+		bIsAudioThreadRunning = false;
 	}
+	check(!bIsAudioThreadRunning);
 }
 
 void FAudioThread::ResumeAudioThread()
 {
-	check(IsInGameThread());
-	if (SuspensionCount > 0)
+	check(FPlatformTLS::GetCurrentThreadId() == GGameThreadId);
+	if (bIsAudioThreadSuspended && CVarSuspendAudioThread.GetValueOnGameThread() == 0)
 	{
-		--SuspensionCount;
-
-		if (SuspensionCount == 0)
-		{
-			GAudioThreadId = CachedAudioThreadId;
-			bIsAudioThreadRunning = true;
-
-			FPlatformMisc::MemoryBarrier();
-		}
+		GAudioThreadId = CachedAudioThreadId;
+		CachedAudioThreadId = 0;
+		bIsAudioThreadSuspended = false;
+		FPlatformMisc::MemoryBarrier();
+		bIsAudioThreadRunning = true;
 	}
 }
 
@@ -170,37 +147,8 @@ void FAudioThread::Exit()
 
 uint32 FAudioThread::Run()
 {
-	FMemory::SetupTLSCachesOnCurrentThread();
 	FPlatformProcess::SetupAudioThread();
-
-#if PLATFORM_WINDOWS
-	if ( !FPlatformMisc::IsDebuggerPresent() || GAlwaysReportCrash )
-	{
-#if !PLATFORM_SEH_EXCEPTIONS_DISABLED
-		__try
-#endif
-		{
-			AudioThreadMain( TaskGraphBoundSyncEvent );
-		}
-#if !PLATFORM_SEH_EXCEPTIONS_DISABLED
-		__except(ReportCrash(GetExceptionInformation()))
-		{
-			AudioThreadError = GErrorHist;
-
-			// Use a memory barrier to ensure that the game thread sees the write to AudioThreadError before
-			// the write to GIsAudioThreadHealthy.
-			FPlatformMisc::MemoryBarrier();
-
-			bIsAudioThreadHealthy = false;
-		}
-#endif
-	}
-	else
-#endif // PLATFORM_WINDOWS
-	{
-		AudioThreadMain( TaskGraphBoundSyncEvent );
-	}
-	FMemory::ClearAndDisableTLSCachesOnCurrentThread();
+	AudioThreadMain( TaskGraphBoundSyncEvent );
 	return 0;
 }
 
@@ -218,6 +166,7 @@ void FAudioThread::SetUseThreadedAudio(const bool bInUseThreadedAudio)
 
 void FAudioThread::RunCommandOnAudioThread(TFunction<void()> InFunction, const TStatId InStatId)
 {
+	check(FPlatformTLS::GetCurrentThreadId() == GGameThreadId);
 	if (bIsAudioThreadRunning)
 	{
 		FFunctionGraphTask::CreateAndDispatchWhenReady(MoveTemp(InFunction), InStatId, nullptr, ENamedThreads::AudioThread);
@@ -233,34 +182,36 @@ void FAudioThread::RunCommandOnGameThread(TFunction<void()> InFunction, const TS
 {
 	if (bIsAudioThreadRunning)
 	{
+		check(GAudioThreadId && FPlatformTLS::GetCurrentThreadId() == GAudioThreadId);
 		FFunctionGraphTask::CreateAndDispatchWhenReady(MoveTemp(InFunction), InStatId, nullptr, ENamedThreads::GameThread);
 	}
 	else
 	{
+		check(FPlatformTLS::GetCurrentThreadId() == GGameThreadId);
 		FScopeCycleCounter ScopeCycleCounter(InStatId);
 		InFunction();
 	}
 }
 
-static FString BuildAudioThreadName( uint32 ThreadIndex )
-{
-	return FString::Printf( TEXT( "%s %u" ), *FName( NAME_AudioThread ).GetPlainNameString(), ThreadIndex );
-}
-
 void FAudioThread::StartAudioThread()
 {
+	check(FPlatformTLS::GetCurrentThreadId() == GGameThreadId);
+
+	check(!bIsAudioThreadRunning);
+	check(!bIsAudioThreadSuspended);
 	if (bUseThreadedAudio)
 	{
 		check(GAudioThread == nullptr);
 
 		static uint32 ThreadCount = 0;
+		check(!ThreadCount); // we should not stop and restart the audio thread; it is complexity we don't need.
 
 		bIsAudioThreadRunning = true;
 
 		// Create the audio thread.
 		AudioThreadRunnable = new FAudioThread();
 
-		GAudioThread = FRunnableThread::Create(AudioThreadRunnable, *BuildAudioThreadName(ThreadCount), 0, TPri_BelowNormal, FPlatformAffinity::GetAudioThreadMask());
+		GAudioThread = FRunnableThread::Create(AudioThreadRunnable, *FName(NAME_AudioThread).GetPlainNameString(), 0, TPri_BelowNormal, FPlatformAffinity::GetAudioThreadMask());
 
 		// Wait for audio thread to have taskgraph bound before we dispatch any tasks for it.
 		((FAudioThread*)AudioThreadRunnable)->TaskGraphBoundSyncEvent->Wait();
@@ -271,16 +222,22 @@ void FAudioThread::StartAudioThread()
 		Fence.Wait();
 
 		ThreadCount++;
+
+		if (CVarSuspendAudioThread.GetValueOnGameThread() != 0)
+		{
+			SuspendAudioThread();
+		}
 	}
 }
 
 void FAudioThread::StopAudioThread()
 {
-	// This function is not thread-safe. Ensure it is only called by the main game thread.
-	check( IsInGameThread() );
+	check(FPlatformTLS::GetCurrentThreadId() == GGameThreadId);
+	check(!bIsAudioThreadSuspended || CVarSuspendAudioThread.GetValueOnGameThread() != 0);
 
-	if (!bIsAudioThreadRunning)
+	if (!bIsAudioThreadRunning && CachedAudioThreadId == 0)
 	{
+		check(!bUseThreadedAudio);
 		return;
 	}
 
@@ -289,16 +246,6 @@ void FAudioThread::StopAudioThread()
 
 	FGraphEventRef QuitTask = TGraphTask<FReturnGraphTask>::CreateTask(nullptr, ENamedThreads::GameThread).ConstructAndDispatchWhenReady(ENamedThreads::AudioThread);
 
-	// Busy wait while BP debugging, to avoid opportunistic execution of game thread tasks
-	// If the game thread is already executing tasks, then we have no choice but to spin
-	if (GIntraFrameDebuggingGameThread || FTaskGraphInterface::Get().IsThreadProcessingTasks(ENamedThreads::GameThread) ) 
-	{
-		while ((QuitTask.GetReference() != nullptr) && !QuitTask->IsComplete())
-		{
-			FPlatformProcess::Sleep(0.0f);
-		}
-	}
-	else
 	{
 		QUICK_SCOPE_CYCLE_COUNTER(STAT_StopAudioThread);
 		FTaskGraphInterface::Get().WaitUntilTaskCompletes(QuitTask, ENamedThreads::GameThread_Local);
@@ -317,30 +264,9 @@ void FAudioThread::StopAudioThread()
 	AudioThreadRunnable = nullptr;
 }
 
-void FAudioThread::CheckAudioThreadHealth()
-{
-	if(!bIsAudioThreadHealthy)
-	{
-		GErrorHist[0] = 0;
-		GIsCriticalError = false;
-		UE_LOG(LogAudio, Fatal, TEXT("Audio thread exception:\r\n%s"),*AudioThreadError);
-	}
-
-	if (IsInGameThread())
-	{
-		GLog->FlushThreadedLogs();
-		SCOPE_CYCLE_COUNTER(STAT_PumpMessages);
-		FPlatformMisc::PumpMessages(false);
-	}
-}
-
-bool FAudioThread::IsAudioThreadHealthy()
-{
-	return bIsAudioThreadHealthy;
-}
-
 void FAudioCommandFence::BeginFence()
 {
+	check(FPlatformTLS::GetCurrentThreadId() == GGameThreadId); // this could be relaxed, but for now, we are going to require all fences are set from the GT
 	if (FAudioThread::IsAudioThreadRunning())
 	{
 		DECLARE_CYCLE_STAT(TEXT("FNullGraphTask.FenceAudioCommand"),
@@ -350,101 +276,47 @@ void FAudioCommandFence::BeginFence()
 		CompletionEvent = TGraphTask<FNullGraphTask>::CreateTask(NULL, ENamedThreads::GameThread).ConstructAndDispatchWhenReady(
 			GET_STATID(STAT_FNullGraphTask_FenceAudioCommand), ENamedThreads::AudioThread);
 	}
+	else
+	{
+		CompletionEvent = nullptr;
+	}
 }
 
 bool FAudioCommandFence::IsFenceComplete() const
 {
-	if (!FAudioThread::IsAudioThreadRunning())
-	{
-		return true;
-	}
-	check(IsInGameThread() || IsInAsyncLoadingThread());
-	FAudioThread::CheckAudioThreadHealth();
+	check(FPlatformTLS::GetCurrentThreadId() == GGameThreadId); // this could be relaxed, but for now, we are going to require all fences are set from the GT
 	if (!CompletionEvent.GetReference() || CompletionEvent->IsComplete())
 	{
-		CompletionEvent = NULL; // this frees the handle for other uses, the NULL state is considered completed
+		CompletionEvent = nullptr; // this frees the handle for other uses, the NULL state is considered completed
 		return true;
 	}
+	check(FAudioThread::IsAudioThreadRunning());
 	return false;
 }
 
-static int32 GTimeToBlockOnAudioFence = 1;
-static FAutoConsoleVariableRef CVarTimeToBlockOnAudioFence(
-	TEXT("g.TimeToBlockOnAudioFence"),
-	GTimeToBlockOnAudioFence,
-	TEXT("Number of milliseconds the game thread should block when waiting on a audio thread fence.")
-	);
-
-/**
- * Block the game thread waiting for a task to finish on the audio thread.
- */
-static void GameThreadWaitForTask(const FGraphEventRef& Task, bool bEmptyGameThreadTasks = false)
-{
-	check(IsInGameThread());
-	check(IsValidRef(Task));
-
-	if (!Task->IsComplete())
-	{
-		SCOPE_CYCLE_COUNTER(STAT_GameIdleTime);
-		{
-			static int32 NumRecursiveCalls = 0;
-		
-			// Check for recursion. It's not completely safe but because we pump messages while 
-			// blocked it is expected.
-			NumRecursiveCalls++;
-			if (NumRecursiveCalls > 1)
-			{
-				UE_LOG(LogAudio,Warning,TEXT("Wait on Audio Thread called recursively! %d calls on the stack."), NumRecursiveCalls);
-			}
-			if (NumRecursiveCalls > 1 || FTaskGraphInterface::Get().IsThreadProcessingTasks(ENamedThreads::GameThread))
-			{
-				bEmptyGameThreadTasks = false; // we don't do this on recursive calls or if we are at a blueprint breakpoint
-			}
-
-			// Grab an event from the pool and fire off a task to trigger it.
-			FEvent* Event = FPlatformProcess::GetSynchEventFromPool();
-			FTaskGraphInterface::Get().TriggerEventWhenTaskCompletes(Event, Task, ENamedThreads::GameThread);
-
-			// Check audio thread health needs to be called from time to
-			// time in order to make sure the audio thread hasn't crashed :)
-			bool bDone;
-			uint32 WaitTime = FMath::Clamp<uint32>(GTimeToBlockOnAudioFence, 0, 33);
-			do
-			{
-				FAudioThread::CheckAudioThreadHealth();
-				if (bEmptyGameThreadTasks)
-				{
-					// process gamethread tasks if there are any
-					FTaskGraphInterface::Get().ProcessThreadUntilIdle(ENamedThreads::GameThread);
-				}
-				bDone = Event->Wait(WaitTime);
-			}
-			while (!bDone);
-
-			// Return the event to the pool and decrement the recursion counter.
-			FPlatformProcess::ReturnSynchEventToPool(Event);
-			Event = nullptr;
-			NumRecursiveCalls--;
-		}
-	}
-}
-
 /**
  * Waits for pending fence commands to retire.
  */
 void FAudioCommandFence::Wait(bool bProcessGameThreadTasks) const
 {
-	if (!IsFenceComplete())
+	if (!IsFenceComplete()) // this checks the current thread
 	{
-#if 0
-		// on most platforms this is a better solution because it doesn't spin
-		// windows needs to pump messages
-		if (bProcessGameThreadTasks)
+		QUICK_SCOPE_CYCLE_COUNTER(STAT_FAudioCommandFence_Wait);
+		const double StartTime = FPlatformTime::Seconds();
+		FEvent* Event = FPlatformProcess::GetSynchEventFromPool();
+		FTaskGraphInterface::Get().TriggerEventWhenTaskCompletes(Event, CompletionEvent, ENamedThreads::GameThread);
+
+		bool bDone;
+		const uint32 WaitTime = 35;
+		do
 		{
-			QUICK_SCOPE_CYCLE_COUNTER(STAT_FAudioCommandFence_Wait);
-			FTaskGraphInterface::Get().WaitUntilTaskCompletes(CompletionEvent, ENamedThreads::GameThread);
-		}
-#endif
-		GameThreadWaitForTask(CompletionEvent, bProcessGameThreadTasks);
+			bDone = Event->Wait(WaitTime);
+			float ThisTime = FPlatformTime::Seconds() - StartTime;
+			UE_CLOG(ThisTime > .036f, LogAudio, Warning, TEXT("Waited %fms for audio thread."), ThisTime * 1000.0f);
+		} while (!bDone);
+
+		// Return the event to the pool and decrement the recursion counter.
+		FPlatformProcess::ReturnSynchEventToPool(Event);
+		Event = nullptr;
 	}
 }
diff --git a/Engine/Source/Runtime/Engine/Private/Blueprint.cpp b/Engine/Source/Runtime/Engine/Private/Blueprint.cpp
--- a/Engine/Source/Runtime/Engine/Private/Blueprint.cpp
+++ b/Engine/Source/Runtime/Engine/Private/Blueprint.cpp
@@ -638,20 +638,8 @@ void UBlueprint::PostLoad()
 	FBlueprintEditorUtils::UpdateStalePinWatches(this);
 #endif // WITH_EDITORONLY_DATA
 
-	{
-		TArray<UEdGraph*> Graphs;
-		GetAllGraphs(Graphs);
-		for (auto It = Graphs.CreateIterator(); It; ++It)
-		{
-			UEdGraph* const Graph = *It;
-			const UEdGraphSchema* Schema = Graph->GetSchema();
-			Schema->BackwardCompatibilityNodeConversion(Graph, true);
-		}
-	}
-
 	FStructureEditorUtils::RemoveInvalidStructureMemberVariableFromBlueprint(this);
 
-
 #if WITH_EDITOR
 	// Do not want to run this code without the editor present nor when running commandlets.
 	if(GEditor && GIsEditor && !IsRunningCommandlet())
diff --git a/Engine/Source/Runtime/Engine/Private/BlueprintGeneratedClass.cpp b/Engine/Source/Runtime/Engine/Private/BlueprintGeneratedClass.cpp
--- a/Engine/Source/Runtime/Engine/Private/BlueprintGeneratedClass.cpp
+++ b/Engine/Source/Runtime/Engine/Private/BlueprintGeneratedClass.cpp
@@ -42,9 +42,6 @@ void UBlueprintGeneratedClass::PostLoad()
 	UObject* ClassCDO = GetDefaultObject();
 
 	// Go through the CDO of the class, and make sure we don't have any legacy components that aren't instanced hanging on.
-	TArray<UObject*> SubObjects;
-	GetObjectsWithOuter(ClassCDO, SubObjects);
-
 	struct FCheckIfComponentChildHelper
 	{
 		static bool IsComponentChild(UObject* CurrObj, const UObject* CDO)
@@ -55,14 +52,14 @@ void UBlueprintGeneratedClass::PostLoad()
 		};
 	};
 
-	for (UObject* CurrObj :	SubObjects)
+	ForEachObjectWithOuter(ClassCDO, [ClassCDO](UObject* CurrObj)
 	{
 		const bool bComponentChild = FCheckIfComponentChildHelper::IsComponentChild(CurrObj, ClassCDO);
 		if (!CurrObj->IsDefaultSubobject() && !CurrObj->IsRooted() && !bComponentChild)
 		{
 			CurrObj->MarkPendingKill();
 		}
-	}
+	});
 
 #if WITH_EDITORONLY_DATA
 	if (GetLinkerUE4Version() < VER_UE4_CLASS_NOTPLACEABLE_ADDED)
@@ -1046,13 +1043,11 @@ void UBlueprintGeneratedClass::GetPreloadDependencies(TArray<UObject*>& OutDeps)
 	UObject *CDO = GetDefaultObject();
 	if (CDO)
 	{
-		TArray<UObject*> Subobjects;
-		GetObjectsWithOuter(CDO, Subobjects);
-		for (UObject* SubObj : Subobjects)
+		ForEachObjectWithOuter(CDO, [&OutDeps](UObject* SubObj)
 		{
 			OutDeps.Add(SubObj->GetClass());
 			OutDeps.Add(SubObj->GetArchetype());
-		}
+		});
 	}
 }
 
diff --git a/Engine/Source/Runtime/Engine/Private/BoneContainer.cpp b/Engine/Source/Runtime/Engine/Private/BoneContainer.cpp
--- a/Engine/Source/Runtime/Engine/Private/BoneContainer.cpp
+++ b/Engine/Source/Runtime/Engine/Private/BoneContainer.cpp
@@ -189,7 +189,40 @@ void FBoneContainer::CacheRequiredAnimCurveUids()
 		if (Mapping != nullptr)
 		{
 			AnimCurveNameUids.Reset();
+			// fill name array
+			TArray<FName> CurveNames;
+			Mapping->FillNameArray(CurveNames);
 			Mapping->FillUidArray(AnimCurveNameUids);
+
+			// if the linked joints don't exists in RequiredBones, remove itself
+			if (CurveNames.Num() > 0)
+			{
+				for (int32 CurveNameIndex = CurveNames.Num() - 1; CurveNameIndex >=0 ; --CurveNameIndex)
+				{
+					const FCurveMetaData* CurveMetaData = Mapping->GetCurveMetaData(CurveNames[CurveNameIndex]);
+					if (CurveMetaData && CurveMetaData->LinkedBones.Num() > 0)
+					{
+						bool bRemove = true;
+						for (int32 LinkedBoneIndex = 0; LinkedBoneIndex < CurveMetaData->LinkedBones.Num(); ++LinkedBoneIndex)
+						{
+							const FBoneReference& BoneReference = CurveMetaData->LinkedBones[LinkedBoneIndex];
+							// we want to make sure all the joints are removed from RequiredBones before removing this UID
+							if (BoneReference.GetCompactPoseIndex(*this) != INDEX_NONE)
+							{
+								// still has some joint that matters, do not remove
+								bRemove = false;
+								break;
+							}
+						}
+
+						if (bRemove)
+						{
+							//remove the UID
+							AnimCurveNameUids.RemoveAt(CurveNameIndex);
+						}
+					}
+				}
+			}
 		}
 	}
 }
diff --git a/Engine/Source/Runtime/Engine/Private/Character.cpp b/Engine/Source/Runtime/Engine/Private/Character.cpp
--- a/Engine/Source/Runtime/Engine/Private/Character.cpp
+++ b/Engine/Source/Runtime/Engine/Private/Character.cpp
@@ -341,11 +341,11 @@ void ACharacter::SetReplicateMovement(bool bInReplicateMovement)
 {
 	Super::SetReplicateMovement(bInReplicateMovement);
 
-	if (CharacterMovement != nullptr)
+	if (CharacterMovement != nullptr && Role == ROLE_Authority)
 	{
 		// Set prediction data time stamp to current time to stop extrapolating
 		// from time bReplicateMovement was turned off to when it was turned on again
-		FNetworkPredictionData_Server* NetworkPrediction = CharacterMovement->GetPredictionData_Server();
+		FNetworkPredictionData_Server* NetworkPrediction = CharacterMovement->HasPredictionData_Server() ? CharacterMovement->GetPredictionData_Server() : nullptr;
 
 		if (NetworkPrediction != nullptr)
 		{
diff --git a/Engine/Source/Runtime/Engine/Private/Components/CharacterMovementComponent.cpp b/Engine/Source/Runtime/Engine/Private/Components/CharacterMovementComponent.cpp
--- a/Engine/Source/Runtime/Engine/Private/Components/CharacterMovementComponent.cpp
+++ b/Engine/Source/Runtime/Engine/Private/Components/CharacterMovementComponent.cpp
@@ -4377,6 +4377,7 @@ void UCharacterMovementComponent::PhysWalking(float deltaTime, int32 Iterations)
 				Hit.TraceEnd = Hit.TraceStart + FVector(0.f, 0.f, MAX_FLOOR_DIST);
 				const FVector RequestedAdjustment = GetPenetrationAdjustment(Hit);
 				ResolvePenetration(RequestedAdjustment, Hit, UpdatedComponent->GetComponentQuat());
+				bForceNextFloorCheck = true;
 			}
 
 			// check if just entered water
@@ -4977,7 +4978,7 @@ void UCharacterMovementComponent::OnTeleported()
 	{
 		return;
 	}
-	bool bWasFalling = (MovementMode == MOVE_Falling);
+
 	bJustTeleported = true;
 
 	// Find floor at current location
@@ -4997,19 +4998,33 @@ void UCharacterMovementComponent::OnTeleported()
 		CurrentFloor.Clear();
 	}
 
-	// If we were walking but no longer have a valid base or floor, start falling.
-	const float SavedVelocityZ = Velocity.Z;
-	SetDefaultMovementMode();
-	if ((MovementMode == MOVE_Walking) && (!CurrentFloor.IsWalkableFloor() || (OldBase && !NewBase)))
+	const bool bWasFalling = (MovementMode == MOVE_Falling);
+	const bool bWasSwimming = (MovementMode == DefaultWaterMovementMode) || (MovementMode == MOVE_Swimming);
+
+	if (CanEverSwim() && IsInWater())
 	{
-		// If we are walking but no longer have a valid base or floor, start falling.
-		Velocity.Z = SavedVelocityZ;
-		SetMovementMode(MOVE_Falling);
+		if (!bWasSwimming)
+		{
+			SetMovementMode(DefaultWaterMovementMode);
+		}
 	}
-
-	if (bWasFalling && IsMovingOnGround())
+	else if (!CurrentFloor.IsWalkableFloor() || (OldBase && !NewBase))
 	{
-		ProcessLanded(CurrentFloor.HitResult, 0.f, 0);
+		if (!bWasFalling && MovementMode != MOVE_Flying && MovementMode != MOVE_Custom)
+		{
+			SetMovementMode(MOVE_Falling);
+		}
+	}
+	else if (NewBase)
+	{
+		if (bWasSwimming)
+		{
+			SetMovementMode(DefaultLandMovementMode);
+		}
+		else if (bWasFalling)
+		{
+			ProcessLanded(CurrentFloor.HitResult, 0.f, 0);
+		}
 	}
 
 	MaybeSaveBaseLocation();
@@ -6105,20 +6120,16 @@ void UCharacterMovementComponent::HandleImpact(const FHitResult& Impact, float T
 
 void UCharacterMovementComponent::ApplyImpactPhysicsForces(const FHitResult& Impact, const FVector& ImpactAcceleration, const FVector& ImpactVelocity)
 {
-	if (bEnablePhysicsInteraction && Impact.bBlockingHit)
+	if (bEnablePhysicsInteraction && Impact.bBlockingHit )
 	{
-		UPrimitiveComponent* ImpactComponent = Impact.GetComponent();
-		if (ImpactComponent != NULL && ImpactComponent->IsAnySimulatingPhysics())
+		if (UPrimitiveComponent* ImpactComponent = Impact.GetComponent())
 		{
-			FVector ForcePoint = Impact.ImpactPoint;
-
 			FBodyInstance* BI = ImpactComponent->GetBodyInstance(Impact.BoneName);
-
-			float BodyMass = 1.0f;
-
-			if (BI != NULL)
+			if(BI != nullptr && BI->IsInstanceSimulatingPhysics())
 			{
-				BodyMass = FMath::Max(BI->GetBodyMass(), 1.0f);
+				FVector ForcePoint = Impact.ImpactPoint;
+
+				const float BodyMass = FMath::Max(BI->GetBodyMass(), 1.0f);
 
 				if(bPushForceUsingZOffset)
 				{
@@ -6132,43 +6143,43 @@ void UCharacterMovementComponent::ApplyImpactPhysicsForces(const FHitResult& Imp
 						ForcePoint.Z = Center.Z + Extents.Z * PushForcePointZOffsetFactor;
 					}
 				}
-			}
 
-			FVector Force = Impact.ImpactNormal * -1.0f;
+				FVector Force = Impact.ImpactNormal * -1.0f;
 
-			float PushForceModificator = 1.0f;
+				float PushForceModificator = 1.0f;
 
-			const FVector ComponentVelocity = ImpactComponent->GetPhysicsLinearVelocity();
-			const FVector VirtualVelocity = ImpactAcceleration.IsZero() ? ImpactVelocity : ImpactAcceleration.GetSafeNormal() * GetMaxSpeed();
+				const FVector ComponentVelocity = ImpactComponent->GetPhysicsLinearVelocity();
+				const FVector VirtualVelocity = ImpactAcceleration.IsZero() ? ImpactVelocity : ImpactAcceleration.GetSafeNormal() * GetMaxSpeed();
 
-			float Dot = 0.0f;
+				float Dot = 0.0f;
 
-			if (bScalePushForceToVelocity && !ComponentVelocity.IsNearlyZero())
-			{			
-				Dot = ComponentVelocity | VirtualVelocity;
+				if (bScalePushForceToVelocity && !ComponentVelocity.IsNearlyZero())
+				{			
+					Dot = ComponentVelocity | VirtualVelocity;
 
-				if (Dot > 0.0f && Dot < 1.0f)
-				{
-					PushForceModificator *= Dot;
+					if (Dot > 0.0f && Dot < 1.0f)
+					{
+						PushForceModificator *= Dot;
+					}
 				}
-			}
 
-			if (bPushForceScaledToMass)
-			{
-				PushForceModificator *= BodyMass;
-			}
+				if (bPushForceScaledToMass)
+				{
+					PushForceModificator *= BodyMass;
+				}
 
-			Force *= PushForceModificator;
+				Force *= PushForceModificator;
 
-			if (ComponentVelocity.IsNearlyZero())
-			{
-				Force *= InitialPushForceFactor;
-				ImpactComponent->AddImpulseAtLocation(Force, ForcePoint, Impact.BoneName);
-			}
-			else
-			{
-				Force *= PushForceFactor;
-				ImpactComponent->AddForceAtLocation(Force, ForcePoint, Impact.BoneName);
+				if (ComponentVelocity.IsNearlyZero())
+				{
+					Force *= InitialPushForceFactor;
+					ImpactComponent->AddImpulseAtLocation(Force, ForcePoint, Impact.BoneName);
+				}
+				else
+				{
+					Force *= PushForceFactor;
+					ImpactComponent->AddForceAtLocation(Force, ForcePoint, Impact.BoneName);
+				}
 			}
 		}
 	}
@@ -7140,6 +7151,15 @@ FNetworkPredictionData_Server_Character* UCharacterMovementComponent::GetPredict
 	return static_cast<class FNetworkPredictionData_Server_Character*>(GetPredictionData_Server());
 }
 
+bool UCharacterMovementComponent::HasPredictionData_Client() const
+{
+	return (ClientPredictionData != nullptr) && HasValidData();
+}
+
+bool UCharacterMovementComponent::HasPredictionData_Server() const
+{
+	return (ServerPredictionData != nullptr) && HasValidData();
+}
 
 void UCharacterMovementComponent::ResetPredictionData_Client()
 {
@@ -8623,13 +8643,16 @@ void UCharacterMovementComponent::SetAvoidanceEnabled(bool bEnable)
 	{
 		bUseRVOAvoidance = bEnable;
 
+		// reset id, RegisterMovementComponent call is required to initialize update timers in avoidance manager
+		AvoidanceUID = 0;
+
 		// this is a safety check - it's possible to not have CharacterOwner at this point if this function gets
 		// called too early
 		ensure(GetCharacterOwner());
 		if (GetCharacterOwner() != nullptr)
 		{
 			UAvoidanceManager* AvoidanceManager = GetWorld()->GetAvoidanceManager();
-			if (AvoidanceManager && bEnable && AvoidanceUID == 0)
+			if (AvoidanceManager && bEnable)
 			{
 				AvoidanceManager->RegisterMovementComponent(this, AvoidanceWeight);
 			}
diff --git a/Engine/Source/Runtime/Engine/Private/Components/ChildActorComponent.cpp b/Engine/Source/Runtime/Engine/Private/Components/ChildActorComponent.cpp
--- a/Engine/Source/Runtime/Engine/Private/Components/ChildActorComponent.cpp
+++ b/Engine/Source/Runtime/Engine/Private/Components/ChildActorComponent.cpp
@@ -491,6 +491,21 @@ void UChildActorComponent::CreateChildActor()
 					ChildActor->AttachToComponent(this, FAttachmentTransformRules::SnapToTargetNotIncludingScale);
 
 					SetIsReplicated(ChildActor->GetIsReplicated());
+
+					if (CachedInstanceData)
+					{
+						for (const FChildActorComponentInstanceData::FAttachedActorInfo& AttachedActorInfo : CachedInstanceData->AttachedActors)
+						{
+							AActor* AttachedActor = AttachedActorInfo.Actor.Get();
+
+							if (AttachedActor && AttachedActor->GetAttachParentActor() == nullptr)
+							{
+								AttachedActor->AttachToActor(ChildActor, FAttachmentTransformRules::KeepWorldTransform, AttachedActorInfo.SocketName);
+								AttachedActor->SetActorRelativeTransform(AttachedActorInfo.RelativeTransform);
+							}
+						}
+					}
+
 				}
 			}
 		}
diff --git a/Engine/Source/Runtime/Engine/Private/Components/MovementComponent.cpp b/Engine/Source/Runtime/Engine/Private/Components/MovementComponent.cpp
--- a/Engine/Source/Runtime/Engine/Private/Components/MovementComponent.cpp
+++ b/Engine/Source/Runtime/Engine/Private/Components/MovementComponent.cpp
@@ -504,6 +504,20 @@ bool UMovementComponent::K2_MoveUpdatedComponent(FVector Delta, FRotator NewRota
 }
 
 
+// Typically we want to depenetrate regardless of direction, so we can get all the way out of penetration quickly.
+// Our rules for "moving with depenetration normal" only get us so far out of the object. We'd prefer to pop out by the full MTD amount.
+// Depenetration moves (in ResolvePenetration) then ignore blocking overlaps to be able to move out by the MTD amount.
+static int32 MoveIgnoreFirstBlockingOverlap = 0;
+
+static FAutoConsoleVariableRef CVarMoveIgnoreFirstBlockingOverlap(
+	TEXT("p.MoveIgnoreFirstBlockingOverlap"),
+	MoveIgnoreFirstBlockingOverlap,
+	TEXT("Whether to ignore the first blocking overlap in SafeMoveUpdatedComponent (if moving out from object and starting in penetration).\n")
+	TEXT("The 'p.InitialOverlapTolerance' setting determines the 'move out' rules, but by default we always try to depenetrate first (not ignore the hit).\n")
+	TEXT("0: Disable (do not ignore), 1: Enable (ignore)"),
+	ECVF_Default);
+
+
 bool UMovementComponent::SafeMoveUpdatedComponent(const FVector& Delta, const FQuat& NewRotation, bool bSweep, FHitResult& OutHit, ETeleportType Teleport)
 {
 	if (UpdatedComponent == NULL)
@@ -512,7 +526,14 @@ bool UMovementComponent::SafeMoveUpdatedComponent(const FVector& Delta, const FQ
 		return false;
 	}
 
-	bool bMoveResult = MoveUpdatedComponent(Delta, NewRotation, bSweep, &OutHit, Teleport);
+	bool bMoveResult = false;
+
+	// Scope for move flags
+	{
+		// Conditionally ignore blocking overlaps (based on CVar)
+		TGuardValue<EMoveComponentFlags> ScopedFlagRestore(MoveComponentFlags, MoveIgnoreFirstBlockingOverlap ? MoveComponentFlags : (MoveComponentFlags | MOVECOMP_NeverIgnoreBlockingOverlaps));
+		bMoveResult = MoveUpdatedComponent(Delta, NewRotation, bSweep, &OutHit, Teleport);
+	}
 
 	// Handle initial penetrations
 	if (OutHit.bStartPenetrating && UpdatedComponent)
diff --git a/Engine/Source/Runtime/Engine/Private/Components/PrimitiveComponent.cpp b/Engine/Source/Runtime/Engine/Private/Components/PrimitiveComponent.cpp
--- a/Engine/Source/Runtime/Engine/Private/Components/PrimitiveComponent.cpp
+++ b/Engine/Source/Runtime/Engine/Private/Components/PrimitiveComponent.cpp
@@ -1127,16 +1127,29 @@ void UPrimitiveComponent::PushHoveredToProxy(const bool bInHovered)
 	}
 }
 
-void UPrimitiveComponent::SetCullDistance(float NewCullDistance)
+void UPrimitiveComponent::SetCullDistance(const float NewCullDistance)
 {
-	if (NewCullDistance >= 0)
+	if (NewCullDistance >= 0.f && NewCullDistance != LDMaxDrawDistance)
 	{
+		const float OldLDMaxDrawDistance = LDMaxDrawDistance;
+
 		LDMaxDrawDistance = NewCullDistance;
 	
-		if (LDMaxDrawDistance < CachedMaxDrawDistance)
+		if (CachedMaxDrawDistance == 0.f || LDMaxDrawDistance < CachedMaxDrawDistance)
 		{
 			SetCachedMaxDrawDistance(LDMaxDrawDistance);
 		}
+		else if (OldLDMaxDrawDistance == CachedMaxDrawDistance)
+		{
+			if (UWorld* World = GetWorld())
+			{
+				World->UpdateCullDistanceVolumes(nullptr, this);
+			}
+			else
+			{
+				SetCachedMaxDrawDistance(LDMaxDrawDistance);
+			}
+		}
 	}
 }
 
@@ -2531,7 +2544,7 @@ void UPrimitiveComponent::UpdateOverlaps(const TArray<FOverlapInfo>* NewPendingO
 					Params.bIgnoreBlocks = true;	//We don't care about blockers since we only route overlap events to real overlaps
 					FCollisionResponseParams ResponseParam;
 					InitSweepCollisionParams(Params, ResponseParam);
-					MyWorld->ComponentOverlapMulti(Overlaps, this, GetComponentLocation(), GetComponentQuat(), Params);
+					ComponentOverlapMulti(Overlaps, MyWorld, GetComponentLocation(), GetComponentQuat(), GetCollisionObjectType(), Params);
 
 					for( int32 ResultIdx=0; ResultIdx<Overlaps.Num(); ResultIdx++ )
 					{
diff --git a/Engine/Source/Runtime/Engine/Private/CurveTable.cpp b/Engine/Source/Runtime/Engine/Private/CurveTable.cpp
--- a/Engine/Source/Runtime/Engine/Private/CurveTable.cpp
+++ b/Engine/Source/Runtime/Engine/Private/CurveTable.cpp
@@ -530,4 +530,13 @@ bool FCurveTableRowHandle::operator==(const FCurveTableRowHandle& Other) const
 bool FCurveTableRowHandle::operator!=(const FCurveTableRowHandle& Other) const
 {
 	return ((Other.CurveTable != CurveTable) || (Other.RowName != RowName));
+}
+
+void FCurveTableRowHandle::PostSerialize(const FArchive& Ar)
+{
+	if (Ar.IsSaving() && !IsNull() && CurveTable)
+	{
+		// Note which row we are pointing to for later searching
+		Ar.MarkSearchableName(CurveTable, RowName);
+	}
 }
\ No newline at end of file
diff --git a/Engine/Source/Runtime/Engine/Private/DataTable.cpp b/Engine/Source/Runtime/Engine/Private/DataTable.cpp
--- a/Engine/Source/Runtime/Engine/Private/DataTable.cpp
+++ b/Engine/Source/Runtime/Engine/Private/DataTable.cpp
@@ -607,6 +607,15 @@ bool FDataTableRowHandle::operator != (FDataTableRowHandle const& Other) const
 	return DataTable != Other.DataTable || RowName != Other.RowName;
 }
 
+void FDataTableRowHandle::PostSerialize(const FArchive& Ar)
+{
+	if (Ar.IsSaving() && !IsNull() && DataTable)
+	{
+		// Note which row we are pointing to for later searching
+		Ar.MarkSearchableName(DataTable, RowName);
+	}
+}
+
 bool FDataTableCategoryHandle::operator==(FDataTableCategoryHandle const& Other) const
 {
 	return DataTable == Other.DataTable && ColumnName == Other.ColumnName && RowContents == Other.RowContents;
diff --git a/Engine/Source/Runtime/Engine/Private/EngineUtils.cpp b/Engine/Source/Runtime/Engine/Private/EngineUtils.cpp
--- a/Engine/Source/Runtime/Engine/Private/EngineUtils.cpp
+++ b/Engine/Source/Runtime/Engine/Private/EngineUtils.cpp
@@ -290,9 +290,7 @@ bool EngineUtils::FindOrLoadAssetsByPath(const FString& Path, TArray<UObject*>&
 
 		if (Package)
 		{
-			TArray<UObject*> ObjectsInPackage;
-			GetObjectsWithOuter(Package, ObjectsInPackage);
-			for (UObject* Object : ObjectsInPackage)
+			ForEachObjectWithOuter(Package, [Type, &OutAssets](UObject* Object)
 			{
 				const bool bWantedType = 
 					((EAssetToLoad::ATL_Regular == Type) && Object->IsAsset()) ||
@@ -301,7 +299,7 @@ bool EngineUtils::FindOrLoadAssetsByPath(const FString& Path, TArray<UObject*>&
 				{
 					OutAssets.Add(Object);
 				}
-			}
+			});
 		}
 	}
 	return true;
diff --git a/Engine/Source/Runtime/Engine/Private/GameViewportClient.cpp b/Engine/Source/Runtime/Engine/Private/GameViewportClient.cpp
--- a/Engine/Source/Runtime/Engine/Private/GameViewportClient.cpp
+++ b/Engine/Source/Runtime/Engine/Private/GameViewportClient.cpp
@@ -312,7 +312,7 @@ void UGameViewportClient::Init(struct FWorldContext& WorldContext, UGameInstance
 #if !UE_BUILD_SHIPPING
 				if (NewDeviceResults.bNewDevice)
 				{
-					NewDeviceResults.AudioDevice->UpdateSoundShowFlags(0, GetSoundShowFlags());
+					NewDeviceResults.AudioDevice->ResolveDesiredStats(this);
 				}
 #endif // UE_BUILD_SHIPPING
 
diff --git a/Engine/Source/Runtime/Engine/Private/GameplayStatics.cpp b/Engine/Source/Runtime/Engine/Private/GameplayStatics.cpp
--- a/Engine/Source/Runtime/Engine/Private/GameplayStatics.cpp
+++ b/Engine/Source/Runtime/Engine/Private/GameplayStatics.cpp
@@ -1151,7 +1151,7 @@ class UAudioComponent* UGameplayStatics::SpawnSoundAttached(class USoundBase* So
 		}
 	}
 
-	FAudioDevice::FCreateComponentParams Params(AttachToComponent->GetOwner());
+	FAudioDevice::FCreateComponentParams Params(AttachToComponent->GetWorld(), AttachToComponent->GetOwner());
 	Params.SetLocation(TestLocation);
 	Params.bStopWhenOwnerDestroyed = bStopWhenAttachedToDestroyed;
 	Params.AttenuationSettings = AttenuationSettings;
diff --git a/Engine/Source/Runtime/Engine/Private/HierarchicalInstancedStaticMesh.cpp b/Engine/Source/Runtime/Engine/Private/HierarchicalInstancedStaticMesh.cpp
--- a/Engine/Source/Runtime/Engine/Private/HierarchicalInstancedStaticMesh.cpp
+++ b/Engine/Source/Runtime/Engine/Private/HierarchicalInstancedStaticMesh.cpp
@@ -1440,7 +1440,7 @@ void FHierarchicalStaticMeshSceneProxy::GetDynamicMeshElements(const TArray<cons
 				for (int32 SampleIndex = 0; SampleIndex < 2; SampleIndex++)
 				{
 					// we aren't going to accurately figure out LOD planes for both samples and try all 4 possibilities, rather we just consider the worst case LOD range; the shaders will render correctly
-					float Fac = View->GetTemporalLODDistanceFactor(SampleIndex, bMultipleSections) * SphereRadius * SphereRadius * LODScale * LODScale;
+					float Fac = View->GetTemporalLODDistanceFactor(SampleIndex, bMultipleSections) * SphereRadius * LODScale;
 
 					float FinalCull = MAX_flt;
 
@@ -1620,8 +1620,7 @@ void FHierarchicalStaticMeshSceneProxy::GetDynamicMeshElements(const TArray<cons
 						for (int32 SampleIndex = 0; SampleIndex < 2; SampleIndex++)
 						{
 							// we aren't going to accurately figure out LOD planes for both samples and try all 4 possibilities, rather we just consider the worst case LOD range; the shaders will render correctly
-							float Fac = View->GetTemporalLODDistanceFactor(SampleIndex, bMultipleSections) * SphereRadius * SphereRadius * LODScale * LODScale;
-
+							float Fac = View->GetTemporalLODDistanceFactor(SampleIndex, bMultipleSections) * SphereRadius * LODScale;
 							float FinalCull = MAX_flt;
 
 							if (MinSize > 0.0)
diff --git a/Engine/Source/Runtime/Engine/Private/InstancedStaticMesh.cpp b/Engine/Source/Runtime/Engine/Private/InstancedStaticMesh.cpp
--- a/Engine/Source/Runtime/Engine/Private/InstancedStaticMesh.cpp
+++ b/Engine/Source/Runtime/Engine/Private/InstancedStaticMesh.cpp
@@ -2025,7 +2025,7 @@ void FInstancedStaticMeshVertexFactoryShaderParameters::SetMesh( FRHICommandList
 			for (int32 SampleIndex = 0; SampleIndex < 2; SampleIndex++)
 			{
 				FVector4& InstancingViewZCompare(SampleIndex ? InstancingViewZCompareOne : InstancingViewZCompareZero);
-				float Fac = View.GetTemporalLODDistanceFactor(SampleIndex) * SphereRadius * SphereRadius * LODScale * LODScale;
+				float Fac = View.GetTemporalLODDistanceFactor(SampleIndex) * SphereRadius * LODScale;
 
 				float FinalCull = MAX_flt;
 				if (MinSize > 0.0)
diff --git a/Engine/Source/Runtime/Engine/Private/KismetMathLibrary.cpp b/Engine/Source/Runtime/Engine/Private/KismetMathLibrary.cpp
--- a/Engine/Source/Runtime/Engine/Private/KismetMathLibrary.cpp
+++ b/Engine/Source/Runtime/Engine/Private/KismetMathLibrary.cpp
@@ -1221,6 +1221,10 @@ FRotator UKismetMathLibrary::RotatorFromAxisAndAngle(FVector Axis, float Angle)
 	return FQuat(SafeAxis, FMath::DegreesToRadians(Angle)).Rotator();
 }
 
+void UKismetMathLibrary::RotatorToAxisAndAngle(const FRotator& Rotation, FVector& Axis, float& Angle)
+{
+	Rotation.Quaternion().ToAxisAndAngle(Axis, Angle);
+}
 
 float UKismetMathLibrary::ClampAxis(float Angle)
 {
diff --git a/Engine/Source/Runtime/Engine/Private/LODActor.cpp b/Engine/Source/Runtime/Engine/Private/LODActor.cpp
--- a/Engine/Source/Runtime/Engine/Private/LODActor.cpp
+++ b/Engine/Source/Runtime/Engine/Private/LODActor.cpp
@@ -12,6 +12,7 @@
 #include "UObjectToken.h"
 
 #include "StaticMeshResources.h"
+#include "FrameworkObjectVersion.h"
 
 #if WITH_EDITOR
 #include "Editor.h"
@@ -185,6 +186,32 @@ void ALODActor::PostLoad()
 {
 	Super::PostLoad();
 	UpdateRegistrationToMatchMaximumLODLevel();
+
+#if WITH_EDITOR
+	if (bRequiresLODScreenSizeConversion)
+	{
+		if (TransitionScreenSize == 0.0f)
+		{
+			TransitionScreenSize = 1.0f;
+		}
+		else
+		{
+			const float HalfFOV = PI * 0.25f;
+			const float ScreenWidth = 1920.0f;
+			const float ScreenHeight = 1080.0f;
+			const FPerspectiveMatrix ProjMatrix(HalfFOV, ScreenWidth, ScreenHeight, 1.0f);
+			FBoxSphereBounds Bounds = GetStaticMeshComponent()->CalcBounds(FTransform());
+
+			// legacy transition screen size was previously a screen AREA fraction using resolution-scaled values, so we need to convert to distance first to correctly calculate the threshold
+			const float ScreenArea = TransitionScreenSize * (ScreenWidth * ScreenHeight);
+			const float ScreenRadius = FMath::Sqrt(ScreenArea / PI);
+			const float ScreenDistance = FMath::Max(ScreenWidth / 2.0f * ProjMatrix.M[0][0], ScreenHeight / 2.0f * ProjMatrix.M[1][1]) * Bounds.SphereRadius / ScreenRadius;
+
+			// Now convert using the query function
+			TransitionScreenSize = ComputeBoundsScreenSize(FVector::ZeroVector, Bounds.SphereRadius, FVector(0.0f, 0.0f, ScreenDistance), ProjMatrix);
+		}
+	}
+#endif
 }
 
 void ALODActor::UpdateRegistrationToMatchMaximumLODLevel()
@@ -642,7 +669,7 @@ void ALODActor::CleanSubObjectsArray()
 
 void ALODActor::RecalculateDrawingDistance(const float InTransitionScreenSize)
 {
-	// At the moment this assumes a fixed field of view of 90 degrees (horizontal and vertical axi)
+	// At the moment this assumes a fixed field of view of 90 degrees (horizontal and vertical axes)
 	static const float FOVRad = 90.0f * (float)PI / 360.0f;
 	static const FMatrix ProjectionMatrix = FPerspectiveMatrix(FOVRad, 1920, 1080, 0.01f);
 	FBoxSphereBounds Bounds = GetStaticMeshComponent()->CalcBounds(FTransform());
@@ -650,8 +677,12 @@ void ALODActor::RecalculateDrawingDistance(const float InTransitionScreenSize)
 	// Get projection multiple accounting for view scaling.
 	const float ScreenMultiple = FMath::Max(1920.0f / 2.0f * ProjectionMatrix.M[0][0],
 		1080.0f / 2.0f * ProjectionMatrix.M[1][1]);
-	// (ScreenMultiple * SphereRadius) / Sqrt(Screensize * 1920 * 1080.0f * PI) = Distance
-	LODDrawDistance = (ScreenMultiple * Bounds.SphereRadius) / FMath::Sqrt((InTransitionScreenSize * 1920.0f * 1080.0f) / PI);
+
+	// ScreenSize is the projected diameter, so halve it
+	const float ScreenRadius = FMath::Max(SMALL_NUMBER, InTransitionScreenSize * 0.5f);
+
+	// Invert the calcs in ComputeBoundsScreenSize
+	LODDrawDistance = FMath::Sqrt(FMath::Abs((FMath::Square((ScreenMultiple * Bounds.SphereRadius) / ScreenRadius)) + FMath::Square(Bounds.SphereRadius)));
 
 	UpdateSubActorLODParents();
 }
@@ -709,6 +740,17 @@ void ALODActor::OnCVarsChanged()
 	}
 }
 
+#if WITH_EDITOR
+void ALODActor::Serialize(FArchive& Ar)
+{
+	Super::Serialize(Ar);
+
+	Ar.UsingCustomVersion(FFrameworkObjectVersion::GUID);
+
+	bRequiresLODScreenSizeConversion = Ar.CustomVer(FFrameworkObjectVersion::GUID) < FFrameworkObjectVersion::LODsUseResolutionIndependentScreenSize;
+}
+#endif
+
 //////////////////////////////////////////////////////////////////////////
 // AHLODMeshCullingVolume
 
diff --git a/Engine/Source/Runtime/Engine/Private/Level.cpp b/Engine/Source/Runtime/Engine/Private/Level.cpp
--- a/Engine/Source/Runtime/Engine/Private/Level.cpp
+++ b/Engine/Source/Runtime/Engine/Private/Level.cpp
@@ -1331,18 +1331,15 @@ void ULevel::PostEditUndo()
 
 	// Non-transactional actors may disappear from the actors list but still exist, so we need to re-add them
 	// Likewise they won't get recreated if we undo to before they were deleted, so we'll have nulls in the actors list to remove
-	//Actors.Remove(nullptr); // removed because TTransArray exploded (undo followed by redo ends up with a different ArrayNum to originally)
 	TSet<AActor*> ActorsSet(Actors);
-	TArray<UObject *> InnerObjects;
-	GetObjectsWithOuter(this, InnerObjects, /*bIncludeNestedObjects*/ false, /*ExclusionFlags*/ RF_NoFlags, /* InternalExclusionFlags */ EInternalObjectFlags::PendingKill);
-	for (UObject* InnerObject : InnerObjects)
+	ForEachObjectWithOuter(this, [&ActorsSet, this](UObject* InnerObject)
 	{
 		AActor* InnerActor = Cast<AActor>(InnerObject);
 		if (InnerActor && !ActorsSet.Contains(InnerActor))
 		{
 			Actors.Add(InnerActor);
 		}
-	}
+	}, /*bIncludeNestedObjects*/ false, /*ExclusionFlags*/ RF_NoFlags, /* InternalExclusionFlags */ EInternalObjectFlags::PendingKill);
 
 	MarkLevelBoundsDirty();
 }
@@ -1760,17 +1757,15 @@ bool ULevel::HasAnyActorsOfType(UClass *SearchType)
 TArray<UBlueprint*> ULevel::GetLevelBlueprints() const
 {
 	TArray<UBlueprint*> LevelBlueprints;
-	TArray<UObject*> LevelChildren;
-	GetObjectsWithOuter(this, LevelChildren, false, RF_NoFlags, EInternalObjectFlags::PendingKill);
 
-	for (UObject* LevelChild : LevelChildren)
+	ForEachObjectWithOuter(this, [&LevelBlueprints](UObject* LevelChild)
 	{
 		UBlueprint* LevelChildBP = Cast<UBlueprint>(LevelChild);
 		if (LevelChildBP)
 		{
 			LevelBlueprints.Add(LevelChildBP);
 		}
-	}
+	}, false, RF_NoFlags, EInternalObjectFlags::PendingKill);
 
 	return LevelBlueprints;
 }
diff --git a/Engine/Source/Runtime/Engine/Private/OpusAudioInfo.cpp b/Engine/Source/Runtime/Engine/Private/OpusAudioInfo.cpp
--- a/Engine/Source/Runtime/Engine/Private/OpusAudioInfo.cpp
+++ b/Engine/Source/Runtime/Engine/Private/OpusAudioInfo.cpp
@@ -14,7 +14,6 @@
 
 #define USE_UE4_MEM_ALLOC 1
 #define OPUS_MAX_FRAME_SIZE_MS 120
-#define SAMPLE_SIZE			( ( uint32 )sizeof( short ) )
 
 ///////////////////////////////////////////////////////////////////////////////////////
 // Followed pattern used in opus_multistream_encoder.c - this will allow us to setup //
@@ -97,46 +96,21 @@ private:
 /*------------------------------------------------------------------------------------
 FOpusAudioInfo.
 ------------------------------------------------------------------------------------*/
-FOpusAudioInfo::FOpusAudioInfo(void)
-	: OpusDecoderWrapper(NULL)
-	, SrcBufferData(NULL)
-	, SrcBufferDataSize(0)
-	, SrcBufferOffset(0)
-	, AudioDataOffset(0)
-	, SampleRate(0)
-	, TrueSampleCount(0)
-	, CurrentSampleCount(0)
-	, NumChannels(0)
-	, MaxFrameSizeSamples(0)
-	, SampleStride(0)
-	, LastPCMByteSize(0)
-	, LastPCMOffset(0)
-	, bStoringEndOfFile(false)
-	, StreamingSoundWave(NULL)
-	, CurrentChunkIndex(0)
+FOpusAudioInfo::FOpusAudioInfo()
+	: OpusDecoderWrapper(nullptr)
 {
 }
 
-FOpusAudioInfo::~FOpusAudioInfo(void)
+FOpusAudioInfo::~FOpusAudioInfo()
 {
-	if (OpusDecoderWrapper != NULL)
+	if (OpusDecoderWrapper != nullptr)
 	{
 		delete OpusDecoderWrapper;
+		OpusDecoderWrapper = nullptr;
 	}
 }
 
-size_t FOpusAudioInfo::Read(void *ptr, uint32 size)
-{
-	size_t BytesToRead = FMath::Min(size, SrcBufferDataSize - SrcBufferOffset);
-	if (BytesToRead > 0)
-	{
-		FMemory::Memcpy(ptr, SrcBufferData + SrcBufferOffset, BytesToRead);
-		SrcBufferOffset += BytesToRead;
-	}
-	return(BytesToRead);
-}
-
-bool FOpusAudioInfo::ReadCompressedInfo(const uint8* InSrcBufferData, uint32 InSrcBufferDataSize, struct FSoundQualityInfo* QualityInfo)
+bool FOpusAudioInfo::ParseHeader(const uint8* InSrcBufferData, uint32 InSrcBufferDataSize, struct FSoundQualityInfo* QualityInfo)
 {
 	SrcBufferData = InSrcBufferData;
 	SrcBufferDataSize = InSrcBufferDataSize;
@@ -144,7 +118,7 @@ bool FOpusAudioInfo::ReadCompressedInfo(const uint8* InSrcBufferData, uint32 InS
 	CurrentSampleCount = 0;
 
 	// Read Identifier, True Sample Count, Number of channels and Frames to Encode first
-	if (!InSrcBufferData || FCStringAnsi::Strcmp((char*)InSrcBufferData, OPUS_ID_STRING) != 0)
+	if (!SrcBufferData || FCStringAnsi::Strcmp((char*)SrcBufferData, OPUS_ID_STRING) != 0)
 	{
 		return false;
 	}
@@ -152,306 +126,59 @@ bool FOpusAudioInfo::ReadCompressedInfo(const uint8* InSrcBufferData, uint32 InS
 	{
 		SrcBufferOffset += FCStringAnsi::Strlen(OPUS_ID_STRING) + 1;
 	}
+
 	Read(&SampleRate, sizeof(uint16));
 	Read(&TrueSampleCount, sizeof(uint32));
 	Read(&NumChannels, sizeof(uint8));
+
 	uint16 SerializedFrames = 0;
 	Read(&SerializedFrames, sizeof(uint16));
-	AudioDataOffset = SrcBufferOffset;
-	SampleStride = SAMPLE_SIZE * NumChannels;
 
-	check(OpusDecoderWrapper == NULL);
-	OpusDecoderWrapper = new FOpusDecoderWrapper(SampleRate, NumChannels);
-	if (!OpusDecoderWrapper->WasInitialisedSuccessfully())
-	{
-		delete OpusDecoderWrapper;
-		OpusDecoderWrapper = NULL;
-		return false;
-	}
-
-	MaxFrameSizeSamples = (SampleRate * OPUS_MAX_FRAME_SIZE_MS) / 1000;
-
-	LastDecodedPCM.Empty(MaxFrameSizeSamples * SampleStride);
-	LastDecodedPCM.AddUninitialized(MaxFrameSizeSamples * SampleStride);
+	// Store the offset to where the audio data begins
+	AudioDataOffset = SrcBufferOffset;
 
+	// Write out the the header info
 	if (QualityInfo)
 	{
 		QualityInfo->SampleRate = SampleRate;
 		QualityInfo->NumChannels = NumChannels;
-		QualityInfo->SampleDataSize = TrueSampleCount * QualityInfo->NumChannels * SAMPLE_SIZE;
+		QualityInfo->SampleDataSize = TrueSampleCount * QualityInfo->NumChannels * sizeof(int16);
 		QualityInfo->Duration = (float)TrueSampleCount / QualityInfo->SampleRate;
 	}
 
 	return true;
 }
 
-bool FOpusAudioInfo::ReadCompressedData(uint8* Destination, bool bLooping, uint32 BufferSize /*= 0*/)
-{
-	check(OpusDecoderWrapper);
-	check(Destination);
-
-	SCOPE_CYCLE_COUNTER(STAT_OpusDecompressTime);
-
-	// Write out any PCM data that was decoded during the last request
-	uint32 RawPCMOffset = WriteFromDecodedPCM(Destination, BufferSize);
-
-	bool bLooped = false;
-
-	if (bStoringEndOfFile && LastPCMByteSize > 0)
-	{
-		// delayed returning looped because we hadn't read the entire buffer
-		bLooped = true;
-		bStoringEndOfFile = false;
-	}
-
-	while (RawPCMOffset < BufferSize)
-	{
-		uint16 FrameSize = 0;
-		Read(&FrameSize, sizeof(uint16));
-		int32 DecodedSamples = DecompressToPCMBuffer(FrameSize);
-		if (DecodedSamples < 0)
-		{
-			LastPCMByteSize = 0;
-			ZeroBuffer(Destination + RawPCMOffset, BufferSize - RawPCMOffset);
-			return false;
-		}
-		else
-		{
-			LastPCMByteSize = IncrementCurrentSampleCount(DecodedSamples) * SampleStride;
-			RawPCMOffset += WriteFromDecodedPCM(Destination + RawPCMOffset, BufferSize - RawPCMOffset);
-
-			if (SrcBufferOffset >= SrcBufferDataSize)
-			{
-				// check whether all decoded PCM was written
-				if (LastPCMByteSize == 0)
-				{
-					bLooped = true;
-				}
-				else
-				{
-					bStoringEndOfFile = true;
-				}
-				if (bLooping)
-				{
-					SrcBufferOffset = AudioDataOffset;
-					CurrentSampleCount = 0;
-				}
-				else
-				{
-					RawPCMOffset += ZeroBuffer(Destination + RawPCMOffset, BufferSize - RawPCMOffset);
-				}
-			}
-		}
-	}
-
-	return bLooped;
-}
-
-void FOpusAudioInfo::ExpandFile(uint8* DstBuffer, struct FSoundQualityInfo* QualityInfo)
-{
-	check(OpusDecoderWrapper);
-	check(DstBuffer);
-	check(QualityInfo);
-	check(QualityInfo->SampleDataSize <= SrcBufferDataSize);
-
-	// Ensure we're at the start of the audio data
-	SrcBufferOffset = AudioDataOffset;
-
-	uint32 RawPCMOffset = 0;
-
-	while (RawPCMOffset < QualityInfo->SampleDataSize)
-	{
-		uint16 FrameSize = 0;
-		Read(&FrameSize, sizeof(uint16));
-		int32 DecodedSamples = DecompressToPCMBuffer(FrameSize);
-
-		if (DecodedSamples < 0)
-		{
-			RawPCMOffset += ZeroBuffer(DstBuffer + RawPCMOffset, QualityInfo->SampleDataSize - RawPCMOffset);
-		}
-		else
-		{
-			LastPCMByteSize = IncrementCurrentSampleCount(DecodedSamples) * SampleStride;
-			RawPCMOffset += WriteFromDecodedPCM(DstBuffer + RawPCMOffset, QualityInfo->SampleDataSize - RawPCMOffset);
-		}
-	}
-}
-
-bool FOpusAudioInfo::StreamCompressedInfo(USoundWave* Wave, struct FSoundQualityInfo* QualityInfo)
-{
-	StreamingSoundWave = Wave;
-
-	// Get the first chunk of audio data (should always be loaded)
-	CurrentChunkIndex = 0;
-	const uint8* FirstChunk = IStreamingManager::Get().GetAudioStreamingManager().GetLoadedChunk(StreamingSoundWave, CurrentChunkIndex);
-
-	if (FirstChunk)
-	{
-		return ReadCompressedInfo(FirstChunk, Wave->RunningPlatformData->Chunks[0].DataSize, QualityInfo);
-	}
-
-	return false;
-}
-
-bool FOpusAudioInfo::StreamCompressedData(uint8* Destination, bool bLooping, uint32 BufferSize)
+bool FOpusAudioInfo::CreateDecoder()
 {
-	check(OpusDecoderWrapper);
-	check(Destination);
-
-	SCOPE_CYCLE_COUNTER(STAT_OpusDecompressTime);
-
-	UE_LOG(LogAudio, Log, TEXT("Streaming compressed data from SoundWave'%s' - Chunk %d, Offset %d"), *StreamingSoundWave->GetName(), CurrentChunkIndex, SrcBufferOffset);
-
-	// Write out any PCM data that was decoded during the last request
-	uint32 RawPCMOffset = WriteFromDecodedPCM(Destination, BufferSize);
-
-	// If next chunk wasn't loaded when last one finished reading, try to get it again now
-	if (SrcBufferData == NULL)
-	{
-		SrcBufferData = IStreamingManager::Get().GetAudioStreamingManager().GetLoadedChunk(StreamingSoundWave, CurrentChunkIndex);
-		if (SrcBufferData)
-		{
-			SrcBufferDataSize = StreamingSoundWave->RunningPlatformData->Chunks[CurrentChunkIndex].DataSize;
-			SrcBufferOffset = CurrentChunkIndex == 0 ? AudioDataOffset : 0;
-		}
-		else
-		{
-			// Still not loaded, zero remainder of current buffer
-			UE_LOG(LogAudio, Warning, TEXT("Unable to read from chunk %d of SoundWave'%s'"), CurrentChunkIndex, *StreamingSoundWave->GetName());
-			ZeroBuffer(Destination + RawPCMOffset, BufferSize - RawPCMOffset);
-			return false;
-		}
-	}
-
-	bool bLooped = false;
-
-	if (bStoringEndOfFile && LastPCMByteSize > 0)
-	{
-		// delayed returning looped because we hadn't read the entire buffer
-		bLooped = true;
-		bStoringEndOfFile = false;
-	}
-
-	while (RawPCMOffset < BufferSize)
+	check(OpusDecoderWrapper == nullptr);
+	OpusDecoderWrapper = new FOpusDecoderWrapper(SampleRate, NumChannels);
+	if (!OpusDecoderWrapper->WasInitialisedSuccessfully())
 	{
-		uint16 FrameSize = 0;
-		Read(&FrameSize, sizeof(uint16));
-		int32 DecodedSamples = DecompressToPCMBuffer(FrameSize);
-
-		if (DecodedSamples < 0)
-		{
-			LastPCMByteSize = 0;
-			ZeroBuffer(Destination + RawPCMOffset, BufferSize - RawPCMOffset);
-			return false;
-		}
-		else
-		{
-			LastPCMByteSize = IncrementCurrentSampleCount(DecodedSamples) * SampleStride;
-
-			RawPCMOffset += WriteFromDecodedPCM(Destination + RawPCMOffset, BufferSize - RawPCMOffset);
-
-			// Have we reached the end of buffer
-			if (SrcBufferOffset >= SrcBufferDataSize)
-			{
-				// Special case for the last chunk of audio
-				if (CurrentChunkIndex == StreamingSoundWave->RunningPlatformData->NumChunks - 1)
-				{
-					// check whether all decoded PCM was written
-					if (LastPCMByteSize == 0)
-					{
-						bLooped = true;
-					}
-					else
-					{
-						bStoringEndOfFile = true;
-					}
-					if (bLooping)
-					{
-						CurrentChunkIndex = 0;
-						SrcBufferOffset = AudioDataOffset;
-						CurrentSampleCount = 0;
-					}
-					else
-					{
-						RawPCMOffset += ZeroBuffer(Destination + RawPCMOffset, BufferSize - RawPCMOffset);
-					}
-				}
-				else
-				{
-					CurrentChunkIndex++;
-					SrcBufferOffset = 0;
-				}
-
-				SrcBufferData = IStreamingManager::Get().GetAudioStreamingManager().GetLoadedChunk(StreamingSoundWave, CurrentChunkIndex);
-				if (SrcBufferData)
-				{
-					UE_LOG(LogAudio, Log, TEXT("Incremented current chunk from SoundWave'%s' - Chunk %d, Offset %d"), *StreamingSoundWave->GetName(), CurrentChunkIndex, SrcBufferOffset);
-					SrcBufferDataSize = StreamingSoundWave->RunningPlatformData->Chunks[CurrentChunkIndex].DataSize;
-				}
-				else
-				{
-					SrcBufferDataSize = 0;
-					RawPCMOffset += ZeroBuffer(Destination + RawPCMOffset, BufferSize - RawPCMOffset);
-				}
-			}
-		}
+		delete OpusDecoderWrapper;
+		OpusDecoderWrapper = nullptr;
+		return false;
 	}
 
-	return bLooped;
+	return true;
 }
 
-int32 FOpusAudioInfo::DecompressToPCMBuffer(uint16 FrameSize)
+uint32 FOpusAudioInfo::GetFrameSize()
 {
-	if (SrcBufferOffset + FrameSize > SrcBufferDataSize)
-	{
-		// if frame size is too large, something has gone wrong
-		return -1;
-	}
-
-	const uint8* SrcPtr = SrcBufferData + SrcBufferOffset;
-	SrcBufferOffset += FrameSize;
-	LastPCMOffset = 0;
-	return OpusDecoderWrapper->Decode(SrcPtr, FrameSize, (int16*)LastDecodedPCM.GetData(), MaxFrameSizeSamples);
+	// Opus format has variable frame size at the head of each frame...
+	// We have to assume that the SrcBufferOffset is at the correct location for the read
+	uint16 FrameSize = 0;
+	Read(&FrameSize, sizeof(uint16));
+	return (uint32)FrameSize;
 }
 
-uint32 FOpusAudioInfo::IncrementCurrentSampleCount(uint32 NewSamples)
+uint32 FOpusAudioInfo::GetMaxFrameSizeSamples() const
 {
-	if (CurrentSampleCount + NewSamples > TrueSampleCount)
-	{
-		NewSamples = TrueSampleCount - CurrentSampleCount;
-		CurrentSampleCount = TrueSampleCount;
-	}
-	else
-	{
-		CurrentSampleCount += NewSamples;
-	}
-	return NewSamples;
+	return SampleRate * OPUS_MAX_FRAME_SIZE_MS / 1000;
 }
 
-uint32 FOpusAudioInfo::WriteFromDecodedPCM(uint8* Destination, uint32 BufferSize)
+int32 FOpusAudioInfo::Decode(const uint8* FrameData, uint16 FrameSize, int16* OutPCMData, int32 SampleSize)
 {
-	uint32 BytesToCopy = FMath::Min(BufferSize, LastPCMByteSize - LastPCMOffset);
-	if (BytesToCopy > 0)
-	{
-		FMemory::Memcpy(Destination, LastDecodedPCM.GetData() + LastPCMOffset, BytesToCopy);
-		LastPCMOffset += BytesToCopy;
-		if (LastPCMOffset >= LastPCMByteSize)
-		{
-			LastPCMOffset = 0;
-			LastPCMByteSize = 0;
-		}
-	}
-	return BytesToCopy;
+	return OpusDecoderWrapper->Decode(FrameData, FrameSize, OutPCMData, SampleSize);
 }
 
-uint32	FOpusAudioInfo::ZeroBuffer(uint8* Destination, uint32 BufferSize)
-{
-	check(Destination);
-
-	if (BufferSize > 0)
-	{
-		FMemory::Memzero(Destination, BufferSize);
-		return BufferSize;
-	}
-	return 0;
-}
diff --git a/Engine/Source/Runtime/Engine/Private/PlayerController.cpp b/Engine/Source/Runtime/Engine/Private/PlayerController.cpp
--- a/Engine/Source/Runtime/Engine/Private/PlayerController.cpp
+++ b/Engine/Source/Runtime/Engine/Private/PlayerController.cpp
@@ -2806,7 +2806,7 @@ void APlayerController::DisplayDebug(class UCanvas* Canvas, const FDebugDisplayI
 {
 	Super::DisplayDebug(Canvas, DebugDisplay, YL, YPos);
 
-	FDisplayDebugManager DisplayDebugManager = Canvas->DisplayDebugManager;
+	FDisplayDebugManager& DisplayDebugManager = Canvas->DisplayDebugManager;
 	DisplayDebugManager.SetDrawColor(FColor(255, 255, 0));
 	DisplayDebugManager.DrawString(FString::Printf(TEXT("STATE %s"), *GetStateName().ToString()));
 
@@ -3989,21 +3989,25 @@ void APlayerController::TickActor( float DeltaSeconds, ELevelTick TickType, FAct
 		// skip updates if pawn lost autonomous proxy role (e.g. TurnOff() call)
 		if (GetPawn() && !GetPawn()->IsPendingKill() && GetPawn()->GetRemoteRole() == ROLE_AutonomousProxy && GetPawn()->bReplicateMovement)
 		{
-			INetworkPredictionInterface* NetworkPredictionInterface = Cast<INetworkPredictionInterface>(GetPawn()->GetMovementComponent());
-			if (NetworkPredictionInterface)
+			UMovementComponent* PawnMovement = GetPawn()->GetMovementComponent();
+			INetworkPredictionInterface* NetworkPredictionInterface = Cast<INetworkPredictionInterface>(PawnMovement);
+			if (NetworkPredictionInterface && IsValid(PawnMovement->UpdatedComponent))
 			{
-				FNetworkPredictionData_Server* ServerData = NetworkPredictionInterface->GetPredictionData_Server();
-				const float TimeSinceUpdate = ServerData ? GetWorld()->GetTimeSeconds() - ServerData->ServerTimeStamp : 0.f;
-				const float PawnTimeSinceUpdate = TimeSinceUpdate * GetPawn()->CustomTimeDilation;
-				if (PawnTimeSinceUpdate > FMath::Max<float>(DeltaSeconds+0.06f,AGameNetworkManager::StaticClass()->GetDefaultObject<AGameNetworkManager>()->MAXCLIENTUPDATEINTERVAL * GetPawn()->GetActorTimeDilation()))
+				FNetworkPredictionData_Server* ServerData = NetworkPredictionInterface->HasPredictionData_Server() ? NetworkPredictionInterface->GetPredictionData_Server() : nullptr;
+				if (ServerData)
 				{
-					//UE_LOG(LogPlayerController, Warning, TEXT("ForcedMovementTick. PawnTimeSinceUpdate: %f, DeltaSeconds: %f, DeltaSeconds+: %f"), PawnTimeSinceUpdate, DeltaSeconds, DeltaSeconds+0.06f);
-					const USkeletalMeshComponent* PawnMesh = GetPawn()->FindComponentByClass<USkeletalMeshComponent>();
-					if (!PawnMesh || !PawnMesh->IsSimulatingPhysics())
+					const float TimeSinceUpdate = GetWorld()->GetTimeSeconds() - ServerData->ServerTimeStamp;
+					const float PawnTimeSinceUpdate = TimeSinceUpdate * GetPawn()->CustomTimeDilation;
+					if (PawnTimeSinceUpdate > FMath::Max<float>(DeltaSeconds+0.06f, AGameNetworkManager::StaticClass()->GetDefaultObject<AGameNetworkManager>()->MAXCLIENTUPDATEINTERVAL * GetPawn()->GetActorTimeDilation()))
 					{
-						NetworkPredictionInterface->ForcePositionUpdate(PawnTimeSinceUpdate);
-						ServerData->ServerTimeStamp = GetWorld()->GetTimeSeconds();
-					}					
+						//UE_LOG(LogPlayerController, Warning, TEXT("ForcedMovementTick. PawnTimeSinceUpdate: %f, DeltaSeconds: %f, DeltaSeconds+: %f"), PawnTimeSinceUpdate, DeltaSeconds, DeltaSeconds+0.06f);
+						const USkeletalMeshComponent* PawnMesh = GetPawn()->FindComponentByClass<USkeletalMeshComponent>();
+						if (!PawnMesh || !PawnMesh->IsSimulatingPhysics())
+						{
+							NetworkPredictionInterface->ForcePositionUpdate(PawnTimeSinceUpdate);
+							ServerData->ServerTimeStamp = GetWorld()->GetTimeSeconds();
+						}
+					}
 				}
 			}
 		}
diff --git a/Engine/Source/Runtime/Engine/Private/SceneManagement.cpp b/Engine/Source/Runtime/Engine/Private/SceneManagement.cpp
--- a/Engine/Source/Runtime/Engine/Private/SceneManagement.cpp
+++ b/Engine/Source/Runtime/Engine/Private/SceneManagement.cpp
@@ -306,42 +306,74 @@ FLightMapInteraction FLightMapInteraction::Texture(
 	return Result;
 }
 
-float ComputeBoundsScreenSize( const FVector4& Origin, const float SphereRadius, const FSceneView& View )
+float ComputeBoundsScreenRadiusSquared(const FVector4& BoundsOrigin, const float SphereRadius, const FVector4& ViewOrigin, const FMatrix& ProjMatrix)
 {
-	const float DistSqr = FVector::DistSquared( Origin, View.ViewMatrices.GetViewOrigin() );
+	const float DistSqr = FVector::DistSquared(BoundsOrigin, ViewOrigin);
 
 	// Get projection multiple accounting for view scaling.
-	const float ScreenMultiple = FMath::Max(View.ViewRect.Width() / 2.0f * View.ViewMatrices.GetProjectionMatrix().M[0][0],
-		View.ViewRect.Height() / 2.0f * View.ViewMatrices.GetProjectionMatrix().M[1][1]);
+	const float ScreenMultiple = FMath::Max(0.5f * ProjMatrix.M[0][0], 0.5f * ProjMatrix.M[1][1]);
+
+	// Calculate screen-space projected radius
+	return FMath::Square(ScreenMultiple * SphereRadius) / FMath::Max(1.0f, DistSqr);
+}
+
+/** Runtime comparison version of ComputeTemporalLODBoundsScreenSize that avoids a square root */
+static float ComputeTemporalLODBoundsScreenRadiusSquared(const FVector& Origin, const float SphereRadius, const FSceneView& View, int32 SampleIndex)
+{
+	return ComputeBoundsScreenRadiusSquared(Origin, SphereRadius, View.GetTemporalLODOrigin(SampleIndex), View.ViewMatrices.GetProjectionMatrix());
+}
 
-	// Approximate number of pixels the sphere covers
-	const float ScreenArea = PI * FMath::Square( ScreenMultiple * SphereRadius ) / FMath::Max( DistSqr, 1.0f );
-	return ScreenArea / View.ViewRect.Area();
+float ComputeBoundsScreenRadiusSquared(const FVector4& Origin, const float SphereRadius, const FSceneView& View)
+{
+	return ComputeBoundsScreenRadiusSquared(Origin, SphereRadius, View.ViewMatrices.GetViewOrigin(), View.ViewMatrices.GetProjectionMatrix());
+}
+
+float ComputeBoundsScreenSize( const FVector4& Origin, const float SphereRadius, const FSceneView& View )
+{
+	return ComputeBoundsScreenSize(Origin, SphereRadius, View.ViewMatrices.GetViewOrigin(), View.ViewMatrices.GetProjectionMatrix());
 }
 
 float ComputeTemporalLODBoundsScreenSize( const FVector& Origin, const float SphereRadius, const FSceneView& View, int32 SampleIndex )
 {
-	const float DistSqr =  (Origin - View.GetTemporalLODOrigin(SampleIndex)).SizeSquared();
+	return ComputeBoundsScreenSize(Origin, SphereRadius, View.GetTemporalLODOrigin(SampleIndex), View.ViewMatrices.GetProjectionMatrix());
+}
+
+float ComputeBoundsScreenSize(const FVector4& BoundsOrigin, const float SphereRadius, const FVector4& ViewOrigin, const FMatrix& ProjMatrix)
+{
+	const float Dist = FVector::Dist(BoundsOrigin, ViewOrigin);
+
+	// Get projection multiple accounting for view scaling.
+	const float ScreenMultiple = FMath::Max(0.5f * ProjMatrix.M[0][0], 0.5f * ProjMatrix.M[1][1]);
 
+	// Calculate screen-space projected radius
+	const float ScreenRadius = ScreenMultiple * SphereRadius / FMath::Max(1.0f, Dist);
+
+	// For clarity, we end up comparing the diameter
+	return ScreenRadius * 2.0f;
+}
+
+float ComputeBoundsDrawDistance(const float ScreenSize, const float SphereRadius, const FMatrix& ProjMatrix)
+{
 	// Get projection multiple accounting for view scaling.
-	const float ScreenMultiple = FMath::Max(View.ViewRect.Width() / 2.0f * View.ViewMatrices.GetProjectionMatrix().M[0][0],
-		View.ViewRect.Height() / 2.0f * View.ViewMatrices.GetProjectionMatrix().M[1][1]);
+	const float ScreenMultiple = FMath::Max(0.5f * ProjMatrix.M[0][0], 0.5f * ProjMatrix.M[1][1]);
+
+	// ScreenSize is the projected diameter, so halve it
+	const float ScreenRadius = FMath::Max(SMALL_NUMBER, ScreenSize * 0.5f);
 
-	// Approximate number of pixels the sphere covers
-	const float ScreenArea = PI * FMath::Square( ScreenMultiple * SphereRadius ) / FMath::Max( DistSqr, 1.0f );
-	return ScreenArea / View.ViewRect.Area();
+	// Invert the calcs in ComputeBoundsScreenSize
+	return (ScreenMultiple * SphereRadius) / ScreenRadius;
 }
 
 int8 ComputeTemporalStaticMeshLOD( const FStaticMeshRenderData* RenderData, const FVector4& Origin, const float SphereRadius, const FSceneView& View, int32 MinLOD, float FactorScale, int32 SampleIndex )
 {
 	const int32 NumLODs = MAX_STATIC_MESH_LODS;
 
-	const float ScreenSize = ComputeTemporalLODBoundsScreenSize(Origin, SphereRadius, View, SampleIndex) * FactorScale;
+	const float ScreenRadiusSquared = ComputeTemporalLODBoundsScreenRadiusSquared(Origin, SphereRadius, View, SampleIndex) * FactorScale * FactorScale;
 
 	// Walk backwards and return the first matching LOD
 	for(int32 LODIndex = NumLODs - 1 ; LODIndex >= 0 ; --LODIndex)
 	{
-		if(RenderData->ScreenSize[LODIndex] > ScreenSize)
+		if(FMath::Square(RenderData->ScreenSize[LODIndex] * 0.5f) > ScreenRadiusSquared)
 		{
 			return FMath::Max(LODIndex, MinLOD);
 		}
@@ -356,12 +388,12 @@ int8 ComputeStaticMeshLOD( const FStaticMeshRenderData* RenderData, const FVecto
 	{
 		const int32 NumLODs = MAX_STATIC_MESH_LODS;
 
-		const float ScreenSize = ComputeBoundsScreenSize(Origin, SphereRadius, View) * FactorScale;
+		const float ScreenRadiusSquared = ComputeBoundsScreenRadiusSquared(Origin, SphereRadius, View) * FactorScale * FactorScale;
 
 		// Walk backwards and return the first matching LOD
 		for (int32 LODIndex = NumLODs - 1; LODIndex >= 0; --LODIndex)
 		{
-			if (RenderData->ScreenSize[LODIndex] > ScreenSize)
+			if (FMath::Square(RenderData->ScreenSize[LODIndex] * 0.5f) > ScreenRadiusSquared)
 			{
 				return FMath::Max(LODIndex, MinLOD);
 			}
@@ -397,7 +429,7 @@ FLODMask ComputeLODForMeshes( const TIndirectArray<class FStaticMesh>& StaticMes
 			{
 				int32 MinLODFound = INT_MAX;
 				bool bFoundLOD = false;
-				const float ScreenSize = ComputeTemporalLODBoundsScreenSize(Origin, SphereRadius, View, SampleIndex);
+				const float ScreenRadiusSquared = ComputeTemporalLODBoundsScreenRadiusSquared(Origin, SphereRadius, View, SampleIndex);
 
 				for(int32 MeshIndex = NumMeshes-1 ; MeshIndex >= 0 ; --MeshIndex)
 				{
@@ -405,7 +437,7 @@ FLODMask ComputeLODForMeshes( const TIndirectArray<class FStaticMesh>& StaticMes
 
 					float MeshScreenSize = Mesh.ScreenSize * ScreenSizeScale;
 
-					if(MeshScreenSize >= ScreenSize)
+					if(FMath::Square(MeshScreenSize * 0.5f) >= ScreenRadiusSquared)
 					{
 						LODToRender.SetLODSample(Mesh.LODIndex, SampleIndex);
 						bFoundLOD = true;
@@ -425,7 +457,7 @@ FLODMask ComputeLODForMeshes( const TIndirectArray<class FStaticMesh>& StaticMes
 		{
 			int32 MinLODFound = INT_MAX;
 			bool bFoundLOD = false;
-			const float ScreenSize = ComputeBoundsScreenSize(Origin, SphereRadius, View);
+			const float ScreenRadiusSquared = ComputeBoundsScreenRadiusSquared(Origin, SphereRadius, View);
 
 			for(int32 MeshIndex = NumMeshes-1 ; MeshIndex >= 0 ; --MeshIndex)
 			{
@@ -433,7 +465,7 @@ FLODMask ComputeLODForMeshes( const TIndirectArray<class FStaticMesh>& StaticMes
 
 				float MeshScreenSize = Mesh.ScreenSize * ScreenSizeScale;
 
-				if(MeshScreenSize >= ScreenSize)
+				if(FMath::Square(MeshScreenSize * 0.5f) >= ScreenRadiusSquared)
 				{
 					LODToRender.SetLOD(Mesh.LODIndex);
 					bFoundLOD = true;
diff --git a/Engine/Source/Runtime/Engine/Private/SimpleConstructionScript.cpp b/Engine/Source/Runtime/Engine/Private/SimpleConstructionScript.cpp
--- a/Engine/Source/Runtime/Engine/Private/SimpleConstructionScript.cpp
+++ b/Engine/Source/Runtime/Engine/Private/SimpleConstructionScript.cpp
@@ -1008,14 +1008,12 @@ void USimpleConstructionScript::GenerateListOfExistingNames(TArray<FName>& Curre
 	// <<< End Backwards Compatibility
 	check(Blueprint);
 
-	TArray<UObject*> NativeCDOChildren;
 	UClass* FirstNativeClass = FBlueprintEditorUtils::FindFirstNativeClass(Blueprint->ParentClass);
-	GetObjectsWithOuter(FirstNativeClass->GetDefaultObject(), NativeCDOChildren, false);
 
-	for (UObject* NativeCDOChild : NativeCDOChildren)
+	ForEachObjectWithOuter(FirstNativeClass->GetDefaultObject(), [&CurrentNames](UObject* NativeCDOChild)
 	{
 		CurrentNames.Add(NativeCDOChild->GetFName());
-	}
+	});
 
 	if (Blueprint->SkeletonGeneratedClass)
 	{
diff --git a/Engine/Source/Runtime/Engine/Private/SkeletalMesh.cpp b/Engine/Source/Runtime/Engine/Private/SkeletalMesh.cpp
--- a/Engine/Source/Runtime/Engine/Private/SkeletalMesh.cpp
+++ b/Engine/Source/Runtime/Engine/Private/SkeletalMesh.cpp
@@ -3076,6 +3076,9 @@ void USkeletalMesh::Serialize( FArchive& Ar )
 	}
 #endif
 
+#if WITH_EDITOR
+	bRequiresLODScreenSizeConversion = Ar.CustomVer(FFrameworkObjectVersion::GUID) < FFrameworkObjectVersion::LODsUseResolutionIndependentScreenSize;
+#endif
 }
 
 void USkeletalMesh::AddReferencedObjects(UObject* InThis, FReferenceCollector& Collector)
@@ -3435,6 +3438,14 @@ void USkeletalMesh::PostLoad()
 			}
 		}
 	}
+
+#if WITH_EDITOR
+	if (bRequiresLODScreenSizeConversion)
+	{
+		// Convert screen area to screen size
+		ConvertLegacyLODScreenSize();
+	}
+#endif
 }
 
 void USkeletalMesh::RebuildRefSkeletonNameToIndexMap()
@@ -4441,7 +4452,46 @@ void USkeletalMesh::AddBoneToReductionSetting(int32 LODIndex, FName BoneName)
 	}
 }
 
+void USkeletalMesh::ConvertLegacyLODScreenSize()
+{
+	if (LODInfo.Num() == 1)
+	{
+		// Only one LOD
+		LODInfo[0].ScreenSize = 1.0f;
+	}
+	else
+	{
+		// Use 1080p, 90 degree FOV as a default, as this should not cause runtime regressions in the common case.
+		// LODs will appear different in Persona, however.
+		const float HalfFOV = PI * 0.25f;
+		const float ScreenWidth = 1920.0f;
+		const float ScreenHeight = 1080.0f;
+		const FPerspectiveMatrix ProjMatrix(HalfFOV, ScreenWidth, ScreenHeight, 1.0f);
+		FBoxSphereBounds Bounds = GetBounds();
+
+		// Multiple models, we should have LOD screen area data.
+		for (int32 LODIndex = 0; LODIndex < LODInfo.Num(); ++LODIndex)
+		{
+			FSkeletalMeshLODInfo& LODInfoEntry = LODInfo[LODIndex];
+
+			if (LODInfoEntry.ScreenSize == 0.0f)
+			{
+				LODInfoEntry.ScreenSize = 1.0f;
+			}
+			else
+			{
+				// legacy screen size was scaled by a fixed constant of 320.0f, so its kinda arbitrary. Convert back to distance based metric first.
+				const float ScreenDepth = FMath::Max(ScreenWidth / 2.0f * ProjMatrix.M[0][0], ScreenHeight / 2.0f * ProjMatrix.M[1][1]) * Bounds.SphereRadius / (LODInfoEntry.ScreenSize * 320.0f);
+
+				// Now convert using the query function
+				LODInfoEntry.ScreenSize = ComputeBoundsScreenSize(FVector::ZeroVector, Bounds.SphereRadius, FVector(0.0f, 0.0f, ScreenDepth), ProjMatrix);
+			}
+		}
+	}
+}
+
 #endif // WITH_EDITOR
+
 /*-----------------------------------------------------------------------------
 USkeletalMeshSocket
 -----------------------------------------------------------------------------*/
diff --git a/Engine/Source/Runtime/Engine/Private/SkeletalRender.cpp b/Engine/Source/Runtime/Engine/Private/SkeletalRender.cpp
--- a/Engine/Source/Runtime/Engine/Private/SkeletalRender.cpp
+++ b/Engine/Source/Runtime/Engine/Private/SkeletalRender.cpp
@@ -69,11 +69,7 @@ void FSkeletalMeshObject::UpdateMinDesiredLODLevel(const FSceneView* View, const
 	static const auto* SkeletalMeshLODRadiusScale = IConsoleManager::Get().FindTConsoleVariableDataFloat(TEXT("r.SkeletalMeshLODRadiusScale"));
 	float LODScale = FMath::Clamp(SkeletalMeshLODRadiusScale->GetValueOnRenderThread(), 0.25f, 1.0f);
 
-	const FVector4 ScreenPosition( View->WorldToScreen(Bounds.Origin) );
-	const float ScreenRadius = FMath::Max(View->ViewRect.Width() / 2.0f * View->ViewMatrices.GetProjectionMatrix().M[0][0],
-		View->ViewRect.Height() / 2.0f * View->ViewMatrices.GetProjectionMatrix().M[1][1]) * Bounds.SphereRadius / FMath::Max(ScreenPosition.W,1.0f);
-
-	const float LODFactor = LODScale * ScreenRadius / 320.0f;
+	const float ScreenRadiusSquared = ComputeBoundsScreenRadiusSquared(Bounds.Origin, Bounds.SphereRadius, *View) * LODScale * LODScale;
 
 	check( SkeletalMeshLODInfo.Num() == SkeletalMeshResource->LODModels.Num() );
 
@@ -89,17 +85,17 @@ void FSkeletalMeshObject::UpdateMinDesiredLODLevel(const FSceneView* View, const
 		// Iterate from worst to best LOD
 		for(int32 LODLevel = SkeletalMeshResource->LODModels.Num()-1; LODLevel > 0; LODLevel--) 
 		{
-			// Get DistanceFactor for this LOD
-			float LODDistanceFactor = SkeletalMeshLODInfo[LODLevel].ScreenSize;
+			// Get ScreenSize for this LOD
+			float ScreenSize = SkeletalMeshLODInfo[LODLevel].ScreenSize;
 
 			// If we are considering shifting to a better (lower) LOD, bias with hysteresis.
 			if(LODLevel  <= CurrentLODLevel)
 			{
-				LODDistanceFactor += SkeletalMeshLODInfo[LODLevel].LODHysteresis;
+				ScreenSize += SkeletalMeshLODInfo[LODLevel].LODHysteresis;
 			}
 
 			// If have passed this boundary, use this LOD
-			if(LODDistanceFactor > LODFactor)
+			if(FMath::Square(ScreenSize * 0.5f) > ScreenRadiusSquared)
 			{
 				NewLODLevel = LODLevel;
 				break;
@@ -115,12 +111,12 @@ void FSkeletalMeshObject::UpdateMinDesiredLODLevel(const FSceneView* View, const
 		MinDesiredLODLevel = WorkingMinDesiredLODLevel;
 		LastFrameNumber = FrameNumber;
 
-		WorkingMaxDistanceFactor = LODFactor;
+		WorkingMaxDistanceFactor = ScreenRadiusSquared;
 		WorkingMinDesiredLODLevel = NewLODLevel;
 	}
 	else
 	{
-		WorkingMaxDistanceFactor = FMath::Max(WorkingMaxDistanceFactor, LODFactor);
+		WorkingMaxDistanceFactor = FMath::Max(WorkingMaxDistanceFactor, ScreenRadiusSquared);
 		WorkingMinDesiredLODLevel = FMath::Min(WorkingMinDesiredLODLevel, NewLODLevel);
 	}
 }
diff --git a/Engine/Source/Runtime/Engine/Private/SoundWave.cpp b/Engine/Source/Runtime/Engine/Private/SoundWave.cpp
--- a/Engine/Source/Runtime/Engine/Private/SoundWave.cpp
+++ b/Engine/Source/Runtime/Engine/Private/SoundWave.cpp
@@ -741,13 +741,18 @@ void USoundWave::Parse( FAudioDevice* AudioDevice, const UPTRINT NodeWaveInstanc
 		}
 
 		// Only append to the wave instances list if we're virtual (always append) or we're audible (non-zero volume)
+		bool bAddedWaveInstance = false;
 		if (WaveInstance->GetVolume() > KINDA_SMALL_NUMBER || (bVirtualizeWhenSilent && AudioDevice->VirtualSoundsEnabled()))
 		{
+			bAddedWaveInstance = true;
 			WaveInstances.Add(WaveInstance);
 		}
 
 		// We're still alive.
-		ActiveSound.bFinished = false;
+		if (bAddedWaveInstance || WaveInstance->LoopingMode == LOOP_Forever)
+		{
+			ActiveSound.bFinished = false;
+		}
 
 		// Sanity check
 		if( NumChannels > 2 && WaveInstance->bUseSpatialization && !WaveInstance->bReportedSpatializationWarning)
diff --git a/Engine/Source/Runtime/Engine/Private/StaticMesh.cpp b/Engine/Source/Runtime/Engine/Private/StaticMesh.cpp
--- a/Engine/Source/Runtime/Engine/Private/StaticMesh.cpp
+++ b/Engine/Source/Runtime/Engine/Private/StaticMesh.cpp
@@ -843,7 +843,20 @@ void FStaticMeshRenderData::ResolveSectionInfo(UStaticMesh* Owner)
 			else
 			{
 				const float ViewDistance = CalculateViewDistance(LOD.MaxDeviation, Owner->SourceModels[LODIndex].ReductionSettings.PixelError);
-				ScreenSize[LODIndex] = PI * FMath::Square( 0.5f * Bounds.SphereRadius / ( Bounds.SphereRadius + ViewDistance ) );
+
+				// Generate a projection matrix.
+				// ComputeBoundsScreenSize only uses (0, 0) and (1, 1) of this matrix.
+				const float HalfFOV = PI * 0.25f;
+				const float ScreenWidth = 1920.0f;
+				const float ScreenHeight = 1080.0f;
+				const FPerspectiveMatrix ProjMatrix(HalfFOV, ScreenWidth, ScreenHeight, 1.0f);
+
+				// Note we offset ViewDistance by SphereRadius here because the MaxDeviation is known to be somewhere in the bounds of the mesh. 
+				// It won't necessarily be at the origin. Before adding this factor for very high poly meshes it would calculate a very small deviation 
+				// for LOD1 which translates to a very small ViewDistance and a large (larger than 1) ScreenSize. This meant you could clip the camera 
+				// into the mesh but unless you were near its origin it wouldn't switch to LOD0. Adding SphereRadius to ViewDistance makes it so that 
+				// the distance is to the bounds which corrects the problem.
+				ScreenSize[LODIndex] = ComputeBoundsScreenSize(FVector::ZeroVector, Bounds.SphereRadius, FVector(0.0f, 0.0f, ViewDistance + Bounds.SphereRadius), ProjMatrix);
 			}
 		}
 		else if (Owner->SourceModels.IsValidIndex(LODIndex))
@@ -1159,8 +1172,8 @@ FArchive& operator<<(FArchive& Ar, FMeshBuildSettings& BuildSettings)
 // If static mesh derived data needs to be rebuilt (new format, serialization
 // differences, etc.) replace the version GUID below with a new one.
 // In case of merge conflicts with DDC versions, you MUST generate a new GUID
-// and set this new GUID as the version.        
-#define STATICMESH_DERIVEDDATA_VER TEXT("9028D97FC1F34C33B2F24FA81D5F45F3")
+// and set this new GUID as the version.                                       
+#define STATICMESH_DERIVEDDATA_VER TEXT("A1C08472697A4FC2B1E0A31E4B382B6E")
 
 static const FString& GetStaticMeshDerivedDataVersion()
 {
@@ -2375,6 +2388,7 @@ void UStaticMesh::Serialize(FArchive& Ar)
 		// Need to set a flag rather than do conversion in place as RenderData is not
 		// created until postload and it is needed for bounding information
 		bRequiresLODDistanceConversion = Ar.UE4Ver() < VER_UE4_STATIC_MESH_SCREEN_SIZE_LODS;
+		bRequiresLODScreenSizeConversion = Ar.CustomVer(FFrameworkObjectVersion::GUID) < FFrameworkObjectVersion::LODsUseResolutionIndependentScreenSize;
 	}
 #endif // #if WITH_EDITOR
 
@@ -2514,6 +2528,12 @@ void UStaticMesh::PostLoad()
 			ConvertLegacyLODDistance();
 		}
 
+		if (bRequiresLODScreenSizeConversion)
+		{
+			// Convert screen area to screen size
+			ConvertLegacyLODScreenArea();
+		}
+
 		if (RenderData && GStaticMeshesThatNeedMaterialFixup.Get(this))
 		{
 			FixupZeroTriangleSections();
@@ -3259,6 +3279,50 @@ void UStaticMesh::ConvertLegacyLODDistance()
 	}
 }
 
+void UStaticMesh::ConvertLegacyLODScreenArea()
+{
+	check(SourceModels.Num() > 0);
+	check(SourceModels.Num() <= MAX_STATIC_MESH_LODS);
+
+	if (SourceModels.Num() == 1)
+	{
+		// Only one model, 
+		SourceModels[0].ScreenSize = 1.0f;
+	}
+	else
+	{
+		// Use 1080p, 90 degree FOV as a default, as this should not cause runtime regressions in the common case.
+		const float HalfFOV = PI * 0.25f;
+		const float ScreenWidth = 1920.0f;
+		const float ScreenHeight = 1080.0f;
+		const FPerspectiveMatrix ProjMatrix(HalfFOV, ScreenWidth, ScreenHeight, 1.0f);
+		FBoxSphereBounds Bounds = GetBounds();
+
+		// Multiple models, we should have LOD screen area data.
+		for (int32 ModelIndex = 0; ModelIndex < SourceModels.Num(); ++ModelIndex)
+		{
+			FStaticMeshSourceModel& SrcModel = SourceModels[ModelIndex];
+
+			if (SrcModel.ScreenSize == 0.0f)
+			{
+				SrcModel.ScreenSize = 1.0f;
+				RenderData->ScreenSize[ModelIndex] = SrcModel.ScreenSize;
+			}
+			else
+			{
+				// legacy transition screen size was previously a screen AREA fraction using resolution-scaled values, so we need to convert to distance first to correctly calculate the threshold
+				const float ScreenArea = SrcModel.ScreenSize * (ScreenWidth * ScreenHeight);
+				const float ScreenRadius = FMath::Sqrt(ScreenArea / PI);
+				const float ScreenDistance = FMath::Max(ScreenWidth / 2.0f * ProjMatrix.M[0][0], ScreenHeight / 2.0f * ProjMatrix.M[1][1]) * Bounds.SphereRadius / ScreenRadius;
+
+				// Now convert using the query function
+				SrcModel.ScreenSize = ComputeBoundsScreenSize(FVector::ZeroVector, Bounds.SphereRadius, FVector(0.0f, 0.0f, ScreenDistance), ProjMatrix);
+				RenderData->ScreenSize[ModelIndex] = SrcModel.ScreenSize;
+			}
+		}
+	}
+}
+
 void UStaticMesh::GenerateLodsInPackage()
 {
 	FFormatNamedArguments Args;
diff --git a/Engine/Source/Runtime/Engine/Private/UnrealEngine.cpp b/Engine/Source/Runtime/Engine/Private/UnrealEngine.cpp
--- a/Engine/Source/Runtime/Engine/Private/UnrealEngine.cpp
+++ b/Engine/Source/Runtime/Engine/Private/UnrealEngine.cpp
@@ -6499,6 +6499,22 @@ bool UEngine::PerformError(const TCHAR* Cmd, FOutputDevice& Ar)
 		FPlatformProcess::Sleep(3600);
 		return true;
 	}
+	else if (FParse::Command(&Cmd, TEXT("AUDIOGPF")))
+	{
+		FAudioThread::RunCommandOnAudioThread([]()
+		{
+			*(int32 *)3 = 123;
+		}, TStatId());
+		return true;
+	}
+	else if (FParse::Command(&Cmd, TEXT("AUDIOCHECK")))
+	{
+		FAudioThread::RunCommandOnAudioThread([]()
+		{
+			check(!"Crashing the audio thread via check(0) at your request");
+		}, TStatId());
+		return true;
+	}
 #endif // !UE_BUILD_SHIPPING
 	return false;
 }
@@ -11429,15 +11445,13 @@ void UEngine::CopyPropertiesForUnrelatedObjects(UObject* OldObject, UObject* New
 	FFindInstancedReferenceSubobjectHelper::Duplicate(OldObject, NewObject, ReferenceReplacementMap, ComponentsOnNewObject);
 
 	// Replace anything with an outer of the old object with NULL, unless it already has a replacement
-	TArray<UObject*> ObjectsInOuter;
-	GetObjectsWithOuter(OldObject, ObjectsInOuter, true);
-	for (int32 ObjectIndex = 0; ObjectIndex < ObjectsInOuter.Num(); ObjectIndex++)
+	ForEachObjectWithOuter(OldObject, [&ReferenceReplacementMap](UObject* ObjectInOuter)
 	{
-		if (!ReferenceReplacementMap.Contains(ObjectsInOuter[ObjectIndex]))
+		if (!ReferenceReplacementMap.Contains(ObjectInOuter))
 		{
-			ReferenceReplacementMap.Add(ObjectsInOuter[ObjectIndex], NULL);
+			ReferenceReplacementMap.Add(ObjectInOuter, nullptr);
 		}
-	}
+	});
 
 	// Replace references to old classes and instances on this object with the corresponding new ones
 	UPackage* NewPackage = Cast<UPackage>(NewObject->GetOutermost());
@@ -12540,40 +12554,44 @@ void FAudioDevice::ResolveDesiredStats(FViewportClient* ViewportClient)
 
 void FAudioDevice::UpdateRequestedStat(const uint8 RequestedStat)
 {
-	check(IsInGameThread());
+	if (!IsInAudioThread())
+	{
+		DECLARE_CYCLE_STAT(TEXT("FAudioThreadTask.UpdateRequestedStat"), STAT_AudioUpdateRequestedStat, STATGROUP_TaskGraphTasks);
 
-	DECLARE_CYCLE_STAT(TEXT("FAudioThreadTask.UpdateRequestedStat"), STAT_AudioUpdateRequestedStat, STATGROUP_TaskGraphTasks);
+		FAudioDevice* AudioDevice = this;
+		FAudioThread::RunCommandOnAudioThread([AudioDevice, RequestedStat]()
+		{
+			AudioDevice->UpdateRequestedStat(RequestedStat);
+		}, GET_STATID(STAT_AudioUpdateRequestedStat));
+		return;
+	}
 
-	FAudioDevice* AudioDevice = this;
-	FAudioThread::RunCommandOnAudioThread([AudioDevice, RequestedStat]()
-	{
-		AudioDevice->RequestedAudioStats ^= RequestedStat;
-	}, GET_STATID(STAT_AudioUpdateRequestedStat));
+	RequestedAudioStats ^= RequestedStat;
 }
 
 bool UEngine::ToggleStatSoundWaves(UWorld* World, FCommonViewportClient* ViewportClient, const TCHAR* Stream)
 {
-	if (FAudioDevice* AudioDevice = World->GetAudioDevice())
+	if (AudioDeviceManager)
 	{
-		AudioDevice->UpdateRequestedStat(ERequestedAudioStats::SoundWaves);
+		AudioDeviceManager->ToggleDebugStat(ERequestedAudioStats::SoundWaves);
 	}
 	return true;
 }
 
 bool UEngine::ToggleStatSoundCues(UWorld* World, FCommonViewportClient* ViewportClient, const TCHAR* Stream)
 {
-	if (FAudioDevice* AudioDevice = World->GetAudioDevice())
+	if (AudioDeviceManager)
 	{
-		AudioDevice->UpdateRequestedStat(ERequestedAudioStats::SoundCues);
+		AudioDeviceManager->ToggleDebugStat(ERequestedAudioStats::SoundCues);
 	}
 	return true;
 }
 
 bool UEngine::ToggleStatSoundMixes(UWorld* World, FCommonViewportClient* ViewportClient, const TCHAR* Stream)
 {
-	if (FAudioDevice* AudioDevice = World->GetAudioDevice())
+	if (AudioDeviceManager)
 	{
-		AudioDevice->UpdateRequestedStat(ERequestedAudioStats::SoundMixes);
+		AudioDeviceManager->ToggleDebugStat(ERequestedAudioStats::SoundMixes);
 	}
 	return true;
 }
@@ -12615,7 +12633,7 @@ int32 UEngine::RenderStatSoundWaves(UWorld* World, FViewport* Viewport, FCanvas*
 				FString TheString = *FString::Printf(TEXT("%4i.    %6.2f  %s   Owner: %s   SoundClass: %s"),
 					WaveInstanceInfo.Key->InstanceIndex,
 					WaveInstanceInfo.Key->ActualVolume,
-					*WaveInstanceInfo.Value->SoundName,
+					*WaveInstanceInfo.Key->WaveInstanceName.ToString(),
 					SoundOwner ? *SoundOwner->GetName() : TEXT("None"),
 					*WaveInstanceInfo.Value->SoundClassName.ToString());
 
@@ -13046,6 +13064,4 @@ int32 UEngine::RenderStatSlateBatches(UWorld* World, FViewport* Viewport, FCanva
 }
 #endif
 
-ENGINE_API UGameplayTagsManager* GGameplayTagsManager = nullptr;
-
 #undef LOCTEXT_NAMESPACE
diff --git a/Engine/Source/Runtime/Engine/Private/Vehicles/WheeledVehicleMovementComponent.cpp b/Engine/Source/Runtime/Engine/Private/Vehicles/WheeledVehicleMovementComponent.cpp
--- a/Engine/Source/Runtime/Engine/Private/Vehicles/WheeledVehicleMovementComponent.cpp
+++ b/Engine/Source/Runtime/Engine/Private/Vehicles/WheeledVehicleMovementComponent.cpp
@@ -1798,8 +1798,11 @@ void UWheeledVehicleMovementComponent::SetAvoidanceEnabled(bool bEnable)
 	{
 		bUseRVOAvoidance = bEnable;
 		
+		// reset id, RegisterMovementComponent call is required to initialize update timers in avoidance manager
+		AvoidanceUID = 0;
+
 		UAvoidanceManager* AvoidanceManager = GetWorld()->GetAvoidanceManager();
-		if (AvoidanceManager && bEnable && AvoidanceUID == 0)
+		if (AvoidanceManager && bEnable)
 		{
 			AvoidanceManager->RegisterMovementComponent(this, AvoidanceWeight);
 		}
diff --git a/Engine/Source/Runtime/Engine/Private/WindDirectionalSource.cpp b/Engine/Source/Runtime/Engine/Private/WindDirectionalSource.cpp
--- a/Engine/Source/Runtime/Engine/Private/WindDirectionalSource.cpp
+++ b/Engine/Source/Runtime/Engine/Private/WindDirectionalSource.cpp
@@ -41,13 +41,13 @@ AWindDirectionalSource::AWindDirectionalSource(const FObjectInitializer& ObjectI
 			ArrowComponent->bUseInEditorScaling = true;
 		}
 
-		if (GetSpriteComponent())
+		if (UBillboardComponent* SpriteComp = GetSpriteComponent())
 		{
-			GetSpriteComponent()->Sprite = ConstructorStatics.SpriteTexture.Get();
-			GetSpriteComponent()->RelativeScale3D = FVector(0.5f, 0.5f, 0.5f);
-			GetSpriteComponent()->SpriteInfo.Category = ConstructorStatics.ID_Wind;
-			GetSpriteComponent()->SpriteInfo.DisplayName = ConstructorStatics.NAME_Wind;
-			GetSpriteComponent()->SetupAttachment(Component);
+			SpriteComp->Sprite = ConstructorStatics.SpriteTexture.Get();
+			SpriteComp->RelativeScale3D = FVector(0.5f, 0.5f, 0.5f);
+			SpriteComp->SpriteInfo.Category = ConstructorStatics.ID_Wind;
+			SpriteComp->SpriteInfo.DisplayName = ConstructorStatics.NAME_Wind;
+			SpriteComp->SetupAttachment(Component);
 		}
 	}
 #endif // WITH_EDITORONLY_DATA
@@ -165,6 +165,54 @@ void UWindDirectionalSourceComponent::Deactivate()
 	}
 }
 
+void UWindDirectionalSourceComponent::SetStrength(float InNewStrength)
+{
+	Strength = InNewStrength;
+
+	// Need to update render thread proxy with new data
+	MarkRenderDynamicDataDirty();
+}
+
+void UWindDirectionalSourceComponent::SetSpeed(float InNewSpeed)
+{
+	Speed = InNewSpeed;
+
+	// Need to update render thread proxy with new data
+	MarkRenderDynamicDataDirty();
+}
+
+void UWindDirectionalSourceComponent::SetMinimumGustAmount(float InNewMinGust)
+{
+	MinGustAmount = InNewMinGust;
+
+	// Need to update render thread proxy with new data
+	MarkRenderDynamicDataDirty();
+}
+
+void UWindDirectionalSourceComponent::SetMaximumGustAmount(float InNewMaxGust)
+{
+	MaxGustAmount = InNewMaxGust;
+
+	// Need to update render thread proxy with new data
+	MarkRenderDynamicDataDirty();
+}
+
+void UWindDirectionalSourceComponent::SetRadius(float InNewRadius)
+{
+	Radius = InNewRadius;
+
+	// Need to update render thread proxy with new data
+	MarkRenderDynamicDataDirty();
+}
+
+void UWindDirectionalSourceComponent::SetWindType(EWindSourceType InNewType)
+{
+	bPointWind = InNewType == EWindSourceType::Point;
+
+	// Need to update render thread proxy with new data
+	MarkRenderDynamicDataDirty();
+}
+
 void UWindDirectionalSourceComponent::CreateRenderState_Concurrent()
 {
 	Super::CreateRenderState_Concurrent();
@@ -174,8 +222,20 @@ void UWindDirectionalSourceComponent::CreateRenderState_Concurrent()
 void UWindDirectionalSourceComponent::SendRenderTransform_Concurrent()
 {
 	Super::SendRenderTransform_Concurrent();
-	GetWorld()->Scene->RemoveWindSource(this);
-	GetWorld()->Scene->AddWindSource(this);
+	UpdateSceneData_Concurrent();
+}
+
+void UWindDirectionalSourceComponent::SendRenderDynamicData_Concurrent()
+{
+	Super::SendRenderDynamicData_Concurrent();
+	UpdateSceneData_Concurrent();
+}
+
+void UWindDirectionalSourceComponent::UpdateSceneData_Concurrent()
+{
+	FSceneInterface* Scene = GetWorld()->Scene;
+	Scene->RemoveWindSource(this);
+	Scene->AddWindSource(this);
 }
 
 void UWindDirectionalSourceComponent::DestroyRenderState_Concurrent()
diff --git a/Engine/Source/Runtime/Engine/Private/World.cpp b/Engine/Source/Runtime/Engine/Private/World.cpp
--- a/Engine/Source/Runtime/Engine/Private/World.cpp
+++ b/Engine/Source/Runtime/Engine/Private/World.cpp
@@ -3376,15 +3376,13 @@ void UWorld::CleanupWorld(bool bSessionEnded, bool bCleanupResources, UWorld* Ne
 		TArray<UObject*> WorldObjects;
 
 		// Iterate over all objects to find ones that reside in the same package as the world.
-		GetObjectsWithOuter(GetOutermost(), WorldObjects);
-		for( int32 ObjectIndex=0; ObjectIndex < WorldObjects.Num(); ++ObjectIndex )
+		ForEachObjectWithOuter(GetOutermost(), [this](UObject* CurrentObject)
 		{
-			UObject* CurrentObject = WorldObjects[ObjectIndex];
 			if ( CurrentObject != this )
 			{
 				CurrentObject->ClearFlags( RF_Standalone );
 			}
-		}
+		});
 
 		if (bCleanupResources && WorldType != EWorldType::PIE)
 		{
diff --git a/Engine/Source/Runtime/Engine/Public/Animation/AnimInstanceProxy.h b/Engine/Source/Runtime/Engine/Public/Animation/AnimInstanceProxy.h
--- a/Engine/Source/Runtime/Engine/Public/Animation/AnimInstanceProxy.h
+++ b/Engine/Source/Runtime/Engine/Public/Animation/AnimInstanceProxy.h
@@ -127,6 +127,17 @@ public:
 	{ 
 		return bShouldExtractRootMotion;
 	}
+	
+	void SavePoseSnapshot(USkeletalMeshComponent* InSkeletalMeshComponent, FName SnapshotName);
+
+	struct FPoseSnapshot
+	{
+		TArray<FTransform> LocalTransforms;
+		FName SnapshotName;
+	};
+
+	/** Get a pose snapshot by name */
+	const FPoseSnapshot* GetPoseSnapshot(FName SnapshotName) const;
 
 	/** Access various counters */
 	const FGraphTraversalCounter& GetInitializationCounter() const { return InitializationCounter; }
@@ -612,4 +623,6 @@ private:
 	/** Native state exit bindings */
 	TArray<FNativeStateBinding> NativeStateExitBindings;
 
+	/** Array of snapshots. Each entry contains a name for finding specific pose snapshots */
+	TArray <FPoseSnapshot> PoseSnapshots;
 };
diff --git a/Engine/Source/Runtime/Engine/Public/Audio.h b/Engine/Source/Runtime/Engine/Public/Audio.h
--- a/Engine/Source/Runtime/Engine/Public/Audio.h
+++ b/Engine/Source/Runtime/Engine/Public/Audio.h
@@ -84,7 +84,7 @@ DECLARE_CYCLE_STAT_EXTERN( TEXT( "Decompress Vorbis" ), STAT_VorbisDecompressTim
 DECLARE_CYCLE_STAT_EXTERN( TEXT( "Prepare Audio Decompression" ), STAT_AudioPrepareDecompressionTime, STATGROUP_Audio , );
 DECLARE_CYCLE_STAT_EXTERN( TEXT( "Prepare Vorbis Decompression" ), STAT_VorbisPrepareDecompressionTime, STATGROUP_Audio , );
 DECLARE_CYCLE_STAT_EXTERN( TEXT( "Finding Nearest Location" ), STAT_AudioFindNearestLocation, STATGROUP_Audio , );
-DECLARE_CYCLE_STAT_EXTERN( TEXT( "Decompress Opus" ), STAT_OpusDecompressTime, STATGROUP_Audio , );
+DECLARE_CYCLE_STAT_EXTERN( TEXT( "Decompress Streamed" ), STAT_AudioStreamedDecompressTime, STATGROUP_Audio , );
 DECLARE_CYCLE_STAT_EXTERN( TEXT( "Buffer Creation" ), STAT_AudioResourceCreationTime, STATGROUP_Audio , );
 
 /**
diff --git a/Engine/Source/Runtime/Engine/Public/AudioDecompress.h b/Engine/Source/Runtime/Engine/Public/AudioDecompress.h
--- a/Engine/Source/Runtime/Engine/Public/AudioDecompress.h
+++ b/Engine/Source/Runtime/Engine/Public/AudioDecompress.h
@@ -116,6 +116,124 @@ public:
 	virtual int32 GetCurrentChunkOffset() const {return -1;}
 };
 
+/** 
+ * Default implementation of a streamed compressed audio format.
+ * Can be subclassed to support streaming of a specific asset format. Handles all 
+ * the platform independent aspects of file format streaming for you (dealing with UE4 streamed assets)
+ */
+class IStreamedCompressedInfo : public ICompressedAudioInfo
+{
+public:
+	IStreamedCompressedInfo();
+	virtual ~IStreamedCompressedInfo() {}
+
+	//~ Begin ICompressedInfo Interface
+	ENGINE_API virtual bool ReadCompressedInfo(const uint8* InSrcBufferData, uint32 InSrcBufferDataSize, FSoundQualityInfo* QualityInfo) override;
+	ENGINE_API virtual bool ReadCompressedData(uint8* Destination, bool bLooping, uint32 BufferSize) override;
+	ENGINE_API virtual void SeekToTime(const float SeekTime) override {};
+	ENGINE_API virtual void ExpandFile(uint8* DstBuffer, struct FSoundQualityInfo* QualityInfo) override;
+	ENGINE_API virtual void EnableHalfRate(bool HalfRate) override {};
+	virtual uint32 GetSourceBufferSize() const override { return SrcBufferDataSize; }
+	virtual bool UsesVorbisChannelOrdering() const override { return false; }
+	virtual int GetStreamBufferSize() const override { return  MONO_PCM_BUFFER_SIZE; }
+	virtual bool SupportsStreaming() const override { return true; }
+	virtual bool StreamCompressedInfo(USoundWave* Wave, FSoundQualityInfo* QualityInfo) override;
+	virtual bool StreamCompressedData(uint8* Destination, bool bLooping, uint32 BufferSize) override;
+	virtual int32 GetCurrentChunkIndex() const override { return CurrentChunkIndex; }
+	virtual int32 GetCurrentChunkOffset() const override { return SrcBufferOffset; }
+	//~ End ICompressedInfo Interface
+
+	/** Parse the header information from the input source buffer data. This is dependent on compression format. */
+	virtual bool ParseHeader(const uint8* InSrcBufferData, uint32 InSrcBufferDataSize, FSoundQualityInfo* QualityInfo) = 0;
+
+	/** Create the compression format dependent decoder object. */
+	virtual bool CreateDecoder() = 0;
+
+	/** Decode the input compressed frame data into output PCMData buffer. */
+	virtual int32 Decode(const uint8* FrameData, uint16 FrameSize, int16* OutPCMData, int32 SampleSize) = 0;
+
+	/** Optional method to allow decoder to prepare to loop. */
+	virtual void PrepareToLoop() {}
+
+	/** Return the size of the current compression frame */
+	virtual uint32 GetFrameSize() = 0;
+
+	/** The size of the decode PCM buffer size. */
+	virtual uint32 GetMaxFrameSizeSamples() const = 0;
+
+protected:
+
+	/** Reads from the internal source audio buffer stream of the given data size. */
+	uint32	Read(void *Outbuffer, uint32 DataSize);
+
+	/**
+	* Decompresses a frame of Opus data to PCM buffer
+	*
+	* @param FrameSize Size of the frame in bytes
+	* @return The amount of samples that were decompressed (< 0 indicates error)
+	*/
+	int32 DecompressToPCMBuffer(uint16 FrameSize);
+
+	/**
+	* Adds to the count of samples that have currently been decoded
+	*
+	* @param NewSamples	How many samples have been decoded
+	* @return How many samples were actually part of the true sample count
+	*/
+	uint32 IncrementCurrentSampleCount(uint32 NewSamples);
+
+	/**
+	* Writes data from decoded PCM buffer, taking into account whether some PCM has been written before
+	*
+	* @param Destination	Where to place the decoded sound
+	* @param BufferSize	Size of the destination buffer in bytes
+	* @return				How many bytes were written
+	*/
+	uint32	WriteFromDecodedPCM(uint8* Destination, uint32 BufferSize);
+
+	/**
+	* Zeroes the contents of a buffer
+	*
+	* @param Destination	Buffer to zero
+	* @param BufferSize	Size of the destination buffer in bytes
+	* @return				How many bytes were zeroed
+	*/
+	uint32	ZeroBuffer(uint8* Destination, uint32 BufferSize);
+
+	/** Ptr to the current streamed chunk. */
+	const uint8* SrcBufferData;
+	/** Size of the current streamed chunk. */
+	uint32 SrcBufferDataSize;
+	/** What byte we're currently reading in the streamed chunk. */
+	uint32 SrcBufferOffset;
+	/** Where the actual audio data starts in the current streamed chunk. Accounts for header offset. */
+	uint32 AudioDataOffset;
+	/** Sample rate of the source file */
+	uint16 SampleRate;
+	/** The total sample count of the source file. */
+	uint32 TrueSampleCount;
+	/** How many samples we've currently read in the source file. */
+	uint32 CurrentSampleCount;
+	/** Number of channels (left/right) in th esource file. */
+	uint8 NumChannels;
+	/** The maximum number of samples per decode frame. */
+	uint32 MaxFrameSizeSamples;
+	/** The number of bytes per interleaved sample (NumChannels * sizeof(int16)). */
+	uint32 SampleStride;
+	/** The decoded PCM byte array from the last decoded frame. */
+	TArray<uint8> LastDecodedPCM;
+	/** The amount of PCM data in bytes was decoded last. */
+	uint32 LastPCMByteSize;
+	/** The current offset in the last decoded PCM buffer. */
+	uint32 LastPCMOffset;
+	/** If we're currently reading the final buffer. */
+	bool bStoringEndOfFile;
+    /** Ptr to the sound wave currently streaming. */
+	USoundWave* StreamingSoundWave;
+	/** The current chunk index in the streamed chunks. */
+	int32 CurrentChunkIndex;
+};
+
 
 /**
  * Asynchronous audio decompression
diff --git a/Engine/Source/Runtime/Engine/Public/AudioDevice.h b/Engine/Source/Runtime/Engine/Public/AudioDevice.h
--- a/Engine/Source/Runtime/Engine/Public/AudioDevice.h
+++ b/Engine/Source/Runtime/Engine/Public/AudioDevice.h
@@ -294,6 +294,7 @@ struct FAudioStats
 		FString Description;
 		float ActualVolume;
 		int32 InstanceIndex;
+		FName WaveInstanceName;
 	};
 
 	struct FStatSoundInfo
diff --git a/Engine/Source/Runtime/Engine/Public/AudioDeviceManager.h b/Engine/Source/Runtime/Engine/Public/AudioDeviceManager.h
--- a/Engine/Source/Runtime/Engine/Public/AudioDeviceManager.h
+++ b/Engine/Source/Runtime/Engine/Public/AudioDeviceManager.h
@@ -148,6 +148,9 @@ public:
 	/** Toggles 3d visualization of 3d sounds on/off */
 	void ToggleVisualize3dDebug();
 
+	/** Toggles the given debug stat bitmask for all current audio devices. */
+	void ToggleDebugStat(const uint8 StatBitMask);
+
 	/** Debug solos the given sound class name. Sounds that play with this sound class will be solo'd */
 	void SetDebugSoloSoundClass(const TCHAR* SoundClassName);
 
diff --git a/Engine/Source/Runtime/Engine/Public/AudioThread.h b/Engine/Source/Runtime/Engine/Public/AudioThread.h
--- a/Engine/Source/Runtime/Engine/Public/AudioThread.h
+++ b/Engine/Source/Runtime/Engine/Public/AudioThread.h
@@ -33,16 +33,7 @@ private:
 	*/
 	static bool bIsAudioThreadRunning;
 
-	static uint32 SuspensionCount;
-
-	/**
-	* Polled by the game thread to detect crashes in the audio thread.
-	* If the audio thread crashes, it sets this variable to false.
-	*/
-	volatile static bool bIsAudioThreadHealthy;
-
-	/** If the audio thread has been terminated by an unhandled exception, this contains the error message. */
-	static FString AudioThreadError;
+	static bool bIsAudioThreadSuspended;
 
 	/** The audio thread itself. */
 	static FRunnable* AudioThreadRunnable;
@@ -75,16 +66,8 @@ public:
 	/** Execute a (presumably audio) command on the game thread. If GIsAudioThreadRunning is false the command will execute immediately */
 	static ENGINE_API void RunCommandOnGameThread(TFunction<void()> InFunction, const TStatId InStatId = TStatId());
 
-	/**
-	* Checks if the audio thread is healthy and running.
-	* If it has crashed, UE_LOG is called with the exception information.
-	*/
-	static ENGINE_API void CheckAudioThreadHealth();
-
-	/** Checks if the audio thread is healthy and running, without crashing */
-	static ENGINE_API bool IsAudioThreadHealthy();
-
 	static ENGINE_API void SetUseThreadedAudio(bool bInUseThreadedAudio);
+	static ENGINE_API bool IsUsingThreadedAudio() { return bUseThreadedAudio; }
 
 	static ENGINE_API bool IsAudioThreadRunning() { return bIsAudioThreadRunning; }
 
