Change several blueprints to calll SetAutoActivate instead of calling SetActive to fix cook warnings

diff --git a/Engine/Source/Runtime/Engine/Classes/Components/ActorComponent.h b/Engine/Source/Runtime/Engine/Classes/Components/ActorComponent.h
--- a/Engine/Source/Runtime/Engine/Classes/Components/ActorComponent.h
+++ b/Engine/Source/Runtime/Engine/Classes/Components/ActorComponent.h
@@ -323,6 +323,13 @@ public:
 	UFUNCTION(BlueprintCallable, Category="Components|Activation", meta=(UnsafeDuringActorConstruction="true"))
 	virtual bool IsActive() const;
 
+	/**
+	 * Sets whether the component should be auto activate or not. Only safe during construction scripts.
+	 * @param bNewAutoActivate - The new auto activate state of the component
+	 */
+	UFUNCTION(BlueprintCallable, Category="Components|Activation")
+	virtual void SetAutoActivate(bool bNewAutoActivate);
+
 	/** Sets whether this component can tick when paused. */
 	UFUNCTION(BlueprintCallable, Category="Utilities")
 	void SetTickableWhenPaused(bool bTickableWhenPaused);
diff --git a/Engine/Source/Runtime/Engine/Private/Components/ActorComponent.cpp b/Engine/Source/Runtime/Engine/Private/Components/ActorComponent.cpp
--- a/Engine/Source/Runtime/Engine/Private/Components/ActorComponent.cpp
+++ b/Engine/Source/Runtime/Engine/Private/Components/ActorComponent.cpp
@@ -1447,6 +1447,18 @@ void UActorComponent::SetActive(bool bNewActive, bool bReset)
 	}
 }
 
+void UActorComponent::SetAutoActivate(bool bNewAutoActivate)
+{
+	if (!bRegistered)
+	{
+		bAutoActivate = bNewAutoActivate;
+	}
+	else
+	{
+		UE_LOG(LogActorComponent, Warning, TEXT("SetAutoActivate called on component %s after registration!"), *GetFullName());
+	}
+}
+
 void UActorComponent::ToggleActive()
 {
 	SetActive(!bIsActive);
