Physics: Fixed line traces not working properly in editor worlds when physics substepping was enabled (UE-36408)
- Substepping relies on interpolating transforms over frames, but only game worlds will be ticked, so we now disallow this feature in non-game worlds.

diff --git a/Engine/Source/Runtime/Engine/Public/PhysicsPublic.h b/Engine/Source/Runtime/Engine/Public/PhysicsPublic.h
--- a/Engine/Source/Runtime/Engine/Public/PhysicsPublic.h
+++ b/Engine/Source/Runtime/Engine/Public/PhysicsPublic.h
@@ -521,6 +521,13 @@ public:
 	/** @return Whether physics scene is using substepping */
 	bool IsSubstepping(uint32 SceneType) const
 	{
+		// Substepping relies on interpolating transforms over frames, but only game worlds will be ticked,
+		// so we disallow this feature in non-game worlds.
+		if( !GetOwningWorld()->IsGameWorld() )
+		{
+			return false;
+		}
+
 		if (SceneType == PST_Sync) return bSubstepping;
 		if (SceneType == PST_Async) return bSubsteppingAsync;
 		return false;
