[NO NON-PII COMMIT MESSAGE]

diff --git a/Engine/Source/Runtime/Engine/Private/Components/ActorComponent.cpp b/Engine/Source/Runtime/Engine/Private/Components/ActorComponent.cpp
--- a/Engine/Source/Runtime/Engine/Private/Components/ActorComponent.cpp
+++ b/Engine/Source/Runtime/Engine/Private/Components/ActorComponent.cpp
@@ -1449,13 +1449,13 @@ void UActorComponent::SetActive(bool bNewActive, bool bReset)
 
 void UActorComponent::SetAutoActivate(bool bNewAutoActivate)
 {
-	if (!bRegistered)
+	if (!bRegistered || IsOwnerRunningUserConstructionScript())
 	{
 		bAutoActivate = bNewAutoActivate;
 	}
 	else
 	{
-		UE_LOG(LogActorComponent, Warning, TEXT("SetAutoActivate called on component %s after registration!"), *GetFullName());
+		UE_LOG(LogActorComponent, Warning, TEXT("SetAutoActivate called on component %s after construction!"), *GetFullName());
 	}
 }
 
