Copying //UE4/Orion-Staging to //UE4/Main (Source: //Orion/Dev-General @ 3196521)

diff --git a/Engine/Source/Runtime/Engine/Classes/Engine/DemoNetDriver.h b/Engine/Source/Runtime/Engine/Classes/Engine/DemoNetDriver.h
--- a/Engine/Source/Runtime/Engine/Classes/Engine/DemoNetDriver.h
+++ b/Engine/Source/Runtime/Engine/Classes/Engine/DemoNetDriver.h
@@ -49,39 +49,73 @@ public:
 	float		TimeSeconds;
 };
 
-typedef TArray< FReplayExternalData > FReplayExternalDataArray;
+// Using an indirect array here since FReplayExternalData stores an FBitReader, and it's not safe to store an FArchive directly in a TArray.
+typedef TIndirectArray< FReplayExternalData > FReplayExternalDataArray;
 
 struct FPlaybackPacket
 {
 	TArray< uint8 > Data;
 	float			TimeSeconds;
+	int32			LevelIndex;
 };
 
 enum ENetworkVersionHistory
 {
-	HISTORY_INITIAL				= 1,
-	HISTORY_SAVE_ABS_TIME_MS	= 2,			// We now save the abs demo time in ms for each frame (solves accumulation errors)
-	HISTORY_INCREASE_BUFFER		= 3,			// Increased buffer size of packets, which invalidates old replays
-	HISTORY_SAVE_ENGINE_VERSION	= 4,			// Now saving engine net version + InternalProtocolVersion
-	HISTORY_EXTRA_VERSION		= 5				// We now save engine/game protocol version, checksum, and changelist
+	HISTORY_INITIAL							= 1,
+	HISTORY_SAVE_ABS_TIME_MS				= 2,			// We now save the abs demo time in ms for each frame (solves accumulation errors)
+	HISTORY_INCREASE_BUFFER					= 3,			// Increased buffer size of packets, which invalidates old replays
+	HISTORY_SAVE_ENGINE_VERSION				= 4,			// Now saving engine net version + InternalProtocolVersion
+	HISTORY_EXTRA_VERSION					= 5,			// We now save engine/game protocol version, checksum, and changelist
+	HISTORY_MULTIPLE_LEVELS					= 6,			// Replays support seamless travel between levels
+	HISTORY_MULTIPLE_LEVELS_TIME_CHANGES	= 7				// Save out the time that level changes happen
 };
 
+static const uint32 MIN_SUPPORTED_VERSION = HISTORY_EXTRA_VERSION;
+
 static const uint32 NETWORK_DEMO_MAGIC				= 0x2CF5A13D;
-static const uint32 NETWORK_DEMO_VERSION			= HISTORY_EXTRA_VERSION;
+static const uint32 NETWORK_DEMO_VERSION			= HISTORY_MULTIPLE_LEVELS_TIME_CHANGES;
 
 static const uint32 NETWORK_DEMO_METADATA_MAGIC		= 0x3D06B24E;
 static const uint32 NETWORK_DEMO_METADATA_VERSION	= 0;
 
+USTRUCT()
+struct FLevelNameAndTime
+{
+	GENERATED_USTRUCT_BODY()
+
+	FLevelNameAndTime()
+		: LevelChangeTimeInMS(0)
+	{}
+
+	FLevelNameAndTime(const FString& InLevelName, uint32 InLevelChangeTimeInMS)
+		: LevelName(InLevelName)
+		, LevelChangeTimeInMS(InLevelChangeTimeInMS)
+	{}
+
+	UPROPERTY()
+	FString LevelName;
+
+	UPROPERTY()
+	uint32 LevelChangeTimeInMS;
+
+	friend FArchive& operator<<(FArchive& Ar, FLevelNameAndTime& LevelNameAndTime)
+	{
+		Ar << LevelNameAndTime.LevelName;
+		Ar << LevelNameAndTime.LevelChangeTimeInMS;
+		return Ar;
+	}
+};
+
 struct FNetworkDemoHeader
 {
-	uint32	Magic;							// Magic to ensure we're opening the right file.
-	uint32	Version;						// Version number to detect version mismatches.
-	uint32	NetworkChecksum;				// Network checksum
-	uint32	EngineNetworkProtocolVersion;	// Version of the engine internal network format
-	uint32	GameNetworkProtocolVersion;		// Version of the game internal network format
-	uint32	Changelist;						// Engine changelist built from
-	FString LevelName;						// Name of level loaded for demo
-	TArray<FString> GameSpecificData;		// Area for subclasses to write stuff
+	uint32	Magic;									// Magic to ensure we're opening the right file.
+	uint32	Version;								// Version number to detect version mismatches.
+	uint32	NetworkChecksum;						// Network checksum
+	uint32	EngineNetworkProtocolVersion;			// Version of the engine internal network format
+	uint32	GameNetworkProtocolVersion;				// Version of the game internal network format
+	uint32	Changelist;								// Engine changelist built from
+	TArray<FLevelNameAndTime> LevelNamesAndTimes;	// Name and time changes of levels loaded for demo
+	TArray<FString> GameSpecificData;				// Area for subclasses to write stuff
 
 	FNetworkDemoHeader() :
 		Magic( NETWORK_DEMO_MAGIC ),
@@ -108,9 +142,9 @@ struct FNetworkDemoHeader
 		Ar << Header.Version;
 
 		// Check version
-		if ( Header.Version != NETWORK_DEMO_VERSION )
+		if (Header.Version < MIN_SUPPORTED_VERSION)
 		{
-			UE_LOG( LogDemo, Error, TEXT( "Header.Version != NETWORK_DEMO_VERSION" ) );
+			UE_LOG(LogDemo, Error, TEXT("Header.Version (%d) < MIN_SUPPORTED_VERSION (%d)"), Header.Version, MIN_SUPPORTED_VERSION);
 			Ar.SetError();
 			return Ar;
 		}
@@ -120,7 +154,26 @@ struct FNetworkDemoHeader
 		Ar << Header.GameNetworkProtocolVersion;
 		Ar << Header.Changelist;
 
-		Ar << Header.LevelName;
+		if (Header.Version < HISTORY_MULTIPLE_LEVELS)
+		{
+			FString LevelName;
+			Ar << LevelName;
+			Header.LevelNamesAndTimes.Add(FLevelNameAndTime(LevelName, 0));
+		}
+		else if (Header.Version == HISTORY_MULTIPLE_LEVELS)
+		{
+			TArray<FString> LevelNames;
+			Ar << LevelNames;
+
+			for (const FString& LevelName : LevelNames)
+			{
+				Header.LevelNamesAndTimes.Add(FLevelNameAndTime(LevelName, 0));
+			}
+		}
+		else
+		{
+			Ar << Header.LevelNamesAndTimes;
+		}
 
 		Ar << Header.GameSpecificData;
 
@@ -157,6 +210,9 @@ class ENGINE_API UDemoNetDriver : public UNetDriver
 	/** True if as have paused all of the channels */
 	bool bChannelsArePaused;
 
+	/** Index of LevelNames that is currently loaded */
+	int32 CurrentLevelIndex;
+
 	/** This is our spectator controller that is used to view the demo world from */
 	APlayerController* SpectatorController;
 
@@ -185,7 +241,7 @@ class ENGINE_API UDemoNetDriver : public UNetDriver
 
 	void		SaveCheckpoint();
 	void		TickCheckpoint();
-	void		LoadCheckpoint( FArchive* GotoCheckpointArchive, int64 GotoCheckpointSkipExtraTimeInMS );
+	bool		LoadCheckpoint( FArchive* GotoCheckpointArchive, int64 GotoCheckpointSkipExtraTimeInMS );
 
 	void		SaveExternalData( FArchive& Ar );
 	void		LoadExternalData( FArchive& Ar, const float TimeSeconds );
@@ -207,6 +263,8 @@ class ENGINE_API UDemoNetDriver : public UNetDriver
 	/** Streaming levels waiting to be saved next frame */
 	TArray< UObject* >					NewStreamingLevelsThisFrame;
 
+	bool bRecordMapChanges;
+
 private:
 	bool		bIsFastForwarding;
 	bool		bIsFastForwardingForCheckpoint;
@@ -238,6 +296,12 @@ private:
 	/** Optional time quota for actor replication during recording. Going over this limit effectively lowers the net update frequency of the remaining actors. Negative values are considered unlimited. */
 	float MaxDesiredRecordTimeMS;
 
+	/**
+	 * Maximum time allowed each frame to spend on saving a checkpoint. If 0, it will save the checkpoint in a single frame, regardless of how long it takes.
+	 * See also demo.CheckpointSaveMaxMSPerFrameOverride.
+	 */
+	float CheckpointSaveMaxMSPerFrame;
+
 	/** A player controller that this driver should consider its viewpoint for actor prioritization purposes. */
 	TWeakObjectPtr<APlayerController> ViewerOverride;
 
@@ -247,11 +311,24 @@ private:
 	/** If true, recording will prioritize replicating actors based on the value that AActor::GetReplayPriority returns. */
 	bool bPrioritizeActors;
 
-
-
 	/** If true, will skip recording, but leaves the replay open so that recording can be resumed again. */
 	bool bPauseRecording;
 
+	/** List of levels used in the current replay */
+	TArray<FLevelNameAndTime> LevelNamesAndTimes;
+
+	/** Does the actual work of TickFlush, either on the main thread or in a task thread in parallel with Slate. */
+	void TickFlushInternal(float DeltaSeconds);
+
+	/** Returns true if TickFlush can be called in parallel with the Slate tick. */
+	bool ShouldTickFlushAsyncEndOfFrame() const;
+
+	/** Returns either CheckpointSaveMaxMSPerFrame or the value of demo.CheckpointSaveMaxMSPerFrameOverride if it's >= 0. */
+	float GetCheckpointSaveMaxMSPerFrame() const;
+
+	/** Adds a new level to the level list */
+	void AddNewLevel(const FString& NewLevelName);
+
 public:
 
 	// UNetDriver interface.
@@ -275,6 +352,9 @@ public:
 	virtual AActor* GetActorForGUID(FNetworkGUID InGUID) const override;
 	virtual bool ShouldReceiveRepNotifiesForObject(UObject* Object) const override;
 
+	/** Called when we are already recording but have traveled to a new map to start recording again */
+	bool ContinueListen(FURL& ListenURL);
+
 	/** 
 	 * Scrubs playback to the given time. 
 	 * 
@@ -297,6 +377,14 @@ public:
 	/** Enable or disable prioritization of actors for recording. */
 	void SetActorPrioritizationEnabled(const bool bInPrioritizeActors) { bPrioritizeActors = bInPrioritizeActors; }
 
+	/** Sets CheckpointSaveMaxMSPerFrame. */
+	void SetCheckpointSaveMaxMSPerFrame(const float InCheckpointSaveMaxMSPerFrame) { CheckpointSaveMaxMSPerFrame = InCheckpointSaveMaxMSPerFrame; }
+
+	/** Called by a task thread if the engine is doing async end of frame tasks in parallel with Slate. */
+	void TickFlushAsyncEndOfFrame(float DeltaSeconds);
+
+	const TArray<FLevelNameAndTime>& GetLevelNameAndTimeList() { return LevelNamesAndTimes; }
+
 public:
 
 	// FExec interface
@@ -382,4 +470,12 @@ protected:
 	{
 		return true;
 	}
+
+	void ProcessClientTravelFunction(class AActor* Actor, class UFunction* Function, void* Parameters, struct FOutParmRec* OutParms, struct FFrame* Stack, class UObject* SubObject);
+
+	bool WriteNetworkDemoHeader(FString& Error);
+
+	void ProcessSeamlessTravel(int32 LevelIndex);
+
+	TArray<FQueuedDemoPacket> QueuedPacketsBeforeTravel;
 };
diff --git a/Engine/Source/Runtime/Engine/Classes/Engine/Engine.h b/Engine/Source/Runtime/Engine/Classes/Engine/Engine.h
--- a/Engine/Source/Runtime/Engine/Classes/Engine/Engine.h
+++ b/Engine/Source/Runtime/Engine/Classes/Engine/Engine.h
@@ -2645,7 +2645,8 @@ public:
 	/** @return true if the engine is autosaving a package */
 	virtual bool IsAutosaving() const { return false; }
 
-	/** @return true if this is a "vanilla" product running only Epic-built binaries, no third-party plugins, no game modules, etc. */
+	virtual bool ShouldDoAsyncEndOfFrameTasks() const { return false; }
+
 	bool IsVanillaProduct() const { return bIsVanillaProduct; }
 
 protected:
diff --git a/Engine/Source/Runtime/Engine/Classes/Engine/GameEngine.h b/Engine/Source/Runtime/Engine/Classes/Engine/GameEngine.h
--- a/Engine/Source/Runtime/Engine/Classes/Engine/GameEngine.h
+++ b/Engine/Source/Runtime/Engine/Classes/Engine/GameEngine.h
@@ -57,6 +57,8 @@ public:
 	 */
 	static TSharedRef<SWindow> CreateGameWindow();
 
+	static void SafeFrameChanged();
+
 	/**
 	 * Modifies the game window resolution settings if any overrides have been specified on the command line
 	 *
@@ -106,6 +108,7 @@ public:
 	virtual void ProcessToggleFreezeCommand( UWorld* InWorld ) override;
 	virtual void ProcessToggleFreezeStreamingCommand( UWorld* InWorld ) override;
 	virtual bool NetworkRemapPath(UNetDriver* Driver, FString& Str, bool bReading = true) override;
+	virtual bool ShouldDoAsyncEndOfFrameTasks() const override;
 
 public:
 
diff --git a/Engine/Source/Runtime/Engine/Classes/Engine/GameInstance.h b/Engine/Source/Runtime/Engine/Classes/Engine/GameInstance.h
--- a/Engine/Source/Runtime/Engine/Classes/Engine/GameInstance.h
+++ b/Engine/Source/Runtime/Engine/Classes/Engine/GameInstance.h
@@ -35,6 +35,60 @@ class FOnlineSessionSearchResult;
 DECLARE_MULTICAST_DELEGATE_ThreeParams(FOnPreClientTravel, const FString& /*PendingURL*/, ETravelType /*TravelType*/, bool /*bIsSeamlessTravel*/);
 typedef FOnPreClientTravel::FDelegate FOnPreClientTravelDelegate;
 
+#if WITH_EDITOR
+
+// The result of a UGameInstance PIE operation
+struct FGameInstancePIEResult
+{
+public:
+	// If not, what was the failure reason
+	FText FailureReason;
+
+	// Did the PIE operation succeed?
+	bool bSuccess;
+
+public:
+	static FGameInstancePIEResult Success()
+	{
+		return FGameInstancePIEResult(true, FText::GetEmpty());
+	}
+
+	static FGameInstancePIEResult Failure(const FText& InReason)
+	{
+		return FGameInstancePIEResult(false, InReason);
+	}
+
+	bool IsSuccess() const
+	{
+		return bSuccess;
+	}
+private:
+	FGameInstancePIEResult(bool bWasSuccess, const FText& InReason)
+		: FailureReason(InReason)
+		, bSuccess(bWasSuccess)
+	{
+	}
+};
+
+// Parameters used to initialize / start a PIE game instance
+//@TODO: Some of these are really mutually exclusive and should be refactored (put into a struct to make this easier in the future)
+struct FGameInstancePIEParameters
+{
+	// Are we doing SIE instead of PIE?
+	bool bSimulateInEditor;
+
+	// Were there any BP compile errors?
+	bool bAnyBlueprintErrors;
+
+	// Should we start in spectator mode?
+	bool bStartInSpectatorMode;
+
+	// Is this a dedicated server instance for PIE?
+	bool bRunAsDedicated;
+};
+
+#endif
+
 /**
  * GameInstance: high-level manager object for an instance of the running game.
  * Spawned at game creation and not destroyed until game instance is shut down.
@@ -107,12 +161,19 @@ public:
 
 	/* Called to initialize the game instance for standalone instances of the game */
 	void InitializeStandalone();
+
 #if WITH_EDITOR
 	/* Called to initialize the game instance for PIE instances of the game */
+	virtual FGameInstancePIEResult InitializeForPlayInEditor(int32 PIEInstanceIndex, const FGameInstancePIEParameters& Params);
+
+	/* Called to actually start the game when doing Play/Simulate In Editor */
+	virtual FGameInstancePIEResult StartPlayInEditorGameInstance(ULocalPlayer* LocalPlayer, const FGameInstancePIEParameters& Params);
+
+	DEPRECATED(4.15, "Please override InitializeForPIE instead")
 	virtual bool InitializePIE(bool bAnyBlueprintErrors, int32 PIEInstance, bool bRunAsDedicated);
 
+	DEPRECATED(4.15, "Please override StartPlayInEditorGameInstance instead")
 	virtual bool StartPIEGameInstance(ULocalPlayer* LocalPlayer, bool bInSimulateInEditor, bool bAnyBlueprintErrors, bool bStartInSpectatorMode);
-
 #endif
 
 	class UEngine* GetEngine() const;
@@ -198,6 +259,9 @@ public:
 	/** Called when demo playback fails for any reason */
 	virtual void HandleDemoPlaybackFailure( EDemoPlayFailure::Type FailureType, const FString& ErrorString = TEXT("") ) { }
 
+	/** This gets called when the player scrubs in a replay to a different level */
+	virtual void OnSeamlessTravelDuringReplay() { }
+
 	static void AddReferencedObjects(UObject* InThis, FReferenceCollector& Collector);
 
 	inline FTimerManager& GetTimerManager() const { return *TimerManager; }
@@ -262,6 +326,18 @@ public:
 	/** Returns true if this instance is for a dedicated server world */
 	bool IsDedicatedServerInstance() const;
 
+	/**
+	 * Helper function for traveling to a session that has already been joined via the online platform
+	 * Grabs the URL from the session info and travels
+	 *
+	 * @param ControllerId controller initiating the request
+	 * @param InSessionName name of session to travel to
+	 *
+	 * @return true if able or attempting to travel, false otherwise
+	 */
+#define ADDED_CLIENTTRAVELTOSESSION 1
+	virtual bool ClientTravelToSession(int32 ControllerId, FName InSessionName);
+
 	/** Broadcast a notification that travel is occurring */
 	void NotifyPreClientTravel(const FString& PendingURL, ETravelType TravelType, bool bIsSeamlessTravel);
 	/** @return delegate fired when client travel occurs */
diff --git a/Engine/Source/Runtime/Engine/Classes/Engine/GameViewportClient.h b/Engine/Source/Runtime/Engine/Classes/Engine/GameViewportClient.h
--- a/Engine/Source/Runtime/Engine/Classes/Engine/GameViewportClient.h
+++ b/Engine/Source/Runtime/Engine/Classes/Engine/GameViewportClient.h
@@ -719,6 +719,12 @@ public:
 	}
 #endif
 
+	/** Accessor for delegate called when the engine toggles fullscreen */
+	FOnToggleFullscreen& OnToggleFullscreen()
+	{
+		return ToggleFullscreenDelegate;
+	}
+
 private:
 	/**
 	 * Set a specific stat to either enabled or disabled (returns the number of remaining enabled stats)
@@ -829,6 +835,9 @@ private:
 	/** Delegate called when ticking the game viewport */
 	FOnGameViewportTick TickDelegate;
 
+	/** Delegate called when the engine toggles fullscreen */
+	FOnToggleFullscreen ToggleFullscreenDelegate;
+
 	/** Data needed to display perframe stat tracking when STAT UNIT is enabled */
 	FStatUnitData* StatUnitData;
 
diff --git a/Engine/Source/Runtime/Engine/Classes/Engine/GameViewportDelegates.h b/Engine/Source/Runtime/Engine/Classes/Engine/GameViewportDelegates.h
--- a/Engine/Source/Runtime/Engine/Classes/Engine/GameViewportDelegates.h
+++ b/Engine/Source/Runtime/Engine/Classes/Engine/GameViewportDelegates.h
@@ -52,3 +52,10 @@ DECLARE_MULTICAST_DELEGATE_OneParam(FOnGameViewportClientPlayerAction, int32);
 DECLARE_MULTICAST_DELEGATE_OneParam(FOnGameViewportTick, float);
 
 DECLARE_DELEGATE_RetVal_ThreeParams(bool, FOnGameViewportInputKey, FKey, FModifierKeysState, EInputEvent);
+
+/**
+* Delegate type used by UGameViewportClient for when engine in toggling fullscreen
+*
+* @see UGameViewportClient
+*/
+DECLARE_MULTICAST_DELEGATE_OneParam(FOnToggleFullscreen, bool);
diff --git a/Engine/Source/Runtime/Engine/Classes/Engine/ObjectLibrary.h b/Engine/Source/Runtime/Engine/Classes/Engine/ObjectLibrary.h
--- a/Engine/Source/Runtime/Engine/Classes/Engine/ObjectLibrary.h
+++ b/Engine/Source/Runtime/Engine/Classes/Engine/ObjectLibrary.h
@@ -162,6 +162,9 @@ public:
 	/** Whether to set bIncludeOnlyOnDiskAssets when setting up AR filters for discovering asset data */
 	bool bIncludeOnlyOnDiskAssets;
 
+	/** Whether to set bRecursivePaths when setting up AR filters for discovering asset data */
+	bool bRecursivePaths;
+
 #if WITH_EDITOR
 	/** The paths that we will query again once assets are finished being discovered */
 	TArray<FString> DeferredAssetDataPaths;
diff --git a/Engine/Source/Runtime/Engine/Classes/Engine/PendingNetGame.h b/Engine/Source/Runtime/Engine/Classes/Engine/PendingNetGame.h
--- a/Engine/Source/Runtime/Engine/Classes/Engine/PendingNetGame.h
+++ b/Engine/Source/Runtime/Engine/Classes/Engine/PendingNetGame.h
@@ -50,16 +50,31 @@ namespace EAcceptConnection
 class FNetworkNotify
 {
 public:
-	/** @todo document */
+	/**
+	 * Notification that an incoming connection is pending, giving the interface a chance to reject the request
+	 *
+	 * @return EAcceptConnection indicating willingness to accept the connection at this time
+	 */
 	virtual EAcceptConnection::Type NotifyAcceptingConnection() PURE_VIRTUAL(FNetworkNotify::NotifyAcceptedConnection,return EAcceptConnection::Ignore;);
 
-	/** @todo document */
+	/**
+	 * Notification that a new connection has been created/established as a result of a remote request, previously approved by NotifyAcceptingConnection
+	 *
+	 * @param Connection newly created connection
+	 */
 	virtual void NotifyAcceptedConnection( class UNetConnection* Connection ) PURE_VIRTUAL(FNetworkNotify::NotifyAcceptedConnection,);
 
-	/** @todo document */
+	/**
+	 * Notification that a new channel is being created/opened as a result of a remote request (Actor creation, etc)
+	 *
+	 * @param Channel newly created channel
+	 *
+	 * @return true if the channel should be opened, false if it should be rejected (destroying the channel)
+	 */
 	virtual bool NotifyAcceptingChannel( class UChannel* Channel ) PURE_VIRTUAL(FNetworkNotify::NotifyAcceptingChannel,return false;);
 
-	/** handler for messages sent through a remote connection's control channel
+	/** 
+	 * Handler for messages sent through a remote connection's control channel
 	 * not required to handle the message, but if it reads any data from Bunch, it MUST read the ENTIRE data stream for that message (i.e. use FNetControlMessage<TYPE>::Receive())
 	 */
 	virtual void NotifyControlMessage(UNetConnection* Connection, uint8 MessageType, class FInBunch& Bunch) PURE_VIRTUAL(FNetworkNotify::NotifyReceivedText,);
diff --git a/Engine/Source/Runtime/Engine/Classes/GameFramework/Character.h b/Engine/Source/Runtime/Engine/Classes/GameFramework/Character.h
--- a/Engine/Source/Runtime/Engine/Classes/GameFramework/Character.h
+++ b/Engine/Source/Runtime/Engine/Classes/GameFramework/Character.h
@@ -348,9 +348,6 @@ public:
 	UPROPERTY(BlueprintReadOnly, replicatedUsing=OnRep_IsCrouched, Category=Character)
 	uint32 bIsCrouched:1;
 
-	UPROPERTY( replicated )
-	uint32 bReplayHasRootMotionSources:1;
-
 	/** Handle Crouching replicated from server */
 	UFUNCTION()
 	virtual void OnRep_IsCrouched();
diff --git a/Engine/Source/Runtime/Engine/Classes/GameFramework/GameSession.h b/Engine/Source/Runtime/Engine/Classes/GameFramework/GameSession.h
--- a/Engine/Source/Runtime/Engine/Classes/GameFramework/GameSession.h
+++ b/Engine/Source/Runtime/Engine/Classes/GameFramework/GameSession.h
@@ -12,6 +12,8 @@ class UWorld;
 class APlayerController;
 struct FJoinabilitySettings;
 
+DECLARE_LOG_CATEGORY_EXTERN(LogGameSession, Log, All);
+
 /**
 Acts as a game-specific wrapper around the session interface. The game code makes calls to this when it needs to interact with the session interface.
 A game session exists only the server, while running an online game.
@@ -197,16 +199,6 @@ class ENGINE_API AGameSession : public AInfo
 	 */
 	virtual void UpdateSessionJoinability(FName InSessionName, bool bPublicSearchable, bool bAllowInvites, bool bJoinViaPresence, bool bJoinViaPresenceFriendsOnly);
 
-	/**
-	 * Travel to a session URL (as client) for a given session
-	 *
-	 * @param ControllerId controller initiating the session travel
-	 * @param SessionName name of session to travel to
-	 *
-	 * @return true if successful, false otherwise
-	 */
-	virtual bool TravelToSession(int32 ControllerId, FName InSessionName);
-
     /**
      * Does the session require push to talk
      * @return true if a push to talk keybinding is required or if voice is always enabled
diff --git a/Engine/Source/Runtime/Engine/Classes/GameFramework/RootMotionSource.h b/Engine/Source/Runtime/Engine/Classes/GameFramework/RootMotionSource.h
--- a/Engine/Source/Runtime/Engine/Classes/GameFramework/RootMotionSource.h
+++ b/Engine/Source/Runtime/Engine/Classes/GameFramework/RootMotionSource.h
@@ -121,6 +121,8 @@ enum class ERootMotionFinishVelocityMode : uint8
 	MaintainLastRootMotionVelocity = 0,
 	// Set Velocity to the specified value (for example, 0,0,0 to stop the character)
 	SetVelocity,
+	// Clamp velocity magnitude to the specified value 
+	ClampVelocity,
 };
 
 
diff --git a/Engine/Source/Runtime/Engine/Classes/Particles/Light/ParticleModuleLight.h b/Engine/Source/Runtime/Engine/Classes/Particles/Light/ParticleModuleLight.h
--- a/Engine/Source/Runtime/Engine/Classes/Particles/Light/ParticleModuleLight.h
+++ b/Engine/Source/Runtime/Engine/Classes/Particles/Light/ParticleModuleLight.h
@@ -49,6 +49,14 @@ class UParticleModuleLight : public UParticleModuleLightBase
 	UPROPERTY(EditAnywhere, Category=Light)
 	struct FRawDistributionFloat LightExponent;
 
+	/**
+	* Channels that this light should affect.
+	* Only affect high quality lights
+	* These channels only apply to opaque materials, direct lighting, and dynamic lighting and shadowing.
+	*/
+	UPROPERTY(EditAnywhere, AdvancedDisplay, BlueprintReadOnly, Category = Light)
+	FLightingChannels LightingChannels;
+
 	UPROPERTY(EditAnywhere, Category = Light)
 	bool bHighQualityLights;
 
diff --git a/Engine/Source/Runtime/Engine/Engine.Build.cs b/Engine/Source/Runtime/Engine/Engine.Build.cs
--- a/Engine/Source/Runtime/Engine/Engine.Build.cs
+++ b/Engine/Source/Runtime/Engine/Engine.Build.cs
@@ -9,7 +9,7 @@ public class Engine : ModuleRules
 	{
 		SharedPCHHeaderFile = "Runtime/Engine/Public/Engine.h";
 
-		PublicIncludePathModuleNames.AddRange(new string[] { "Renderer", "PacketHandler" });
+		PublicIncludePathModuleNames.AddRange(new string[] { "Renderer", "PacketHandler", "NetworkReplayStreaming" });
 
 		PrivateIncludePaths.AddRange(
 			new string[] {
diff --git a/Engine/Source/Runtime/Engine/Private/AI/Navigation/NavigationModifier.cpp b/Engine/Source/Runtime/Engine/Private/AI/Navigation/NavigationModifier.cpp
--- a/Engine/Source/Runtime/Engine/Private/AI/Navigation/NavigationModifier.cpp
+++ b/Engine/Source/Runtime/Engine/Private/AI/Navigation/NavigationModifier.cpp
@@ -46,7 +46,9 @@ UClass* FNavigationLinkBase::GetAreaClass() const
 
 void FNavigationLinkBase::InitializeAreaClass(const bool bForceRefresh)
 {
-	if (!bAreaClassInitialized || bForceRefresh)
+	// @hack fix for changes to AreaClass in the editor not taking effect
+	// proper fix will get merged over from Dev-Gen
+	if (!bAreaClassInitialized || bForceRefresh || GIsEditor)
 	{
 		AreaClassOb = (UClass*)AreaClass;
 		bAreaClassInitialized = true;
diff --git a/Engine/Source/Runtime/Engine/Private/AI/Navigation/NavigationSystem.cpp b/Engine/Source/Runtime/Engine/Private/AI/Navigation/NavigationSystem.cpp
--- a/Engine/Source/Runtime/Engine/Private/AI/Navigation/NavigationSystem.cpp
+++ b/Engine/Source/Runtime/Engine/Private/AI/Navigation/NavigationSystem.cpp
@@ -1463,7 +1463,7 @@ const ANavigationData* UNavigationSystem::GetNavDataForProps(const FNavAgentProp
 	if (NavDataInstance == nullptr)
 	{
 		TArray<FNavAgentProperties> AgentPropertiesList;
-		int32 NumNavDatas = AgentToNavDataMap.GetKeys(AgentPropertiesList);
+		AgentToNavDataMap.GenerateKeyArray(AgentPropertiesList);
 		
 		FNavAgentProperties BestFitNavAgent;
 		float BestExcessHeight = -FLT_MAX;
diff --git a/Engine/Source/Runtime/Engine/Private/AI/Navigation/RecastNavMesh.cpp b/Engine/Source/Runtime/Engine/Private/AI/Navigation/RecastNavMesh.cpp
--- a/Engine/Source/Runtime/Engine/Private/AI/Navigation/RecastNavMesh.cpp
+++ b/Engine/Source/Runtime/Engine/Private/AI/Navigation/RecastNavMesh.cpp
@@ -837,6 +837,97 @@ bool ARecastNavMesh::GetPolysInTile(int32 TileIndex, TArray<FNavPoly>& Polys) co
 	return RecastNavMeshImpl && RecastNavMeshImpl->GetPolysInTile(TileIndex, Polys);
 }
 
+bool ARecastNavMesh::GetNavLinksInTile(const int32 TileIndex, TArray<FNavPoly>& Polys, const bool bIncludeLinksFromNeighborTiles) const
+{
+	if (RecastNavMeshImpl == nullptr || RecastNavMeshImpl->DetourNavMesh == nullptr
+		|| TileIndex < 0 || TileIndex >= RecastNavMeshImpl->DetourNavMesh->getMaxTiles())
+	{
+		return false;
+	}
+
+	const dtNavMesh* DetourNavMesh = RecastNavMeshImpl->DetourNavMesh;
+	const int32 InitialLinkCount = Polys.Num();
+
+	const dtMeshTile* Tile = DetourNavMesh->getTile(TileIndex);
+	if (Tile && Tile->header)
+	{
+		const int32 LinkCount = Tile->header->offMeshConCount;
+
+		if (LinkCount > 0)
+		{
+			const int32 BaseIdx = Polys.Num();
+			Polys.AddZeroed(LinkCount);
+
+			const dtPoly* Poly = Tile->polys;
+			for (int32 LinkIndex = 0; LinkIndex < LinkCount; ++LinkIndex, ++Poly)
+			{
+				FNavPoly& OutPoly = Polys[BaseIdx + LinkIndex];
+				const int32 PolyIndex = Tile->header->offMeshBase + LinkIndex;
+				OutPoly.Ref = DetourNavMesh->encodePolyId(Tile->salt, TileIndex, PolyIndex);
+				OutPoly.Center = (Recast2UnrealPoint(&Tile->verts[Poly->verts[0] * 3]) + Recast2UnrealPoint(&Tile->verts[Poly->verts[1] * 3])) / 2;
+			}
+		}
+
+		if (bIncludeLinksFromNeighborTiles)
+		{
+			TArray<const dtMeshTile*> NeighborTiles;
+			NeighborTiles.Reserve(32);
+			for (int32 SideIndex = 0; SideIndex < 8; ++SideIndex)
+			{
+				const int32 StartIndex = NeighborTiles.Num();
+				const int32 NeighborCount = DetourNavMesh->getNeighbourTilesCountAt(Tile->header->x, Tile->header->y, SideIndex);
+				if (NeighborCount > 0)
+				{
+					const unsigned char oppositeSide = (unsigned char)dtOppositeTile(SideIndex);
+
+					NeighborTiles.AddZeroed(NeighborCount);
+					int32 NeighborX = Tile->header->x;
+					int32 NeighborY = Tile->header->y;
+
+					if (DetourNavMesh->getNeighbourCoords(Tile->header->x, Tile->header->y, SideIndex, NeighborX, NeighborY))
+					{
+						DetourNavMesh->getTilesAt(NeighborX, NeighborY, NeighborTiles.GetData() + StartIndex, NeighborCount);
+					}
+
+					for (const dtMeshTile* NeighborTile : NeighborTiles)
+					{
+						if (NeighborTile && NeighborTile->header && NeighborTile->offMeshCons)
+						{
+							const dtTileRef NeighborTileId = DetourNavMesh->getTileRef(NeighborTile);
+
+							for (int32 LinkIndex = 0; LinkIndex < NeighborTile->header->offMeshConCount; ++LinkIndex)
+							{
+								dtOffMeshConnection* targetCon = &NeighborTile->offMeshCons[LinkIndex];
+								if (targetCon->side != oppositeSide)
+								{
+									continue;
+								}
+
+								const unsigned char biDirFlag = targetCon->getBiDirectional() ? DT_LINK_FLAG_OFFMESH_CON_BIDIR : 0;
+
+								const dtPoly* targetPoly = &NeighborTile->polys[targetCon->poly];
+								// Skip off-mesh connections which start location could not be connected at all.
+								if (targetPoly->firstLink == DT_NULL_LINK)
+								{
+									continue;
+								}
+
+								FNavPoly& OutPoly = Polys[Polys.AddZeroed()];
+								OutPoly.Ref = NeighborTileId | targetCon->poly;
+								OutPoly.Center = (Recast2UnrealPoint(&targetCon->pos[0]) + Recast2UnrealPoint(&targetCon->pos[3])) / 2;
+							}
+						}
+					}
+
+					NeighborTiles.Reset();
+				}
+			}
+		}
+	}
+
+	return (Polys.Num() - InitialLinkCount > 0);
+}
+
 int32 ARecastNavMesh::GetNavMeshTilesCount() const
 {
 	int32 NumTiles = 0;
diff --git a/Engine/Source/Runtime/Engine/Private/Animation/AnimInstance.cpp b/Engine/Source/Runtime/Engine/Private/Animation/AnimInstance.cpp
--- a/Engine/Source/Runtime/Engine/Private/Animation/AnimInstance.cpp
+++ b/Engine/Source/Runtime/Engine/Private/Animation/AnimInstance.cpp
@@ -214,18 +214,23 @@ void UAnimInstance::UninitializeAnimation()
 
 	StopAllMontages(0.f);
 
-	for(int32 Index = 0; Index < MontageInstances.Num(); ++Index)
+	if (MontageInstances.Num() > 0)
 	{
-		FAnimMontageInstance* MontageInstance = MontageInstances[Index];
-		if (ensure(MontageInstance != nullptr))
+		for (int32 Index = 0; Index < MontageInstances.Num(); ++Index)
 		{
-			ClearMontageInstanceReferences(*MontageInstance);
-			delete MontageInstance;
+			FAnimMontageInstance* MontageInstance = MontageInstances[Index];
+			if (ensure(MontageInstance != nullptr))
+			{
+				ClearMontageInstanceReferences(*MontageInstance);
+				delete MontageInstance;
+			}
 		}
-	}
 
-	MontageInstances.Empty();
-	ActiveMontagesMap.Empty();
+		MontageInstances.Empty();
+		ActiveMontagesMap.Empty();
+
+		OnAllMontageInstancesEnded.Broadcast();
+	}
 
 	USkeletalMeshComponent* SkelMeshComp = GetSkelMeshComponent();
 	if (SkelMeshComp)
@@ -365,6 +370,10 @@ void UAnimInstance::UpdateAnimation(float DeltaSeconds, bool bNeedsValidRootMoti
 
 	PreUpdateAnimation(DeltaSeconds);
 
+	// need to update montage BEFORE node update or Native Update.
+	// so that node knows where montage is
+	UpdateMontage(DeltaSeconds);
+
 	{
 		SCOPE_CYCLE_COUNTER(STAT_NativeUpdateAnimation);
 		NativeUpdateAnimation(DeltaSeconds);
@@ -379,10 +388,6 @@ PRAGMA_ENABLE_DEPRECATION_WARNINGS
 		BlueprintUpdateAnimation(DeltaSeconds);
 	}
 
-	// need to update montage BEFORE node update
-	// so that node knows where montage is
-	UpdateMontage(DeltaSeconds);
-
 	if(bNeedsValidRootMotion || NeedsImmediateUpdate(DeltaSeconds))
 	{
 		// cant use parallel update, so just do the work here
@@ -862,7 +867,7 @@ void UAnimInstance::DisplayDebug(class UCanvas* Canvas, const FDebugDisplayInfo&
 
 			DisplayDebugManager.SetLinearDrawColor((MontageInstance->IsActive()) ? ActiveColor : TextWhite);
 
-			FString MontageEntry = FString::Printf(TEXT("%i) %s CurrSec: %s NextSec: %s W:%.3f DW:%.3f"), MontageIndex, *MontageInstance->Montage->GetName(), *MontageInstance->GetCurrentSection().ToString(), *MontageInstance->GetNextSection().ToString(), MontageInstance->GetWeight(), MontageInstance->GetDesiredWeight());
+			FString MontageEntry = FString::Printf(TEXT("%i) %s CurrSec: %s NextSec: %s W:%.2f DW:%.2f"), MontageIndex, *MontageInstance->Montage->GetName(), *MontageInstance->GetCurrentSection().ToString(), *MontageInstance->GetNextSection().ToString(), MontageInstance->GetWeight(), MontageInstance->GetDesiredWeight());
 			DisplayDebugManager.DrawString(MontageEntry, Indent);
 		}
 	}
@@ -1279,25 +1284,43 @@ void UAnimInstance::EndNotifyStates()
 //to debug montage weight
 #define DEBUGMONTAGEWEIGHT 0
 
-float UAnimInstance::GetSlotNodeGlobalWeight(FName SlotNodeName) const
+float UAnimInstance::GetSlotNodeGlobalWeight(const FName& SlotNodeName) const
 {
 	return GetProxyOnGameThread<FAnimInstanceProxy>().GetSlotNodeGlobalWeight(SlotNodeName);
 }
 
-float UAnimInstance::GetSlotMontageGlobalWeight(FName SlotNodeName) const
+float UAnimInstance::GetSlotMontageGlobalWeight(const FName& SlotNodeName) const
 {
 	return GetProxyOnAnyThread<FAnimInstanceProxy>().GetSlotMontageGlobalWeight(SlotNodeName);
 }
 
+float UAnimInstance::GetSlotMontageLocalWeight(const FName& SlotNodeName) const
+{
+	return GetProxyOnAnyThread<FAnimInstanceProxy>().GetSlotMontageLocalWeight(SlotNodeName);
+}
+
+float UAnimInstance::CalcSlotMontageLocalWeight(const FName& SlotNodeName) const
+{
+	return GetProxyOnAnyThread<FAnimInstanceProxy>().CalcSlotMontageLocalWeight(SlotNodeName);
+}
+
 float UAnimInstance::GetCurveValue(FName CurveName)
+{
+	float Value = 0.f;
+	GetCurveValue(CurveName, Value);
+	return Value;
+}
+
+bool UAnimInstance::GetCurveValue(FName CurveName, float& OutValue)
 {
 	float* Value = AnimationCurves[(uint8)EAnimCurveType::AttributeCurve].Find(CurveName);
 	if (Value)
 	{
-		return *Value;
+		OutValue = *Value;
+		return true;
 	}
 
-	return 0.f;
+	return false;
 }
 
 void UAnimInstance::SetRootMotionMode(TEnumAsByte<ERootMotionMode::Type> Value)
@@ -1468,6 +1491,11 @@ void UAnimInstance::Montage_Advance(float DeltaSeconds)
 				delete MontageInstance;
 				MontageInstances.RemoveAt(InstanceIndex);
 				--InstanceIndex;
+
+				if (MontageInstances.Num() == 0)
+				{
+					OnAllMontageInstancesEnded.Broadcast();
+				}
 			}
 #if DO_CHECK && WITH_EDITORONLY_DATA && 0
 			else
diff --git a/Engine/Source/Runtime/Engine/Private/Animation/AnimInstanceProxy.cpp b/Engine/Source/Runtime/Engine/Private/Animation/AnimInstanceProxy.cpp
--- a/Engine/Source/Runtime/Engine/Private/Animation/AnimInstanceProxy.cpp
+++ b/Engine/Source/Runtime/Engine/Private/Animation/AnimInstanceProxy.cpp
@@ -199,6 +199,10 @@ void FAnimInstanceProxy::PreUpdate(UAnimInstance* InAnimInstance, float DeltaSec
 
 	NotifyQueue.Reset(InAnimInstance->GetSkelMeshComponent());
 
+#if ENABLE_ANIM_DRAW_DEBUG
+	QueuedDrawDebugItems.Reset();
+#endif
+
 	ClearSlotNodeWeights();
 
 	// Reset the player tick list (but keep it presized)
@@ -261,8 +265,23 @@ void FAnimInstanceProxy::PostUpdate(UAnimInstance* InAnimInstance) const
 		DebugData.AnimNodePoseWatch = MoveTemp(PoseWatchEntriesForThisFrame);
 	}
 #endif
+
 	InAnimInstance->NotifyQueue.Append(NotifyQueue);
 	InAnimInstance->NotifyQueue.ApplyMontageNotifies(*this);
+
+	// Send Queued DrawDebug Commands.
+#if ENABLE_ANIM_DRAW_DEBUG
+	for (const FQueuedDrawDebugItem& DebugItem : QueuedDrawDebugItems)
+	{
+		switch (DebugItem.ItemType)
+		{
+		case EDrawDebugItemType::OnScreenMessage: GEngine->AddOnScreenDebugMessage(INDEX_NONE, 0.f, DebugItem.Color, DebugItem.Message, false, DebugItem.TextScale); break;
+		case EDrawDebugItemType::DirectionalArrow: DrawDebugDirectionalArrow(InAnimInstance->GetSkelMeshComponent()->GetWorld(), DebugItem.StartLoc, DebugItem.EndLoc, DebugItem.Size, DebugItem.Color, DebugItem.bPersistentLines, DebugItem.LifeTime, 0, DebugItem.Thickness); break;
+		case EDrawDebugItemType::Sphere : DrawDebugSphere(InAnimInstance->GetSkelMeshComponent()->GetWorld(), DebugItem.Center, DebugItem.Radius, DebugItem.Segments, DebugItem.Color, DebugItem.bPersistentLines, DebugItem.LifeTime, 0, DebugItem.Thickness); break;
+		case EDrawDebugItemType::Line: DrawDebugLine(InAnimInstance->GetSkelMeshComponent()->GetWorld(), DebugItem.StartLoc, DebugItem.EndLoc, DebugItem.Color, DebugItem.bPersistentLines, DebugItem.LifeTime, 0, DebugItem.Thickness); break;
+		}
+	}
+#endif
 }
 
 void FAnimInstanceProxy::InitializeObjects(UAnimInstance* InAnimInstance)
@@ -595,7 +614,7 @@ void FAnimInstanceProxy::ReinitializeSlotNodes()
 	SlotNodeInitializationCounter.Increment();
 }
 
-void FAnimInstanceProxy::RegisterSlotNodeWithAnimInstance(FName SlotNodeName)
+void FAnimInstanceProxy::RegisterSlotNodeWithAnimInstance(const FName& SlotNodeName)
 {
 	// verify if same slot node name exists
 	// then warn users, this is invalid
@@ -622,7 +641,7 @@ void FAnimInstanceProxy::RegisterSlotNodeWithAnimInstance(FName SlotNodeName)
 	SlotWeightTracker[1].Add(FMontageActiveSlotTracker());
 }
 
-void FAnimInstanceProxy::UpdateSlotNodeWeight(FName SlotNodeName, float InMontageLocalWeight, float InNodeGlobalWeight)
+void FAnimInstanceProxy::UpdateSlotNodeWeight(const FName& SlotNodeName, float InMontageLocalWeight, float InNodeGlobalWeight)
 {
 	const int32* TrackerIndexPtr = SlotNameToTrackerIndex.Find(SlotNodeName);
 	if (TrackerIndexPtr)
@@ -648,7 +667,7 @@ void FAnimInstanceProxy::ClearSlotNodeWeights()
 	}
 }
 
-bool FAnimInstanceProxy::IsSlotNodeRelevantForNotifies(FName SlotNodeName) const
+bool FAnimInstanceProxy::IsSlotNodeRelevantForNotifies(const FName& SlotNodeName) const
 {
 	const int32* TrackerIndexPtr = SlotNameToTrackerIndex.Find(SlotNodeName);
 	if (TrackerIndexPtr)
@@ -660,7 +679,7 @@ bool FAnimInstanceProxy::IsSlotNodeRelevantForNotifies(FName SlotNodeName) const
 	return false;
 }
 
-float FAnimInstanceProxy::GetSlotNodeGlobalWeight(FName SlotNodeName) const
+float FAnimInstanceProxy::GetSlotNodeGlobalWeight(const FName& SlotNodeName) const
 {
 	const int32* TrackerIndexPtr = SlotNameToTrackerIndex.Find(SlotNodeName);
 	if (TrackerIndexPtr)
@@ -672,7 +691,7 @@ float FAnimInstanceProxy::GetSlotNodeGlobalWeight(FName SlotNodeName) const
 	return 0.f;
 }
 
-float FAnimInstanceProxy::GetSlotMontageGlobalWeight(FName SlotNodeName) const
+float FAnimInstanceProxy::GetSlotMontageGlobalWeight(const FName& SlotNodeName) const
 {
 	const int32* TrackerIndexPtr = SlotNameToTrackerIndex.Find(SlotNodeName);
 	if (TrackerIndexPtr)
@@ -684,6 +703,26 @@ float FAnimInstanceProxy::GetSlotMontageGlobalWeight(FName SlotNodeName) const
 	return 0.f;
 }
 
+float FAnimInstanceProxy::GetSlotMontageLocalWeight(const FName& SlotNodeName) const
+{
+	const int32* TrackerIndexPtr = SlotNameToTrackerIndex.Find(SlotNodeName);
+	if (TrackerIndexPtr)
+	{
+		const FMontageActiveSlotTracker& Tracker = SlotWeightTracker[GetSyncGroupReadIndex()][*TrackerIndexPtr];
+		return Tracker.MontageLocalWeight;
+	}
+
+	return 0.f;
+}
+
+float FAnimInstanceProxy::CalcSlotMontageLocalWeight(const FName& SlotNodeName) const
+{
+	float out_SlotNodeLocalWeight, out_SourceWeight, out_TotalNodeWeight;
+	GetSlotWeight(SlotNodeName, out_SlotNodeLocalWeight, out_SourceWeight, out_TotalNodeWeight);
+
+	return out_SlotNodeLocalWeight;
+}
+
 FAnimNode_Base* FAnimInstanceProxy::GetCheckedNodeFromIndexUntyped(int32 NodeIdx, UScriptStruct* RequiredStructType)
 {
 	FAnimNode_Base* NodePtr = nullptr;
@@ -822,7 +861,7 @@ void FAnimInstanceProxy::EvaluateAnimationNode(FPoseContext& Output)
 #define DEBUG_MONTAGEINSTANCE_WEIGHT 1
 #endif // !(UE_BUILD_SHIPPING || UE_BUILD_TEST)
 
-void FAnimInstanceProxy::SlotEvaluatePose(FName SlotNodeName, const FCompactPose& SourcePose, const FBlendedCurve& SourceCurve, float InSourceWeight, FCompactPose& BlendedPose, FBlendedCurve& BlendedCurve, float InBlendWeight, float InTotalNodeWeight)
+void FAnimInstanceProxy::SlotEvaluatePose(const FName& SlotNodeName, const FCompactPose& SourcePose, const FBlendedCurve& SourceCurve, float InSourceWeight, FCompactPose& BlendedPose, FBlendedCurve& BlendedCurve, float InBlendWeight, float InTotalNodeWeight)
 {
 	//Accessing MontageInstances from this function is not safe (as this can be called during Parallel Anim Evaluation!
 	//Any montage data you need to add should be part of MontageEvaluationData
@@ -968,7 +1007,7 @@ void FAnimInstanceProxy::SlotEvaluatePose(FName SlotNodeName, const FCompactPose
 //to debug montage weight
 #define DEBUGMONTAGEWEIGHT 0
 
-void FAnimInstanceProxy::GetSlotWeight(FName const& SlotNodeName, float& out_SlotNodeWeight, float& out_SourceWeight, float& out_TotalNodeWeight) const
+void FAnimInstanceProxy::GetSlotWeight(const FName& SlotNodeName, float& out_SlotNodeWeight, float& out_SourceWeight, float& out_TotalNodeWeight) const
 {
 	// node total weight 
 	float NewSlotNodeWeight = 0.f;
@@ -1064,6 +1103,69 @@ void FAnimInstanceProxy::GatherDebugData(FNodeDebugData& DebugData)
 	}
 }
 
+#if ENABLE_ANIM_DRAW_DEBUG
+
+void FAnimInstanceProxy::AnimDrawDebugOnScreenMessage(const FString& DebugMessage, const FColor& Color, const FVector2D& TextScale)
+{
+	FQueuedDrawDebugItem DrawDebugItem;
+
+	DrawDebugItem.ItemType = EDrawDebugItemType::OnScreenMessage;
+	DrawDebugItem.Message = DebugMessage;
+	DrawDebugItem.Color = Color;
+	DrawDebugItem.TextScale = TextScale;
+
+	QueuedDrawDebugItems.Add(DrawDebugItem);
+}
+
+void FAnimInstanceProxy::AnimDrawDebugDirectionalArrow(const FVector& LineStart, const FVector& LineEnd, float ArrowSize, const FColor& Color, bool bPersistentLines, float LifeTime, float Thickness)
+{
+	FQueuedDrawDebugItem DrawDebugItem;
+
+	DrawDebugItem.ItemType = EDrawDebugItemType::DirectionalArrow;
+	DrawDebugItem.StartLoc = LineStart;
+	DrawDebugItem.EndLoc = LineEnd;
+	DrawDebugItem.Size = ArrowSize;
+	DrawDebugItem.Color = Color;
+	DrawDebugItem.bPersistentLines = bPersistentLines;
+	DrawDebugItem.LifeTime = LifeTime;
+	DrawDebugItem.Thickness = Thickness;
+
+	QueuedDrawDebugItems.Add(DrawDebugItem);
+}
+
+void FAnimInstanceProxy::AnimDrawDebugSphere(const FVector& Center, float Radius, int32 Segments, const FColor& Color, bool bPersistentLines, float LifeTime, float Thickness)
+{
+	FQueuedDrawDebugItem DrawDebugItem;
+
+	DrawDebugItem.ItemType = EDrawDebugItemType::Sphere;
+	DrawDebugItem.Center = Center;
+	DrawDebugItem.Radius = Radius;
+	DrawDebugItem.Segments = Segments;
+	DrawDebugItem.Color = Color;
+	DrawDebugItem.bPersistentLines = bPersistentLines;
+	DrawDebugItem.LifeTime = LifeTime;
+	DrawDebugItem.Thickness = Thickness;
+
+	QueuedDrawDebugItems.Add(DrawDebugItem);
+}
+
+void FAnimInstanceProxy::AnimDrawDebugLine(const FVector& StartLoc, const FVector& EndLoc, const FColor& Color, bool bPersistentLines, float LifeTime, float Thickness)
+{
+	FQueuedDrawDebugItem DrawDebugItem;
+
+	DrawDebugItem.ItemType = EDrawDebugItemType::Line;
+	DrawDebugItem.StartLoc = StartLoc;
+	DrawDebugItem.EndLoc = EndLoc;
+	DrawDebugItem.Color = Color;
+	DrawDebugItem.bPersistentLines = bPersistentLines;
+	DrawDebugItem.LifeTime = LifeTime;
+	DrawDebugItem.Thickness = Thickness;
+
+	QueuedDrawDebugItems.Add(DrawDebugItem);
+}
+
+#endif // ENABLE_ANIM_DRAW_DEBUG
+
 float FAnimInstanceProxy::GetInstanceAssetPlayerLength(int32 AssetPlayerIndex)
 {
 	if(FAnimNode_AssetPlayerBase* PlayerNode = GetNodeFromIndex<FAnimNode_AssetPlayerBase>(AssetPlayerIndex))
diff --git a/Engine/Source/Runtime/Engine/Private/Animation/AnimNode_StateMachine.cpp b/Engine/Source/Runtime/Engine/Private/Animation/AnimNode_StateMachine.cpp
--- a/Engine/Source/Runtime/Engine/Private/Animation/AnimNode_StateMachine.cpp
+++ b/Engine/Source/Runtime/Engine/Private/Animation/AnimNode_StateMachine.cpp
@@ -77,15 +77,20 @@ void FAnimationActiveTransitionEntry::InitializeCustomGraphLinks(const FAnimatio
 {
 	if (TransitionRule.CustomResultNodeIndex != INDEX_NONE)
 	{
-		const IAnimClassInterface* AnimBlueprintClass = Context.GetAnimClass();
-		CustomTransitionGraph.LinkID = AnimBlueprintClass->GetAnimNodeProperties().Num() - 1 - TransitionRule.CustomResultNodeIndex; //@TODO: Crazysauce
-		FAnimationInitializeContext InitContext(Context.AnimInstanceProxy);
-		CustomTransitionGraph.Initialize(InitContext);
-
-		for (int32 Index = 0; Index < TransitionRule.PoseEvaluatorLinks.Num(); ++Index)
+		if (const IAnimClassInterface* AnimBlueprintClass = Context.GetAnimClass())
 		{
-			FAnimNode_TransitionPoseEvaluator* PoseEvaluator = GetNodeFromPropertyIndex<FAnimNode_TransitionPoseEvaluator>(Context.AnimInstanceProxy->GetAnimInstanceObject(), AnimBlueprintClass, TransitionRule.PoseEvaluatorLinks[Index]);
-			PoseEvaluators.Add(PoseEvaluator);
+			CustomTransitionGraph.LinkID = AnimBlueprintClass->GetAnimNodeProperties().Num() - 1 - TransitionRule.CustomResultNodeIndex; //@TODO: Crazysauce
+			FAnimationInitializeContext InitContext(Context.AnimInstanceProxy);
+			CustomTransitionGraph.Initialize(InitContext);
+
+			if (Context.AnimInstanceProxy)
+			{
+				for (int32 Index = 0; Index < TransitionRule.PoseEvaluatorLinks.Num(); ++Index)
+				{
+					FAnimNode_TransitionPoseEvaluator* PoseEvaluator = GetNodeFromPropertyIndex<FAnimNode_TransitionPoseEvaluator>(Context.AnimInstanceProxy->GetAnimInstanceObject(), AnimBlueprintClass, TransitionRule.PoseEvaluatorLinks[Index]);
+					PoseEvaluators.Add(PoseEvaluator);
+				}
+			}
 		}
 	}
 
@@ -431,15 +436,20 @@ void FAnimNode_StateMachine::Update(const FAnimationUpdateContext& Context)
 			// Push the transition onto the stack
 			const FAnimationTransitionBetweenStates& ReferenceTransition = GetTransitionInfo(PotentialTransition.TransitionRule->TransitionIndex);
 			FAnimationActiveTransitionEntry* NewTransition = new (ActiveTransitionArray) FAnimationActiveTransitionEntry(NextState, ExistingWeightOfNextState, PreviousTransitionForNextState, PreviousState, ReferenceTransition);
-			NewTransition->InitializeCustomGraphLinks(Context, *(PotentialTransition.TransitionRule));
+			if (NewTransition && PotentialTransition.TransitionRule)
+			{
+				NewTransition->InitializeCustomGraphLinks(Context, *(PotentialTransition.TransitionRule));
 
-			NewTransition->SourceTransitionIndices = PotentialTransition.SourceTransitionIndices;
+#if WITH_EDITORONLY_DATA
+				NewTransition->SourceTransitionIndices = PotentialTransition.SourceTransitionIndices;
+#endif
 
-			if(!bFirstUpdate)
-			{
-				Context.AnimInstanceProxy->AddAnimNotifyFromGeneratedClass(NewTransition->StartNotify);
+				if (!bFirstUpdate)
+				{
+					Context.AnimInstanceProxy->AddAnimNotifyFromGeneratedClass(NewTransition->StartNotify);
+				}
 			}
-			
+
 			SetState(Context, NextState);
 
 			TransitionCountThisFrame++;
diff --git a/Engine/Source/Runtime/Engine/Private/Animation/AnimNotifyState_Trail.cpp b/Engine/Source/Runtime/Engine/Private/Animation/AnimNotifyState_Trail.cpp
--- a/Engine/Source/Runtime/Engine/Private/Animation/AnimNotifyState_Trail.cpp
+++ b/Engine/Source/Runtime/Engine/Private/Animation/AnimNotifyState_Trail.cpp
@@ -64,6 +64,21 @@ UParticleSystem* UAnimNotifyState_Trail::GetOverridenPSTemplate(class USkeletalM
 	return OverridePSTemplate(MeshComp, Animation);
 }
 
+float UAnimNotifyState_Trail::GetCurveWidth(USkeletalMeshComponent* MeshComp) const
+{
+	UAnimInstance* AnimInst = MeshComp->GetAnimInstance();
+	float Width = 1.0f;
+	if (WidthScaleCurve != NAME_None && AnimInst)
+	{
+		if (!AnimInst->GetCurveValue(WidthScaleCurve, Width))
+		{
+			// Fallback to 1.f if curve was not found
+			Width = 1.f;
+		}
+	}
+	return Width;
+}
+
 void UAnimNotifyState_Trail::NotifyBegin(class USkeletalMeshComponent * MeshComp, class UAnimSequenceBase * Animation, float TotalDuration)
 {
 	bool bError = ValidateInput(MeshComp);
@@ -88,12 +103,7 @@ void UAnimNotifyState_Trail::NotifyBegin(class USkeletalMeshComponent * MeshComp
 	ParticleSystemComponentArray Children;
 	GetCandidateSystems(*MeshComp, Children);
 
-	UAnimInstance* AnimInst = MeshComp->GetAnimInstance();
-	float Width = 1.0f;
-	if (WidthScaleCurve != NAME_None && AnimInst)
-	{
-		Width = AnimInst->GetCurveValue(WidthScaleCurve);
-	}
+	float Width = GetCurveWidth(MeshComp);
 
 	UParticleSystemComponent* RecycleCandidates[3] = {nullptr, nullptr, nullptr}; // in order of priority
 	bool bFoundExistingTrail = false;
@@ -203,13 +213,8 @@ void UAnimNotifyState_Trail::NotifyTick(class USkeletalMeshComponent * MeshComp,
 
 	ParticleSystemComponentArray Children;
 	GetCandidateSystems(*MeshComp, Children);
-
-	UAnimInstance* AnimInst = MeshComp->GetAnimInstance();
-	float Width = 1.0f;
-	if (WidthScaleCurve != NAME_None && AnimInst)
-	{
-		Width = AnimInst->GetCurveValue(WidthScaleCurve);
-	}
+	
+	float Width = GetCurveWidth(MeshComp);
 
 	for (UParticleSystemComponent* ParticleComp : Children)
 	{
@@ -264,6 +269,31 @@ void UAnimNotifyState_Trail::NotifyEnd(class USkeletalMeshComponent * MeshComp,
 	Received_NotifyEnd(MeshComp, Animation);
 }
 
+UParticleSystemComponent* UAnimNotifyState_Trail::GetParticleSystemComponent(USkeletalMeshComponent* MeshComp) const
+{
+	if (MeshComp == nullptr)
+	{
+		return nullptr;
+	}
+
+	ParticleSystemComponentArray Children;
+	GetCandidateSystems(*MeshComp, Children);
+	for (UParticleSystemComponent* ParticleComp : Children)
+	{
+		if (ParticleComp->IsActive())
+		{
+			UParticleSystemComponent::TrailEmitterArray TrailEmitters;
+			ParticleComp->GetOwnedTrailEmitters(TrailEmitters, this, false);
+			if (TrailEmitters.Num() > 0)
+			{
+				// We have a trail emitter, so return this one
+				return ParticleComp;
+			}
+		}
+	}
+	return nullptr;
+}
+
 bool UAnimNotifyState_Trail::ValidateInput(class USkeletalMeshComponent * MeshComp, bool bReportErrors/* =false */)
 {
 #if WITH_EDITOR
diff --git a/Engine/Source/Runtime/Engine/Private/Animation/AnimNotify_ResetClothingSimulation.cpp b/Engine/Source/Runtime/Engine/Private/Animation/AnimNotify_ResetClothingSimulation.cpp
new file mode 100644
--- /dev/null
+++ b/Engine/Source/Runtime/Engine/Private/Animation/AnimNotify_ResetClothingSimulation.cpp
@@ -0,0 +1,22 @@
+// Copyright 1998-2016 Epic Games, Inc. All Rights Reserved.
+
+#include "EnginePrivate.h"
+#include "Animation/AnimNotifies/AnimNotify_ResetClothingSimulation.h"
+
+UAnimNotify_ResetClothingSimulation::UAnimNotify_ResetClothingSimulation()
+	: Super()
+{
+#if WITH_EDITORONLY_DATA
+	NotifyColor = FColor(90, 220, 255, 255);
+#endif // WITH_EDITORONLY_DATA
+}
+
+void UAnimNotify_ResetClothingSimulation::Notify(class USkeletalMeshComponent* MeshComp, class UAnimSequenceBase* Animation)
+{
+	MeshComp->ForceClothNextUpdateTeleportAndReset();
+}
+
+FString UAnimNotify_ResetClothingSimulation::GetNotifyName_Implementation() const
+{
+	return TEXT("Reset Clothing Sim");
+}
\ No newline at end of file
diff --git a/Engine/Source/Runtime/Engine/Private/Animation/AnimPhysicsSolver.cpp b/Engine/Source/Runtime/Engine/Private/Animation/AnimPhysicsSolver.cpp
--- a/Engine/Source/Runtime/Engine/Private/Animation/AnimPhysicsSolver.cpp
+++ b/Engine/Source/Runtime/Engine/Private/Animation/AnimPhysicsSolver.cpp
@@ -63,7 +63,7 @@ FAnimPhysShape FAnimPhysShape::MakeBox(FVector& Extents)
 	// if we encounter it.
 	if(Extents.SizeSquared() <= SMALL_NUMBER)
 	{
-		UE_LOG(LogAnimation, Warning, TEXT("AnimDynamics: Attempted to create a simulation box with 0 volume, this introduces NaNs into the simulation. Adjusting box extents to (1.0f,1.0f,1.0f)"));
+// 		UE_LOG(LogAnimation, Warning, TEXT("AnimDynamics: Attempted to create a simulation box with 0 volume, this introduces NaNs into the simulation. Adjusting box extents to (1.0f,1.0f,1.0f)"));
 		Extents = FVector(1.0f);
 	}
 
diff --git a/Engine/Source/Runtime/Engine/Private/AudioDevice.cpp b/Engine/Source/Runtime/Engine/Private/AudioDevice.cpp
--- a/Engine/Source/Runtime/Engine/Private/AudioDevice.cpp
+++ b/Engine/Source/Runtime/Engine/Private/AudioDevice.cpp
@@ -2690,6 +2690,15 @@ int32 FAudioDevice::GetSortedActiveWaveInstances(TArray<FWaveInstance*>& WaveIns
 		}
 	}
 
+	// Remove all wave instances from the wave instance list that are stopping due to max concurrency
+	for (int32 i = WaveInstances.Num() - 1; i >= 0; --i)
+	{
+		if (WaveInstances[i]->ShouldStopDueToMaxConcurrency())
+		{
+			WaveInstances.RemoveAtSwap(i, 1, false);
+		}
+	}
+
 	int32 FirstActiveIndex = 0;
 
 	if (WaveInstances.Num() >= 0)
diff --git a/Engine/Source/Runtime/Engine/Private/Character.cpp b/Engine/Source/Runtime/Engine/Private/Character.cpp
--- a/Engine/Source/Runtime/Engine/Private/Character.cpp
+++ b/Engine/Source/Runtime/Engine/Private/Character.cpp
@@ -796,7 +796,7 @@ void ACharacter::PossessedBy(AController* NewController)
 	// If we are controlled remotely, set animation timing to be driven by client's network updates. So timing and events remain in sync.
 	if (Mesh && (GetRemoteRole() == ROLE_AutonomousProxy && GetNetConnection() != nullptr))
 	{
-		Mesh->bAutonomousTickPose = true;
+		Mesh->bOnlyAllowAutonomousTickPose = true;
 	}
 }
 
@@ -813,7 +813,7 @@ void ACharacter::UnPossessed()
 	// We're no longer controlled remotely, resume regular ticking of animations.
 	if (Mesh)
 	{
-		Mesh->bAutonomousTickPose = false;
+		Mesh->bOnlyAllowAutonomousTickPose = false;
 	}
 }
 
@@ -831,7 +831,7 @@ void ACharacter::TornOff()
 	// We're no longer controlled remotely, resume regular ticking of animations.
 	if (Mesh)
 	{
-		Mesh->bAutonomousTickPose = false;
+		Mesh->bOnlyAllowAutonomousTickPose = false;
 	}
 }
 
@@ -1435,8 +1435,6 @@ void ACharacter::PreReplication( IRepChangedPropertyTracker & ChangedPropertyTra
 
 		ChangedPropertyTracker.SetExternalData( Writer.GetData(), Writer.GetNumBits() );
 	}
-
-	bReplayHasRootMotionSources = CharacterMovement ? CharacterMovement->HasRootMotionSources() : false;
 }
 
 void ACharacter::GetLifetimeReplicatedProps( TArray< FLifetimeProperty > & OutLifetimeProps ) const
@@ -1449,7 +1447,6 @@ void ACharacter::GetLifetimeReplicatedProps( TArray< FLifetimeProperty > & OutLi
 	DOREPLIFETIME_CONDITION( ACharacter, ReplicatedMovementMode,			COND_SimulatedOnly );
 	DOREPLIFETIME_CONDITION( ACharacter, bIsCrouched,						COND_SimulatedOnly );
 	DOREPLIFETIME_CONDITION( ACharacter, AnimRootMotionTranslationScale,	COND_SimulatedOnly );
-	DOREPLIFETIME_CONDITION( ACharacter, bReplayHasRootMotionSources,		COND_ReplayOnly );
 
 	// Change the condition of the replicated movement property to not replicate in replays since we handle this specifically via saving this out in external replay data
 	DOREPLIFETIME_CHANGE_CONDITION( AActor, ReplicatedMovement,				COND_SimulatedOrPhysicsNoReplay );
diff --git a/Engine/Source/Runtime/Engine/Private/ChartCreation.cpp b/Engine/Source/Runtime/Engine/Private/ChartCreation.cpp
--- a/Engine/Source/Runtime/Engine/Private/ChartCreation.cpp
+++ b/Engine/Source/Runtime/Engine/Private/ChartCreation.cpp
@@ -296,14 +296,16 @@ void FDumpFPSChartToEndpoint::DumpChart(double InWallClockTimeFromStartOfChartin
 
 struct FDumpFPSChartToAnalyticsArray : public FDumpFPSChartToEndpoint
 {
-	FDumpFPSChartToAnalyticsArray(const FPerformanceTrackingChart& InChart, TArray<FAnalyticsEventAttribute>& InParamArray)
+	FDumpFPSChartToAnalyticsArray(const FPerformanceTrackingChart& InChart, TArray<FAnalyticsEventAttribute>& InParamArray, bool bShouldIncludeClientHWInfo)
 		: FDumpFPSChartToEndpoint(InChart)
 		, ParamArray(InParamArray)
+		, bIncludeClientHWInfo(bShouldIncludeClientHWInfo)
 	{
 	}
 
 protected:
 	TArray<FAnalyticsEventAttribute>& ParamArray;
+	bool bIncludeClientHWInfo;
 protected:
 	virtual void PrintToEndpoint(const FString& Text) override
 	{
@@ -344,7 +346,10 @@ protected:
 		ParamArray.Add(FAnalyticsEventAttribute(TEXT("TotalHitches"), TotalHitchCount));
 		ParamArray.Add(FAnalyticsEventAttribute(TEXT("TotalGameBoundHitches"), Chart.TotalGameThreadBoundHitchCount));
 		ParamArray.Add(FAnalyticsEventAttribute(TEXT("TotalRenderBoundHitches"), Chart.TotalRenderThreadBoundHitchCount));
-		ParamArray.Add(FAnalyticsEventAttribute(TEXT("TotalGPUBoundHitches"), Chart.TotalGPUBoundHitchCount));
+		if (bIncludeClientHWInfo)
+		{
+			ParamArray.Add(FAnalyticsEventAttribute(TEXT("TotalGPUBoundHitches"), Chart.TotalGPUBoundHitchCount));
+		}
 		ParamArray.Add(FAnalyticsEventAttribute(TEXT("TotalTimeInHitchFrames"), TotalTimeSpentInHitchBuckets));
 		ParamArray.Add(FAnalyticsEventAttribute(TEXT("HitchesPerMinute"), Chart.GetAvgHitchesPerMinute()));
 
@@ -389,21 +394,22 @@ protected:
 		ParamArray.Add(FAnalyticsEventAttribute(TEXT("OS"), FString::Printf(TEXT("%s %s"), *OSMajor, *OSMinor)));
 		ParamArray.Add(FAnalyticsEventAttribute(TEXT("CPU"), FString::Printf(TEXT("%s %s"), *CPUVendor, *CPUBrand)));
 
-		//@TODO: Cut these out for servers
-		ParamArray.Add(FAnalyticsEventAttribute(TEXT("DesktopGPU"), DesktopGPUBrand)); //@TODO: Cut this one out entirely?
-		ParamArray.Add(FAnalyticsEventAttribute(TEXT("GPUAdapter"), ActualGPUBrand));
-
-		ParamArray.Add(FAnalyticsEventAttribute(TEXT("ResolutionQuality"), ScalabilityQuality.ResolutionQuality));
-		ParamArray.Add(FAnalyticsEventAttribute(TEXT("ViewDistanceQuality"), ScalabilityQuality.ViewDistanceQuality));
-		ParamArray.Add(FAnalyticsEventAttribute(TEXT("AntiAliasingQuality"), ScalabilityQuality.AntiAliasingQuality));
-		ParamArray.Add(FAnalyticsEventAttribute(TEXT("ShadowQuality"), ScalabilityQuality.ShadowQuality));
-		ParamArray.Add(FAnalyticsEventAttribute(TEXT("PostProcessQuality"), ScalabilityQuality.PostProcessQuality));
-		ParamArray.Add(FAnalyticsEventAttribute(TEXT("TextureQuality"), ScalabilityQuality.TextureQuality));
-		ParamArray.Add(FAnalyticsEventAttribute(TEXT("FXQuality"), ScalabilityQuality.EffectsQuality));
-		ParamArray.Add(FAnalyticsEventAttribute(TEXT("FoliageQuality"), ScalabilityQuality.FoliageQuality));
-		ParamArray.Add(FAnalyticsEventAttribute(TEXT("PercentGPUBound"), FString::Printf(TEXT("%4.2f"), BoundGPUPct)));
-		ParamArray.Add(FAnalyticsEventAttribute(TEXT("AvgGPUTime"), FString::Printf(TEXT("%4.2f"), AvgGPUFrameTime)));
-		//@ENDTODO:
+		if (bIncludeClientHWInfo)
+		{
+			ParamArray.Add(FAnalyticsEventAttribute(TEXT("DesktopGPU"), DesktopGPUBrand)); //@TODO: Cut this one out entirely?
+			ParamArray.Add(FAnalyticsEventAttribute(TEXT("GPUAdapter"), ActualGPUBrand));
+
+			ParamArray.Add(FAnalyticsEventAttribute(TEXT("ResolutionQuality"), ScalabilityQuality.ResolutionQuality));
+			ParamArray.Add(FAnalyticsEventAttribute(TEXT("ViewDistanceQuality"), ScalabilityQuality.ViewDistanceQuality));
+			ParamArray.Add(FAnalyticsEventAttribute(TEXT("AntiAliasingQuality"), ScalabilityQuality.AntiAliasingQuality));
+			ParamArray.Add(FAnalyticsEventAttribute(TEXT("ShadowQuality"), ScalabilityQuality.ShadowQuality));
+			ParamArray.Add(FAnalyticsEventAttribute(TEXT("PostProcessQuality"), ScalabilityQuality.PostProcessQuality));
+			ParamArray.Add(FAnalyticsEventAttribute(TEXT("TextureQuality"), ScalabilityQuality.TextureQuality));
+			ParamArray.Add(FAnalyticsEventAttribute(TEXT("FXQuality"), ScalabilityQuality.EffectsQuality));
+			ParamArray.Add(FAnalyticsEventAttribute(TEXT("FoliageQuality"), ScalabilityQuality.FoliageQuality));
+			ParamArray.Add(FAnalyticsEventAttribute(TEXT("PercentGPUBound"), FString::Printf(TEXT("%4.2f"), BoundGPUPct)));
+			ParamArray.Add(FAnalyticsEventAttribute(TEXT("AvgGPUTime"), FString::Printf(TEXT("%4.2f"), AvgGPUFrameTime)));
+		}
 
 		ParamArray.Add(FAnalyticsEventAttribute(TEXT("AvgFPS"), FString::Printf(TEXT("%4.2f"), AvgFPS)));
 		ParamArray.Add(FAnalyticsEventAttribute(TEXT("TimeDisregarded"), FString::Printf(TEXT("%4.2f"), TimeDisregarded)));
@@ -752,7 +758,7 @@ void FPerformanceTrackingChart::DumpChartToAnalyticsParams(const FString& InMapN
 	if ((TotalTime > 0.f) && (NumFrames > 0))
 	{
 		// Dump all the basic stats
-		FDumpFPSChartToAnalyticsArray AnalyticsEndpoint(*this, InParamArray);
+		FDumpFPSChartToAnalyticsArray AnalyticsEndpoint(*this, InParamArray, bIncludeClientHWInfo);
 		AnalyticsEndpoint.DumpChart(AccumulatedChartTime, InMapName);
 
 		if (bIncludeClientHWInfo)
@@ -813,8 +819,7 @@ void FPerformanceTrackingChart::DumpChartToAnalyticsParams(const FString& InMapN
 			}
 
 			// Screen percentage (3D render resolution)
-			//@TODO: Change to use actual cvar (r.screenpercentage), right now this is just a dupe of ResolutionQuality above
-			InParamArray.Add(FAnalyticsEventAttribute(TEXT("ScreenPct"), UserSettingsObj->ScalabilityQuality.ResolutionQuality));
+			InParamArray.Add(FAnalyticsEventAttribute(TEXT("ScreenPct"), Scalability::GetResolutionScreenPercentage()));
 
 			// Window mode and window/monitor resolution
 			const EWindowMode::Type FullscreenMode = UserSettingsObj->GetLastConfirmedFullscreenMode();
diff --git a/Engine/Source/Runtime/Engine/Private/Components/CharacterMovementComponent.cpp b/Engine/Source/Runtime/Engine/Private/Components/CharacterMovementComponent.cpp
--- a/Engine/Source/Runtime/Engine/Private/Components/CharacterMovementComponent.cpp
+++ b/Engine/Source/Runtime/Engine/Private/Components/CharacterMovementComponent.cpp
@@ -913,7 +913,7 @@ void UCharacterMovementComponent::OnMovementModeChanged(EMovementMode PreviousMo
 		}
 	}
 
-	if (MovementMode != MOVE_Falling && PreviousMovementMode == MOVE_Falling && PathFollowingComp.IsValid())
+	if (MovementMode == MOVE_Falling && PreviousMovementMode != MOVE_Falling && PathFollowingComp.IsValid())
 	{
 		PathFollowingComp->OnStartedFalling();
 	}
@@ -1328,8 +1328,6 @@ void UCharacterMovementComponent::SimulatedTick(float DeltaSeconds)
 
 			// Set root motion states to that of repped in state
 			CurrentRootMotion.UpdateStateFrom(RootMotionRepMove.RootMotion.AuthoritativeRootMotion, true);
-			CurrentRootMotion.bIsAdditiveVelocityApplied = RootMotionRepMove.RootMotion.AuthoritativeRootMotion.bIsAdditiveVelocityApplied;
-			CurrentRootMotion.LastPreAdditiveVelocity = RootMotionRepMove.RootMotion.AuthoritativeRootMotion.LastPreAdditiveVelocity;
 
 			// Clear out existing RootMotionRepMoves since we've consumed the most recent
 			UE_LOG(LogRootMotion, Log,  TEXT("\tClearing old moves in SimulatedTick (%d)"), MoveIndex+1);
@@ -1357,6 +1355,7 @@ void UCharacterMovementComponent::SimulatedTick(float DeltaSeconds)
 			CharacterOwner->RootMotionRepMoves.Empty();
 			CharacterOwner->OnRep_ReplicatedMovement();
 			CharacterOwner->OnRep_ReplicatedBasedMovement();
+			ApplyNetworkMovementMode(GetCharacterOwner()->GetReplicatedMovementMode());
 		}
 
 		// Avoid moving the mesh during movement if SmoothClientPosition will take care of it.
@@ -1898,6 +1897,12 @@ void UCharacterMovementComponent::PerformMovement(float DeltaSeconds)
 		// Character::LaunchCharacter() has been deferred until now.
 		HandlePendingLaunch();
 
+		// Update saved LastPreAdditiveVelocity with any external changes to character Velocity that happened due to ApplyAccumulatedForces/HandlePendingLaunch
+		if( CurrentRootMotion.HasAdditiveVelocity() )
+		{
+			CurrentRootMotion.LastPreAdditiveVelocity += (Velocity-OldVelocity);
+		}
+
 		// Prepare Root Motion (generate/accumulate from root motion sources to be used later)
 		if( HasRootMotionSources() && !CharacterOwner->bClientUpdating && !CharacterOwner->bServerMoveIgnoreRootMotion)
 		{
@@ -7061,6 +7066,17 @@ bool UCharacterMovementComponent::ClientUpdatePositionAfterServerUpdate()
 	if (ClientData->SavedMoves.Num() == 0)
 	{
 		UE_LOG(LogNetPlayerMovement, VeryVerbose, TEXT("ClientUpdatePositionAfterServerUpdate No saved moves to replay"), ClientData->SavedMoves.Num());
+
+		// With no saved moves to resimulate, the move the server updated us with is the last move we've done, no resimulation needed.
+		CharacterOwner->bClientResimulateRootMotion = false;
+		if (CharacterOwner->bClientResimulateRootMotionSources)
+		{
+			// With no resimulation, we just update our current root motion to what the server sent us
+			UE_LOG(LogRootMotion, VeryVerbose, TEXT("CurrentRootMotion getting updated to ServerUpdate state: %s"), *CharacterOwner->GetName());
+			CurrentRootMotion.UpdateStateFrom(CharacterOwner->SavedRootMotion);
+			CharacterOwner->bClientResimulateRootMotionSources = false;
+		}
+
 		return false;
 	}
 
@@ -8555,8 +8571,6 @@ void UCharacterMovementComponent::ClientAdjustRootMotionSourcePosition_Implement
 	// Server disagrees with Client on Root Motion state.
 	if( CharacterOwner->bClientResimulateRootMotionSources || (ServerRootMotion != ClientData->LastAckedMove->SavedRootMotion) )
 	{
-		//UE_LOG(LogRootMotion, Warning,  TEXT("\tServer disagrees with Client's track position!! ServerTrackPosition: %f, ClientTrackPosition: %f, DeltaTrackPosition: %f. TimeStamp: %f"),
-		//	ServerMontageTrackPosition, ClientData->LastAckedMove->RootMotionTrackPosition, (ServerMontageTrackPosition - ClientData->LastAckedMove->RootMotionTrackPosition), TimeStamp);
 		if (!CharacterOwner->bClientResimulateRootMotionSources)
 		{
 			UE_LOG(LogNetPlayerMovement, VeryVerbose, TEXT("ClientAdjustRootMotionSourcePosition called, server/LastAckedMove mismatch"));
@@ -8915,37 +8929,47 @@ void UCharacterMovementComponent::ApplyWorldOffset(const FVector& InOffset, bool
 void UCharacterMovementComponent::TickCharacterPose(float DeltaTime)
 {
 	check(CharacterOwner && CharacterOwner->GetMesh());
+	USkeletalMeshComponent* CharacterMesh = CharacterOwner->GetMesh();
 
-	// Keep track of if we're playing root motion, just in case the root motion montage ends this frame.
-	bool bWasPlayingRootMotion = CharacterOwner->IsPlayingRootMotion();
-
-	CharacterOwner->GetMesh()->TickPose(DeltaTime, true);
+	// bAutonomousTickPose is set, we control TickPose from the Character's Movement and Networking updates, and bypass the Component's update.
+	// (Or Simulating Root Motion for remote clients)
+	CharacterMesh->bIsAutonomousTickPose = true;
 
-	// Grab root motion now that we have ticked the pose
-	if (CharacterOwner->IsPlayingRootMotion() || bWasPlayingRootMotion)
+	if (CharacterMesh->ShouldTickPose())
 	{
-		FRootMotionMovementParams RootMotion = CharacterOwner->GetMesh()->ConsumeRootMotion();
-		if (RootMotion.bHasRootMotion)
+		// Keep track of if we're playing root motion, just in case the root motion montage ends this frame.
+		const bool bWasPlayingRootMotion = CharacterOwner->IsPlayingRootMotion();
+
+		CharacterMesh->TickPose(DeltaTime, true);
+
+		// Grab root motion now that we have ticked the pose
+		if (CharacterOwner->IsPlayingRootMotion() || bWasPlayingRootMotion)
 		{
-			RootMotion.ScaleRootMotionTranslation(CharacterOwner->GetAnimRootMotionTranslationScale());
-			RootMotionParams.Accumulate(RootMotion);
-		}
+			FRootMotionMovementParams RootMotion = CharacterMesh->ConsumeRootMotion();
+			if (RootMotion.bHasRootMotion)
+			{
+				RootMotion.ScaleRootMotionTranslation(CharacterOwner->GetAnimRootMotionTranslationScale());
+				RootMotionParams.Accumulate(RootMotion);
+			}
 
 #if !(UE_BUILD_SHIPPING)
-		// Debugging
-		{
-			FAnimMontageInstance* RootMotionMontageInstance = CharacterOwner->GetRootMotionAnimMontageInstance();
-			UE_LOG(LogRootMotion, Log, TEXT("UCharacterMovementComponent::TickCharacterPose Role: %s, RootMotionMontage: %s, MontagePos: %f, DeltaTime: %f, ExtractedRootMotion: %s, AccumulatedRootMotion: %s")
-				, *UEnum::GetValueAsString(TEXT("Engine.ENetRole"), CharacterOwner->Role)
-				, *GetNameSafe(RootMotionMontageInstance ? RootMotionMontageInstance->Montage : NULL)
-				, RootMotionMontageInstance ? RootMotionMontageInstance->GetPosition() : -1.f
-				, DeltaTime
-				, *RootMotion.GetRootMotionTransform().GetTranslation().ToCompactString()
-				, *RootMotionParams.GetRootMotionTransform().GetTranslation().ToCompactString()
-				);
-		}
+			// Debugging
+			{
+				FAnimMontageInstance* RootMotionMontageInstance = CharacterOwner->GetRootMotionAnimMontageInstance();
+				UE_LOG(LogRootMotion, Log, TEXT("UCharacterMovementComponent::TickCharacterPose Role: %s, RootMotionMontage: %s, MontagePos: %f, DeltaTime: %f, ExtractedRootMotion: %s, AccumulatedRootMotion: %s")
+					, *UEnum::GetValueAsString(TEXT("Engine.ENetRole"), CharacterOwner->Role)
+					, *GetNameSafe(RootMotionMontageInstance ? RootMotionMontageInstance->Montage : NULL)
+					, RootMotionMontageInstance ? RootMotionMontageInstance->GetPosition() : -1.f
+					, DeltaTime
+					, *RootMotion.GetRootMotionTransform().GetTranslation().ToCompactString()
+					, *RootMotionParams.GetRootMotionTransform().GetTranslation().ToCompactString()
+					);
+			}
 #endif // !(UE_BUILD_SHIPPING)
+		}
 	}
+
+	CharacterMesh->bIsAutonomousTickPose = false;
 }
 
 /** 
@@ -8954,11 +8978,6 @@ void UCharacterMovementComponent::TickCharacterPose(float DeltaTime)
 
 bool UCharacterMovementComponent::HasRootMotionSources() const
 {
-	if ( NetworkSmoothingMode == ENetworkSmoothingMode::Replay && CharacterOwner && CharacterOwner->Role == ROLE_SimulatedProxy && CharacterOwner->bReplayHasRootMotionSources )
-	{
-		return true;
-	}
-
 	return CurrentRootMotion.HasActiveRootMotionSources() || (CharacterOwner && CharacterOwner->IsPlayingRootMotion() && CharacterOwner->GetMesh());
 }
 
@@ -9036,8 +9055,8 @@ void UCharacterMovementComponent::ConvertRootMotionServerIDsToLocalIDs(const FRo
 	// RootMotion data are the server LocalIDs. When processing an FRootMotionSourceGroup
 	// for use on clients, we want to map server LocalIDs to our LocalIDs.
 	// We save off these mappings for quicker access and to save having to
-	// "find best match" every time we recieve server data.
-	for (auto& ServerRootMotionSource : InOutServerRootMotion.RootMotionSources)
+	// "find best match" every time we receive server data.
+	for (TSharedPtr<FRootMotionSource>& ServerRootMotionSource : InOutServerRootMotion.RootMotionSources)
 	{
 		if (ServerRootMotionSource.IsValid())
 		{
@@ -9049,7 +9068,7 @@ void UCharacterMovementComponent::ConvertRootMotionServerIDsToLocalIDs(const FRo
 				bool bMappingFound = false;
 				for (FRootMotionServerToLocalIDMapping& Mapping : RootMotionIDMappings)
 				{
-					if (Mapping.IsStillValid(TimeStamp) && ServerID == Mapping.ServerID)
+					if ((ServerID == Mapping.ServerID) && Mapping.IsStillValid(TimeStamp))
 					{
 						ServerRootMotionSource->LocalID = Mapping.LocalID;
 						Mapping.TimeStamp = TimeStamp;
@@ -9059,13 +9078,28 @@ void UCharacterMovementComponent::ConvertRootMotionServerIDsToLocalIDs(const FRo
 				}
 				if (bMappingFound)
 				{
+#if !(UE_BUILD_SHIPPING || UE_BUILD_TEST)
+					// We rely on this rule (Matches) being always true, so in non-shipping builds make sure it never breaks.
+					for (const TSharedPtr<FRootMotionSource>& LocalRootMotionSource : LocalRootMotionToMatchWith.RootMotionSources)
+					{
+						if (LocalRootMotionSource->LocalID == ServerRootMotionSource->LocalID)
+						{
+							ensureMsgf(LocalRootMotionSource->Matches(ServerRootMotionSource.Get()), 
+								TEXT("Character(%s) Local RootMotionSource(%s) has the same LocalID(%d) as a non-matching ServerRootMotionSource(%s)!"),
+								*GetNameSafe(CharacterOwner), *LocalRootMotionSource->ToSimpleString(), LocalRootMotionSource->LocalID, *ServerRootMotionSource->ToSimpleString());
+
+							// Don't break here, test if there are multiple matches.
+						}
+					}
+#endif
+
 					continue; // We've found the correct LocalID, done with this one, process next ServerRootMotionSource
 				}
 			}
 
 			// If no mapping found, find match out of those that are not already mapped
 			bool bMatchFound = false;
-			for (const auto& LocalRootMotionSource : LocalRootMotionToMatchWith.RootMotionSources)
+			for (const TSharedPtr<FRootMotionSource>& LocalRootMotionSource : LocalRootMotionToMatchWith.RootMotionSources)
 			{
 				if (LocalRootMotionSource.IsValid())
 				{
@@ -9077,12 +9111,13 @@ void UCharacterMovementComponent::ConvertRootMotionServerIDsToLocalIDs(const FRo
 						bool bMappingFound = false;
 						for (FRootMotionServerToLocalIDMapping& Mapping : RootMotionIDMappings)
 						{
-							if (Mapping.IsStillValid(TimeStamp) && LocalID == Mapping.LocalID)
+							if ((LocalID == Mapping.LocalID) && Mapping.IsStillValid(TimeStamp))
 							{
 								bMappingFound = true;
 								break; // Found it, don't need to search any more mappings
 							}
 						}
+
 						if (bMappingFound)
 						{
 							continue; // We found a ServerID matching this LocalID, so we don't try to match this
@@ -9098,36 +9133,25 @@ void UCharacterMovementComponent::ConvertRootMotionServerIDsToLocalIDs(const FRo
 
 						// Add to Mapping
 						{
-							FRootMotionServerToLocalIDMapping NewMapping;
-							NewMapping.LocalID = LocalID;
-							NewMapping.ServerID = ServerID;
-							NewMapping.TimeStamp = TimeStamp;
-
-							// Remove oldest mapping if the mapping is full
-							if (RootMotionIDMappings.Num() >= (uint32)ERootMotionMapping::MapSize)
+							// Remove out of date mappings, they can never be used again.
+							for (int32 MappingIndex = RootMotionIDMappings.Num()-1; MappingIndex >= 0; MappingIndex--)
 							{
-								// Find index of oldest mapping
-								uint32 OldestIndex = 0;
-								float OldestTimeStamp = RootMotionIDMappings[0].TimeStamp;
-								for (int i = 1; i < RootMotionIDMappings.Num(); ++i)
+								if (!RootMotionIDMappings[MappingIndex].IsStillValid(TimeStamp))
 								{
-									if (RootMotionIDMappings[i].TimeStamp < OldestTimeStamp)
-									{
-										OldestIndex = i;
-										OldestTimeStamp = RootMotionIDMappings[i].TimeStamp;
-									}
+									// Most recent entries added last, so we can cull the top of the list.
+									RootMotionIDMappings.RemoveAt(0, MappingIndex+1, false);
+									break;
 								}
-
-								// Remove oldest mapping
-								RootMotionIDMappings.RemoveAtSwap(OldestIndex, 1, false);
 							}
 
+							FRootMotionServerToLocalIDMapping NewMapping;
+							NewMapping.LocalID = LocalID;
+							NewMapping.ServerID = ServerID;
+							NewMapping.TimeStamp = TimeStamp;
+
 							RootMotionIDMappings.Add(NewMapping);
 							bMatchFound = true;
 							break; // Stop searching LocalRootMotionSources, we've found a match
-
-							// TODO-RootMotionSource: In non-shipping builds, continue looping through Local root motion sources
-							// and if we match more than one, LogWarning so we can catch bad Matching
 						}
 					}
 				}
diff --git a/Engine/Source/Runtime/Engine/Private/Components/PrimitiveComponent.cpp b/Engine/Source/Runtime/Engine/Private/Components/PrimitiveComponent.cpp
--- a/Engine/Source/Runtime/Engine/Private/Components/PrimitiveComponent.cpp
+++ b/Engine/Source/Runtime/Engine/Private/Components/PrimitiveComponent.cpp
@@ -2933,6 +2933,11 @@ void UPrimitiveComponent::GetResourceSizeEx(FResourceSizeEx& CumulativeResourceS
 	{
 		BodyInstance.GetBodyInstanceResourceSizeEx(CumulativeResourceSize);
 	}
+	if (SceneProxy)
+	{
+		CumulativeResourceSize.AddDedicatedSystemMemoryBytes(SceneProxy->GetMemoryFootprint());
+	}
+
 }
 
 void UPrimitiveComponent::SetRenderCustomDepth(bool bValue)
diff --git a/Engine/Source/Runtime/Engine/Private/Components/SkeletalMeshComponent.cpp b/Engine/Source/Runtime/Engine/Private/Components/SkeletalMeshComponent.cpp
--- a/Engine/Source/Runtime/Engine/Private/Components/SkeletalMeshComponent.cpp
+++ b/Engine/Source/Runtime/Engine/Private/Components/SkeletalMeshComponent.cpp
@@ -784,7 +784,19 @@ bool USkeletalMeshComponent::ShouldTickPose() const
 	// When we stop root motion we go back to ticking after CharacterMovement. Unfortunately that means that we could tick twice that frame.
 	// So only enforce a single tick per frame.
 	const bool bAlreadyTickedThisFrame = PoseTickedThisFrame();
-	return (Super::ShouldTickPose() && IsRegistered() && (AnimScriptInstance || PostProcessAnimInstance) && !bAutonomousTickPose && !bPauseAnims && GetWorld()->AreActorsInitialized() && !bNoSkeletonUpdate && !bAlreadyTickedThisFrame);
+
+	// Autonomous Ticking is allowed to occur multiple times per frame, as we can receive and process multiple networking updates the same frame.
+	const bool bShouldTickBasedOnAutonomousCheck = bIsAutonomousTickPose || (!bOnlyAllowAutonomousTickPose && !bAlreadyTickedThisFrame);
+
+	const bool bIsPlayingNetworkedRootMotionMontage = (GetAnimInstance() != nullptr)
+		? (GetAnimInstance()->RootMotionMode == ERootMotionMode::RootMotionFromMontagesOnly) && (GetAnimInstance()->GetRootMotionMontageInstance() != nullptr)
+		: false;
+
+	// When playing networked Root Motion Montages, we want these to play on dedicated servers and remote clients for networking and position correction purposes.
+	// So we force pose updates in that case to keep root motion and position in sync.
+	const bool bShouldTickBasedOnVisibility = ((MeshComponentUpdateFlag < EMeshComponentUpdateFlag::OnlyTickPoseWhenRendered) || bRecentlyRendered || bIsPlayingNetworkedRootMotionMontage);
+
+	return (bShouldTickBasedOnVisibility && bShouldTickBasedOnAutonomousCheck && IsRegistered() && (AnimScriptInstance || PostProcessAnimInstance) && !bPauseAnims && GetWorld()->AreActorsInitialized() && !bNoSkeletonUpdate);
 }
 
 static FThreadSafeCounter Ticked;
@@ -1483,7 +1495,7 @@ void USkeletalMeshComponent::RefreshBoneTransforms(FActorComponentTickFunction*
 
 	const bool bInvalidCachedCurve = bDoEvaluationRateOptimization && 
 									CurrentAnimCurveMappingNameUids != nullptr &&
-									((CachedCurve.Num() != GetCurveNumber(SkeletalMesh->Skeleton)) || (CachedCurve.UIDList != CurrentAnimCurveMappingNameUids));
+									((CachedCurve.Num() != GetCurveNumber(SkeletalMesh->Skeleton)) || (CachedCurve.UIDList != CurrentAnimCurveMappingNameUids) || (AnimCurves.Num() != GetCurveNumber(SkeletalMesh->Skeleton)));
 
 	const bool bShouldDoEvaluation = !bDoEvaluationRateOptimization || bInvalidCachedBones || bInvalidCachedCurve || !AnimUpdateRateParams->ShouldSkipEvaluation();
 
@@ -1583,7 +1595,7 @@ void USkeletalMeshComponent::RefreshBoneTransforms(FActorComponentTickFunction*
 			}
 			else
 			{
-				PerformAnimationEvaluation(SkeletalMesh, AnimScriptInstance, GetEditableComponentSpaceTransforms(), BoneSpaceTransforms, RootBoneTranslation, AnimEvaluationContext.Curve);
+				PerformAnimationEvaluation(SkeletalMesh, AnimScriptInstance, GetEditableComponentSpaceTransforms(), BoneSpaceTransforms, RootBoneTranslation, AnimCurves);
 			}
 		}
 		else
@@ -1598,7 +1610,7 @@ void USkeletalMeshComponent::RefreshBoneTransforms(FActorComponentTickFunction*
 				LocalEditableSpaceBases.Append(CachedComponentSpaceTransforms);
 				if (CachedCurve.IsValid())
 				{
-					AnimEvaluationContext.Curve.CopyFrom(CachedCurve);
+					AnimCurves.CopyFrom(CachedCurve);
 				}
 			}
 			if(AnimEvaluationContext.bDoUpdate)
@@ -1659,7 +1671,7 @@ void USkeletalMeshComponent::PostAnimEvaluation(FAnimationEvaluationContext& Eva
 
 	if (EvaluationContext.bDuplicateToCacheCurve)
 	{
-		CachedCurve.CopyFrom(EvaluationContext.Curve);
+		CachedCurve.CopyFrom(AnimCurves);
 	}
 	
 	if (EvaluationContext.bDuplicateToCacheBones)
@@ -1695,20 +1707,20 @@ void USkeletalMeshComponent::PostAnimEvaluation(FAnimationEvaluationContext& Eva
 		FillComponentSpaceTransforms(SkeletalMesh, BoneSpaceTransforms, GetEditableComponentSpaceTransforms());
 
 		// interpolate curve
-		EvaluationContext.Curve.LerpTo(CachedCurve, Alpha);
+		AnimCurves.LerpTo(CachedCurve, Alpha);
 	}
 
 	if(AnimScriptInstance)
 	{
 #if WITH_EDITOR
-		GetEditableAnimationCurves() = EvaluationContext.Curve;
+		GetEditableAnimationCurves() = AnimCurves;
 #endif 
 		// curve update happens first
-		AnimScriptInstance->UpdateCurves(EvaluationContext.Curve);
+		AnimScriptInstance->UpdateCurves(AnimCurves);
 
 		for(UAnimInstance* SubInstance : SubInstances)
 		{
-			SubInstance->UpdateCurves(EvaluationContext.Curve);
+			SubInstance->UpdateCurves(AnimCurves);
 		}
 	}
 
@@ -2666,6 +2678,7 @@ void USkeletalMeshComponent::CompleteParallelAnimationEvaluation(bool bDoPostAni
 
 			Exchange(AnimEvaluationContext.ComponentSpaceTransforms, AnimEvaluationContext.bDoInterpolation ? CachedComponentSpaceTransforms : GetEditableComponentSpaceTransforms());
 			Exchange(AnimEvaluationContext.BoneSpaceTransforms, AnimEvaluationContext.bDoInterpolation ? CachedBoneSpaceTransforms : BoneSpaceTransforms);
+			Exchange(AnimEvaluationContext.Curve, AnimEvaluationContext.bDoInterpolation ? CachedCurve : AnimCurves);
 			Exchange(AnimEvaluationContext.RootBoneTranslation, RootBoneTranslation);
 		}
 
diff --git a/Engine/Source/Runtime/Engine/Private/Components/SkinnedMeshComponent.cpp b/Engine/Source/Runtime/Engine/Private/Components/SkinnedMeshComponent.cpp
--- a/Engine/Source/Runtime/Engine/Private/Components/SkinnedMeshComponent.cpp
+++ b/Engine/Source/Runtime/Engine/Private/Components/SkinnedMeshComponent.cpp
@@ -1348,7 +1348,8 @@ class USkeletalMeshSocket const* USkinnedMeshComponent::GetSocketByName(FName In
 	}
 	else
 	{
-		UE_LOG(LogSkinnedMeshComp, Warning, TEXT("GetSocketByName(%s: %s): No SkeletalMesh for %s"), *GetNameSafe(SkeletalMesh), *InSocketName.ToString(), *GetName());
+		UE_LOG(LogSkinnedMeshComp, Warning, TEXT("GetSocketByName(%s): No SkeletalMesh for Component(%s) Actor(%s)"), 
+			*InSocketName.ToString(), *GetName(), *GetNameSafe(GetOuter()));
 	}
 
 	return Socket;
diff --git a/Engine/Source/Runtime/Engine/Private/Components/StaticMeshComponent.cpp b/Engine/Source/Runtime/Engine/Private/Components/StaticMeshComponent.cpp
--- a/Engine/Source/Runtime/Engine/Private/Components/StaticMeshComponent.cpp
+++ b/Engine/Source/Runtime/Engine/Private/Components/StaticMeshComponent.cpp
@@ -718,6 +718,10 @@ bool UStaticMeshComponent::BuildTextureStreamingData(ETextureStreamingBuildType
 			}
 			else if (StreamingTextureData.Num() == 0)
 			{
+				// Reset the validity here even if the bounds don't fit as the material might not use any streaming textures.
+				// This is required as the texture streaming build only mark levels as dirty if they have texture related data.
+				bBuildDataValid = true; 
+
 				// In that case, check if the component refers a streaming texture. If so, the build data is missing.
 				TArray<UMaterialInterface*> UsedMaterials;
 				GetUsedMaterials(UsedMaterials);
diff --git a/Engine/Source/Runtime/Engine/Private/Components/TextRenderComponent.cpp b/Engine/Source/Runtime/Engine/Private/Components/TextRenderComponent.cpp
--- a/Engine/Source/Runtime/Engine/Private/Components/TextRenderComponent.cpp
+++ b/Engine/Source/Runtime/Engine/Private/Components/TextRenderComponent.cpp
@@ -362,7 +362,10 @@ public:
 			check(InMaterial && InFont && InFont->FontCacheType == EFontCacheType::Offline);
 
 			const int32 NumFontPages = InFont->Textures.Num();
-			if (NumFontPages > 0)
+
+			// Checking GIsRequestingExit as a workaround for lighting rebuild command let crash.
+			// Happening because GIsRequestingExit is true preventing the FTextRenderComponentMIDCache from registering into the GGCObjectReferencer
+			if (!GIsRequestingExit && NumFontPages > 0)
 			{
 				TArray<FGuid> FontParameterIds;
 				InMaterial->GetMaterial()->GetAllFontParameterNames(FontParameters, FontParameterIds);
@@ -403,6 +406,7 @@ public:
 			// We can only test for stale MIDs when we created the MIDs ourselves (which we don't do if the outer MID was itself a MID)
 			if (GIsEditor && !InMaterial->IsA<UMaterialInstanceDynamic>())
 			{
+				// We only test against the number of font pages when we created the MIDs
 				bIsStale = MIDs.Num() != InFont->Textures.Num();
 
 				if (!bIsStale)
@@ -422,6 +426,7 @@ public:
 			return bIsStale;
 		}
 
+
 		TArray<UMaterialInstanceDynamic*> MIDs;
 		TArray<FName> FontParameters;
 	};
diff --git a/Engine/Source/Runtime/Engine/Private/DataChannel.cpp b/Engine/Source/Runtime/Engine/Private/DataChannel.cpp
--- a/Engine/Source/Runtime/Engine/Private/DataChannel.cpp
+++ b/Engine/Source/Runtime/Engine/Private/DataChannel.cpp
@@ -960,11 +960,11 @@ FOutBunch* UChannel::PrepBunch(FOutBunch* Bunch, FOutBunch* OutBunch, bool Merge
 		Connection->LastOutBunch = OutBunch;
 
 #if !(UE_BUILD_SHIPPING || UE_BUILD_TEST)
-		if (CVarNetReliableDebug.GetValueOnGameThread() == 1)
+		if (CVarNetReliableDebug.GetValueOnAnyThread() == 1)
 		{
 			UE_LOG(LogNetTraffic, Warning, TEXT("%s. Reliable: %s"), *Describe(), *Bunch->DebugString);
 		}
-		if (CVarNetReliableDebug.GetValueOnGameThread() == 2)
+		if (CVarNetReliableDebug.GetValueOnAnyThread() == 2)
 		{
 			UE_LOG(LogNetTraffic, Warning, TEXT("%s. Reliable: %s"), *Describe(), *Bunch->DebugString);
 			PrintReliableBunchBuffer();
@@ -1493,7 +1493,7 @@ void UActorChannel::Close()
 
 			// Validation checking
 			static const auto ValidateCVar = IConsoleManager::Get().FindTConsoleVariableDataInt(TEXT("net.DormancyValidate"));
-			if ( ValidateCVar && ValidateCVar->GetValueOnGameThread() > 0 )
+			if ( ValidateCVar && ValidateCVar->GetValueOnAnyThread() > 0 )
 			{
 				bKeepReplicators = true;		// We need to keep the replicators around so we can use
 			}
@@ -2299,7 +2299,7 @@ bool UActorChannel::ReplicateActor()
 	}
 
 #if !(UE_BUILD_SHIPPING || UE_BUILD_TEST)
-	if (CVarNetReliableDebug.GetValueOnGameThread() > 0)
+	if (CVarNetReliableDebug.GetValueOnAnyThread() > 0)
 	{
 		Bunch.DebugString = FString::Printf(TEXT("%.2f ActorBunch: %s"), Connection->Driver->Time, *Actor->GetName() );
 	}
diff --git a/Engine/Source/Runtime/Engine/Private/DataReplication.cpp b/Engine/Source/Runtime/Engine/Private/DataReplication.cpp
--- a/Engine/Source/Runtime/Engine/Private/DataReplication.cpp
+++ b/Engine/Source/Runtime/Engine/Private/DataReplication.cpp
@@ -326,7 +326,7 @@ void FReplicationChangelistMgr::Update( const UObject* InObject, const uint32 Re
 	//	2. We check LastCompareIndex > 1 so we can do at least one pass per connection to compare all properties
 	//		This is necessary due to how RemoteRole is manipulated per connection, so we need to give all connections a chance to see if it changed
 	//	3. We ALWAYS compare on bNetInitial to make sure we have a fresh changelist of net initial properties in this case
-	if ( CVarShareShadowState.GetValueOnGameThread() && !RepFlags.bNetInitial && LastCompareIndex > 1 && LastReplicationFrame == ReplicationFrame )
+	if ( CVarShareShadowState.GetValueOnAnyThread() && !RepFlags.bNetInitial && LastCompareIndex > 1 && LastReplicationFrame == ReplicationFrame )
 	{
 		INC_DWORD_STAT_BY( STAT_NetSkippedDynamicProps, 1 );
 		return;
@@ -1020,6 +1020,7 @@ void FObjectReplicator::ReplicateCustomDeltaProperties( FNetBitWriter & Bunch, F
 	ConditionMap[COND_Custom] = true;
 	ConditionMap[COND_ReplayOrOwner] = bIsReplay || bIsOwner;
 	ConditionMap[COND_ReplayOnly] = bIsReplay;
+	ConditionMap[COND_SkipReplay] = !bIsReplay;
 
 	// Make sure net field export group is registered
 	FNetFieldExportGroup* NetFieldExportGroup = OwningChannel->GetOrCreateNetFieldExportGroupForClassNetCache( Object );
@@ -1159,7 +1160,7 @@ bool FObjectReplicator::ReplicateProperties( FOutBunch & Bunch, FReplicationFlag
 	{
 		static const auto* CVar = IConsoleManager::Get().FindTConsoleVariableDataInt( TEXT( "net.RPC.Debug" ) );
 
-		if ( UNLIKELY( CVar && CVar->GetValueOnGameThread() == 1 ) )
+		if ( UNLIKELY( CVar && CVar->GetValueOnAnyThread() == 1 ) )
 		{
 			UE_LOG( LogRepTraffic, Warning,	TEXT("      Sending queued RPCs: %s. Channel[%d] [%.1f bytes]"), *Object->GetName(), OwningChannel->ChIndex, RemoteFunctions->GetNumBits() / 8.f );
 		}
@@ -1269,7 +1270,7 @@ void FObjectReplicator::QueueRemoteFunctionBunch( UFunction* Func, FOutBunch &Bu
 		RemoteFuncInfo[InfoIdx].Calls = 0;
 	}
 	
-	if (++RemoteFuncInfo[InfoIdx].Calls > CVarMaxRPCPerNetUpdate.GetValueOnGameThread())
+	if (++RemoteFuncInfo[InfoIdx].Calls > CVarMaxRPCPerNetUpdate.GetValueOnAnyThread())
 	{
 		UE_LOG(LogRep, Verbose, TEXT("Too many calls (%d) to RPC %s within a single netupdate. Skipping. %s.  LastCallTime: %.2f. CurrentTime: %.2f. LastRelevantTime: %.2f. LastUpdateTime: %.2f "),
 			RemoteFuncInfo[InfoIdx].Calls, *Func->GetName(), *GetPathNameSafe(GetObject()), RemoteFuncInfo[InfoIdx].LastCallTime, OwningChannel->Connection->Driver->Time, OwningChannel->RelevantTime, OwningChannel->LastUpdateTime);
