Copying //UE4/Orion-Staging to //UE4/Main (Source: //Orion/Dev-General @ 3196521)

diff --git a/Engine/Source/Runtime/Engine/Private/DemoNetDriver.cpp b/Engine/Source/Runtime/Engine/Private/DemoNetDriver.cpp
--- a/Engine/Source/Runtime/Engine/Private/DemoNetDriver.cpp
+++ b/Engine/Source/Runtime/Engine/Private/DemoNetDriver.cpp
@@ -47,7 +47,8 @@ static TAutoConsoleVariable<int32> CVarUseAdaptiveReplayUpdateFrequency( TEXT( "
 static TAutoConsoleVariable<int32> CVarDemoAsyncLoadWorld( TEXT( "demo.AsyncLoadWorld" ), 0, TEXT( "If 1, we will use seamless server travel to load the replay world asynchronously" ) );
 static TAutoConsoleVariable<float> CVarCheckpointUploadDelayInSeconds( TEXT( "demo.CheckpointUploadDelayInSeconds" ), 30.0f, TEXT( "" ) );
 static TAutoConsoleVariable<int32> CVarDemoLoadCheckpointGarbageCollect( TEXT( "demo.LoadCheckpointGarbageCollect" ), 1, TEXT("If nonzero, CollectGarbage will be called during LoadCheckpoint after the old actors and connection are cleaned up." ) );
-static TAutoConsoleVariable<float> CVarCheckpointSaveMaxMSPerFrame( TEXT( "demo.CheckpointSaveMaxMSPerFrame" ), 0.0f, TEXT( "Maximum time allowed each frame to spend on saving a checkpoint. If 0, it will save the checkpoint in a single frame, regardless of how long it takes." ) );
+static TAutoConsoleVariable<float> CVarCheckpointSaveMaxMSPerFrameOverride( TEXT( "demo.CheckpointSaveMaxMSPerFrameOverride" ), -1.0f, TEXT( "If >= 0, this value will override the CheckpointSaveMaxMSPerFrame member variable, which is the maximum time allowed each frame to spend on saving a checkpoint. If 0, it will save the checkpoint in a single frame, regardless of how long it takes." ) );
+static TAutoConsoleVariable<int32> CVarDemoClientRecordAsyncEndOfFrame( TEXT( "demo.ClientRecordAsyncEndOfFrame" ), 0, TEXT( "If true, TickFlush will be called on a thread in parallel with Slate." ) );
 
 #if !(UE_BUILD_SHIPPING || UE_BUILD_TEST)
 static TAutoConsoleVariable<int32> CVarDemoForceFailure( TEXT( "demo.ForceFailure" ), 0, TEXT( "" ) );
@@ -190,8 +191,7 @@ public:
 			}
 
 			// We're done
-			Driver->LoadCheckpoint( GotoCheckpointArchive, GotoCheckpointSkipExtraTimeInMS );
-			return true;
+			return Driver->LoadCheckpoint( GotoCheckpointArchive, GotoCheckpointSkipExtraTimeInMS );
 		}
 
 		return false;
@@ -270,6 +270,8 @@ public:
 UDemoNetDriver::UDemoNetDriver(const FObjectInitializer& ObjectInitializer)
 	: Super(ObjectInitializer)
 {
+	CurrentLevelIndex = 0;
+	bRecordMapChanges = false;
 }
 
 void UDemoNetDriver::AddReplayTask( FQueuedReplayTask* NewTask )
@@ -363,6 +365,7 @@ bool UDemoNetDriver::InitBase( bool bInitAsClient, FNetworkNotify* InNotify, con
 		ViewerOverride					= nullptr;
 		bPrioritizeActors				= false;
 		bPauseRecording					= false;
+		CheckpointSaveMaxMSPerFrame		= 0.0f;
 
 		ResetDemoState();
 
@@ -525,22 +528,10 @@ bool UDemoNetDriver::InitConnectInternal( FString& Error )
 
 		if ( bAsyncLoadWorld )
 		{
-			TArray<AController*> Controllers;
-			for ( FConstControllerIterator Iterator = World->GetControllerIterator(); Iterator; ++Iterator )
-			{
-				Controllers.Add( *Iterator );
-			}
-
-			for ( int i = 0; i < Controllers.Num(); i++ )
-			{
-				if ( Controllers[i] )
-				{
-					Controllers[i]->Destroy();
-				}
-			}
+			LevelNamesAndTimes = PlaybackDemoHeader.LevelNamesAndTimes;
 
 			// FIXME: Test for failure!!!
-			World->SeamlessTravel( PlaybackDemoHeader.LevelName, true );
+			ProcessSeamlessTravel(0);
 		}
 		else
 		{
@@ -548,7 +539,7 @@ bool UDemoNetDriver::InitConnectInternal( FString& Error )
 			FString LoadMapError;
 
 			FURL LocalDemoURL;
-			LocalDemoURL.Map = PlaybackDemoHeader.LevelName;
+			LocalDemoURL.Map = PlaybackDemoHeader.LevelNamesAndTimes[0].LevelName;
 
 			FWorldContext * WorldContext = GEngine->GetWorldContextFromWorld( GetWorld() );
 
@@ -613,6 +604,8 @@ bool UDemoNetDriver::InitListen( FNetworkNotify* InNotify, FURL& ListenURL, bool
 
 	const TCHAR* FriendlyNameOption = ListenURL.GetOption( TEXT("DemoFriendlyName="), nullptr );
 
+	bRecordMapChanges = ListenURL.GetOption(TEXT("RecordMapChanges"), nullptr) != nullptr;
+
 	TArray< FString > UserNames;
 	AGameStateBase* GameState = GetWorld()->GetGameState();
 
@@ -638,6 +631,42 @@ bool UDemoNetDriver::InitListen( FNetworkNotify* InNotify, FURL& ListenURL, bool
 		FNetworkVersion::GetReplayVersion(),
 		FOnStreamReadyDelegate::CreateUObject( this, &UDemoNetDriver::ReplayStreamingReady ) );
 
+	AddNewLevel(World->GetOuter()->GetName());
+
+	bool Result = WriteNetworkDemoHeader(Error);
+
+	// Spawn the demo recording spectator.
+	SpawnDemoRecSpectator( Connection, ListenURL );
+
+	return Result;
+}
+
+bool UDemoNetDriver::ContinueListen(FURL& ListenURL)
+{
+	if (IsRecording() && ensure(IsRecordingPaused()))
+	{
+		++CurrentLevelIndex;
+
+		AddNewLevel(World->GetOuter()->GetName());
+
+		FString Error;
+		WriteNetworkDemoHeader(Error);
+
+		ReplayStreamer->RefreshHeader();
+
+		PauseRecording(false);
+		SpawnDemoRecSpectator(ClientConnections[0], ListenURL);
+
+		// Force a checkpoint to be created in the next tick - We need a checkpoint right after travelling so that scrubbing
+		// from a different level will have essentially an "empty" checkpoint to work from.
+		LastCheckpointTime = -1 * CVarCheckpointUploadDelayInSeconds.GetValueOnGameThread();
+		return true;
+	}
+	return false;
+}
+
+bool UDemoNetDriver::WriteNetworkDemoHeader(FString& Error)
+{
 	FArchive* FileAr = ReplayStreamer->GetHeaderArchive();
 
 	if( !FileAr )
@@ -648,7 +677,7 @@ bool UDemoNetDriver::InitListen( FNetworkNotify* InNotify, FURL& ListenURL, bool
 
 	FNetworkDemoHeader DemoHeader;
 
-	DemoHeader.LevelName = World->GetCurrentLevel()->GetOutermost()->GetName();
+	DemoHeader.LevelNamesAndTimes = LevelNamesAndTimes;
 
 	WriteGameSpecificDemoHeader(DemoHeader.GameSpecificData);
 
@@ -656,9 +685,6 @@ bool UDemoNetDriver::InitListen( FNetworkNotify* InNotify, FURL& ListenURL, bool
 	(*FileAr) << DemoHeader;
 	FileAr->Flush();
 
-	// Spawn the demo recording spectator.
-	SpawnDemoRecSpectator( Connection, ListenURL );
-
 	return true;
 }
 
@@ -672,7 +698,28 @@ bool UDemoNetDriver::IsPlaying() const
 	return ServerConnection != nullptr && ServerConnection->State != USOCK_Closed;
 }
 
-void UDemoNetDriver::TickFlush( float DeltaSeconds )
+bool UDemoNetDriver::ShouldTickFlushAsyncEndOfFrame() const
+{
+	return	GEngine && GEngine->ShouldDoAsyncEndOfFrameTasks() && CVarDemoClientRecordAsyncEndOfFrame.GetValueOnAnyThread() != 0 && World && World->IsRecordingClientReplay();
+}
+
+void UDemoNetDriver::TickFlush(float DeltaSeconds)
+{
+	if (!ShouldTickFlushAsyncEndOfFrame())
+	{
+		TickFlushInternal(DeltaSeconds);
+	}
+}
+
+void UDemoNetDriver::TickFlushAsyncEndOfFrame(float DeltaSeconds)
+{
+	if (ShouldTickFlushAsyncEndOfFrame())
+	{
+		TickFlushInternal(DeltaSeconds);
+	}
+}
+
+void UDemoNetDriver::TickFlushInternal( float DeltaSeconds )
 {
 	// Set the context on the world for this driver's level collection.
 	const FLevelCollection* FoundCollection = nullptr;
@@ -692,7 +739,7 @@ void UDemoNetDriver::TickFlush( float DeltaSeconds )
 
 	Super::TickFlush( DeltaSeconds );
 
-	if ( !IsRecording() || bPauseRecording )
+	if (!IsRecording())
 	{
 		// Nothing to do
 		return;
@@ -710,6 +757,11 @@ void UDemoNetDriver::TickFlush( float DeltaSeconds )
 		return;
 	}
 
+	if (bPauseRecording)
+	{
+		return;
+	}
+
 	FArchive* FileAr = ReplayStreamer->GetStreamingArchive();
 
 	if ( FileAr == nullptr )
@@ -894,7 +946,7 @@ bool UDemoNetDriver::ShouldClientDestroyTearOffActors() const
 
 bool UDemoNetDriver::ShouldSkipRepNotifies() const
 {
-	if ( CVarDemoFastForwardSkipRepNotifies.GetValueOnGameThread() != 0 )
+	if ( CVarDemoFastForwardSkipRepNotifies.GetValueOnAnyThread() != 0 )
 	{
 		return bIsFastForwarding;
 	}
@@ -1055,6 +1107,22 @@ static void SerializeGuidCache( TSharedPtr< class FNetGUIDCache > GuidCache, FAr
 	}
 }
 
+float UDemoNetDriver::GetCheckpointSaveMaxMSPerFrame() const
+{
+	const float CVarValue = CVarCheckpointSaveMaxMSPerFrameOverride.GetValueOnAnyThread();
+	if (CVarValue >= 0.0f)
+	{
+		return CVarValue;
+	}
+
+	return CheckpointSaveMaxMSPerFrame;
+}
+
+void UDemoNetDriver::AddNewLevel(const FString& NewLevelName)
+{
+	LevelNamesAndTimes.Add(FLevelNameAndTime(NewLevelName, ReplayStreamer->GetTotalDemoTime()));
+}
+
 void UDemoNetDriver::SaveCheckpoint()
 {
 	DECLARE_SCOPE_CYCLE_COUNTER(TEXT("SaveCheckpoint time"), STAT_ReplayCheckpointSaveTime, STATGROUP_Net);
@@ -1165,7 +1233,7 @@ void UDemoNetDriver::TickCheckpoint()
 	// Set bResendAllDataSinceOpen to true to signify that we want to do this
 	ClientConnections[0]->bResendAllDataSinceOpen = true;
 
-	const double CheckpointMaxUploadTimePerFrame = ( double )CVarCheckpointSaveMaxMSPerFrame.GetValueOnGameThread() / 1000;
+	const double CheckpointMaxUploadTimePerFrame = ( double )GetCheckpointSaveMaxMSPerFrame() / 1000;
 
 	for ( auto ActorIt = PendingCheckpointActors.CreateIterator(); ActorIt; ++ActorIt )
 	{
@@ -1213,6 +1281,7 @@ void UDemoNetDriver::TickCheckpoint()
 		//
 		// We're done saving this checkpoint
 		//
+		*CheckpointArchive << CurrentLevelIndex;
 
 		// First, save the current guid cache
 		SerializeGuidCache( GuidCache, CheckpointArchive );
@@ -1300,7 +1369,7 @@ void UDemoNetDriver::LoadExternalData( FArchive& Ar, const float TimeSeconds )
 
 		FReplayExternalDataArray& ExternalDataArray = ExternalDataToObjectMap.FindOrAdd( NetGUID );
 
-		ExternalDataArray.Add( FReplayExternalData( Reader, TimeSeconds ) );
+		ExternalDataArray.Add( new FReplayExternalData( Reader, TimeSeconds ) );
 	}
 }
 
@@ -1373,8 +1442,8 @@ void UDemoNetDriver::TickDemoRecord( float DeltaSeconds )
 
 	ReplayStreamer->UpdateTotalDemoTime( GetDemoCurrentTimeInMS() );
 
-	float MAX_RECORD_HZ	= CVarDemoRecordHz.GetValueOnGameThread();
-	float MIN_RECORD_HZ = CVarDemoMinRecordHz.GetValueOnGameThread();
+	float MAX_RECORD_HZ	= CVarDemoRecordHz.GetValueOnAnyThread();
+	float MIN_RECORD_HZ = CVarDemoMinRecordHz.GetValueOnAnyThread();
 	if (MIN_RECORD_HZ > MAX_RECORD_HZ)
 	{
 		// make sure min and max are sane
@@ -1401,7 +1470,7 @@ void UDemoNetDriver::TickDemoRecord( float DeltaSeconds )
 		ServerTickTime	= 1.0 / ServerTickTime;
 	}
 
-	const bool bUseAdapativeNetFrequency = CVarUseAdaptiveReplayUpdateFrequency.GetValueOnGameThread() > 0;
+	const bool bUseAdapativeNetFrequency = CVarUseAdaptiveReplayUpdateFrequency.GetValueOnAnyThread() > 0;
 
 	// Build priority list
 	PrioritizedActors.Reset( GetNetworkObjectList().GetObjects().Num() );
@@ -1579,11 +1648,11 @@ void UDemoNetDriver::TickDemoRecord( float DeltaSeconds )
 	WriteDemoFrameFromQueuedDemoPackets( *FileAr, CastChecked< UDemoNetConnection >( ClientConnections[0] )->QueuedDemoPackets );
 
 	// Save a checkpoint if it's time
-	if ( CVarEnableCheckpoints.GetValueOnGameThread() == 1 )
+	if ( CVarEnableCheckpoints.GetValueOnAnyThread() == 1 )
 	{
 		check( !PendingCheckpointActors.Num() );		// We early out above, so this shouldn't be possible
 
-		const double CHECKPOINT_DELAY = CVarCheckpointUploadDelayInSeconds.GetValueOnGameThread();
+		const double CHECKPOINT_DELAY = CVarCheckpointUploadDelayInSeconds.GetValueOnAnyThread();
 
 		if ( DemoCurrentTime - LastCheckpointTime > CHECKPOINT_DELAY )
 		{
@@ -1657,6 +1726,13 @@ bool UDemoNetDriver::ReadDemoFrameIntoPlaybackPackets( FArchive& Ar )
 		return false;
 	}
 
+	int32 ReadCurrentLevelIndex = 0;
+
+	if (PlaybackDemoHeader.Version >= HISTORY_MULTIPLE_LEVELS)
+	{
+		Ar << ReadCurrentLevelIndex;
+	}
+
 	float TimeSeconds = 0.0f;
 
 	Ar << TimeSeconds;
@@ -1766,12 +1842,38 @@ bool UDemoNetDriver::ReadDemoFrameIntoPlaybackPackets( FArchive& Ar )
 		FPlaybackPacket& Packet = *( new( PlaybackPackets )FPlaybackPacket );
 		Packet.Data.AddUninitialized( PacketBytes );
 		Packet.TimeSeconds = TimeSeconds;
+		Packet.LevelIndex = ReadCurrentLevelIndex;
 		FMemory::Memcpy( Packet.Data.GetData(), ReadBuffer, PacketBytes );
 	}
 
 	return true;
 }
 
+void UDemoNetDriver::ProcessSeamlessTravel(int32 LevelIndex)
+{
+	// Destroy all player controllers since FSeamlessTravelHandler will not destroy them.
+	TArray<AController*> Controllers;
+	for (FConstControllerIterator Iterator = World->GetControllerIterator(); Iterator; ++Iterator)
+	{
+		Controllers.Add(*Iterator);
+	}
+
+	for (int i = 0; i < Controllers.Num(); i++)
+	{
+		if (Controllers[i])
+		{
+			// bNetForce is true so that the replicated spectator player controller will
+			// be destroyed as well.
+			Controllers[i]->Destroy(true);
+		}
+	}
+
+	// Set this to nullptr since we just destroyed it.
+	SpectatorController = nullptr;
+
+	World->SeamlessTravel(PlaybackDemoHeader.LevelNamesAndTimes[LevelIndex].LevelName, true);
+}
+
 bool UDemoNetDriver::ConditionallyReadDemoFrameIntoPlaybackPackets( FArchive& Ar )
 {
 	if ( PlaybackPackets.Num() > 0 )
@@ -1856,6 +1958,14 @@ bool UDemoNetDriver::ConditionallyProcessPlaybackPackets()
 		return false;
 	}
 
+	if (PlaybackPackets[0].LevelIndex != CurrentLevelIndex)
+	{
+		GetWorld()->GetGameInstance()->OnSeamlessTravelDuringReplay();
+		CurrentLevelIndex = PlaybackPackets[0].LevelIndex;
+		ProcessSeamlessTravel(CurrentLevelIndex);
+		return false;
+	}
+
 	const bool Result = ProcessPacket( PlaybackPackets[0].Data.GetData(), PlaybackPackets[0].Data.Num() );
 
 	PlaybackPackets.RemoveAt( 0 );
@@ -1903,6 +2013,8 @@ bool UDemoNetDriver::ProcessPacket( uint8* Data, int32 Count )
 
 void UDemoNetDriver::WriteDemoFrameFromQueuedDemoPackets( FArchive& Ar, TArray<FQueuedDemoPacket>& QueuedPackets )
 {
+	Ar << CurrentLevelIndex;
+	
 	// Save total absolute demo time in seconds
 	Ar << DemoCurrentTime;
 
@@ -2277,6 +2389,16 @@ void UDemoNetDriver::ReplayStreamingReady( bool bSuccess, bool bRecord )
 			return;
 		}
 
+		const TCHAR* const SkipToLevelIndexOption = DemoURL.GetOption(TEXT("SkipToLevelIndex="), nullptr);
+		if (SkipToLevelIndexOption)
+		{
+			int32 Index = FCString::Atoi(SkipToLevelIndexOption);
+			if (Index < LevelNamesAndTimes.Num())
+			{
+				AddReplayTask(new FGotoTimeInSecondsTask(this, (float)LevelNamesAndTimes[Index].LevelChangeTimeInMS / 1000.0f));
+			}
+		}
+
 		if ( ReplayStreamer->IsLive() && ReplayStreamer->GetTotalDemoTime() > 15 * 1000 )
 		{
 			// If the load time wasn't very long, jump to end now
@@ -2306,7 +2428,7 @@ FReplayExternalDataArray* UDemoNetDriver::GetExternalDataArrayForObject( UObject
 	return ExternalDataToObjectMap.Find( NetworkGUID );
 }
 
-void UDemoNetDriver::LoadCheckpoint( FArchive* GotoCheckpointArchive, int64 GotoCheckpointSkipExtraTimeInMS )
+bool UDemoNetDriver::LoadCheckpoint( FArchive* GotoCheckpointArchive, int64 GotoCheckpointSkipExtraTimeInMS )
 {
 	DECLARE_SCOPE_CYCLE_COUNTER(TEXT("LoadCheckpoint time"), STAT_ReplayCheckpointLoadTime, STATGROUP_Net);
 
@@ -2314,6 +2436,59 @@ void UDemoNetDriver::LoadCheckpoint( FArchive* GotoCheckpointArchive, int64 Goto
 	check( !bIsFastForwardingForCheckpoint );
 	check( !bIsFastForwarding );
 
+	int32 LevelForCheckpoint = 0;
+
+	if (PlaybackDemoHeader.Version >= HISTORY_MULTIPLE_LEVELS)
+	{
+		if (GotoCheckpointArchive->TotalSize() > 0)
+		{
+			*GotoCheckpointArchive << LevelForCheckpoint;
+		}
+	}
+
+	if (LevelForCheckpoint != CurrentLevelIndex)
+	{
+		GetWorld()->GetGameInstance()->OnSeamlessTravelDuringReplay();
+
+		for (FActorIterator It(GetWorld()); It; ++It)
+		{
+			GetWorld()->DestroyActor(*It, true);
+		}
+
+		// Clean package map to prepare to restore it to the checkpoint state
+		GuidCache->ObjectLookup.Empty();
+		GuidCache->NetGUIDLookup.Empty();
+
+		GuidCache->NetFieldExportGroupMap.Empty();
+		GuidCache->NetFieldExportGroupPathToIndex.Empty();
+		GuidCache->NetFieldExportGroupIndexToPath.Empty();
+
+		SpectatorController = nullptr;
+
+		ServerConnection->Close();
+		ServerConnection->CleanUp();
+
+		// Recreate the server connection - this is done so that when we execute the code the below again when we read in the
+		// checkpoint again after the server travel is finished, we'll have a clean server connection to work with.
+		ServerConnection = NewObject<UNetConnection>(GetTransientPackage(), UDemoNetConnection::StaticClass());
+
+		FURL ConnectURL;
+		ConnectURL.Map = DemoURL.Map;
+		ServerConnection->InitConnection(this, USOCK_Pending, ConnectURL, 1000000);
+
+		CollectGarbage(GARBAGE_COLLECTION_KEEPFLAGS);
+
+		ProcessSeamlessTravel(LevelForCheckpoint);
+		CurrentLevelIndex = LevelForCheckpoint;
+
+		if (GotoCheckpointArchive->TotalSize() != 0 && GotoCheckpointArchive->TotalSize() != INDEX_NONE)
+		{
+			GotoCheckpointArchive->Seek(0);
+		}
+
+		return false;
+	}
+
 	// Save off the current spectator position
 	// Check for NULL, which can be the case if we haven't played any of the demo yet but want to fast forward (joining live game for example)
 	if ( SpectatorController != NULL )
@@ -2499,7 +2674,7 @@ void UDemoNetDriver::LoadCheckpoint( FArchive* GotoCheckpointArchive, int64 Goto
 			SkipTimeInternal( ( float )GotoCheckpointSkipExtraTimeInMS / 1000.0f, true, true );
 		}
 
-		return;
+		return true;
 	}
 
 	int32 NumValues = 0;
@@ -2577,6 +2752,8 @@ void UDemoNetDriver::LoadCheckpoint( FArchive* GotoCheckpointArchive, int64 Goto
 			SpectatorController->SetViewTarget( ViewTarget );
 		}
 	}
+
+	return true;
 }
 
 bool UDemoNetDriver::ShouldQueueBunchesForActorGUID(FNetworkGUID InGUID) const
diff --git a/Engine/Source/Runtime/Engine/Private/GameEngine.cpp b/Engine/Source/Runtime/Engine/Private/GameEngine.cpp
--- a/Engine/Source/Runtime/Engine/Private/GameEngine.cpp
+++ b/Engine/Source/Runtime/Engine/Private/GameEngine.cpp
@@ -57,6 +57,12 @@ static FAutoConsoleVariableRef CvarSlowFrameLoggingThreshold(
 	ECVF_Default
 	);
 
+static int32 GDoAsyncEndOfFrameTasks = 0;
+static FAutoConsoleVariableRef CVarDoAsyncEndOfFrameTasks(
+	TEXT("tick.DoAsyncEndOfFrameTasks"),
+	GDoAsyncEndOfFrameTasks,
+	TEXT("Experimental option to run various things concurrently with the HUD render.")
+	);
 
 /** Benchmark results to the log */
 static void RunSynthBenchmark(const TArray<FString>& Args)
@@ -747,6 +753,11 @@ bool UGameEngine::NetworkRemapPath(UNetDriver* Driver, FString& Str, bool bReadi
 	return false;
 }
 
+bool UGameEngine::ShouldDoAsyncEndOfFrameTasks() const
+{
+	return FApp::ShouldUseThreadingForPerformance() && ENamedThreads::RenderThread != ENamedThreads::GameThread && !!GDoAsyncEndOfFrameTasks;
+}
+
 /*-----------------------------------------------------------------------------
 	Command line executor.
 -----------------------------------------------------------------------------*/
diff --git a/Engine/Source/Runtime/Engine/Private/GameFramework/RootMotionSource.cpp b/Engine/Source/Runtime/Engine/Private/GameFramework/RootMotionSource.cpp
--- a/Engine/Source/Runtime/Engine/Private/GameFramework/RootMotionSource.cpp
+++ b/Engine/Source/Runtime/Engine/Private/GameFramework/RootMotionSource.cpp
@@ -1652,23 +1652,24 @@ void FRootMotionSourceGroup::RemoveRootMotionSourceByID(uint16 RootMotionSourceI
 
 void FRootMotionSourceGroup::UpdateStateFrom(const FRootMotionSourceGroup& GroupToTakeStateFrom, bool bMarkForSimulatedCatchup)
 {
+	bIsAdditiveVelocityApplied = GroupToTakeStateFrom.bIsAdditiveVelocityApplied;
+	LastPreAdditiveVelocity = GroupToTakeStateFrom.LastPreAdditiveVelocity;
+
 	// For each matching Source in GroupToTakeStateFrom, move state over to this group's Sources
 	// Can do all matching with LocalID only, since anything passed into this function should have
 	// already been "matched" to LocalIDs
-	for (const auto& TakeFromRootMotionSource : GroupToTakeStateFrom.RootMotionSources)
+	for (const TSharedPtr<FRootMotionSource>& TakeFromRootMotionSource : GroupToTakeStateFrom.RootMotionSources)
 	{
-		if (TakeFromRootMotionSource.IsValid() && TakeFromRootMotionSource->LocalID != (uint16)ERootMotionSourceID::Invalid)
+		if (TakeFromRootMotionSource.IsValid() && (TakeFromRootMotionSource->LocalID != (uint16)ERootMotionSourceID::Invalid))
 		{
-			for (const auto& RootMotionSource : RootMotionSources)
+			for (const TSharedPtr<FRootMotionSource>& RootMotionSource : RootMotionSources)
 			{
-				if (RootMotionSource.IsValid() && RootMotionSource->LocalID == TakeFromRootMotionSource->LocalID)
+				if (RootMotionSource.IsValid() && (RootMotionSource->LocalID == TakeFromRootMotionSource->LocalID))
 				{
-					if (!RootMotionSource->Matches(TakeFromRootMotionSource.Get()))
-					{
-						// If we hit this case, we have an issue with Matches() and/or LocalIDs being mapped to invalid "partners"
-						checkf(false, TEXT("FRootMotionSource has the same LocalID as a non-matching Source!"));
-						break;
-					}
+					// We rely on the 'Matches' rule to be exact, verify that it is still correct here.
+					// If not, we're matching different root motion sources, or we're using properties that change over time for matching.
+					ensureMsgf(RootMotionSource->Matches(TakeFromRootMotionSource.Get()), TEXT("UpdateStateFrom RootMotionSource(%s) has the same LocalID(%d) as a non-matching TakeFromRootMotionSource(%s)!"),
+						*RootMotionSource->ToSimpleString(), RootMotionSource->LocalID, *TakeFromRootMotionSource->ToSimpleString());
 
 					const bool bSuccess = RootMotionSource->UpdateStateFrom(TakeFromRootMotionSource.Get(), bMarkForSimulatedCatchup);
 					if (bSuccess)
@@ -1825,6 +1826,14 @@ FRootMotionSourceGroup& FRootMotionSourceGroup::operator=(const FRootMotionSourc
 
 bool FRootMotionSourceGroup::operator==(const FRootMotionSourceGroup& Other) const
 {
+	if (bHasAdditiveSources != Other.bHasAdditiveSources || 
+		bHasOverrideSources != Other.bHasOverrideSources ||
+		bIsAdditiveVelocityApplied != Other.bIsAdditiveVelocityApplied ||
+		!LastPreAdditiveVelocity.Equals(Other.LastPreAdditiveVelocity, 1.f))
+	{
+		return false;
+	}
+
 	// Both invalid structs or both valid and Pointer compare (???) // deep comparison equality
 	if (RootMotionSources.Num() != Other.RootMotionSources.Num())
 	{
diff --git a/Engine/Source/Runtime/Engine/Private/GameInstance.cpp b/Engine/Source/Runtime/Engine/Private/GameInstance.cpp
--- a/Engine/Source/Runtime/Engine/Private/GameInstance.cpp
+++ b/Engine/Source/Runtime/Engine/Private/GameInstance.cpp
@@ -9,8 +9,9 @@
 #include "Engine/NetworkObjectList.h"
 #include "GameFramework/GameModeBase.h"
 #include "GameFramework/OnlineSession.h"
+#include "GameFramework/GameSession.h"
 #include "GameFramework/PlayerState.h"
-#include "GameFramework/GameModeBase.h"
+#include "Net/OnlineEngineInterface.h"
 
 #if WITH_EDITOR
 #include "UnrealEd.h"
@@ -106,21 +107,22 @@ void UGameInstance::InitializeStandalone()
 }
 
 #if WITH_EDITOR
-bool UGameInstance::InitializePIE(bool bAnyBlueprintErrors, int32 PIEInstance, bool bRunAsDedicated)
+
+FGameInstancePIEResult UGameInstance::InitializeForPlayInEditor(int32 PIEInstanceIndex, const FGameInstancePIEParameters& Params)
 {
 	UEditorEngine* const EditorEngine = CastChecked<UEditorEngine>(GetEngine());
 
 	// Look for an existing pie world context, may have been created before
-	WorldContext = EditorEngine->GetWorldContextFromPIEInstance(PIEInstance);
+	WorldContext = EditorEngine->GetWorldContextFromPIEInstance(PIEInstanceIndex);
 
 	if (!WorldContext)
 	{
 		// If not, create a new one
 		WorldContext = &EditorEngine->CreateNewWorldContext(EWorldType::PIE);
-		WorldContext->PIEInstance = PIEInstance;
+		WorldContext->PIEInstance = PIEInstanceIndex;
 	}
 
-	WorldContext->RunAsDedicated = bRunAsDedicated;
+	WorldContext->RunAsDedicated = Params.bRunAsDedicated;
 
 	WorldContext->OwningGameInstance = this;
 	
@@ -161,8 +163,7 @@ bool UGameInstance::InitializePIE(bool bAnyBlueprintErrors, int32 PIEInstance, b
 	// failed to create the world!
 	if (NewWorld == nullptr)
 	{
-		FMessageDialog::Open(EAppMsgType::Ok, NSLOCTEXT("UnrealEd", "Error_FailedCreateEditorPreviewWorld", "Failed to create editor preview world."));
-		return false;
+		return FGameInstancePIEResult::Failure(NSLOCTEXT("UnrealEd", "Error_FailedCreateEditorPreviewWorld", "Failed to create editor preview world."));
 	}
 
 	NewWorld->SetGameInstance(this);
@@ -171,7 +172,7 @@ bool UGameInstance::InitializePIE(bool bAnyBlueprintErrors, int32 PIEInstance, b
 
 	// make sure we can clean up this world!
 	NewWorld->ClearFlags(RF_Standalone);
-	NewWorld->bKismetScriptError = bAnyBlueprintErrors;
+	NewWorld->bKismetScriptError = Params.bAnyBlueprintErrors;
 
 	// Do a GC pass if necessary to remove any potentially unreferenced objects
 	if(bNeedsGarbageCollection)
@@ -181,10 +182,29 @@ bool UGameInstance::InitializePIE(bool bAnyBlueprintErrors, int32 PIEInstance, b
 
 	Init();
 
+	// Give the deprecated method a chance to fail as well
+	FGameInstancePIEResult InitResult = FGameInstancePIEResult::Success();
+
+	if (InitResult.IsSuccess())
+	{
+		PRAGMA_DISABLE_DEPRECATION_WARNINGS
+		InitResult = InitializePIE(Params.bAnyBlueprintErrors, PIEInstanceIndex, Params.bRunAsDedicated) ?
+			FGameInstancePIEResult::Success() :
+			FGameInstancePIEResult::Failure(NSLOCTEXT("UnrealEd", "Error_CouldntInitInstance", "The game instance failed to Play/Simulate In Editor"));
+		PRAGMA_ENABLE_DEPRECATION_WARNINGS
+	}
+
+	return InitResult;
+}
+
+
+bool UGameInstance::InitializePIE(bool bAnyBlueprintErrors, int32 PIEInstance, bool bRunAsDedicated)
+{
+	// DEPRECATED VERSION
 	return true;
 }
 
-bool UGameInstance::StartPIEGameInstance(ULocalPlayer* LocalPlayer, bool bInSimulateInEditor, bool bAnyBlueprintErrors, bool bStartInSpectatorMode)
+FGameInstancePIEResult UGameInstance::StartPlayInEditorGameInstance(ULocalPlayer* LocalPlayer, const FGameInstancePIEParameters& Params)
 {
 	UEditorEngine* const EditorEngine = CastChecked<UEditorEngine>(GetEngine());
 	ULevelEditorPlaySettings const* PlayInSettings = GetDefault<ULevelEditorPlaySettings>();
@@ -202,8 +222,7 @@ bool UGameInstance::StartPIEGameInstance(ULocalPlayer* LocalPlayer, bool bInSimu
 		}
 		else
 		{
-			FMessageDialog::Open(EAppMsgType::Ok, FText::Format(NSLOCTEXT("UnrealEd", "Error_CouldntLaunchPIEClient", "Couldn't Launch PIE Client: {0}"), FText::FromString(Error)));
-			return false;
+			return FGameInstancePIEResult::Failure(FText::Format(NSLOCTEXT("UnrealEd", "Error_CouldntLaunchPIEClient", "Couldn't Launch PIE Client: {0}"), FText::FromString(Error)));
 		}
 	}
 	else
@@ -214,7 +233,7 @@ bool UGameInstance::StartPIEGameInstance(ULocalPlayer* LocalPlayer, bool bInSimu
 		// make a URL
 		FURL URL;
 		// If the user wants to start in spectator mode, do not use the custom play world for now
-		if (EditorEngine->UserEditedPlayWorldURL.Len() > 0 && !bStartInSpectatorMode)
+		if (EditorEngine->UserEditedPlayWorldURL.Len() > 0 && !Params.bStartInSpectatorMode)
 		{
 			// If the user edited the play world url. Verify that the map name is the same as the currently loaded map.
 			URL = FURL(NULL, *EditorEngine->UserEditedPlayWorldURL, TRAVEL_Absolute);
@@ -227,7 +246,7 @@ bool UGameInstance::StartPIEGameInstance(ULocalPlayer* LocalPlayer, bool bInSimu
 		else
 		{
 			// The user did not edit the url, just build one from scratch.
-			URL = FURL(NULL, *EditorEngine->BuildPlayWorldURL(*PIEMapName, bStartInSpectatorMode), TRAVEL_Absolute);
+			URL = FURL(NULL, *EditorEngine->BuildPlayWorldURL(*PIEMapName, Params.bStartInSpectatorMode), TRAVEL_Absolute);
 		}
 
 		// If a start location is specified, spawn a temporary PlayerStartPIE actor at the start location and use it as the portal.
@@ -235,15 +254,13 @@ bool UGameInstance::StartPIEGameInstance(ULocalPlayer* LocalPlayer, bool bInSimu
 		if (EditorEngine->SpawnPlayFromHereStart(PlayWorld, PlayerStart, EditorEngine->PlayWorldLocation, EditorEngine->PlayWorldRotation) == false)
 		{
 			// failed to create "play from here" playerstart
-			FMessageDialog::Open(EAppMsgType::Ok, NSLOCTEXT("UnrealEd", "Error_FailedCreatePlayFromHerePlayerStart", "Failed to create PlayerStart at desired starting location."));
-			return false;
+			return FGameInstancePIEResult::Failure(NSLOCTEXT("UnrealEd", "Error_FailedCreatePlayFromHerePlayerStart", "Failed to create PlayerStart at desired starting location."));
 		}
 
 		if (!PlayWorld->SetGameMode(URL))
 		{
 			// Setting the game mode failed so bail 
-			FMessageDialog::Open(EAppMsgType::Ok, NSLOCTEXT("UnrealEd", "Error_FailedCreateEditorPreviewWorld", "Failed to create editor preview world."));
-			return false;
+			return FGameInstancePIEResult::Failure(NSLOCTEXT("UnrealEd", "Error_FailedCreateEditorPreviewWorld", "Failed to create editor preview world."));
 		}
 		
 		// Make sure "always loaded" sub-levels are fully loaded
@@ -261,8 +278,7 @@ bool UGameInstance::StartPIEGameInstance(ULocalPlayer* LocalPlayer, bool bInSimu
 			FString Error;
 			if (!LocalPlayer->SpawnPlayActor(URL.ToString(1), Error, PlayWorld))
 			{
-				FMessageDialog::Open(EAppMsgType::Ok, FText::Format(NSLOCTEXT("UnrealEd", "Error_CouldntSpawnPlayer", "Couldn't spawn player: {0}"), FText::FromString(Error)));
-				return false;
+				return FGameInstancePIEResult::Failure(FText::Format(NSLOCTEXT("UnrealEd", "Error_CouldntSpawnPlayer", "Couldn't spawn player: {0}"), FText::FromString(Error)));
 			}
 		}
 
@@ -282,6 +298,24 @@ bool UGameInstance::StartPIEGameInstance(ULocalPlayer* LocalPlayer, bool bInSimu
 		PlayWorld->BeginPlay();
 	}
 
+	// Give the deprecated method a chance to fail as well
+	FGameInstancePIEResult StartResult = FGameInstancePIEResult::Success();
+
+	if (StartResult.IsSuccess())
+	{
+		PRAGMA_DISABLE_DEPRECATION_WARNINGS
+		StartResult = StartPIEGameInstance(LocalPlayer, Params.bSimulateInEditor, Params.bAnyBlueprintErrors, Params.bStartInSpectatorMode) ?
+			FGameInstancePIEResult::Success() :
+			FGameInstancePIEResult::Failure(NSLOCTEXT("UnrealEd", "Error_CouldntInitInstance", "The game instance failed to Play/Simulate In Editor"));
+		PRAGMA_ENABLE_DEPRECATION_WARNINGS
+	}
+
+	return StartResult;
+}
+
+bool UGameInstance::StartPIEGameInstance(ULocalPlayer* LocalPlayer, bool bInSimulateInEditor, bool bAnyBlueprintErrors, bool bStartInSpectatorMode)
+{
+	// DEPRECATED VERSION
 	return true;
 }
 #endif
@@ -770,19 +804,24 @@ void UGameInstance::StartRecordingReplay(const FString& Name, const FString& Fri
 		DemoURL.AddOption(*Option);
 	}
 
-	CurrentWorld->DestroyDemoNetDriver();
+	bool bDestroyedDemoNetDriver = false;
+	if (!CurrentWorld->DemoNetDriver || !CurrentWorld->DemoNetDriver->bRecordMapChanges || !CurrentWorld->DemoNetDriver->IsRecordingPaused())
+	{
+		CurrentWorld->DestroyDemoNetDriver();
+		bDestroyedDemoNetDriver = true; 
 
-	const FName NAME_DemoNetDriver( TEXT( "DemoNetDriver" ) );
+		const FName NAME_DemoNetDriver(TEXT("DemoNetDriver"));
 
-	if ( !GEngine->CreateNamedNetDriver( CurrentWorld, NAME_DemoNetDriver, NAME_DemoNetDriver ) )
-	{
-		UE_LOG( LogDemo, Warning, TEXT( "RecordReplay: failed to create demo net driver!" ) );
-		return;
-	}
+		if (!GEngine->CreateNamedNetDriver(CurrentWorld, NAME_DemoNetDriver, NAME_DemoNetDriver))
+		{
+			UE_LOG(LogDemo, Warning, TEXT("RecordReplay: failed to create demo net driver!"));
+			return;
+		}
 
-	CurrentWorld->DemoNetDriver = Cast< UDemoNetDriver >( GEngine->FindNamedNetDriver( CurrentWorld, NAME_DemoNetDriver ) );
+		CurrentWorld->DemoNetDriver = Cast< UDemoNetDriver >(GEngine->FindNamedNetDriver(CurrentWorld, NAME_DemoNetDriver));
+	}
 
-	check( CurrentWorld->DemoNetDriver != NULL );
+	check(CurrentWorld->DemoNetDriver != nullptr);
 
 	CurrentWorld->DemoNetDriver->SetWorld( CurrentWorld );
 
@@ -795,15 +834,23 @@ void UGameInstance::StartRecordingReplay(const FString& Name, const FString& Fri
 
 	FString Error;
 
-	if ( !CurrentWorld->DemoNetDriver->InitListen( CurrentWorld, DemoURL, false, Error ) )
+	if (bDestroyedDemoNetDriver)
 	{
-		UE_LOG( LogDemo, Warning, TEXT( "Demo recording failed: %s" ), *Error );
-		CurrentWorld->DemoNetDriver = NULL;
+		if (!CurrentWorld->DemoNetDriver->InitListen(CurrentWorld, DemoURL, false, Error))
+		{
+			UE_LOG(LogDemo, Warning, TEXT("Demo recording - InitListen failed: %s"), *Error);
+			CurrentWorld->DemoNetDriver = NULL;
+			return;
+		}
 	}
-	else
+	else if (!CurrentWorld->DemoNetDriver->ContinueListen(DemoURL))
 	{
-		UE_LOG(LogDemo, Log, TEXT( "Num Network Actors: %i" ), CurrentWorld->DemoNetDriver->GetNetworkObjectList().GetObjects().Num() );
+		UE_LOG(LogDemo, Warning, TEXT("Demo recording - ContinueListen failed"));
+		CurrentWorld->DemoNetDriver = NULL;
+		return;
 	}
+
+	UE_LOG(LogDemo, Log, TEXT( "Num Network Actors: %i" ), CurrentWorld->DemoNetDriver->GetNetworkObjectList().GetObjects().Num() );
 }
 
 void UGameInstance::StopRecordingReplay()
@@ -901,6 +948,28 @@ bool UGameInstance::IsDedicatedServerInstance() const
 	}
 }
 
+bool UGameInstance::ClientTravelToSession(int32 ControllerId, FName InSessionName)
+{
+	UWorld* World = GetWorld();
+
+	FString URL;
+	if (UOnlineEngineInterface::Get()->GetResolvedConnectString(World, InSessionName, URL))
+	{
+		APlayerController* PC = UGameplayStatics::GetPlayerController(World, ControllerId);
+		if (PC)
+		{
+			PC->ClientTravel(URL, TRAVEL_Absolute);
+			return true;
+		}
+	}
+	else
+	{
+		UE_LOG(LogGameSession, Warning, TEXT("Failed to resolve session connect string for %s"), *InSessionName.ToString());
+	}
+
+	return false;
+}
+
 void UGameInstance::NotifyPreClientTravel(const FString& PendingURL, ETravelType TravelType, bool bIsSeamlessTravel)
 {
 	OnNotifyPreClientTravel().Broadcast(PendingURL, TravelType, bIsSeamlessTravel);
diff --git a/Engine/Source/Runtime/Engine/Private/GameMode.cpp b/Engine/Source/Runtime/Engine/Private/GameMode.cpp
--- a/Engine/Source/Runtime/Engine/Private/GameMode.cpp
+++ b/Engine/Source/Runtime/Engine/Private/GameMode.cpp
@@ -663,8 +663,6 @@ void AGameMode::AddInactivePlayer(APlayerState* PlayerState, APlayerController*
 			}
 		}
 	}
-
-	PlayerState->OnDeactivated();
 }
 
 bool AGameMode::FindInactivePlayer(APlayerController* PC)
diff --git a/Engine/Source/Runtime/Engine/Private/GameSession.cpp b/Engine/Source/Runtime/Engine/Private/GameSession.cpp
--- a/Engine/Source/Runtime/Engine/Private/GameSession.cpp
+++ b/Engine/Source/Runtime/Engine/Private/GameSession.cpp
@@ -12,7 +12,7 @@
 #include "GameFramework/GameSession.h"
 #include "GameFramework/GameModeBase.h"
 
-DEFINE_LOG_CATEGORY_STATIC(LogGameSession, Log, All);
+DEFINE_LOG_CATEGORY(LogGameSession);
 
 static TAutoConsoleVariable<int32> CVarMaxPlayersOverride( TEXT( "net.MaxPlayersOverride" ), 0, TEXT( "If greater than 0, will override the standard max players count. Useful for testing full servers." ) );
 
@@ -359,28 +359,6 @@ void AGameSession::ReturnToMainMenuHost()
 	}
 }
 
-bool AGameSession::TravelToSession(int32 ControllerId, FName InSessionName)
-{
-	UWorld* World = GetWorld();
-
-	FString URL;
-	if (UOnlineEngineInterface::Get()->GetResolvedConnectString(World, InSessionName, URL))
-	{
-		APlayerController* PC = UGameplayStatics::GetPlayerController(World, ControllerId);
-		if (PC)
-		{
-			PC->ClientTravel(URL, TRAVEL_Absolute);
-			return true;
-		}
-	}
-	else
-	{
-		UE_LOG(LogGameSession, Warning, TEXT("Failed to resolve session connect string for %s"), *InSessionName.ToString());
-	}
-
-	return false;
-}
-
 void AGameSession::PostSeamlessTravel()
 {
 }
diff --git a/Engine/Source/Runtime/Engine/Private/GameViewportClient.cpp b/Engine/Source/Runtime/Engine/Private/GameViewportClient.cpp
--- a/Engine/Source/Runtime/Engine/Private/GameViewportClient.cpp
+++ b/Engine/Source/Runtime/Engine/Private/GameViewportClient.cpp
@@ -2721,6 +2721,9 @@ bool UGameViewportClient::HandleToggleFullscreenCommand()
 		}
 	}
 
+	int32 ResolutionX = GSystemResolution.ResX;
+	int32 ResolutionY = GSystemResolution.ResY;
+
 	// Make sure the user's settings are updated after pressing Alt+Enter to toggle fullscreen.  Note
 	// that we don't need to "apply" the setting change, as we already did that above directly.
 	UGameEngine* GameEngine = Cast<UGameEngine>( GEngine );
@@ -2731,10 +2734,17 @@ bool UGameViewportClient::HandleToggleFullscreenCommand()
 		{
 			UserSettings->SetFullscreenMode( FullScreenMode );
 			UserSettings->ConfirmVideoMode();
+
+			ResolutionX = UserSettings->GetScreenResolution().X;
+			ResolutionY = UserSettings->GetScreenResolution().Y;
+			UGameEngine::ConditionallyOverrideSettings(ResolutionX, ResolutionY, FullScreenMode);
 		}
 	}
 
-	FSystemResolution::RequestResolutionChange(GSystemResolution.ResX, GSystemResolution.ResY, FullScreenMode);
+	FSystemResolution::RequestResolutionChange(ResolutionX, ResolutionY, FullScreenMode);
+
+	ToggleFullscreenDelegate.Broadcast(FullScreenMode != EWindowMode::Windowed);
+
 	return true;
 }
 
diff --git a/Engine/Source/Runtime/Engine/Private/Level.cpp b/Engine/Source/Runtime/Engine/Private/Level.cpp
--- a/Engine/Source/Runtime/Engine/Private/Level.cpp
+++ b/Engine/Source/Runtime/Engine/Private/Level.cpp
@@ -1588,6 +1588,14 @@ void ULevel::InitializeNetworkActors()
 				if (Actor->bNetLoadOnClient)
 				{
 					Actor->bNetStartup = true;
+
+					for (UActorComponent* Component : Actor->GetComponents())
+					{
+						if (Component)
+						{
+							Component->SetIsNetStartupComponent(true);
+						}
+					}
 				}
 
 				if (!bIsServer)
diff --git a/Engine/Source/Runtime/Engine/Private/LevelStreaming.cpp b/Engine/Source/Runtime/Engine/Private/LevelStreaming.cpp
--- a/Engine/Source/Runtime/Engine/Private/LevelStreaming.cpp
+++ b/Engine/Source/Runtime/Engine/Private/LevelStreaming.cpp
@@ -507,6 +507,11 @@ bool ULevelStreaming::RequestLevel(UWorld* PersistentWorld, bool bAllowLevelLoad
 			// Editor immediately blocks on load and we also block if background level streaming is disabled.
 			if (BlockPolicy == AlwaysBlock || (ShouldBeAlwaysLoaded() && BlockPolicy != NeverBlock))
 			{
+				if (IsAsyncLoading())
+				{
+					UE_LOG(LogStreaming, Log, TEXT("ULevelStreaming::RequestLevel(%s) is flushing async loading"), *DesiredPackageName.ToString());
+				}
+
 				// Finish all async loading.
 				FlushAsyncLoading();
 			}
diff --git a/Engine/Source/Runtime/Engine/Private/Materials/HLSLMaterialTranslator.h b/Engine/Source/Runtime/Engine/Private/Materials/HLSLMaterialTranslator.h
--- a/Engine/Source/Runtime/Engine/Private/Materials/HLSLMaterialTranslator.h
+++ b/Engine/Source/Runtime/Engine/Private/Materials/HLSLMaterialTranslator.h
@@ -1683,6 +1683,13 @@ protected:
 
 	virtual int32 CallExpression(FMaterialExpressionKey ExpressionKey,FMaterialCompiler* Compiler) override
 	{
+		// For any translated result not relying on material attributes, we can
+		// discard the attribute ID from the key to allow result sharing
+		if (ExpressionKey.Expression && !ExpressionKey.Expression->IsResultMaterialAttributes(ExpressionKey.OutputIndex))
+		{
+			ExpressionKey.MaterialAttributeID = FGuid(0,0,0,0);
+		}
+
 		// Check if this expression has already been translated.
 		check(ShaderFrequency < SF_NumFrequencies);
 		auto& CurrentFunctionStack = FunctionStacks[ShaderFrequency];
diff --git a/Engine/Source/Runtime/Engine/Private/NetConnection.cpp b/Engine/Source/Runtime/Engine/Private/NetConnection.cpp
--- a/Engine/Source/Runtime/Engine/Private/NetConnection.cpp
+++ b/Engine/Source/Runtime/Engine/Private/NetConnection.cpp
@@ -556,15 +556,25 @@ void UNetConnection::ReceivedRawPacket( void* InData, int32 Count )
 	{
 		const ProcessedPacket UnProcessedPacket = Handler->Incoming(Data, Count);
 
-		Count = FMath::DivideAndRoundUp(UnProcessedPacket.CountBits, 8);
-
-		if (Count > 0)
+		if (!UnProcessedPacket.bError)
 		{
-			Data = UnProcessedPacket.Data;
+			Count = FMath::DivideAndRoundUp(UnProcessedPacket.CountBits, 8);
+
+			if (Count > 0)
+			{
+				Data = UnProcessedPacket.Data;
+			}
+			// This packed has been consumed
+			else
+			{
+				return;
+			}
 		}
-		// This packed has been consumed
 		else
 		{
+			CLOSE_CONNECTION_DUE_TO_SECURITY_VIOLATION(this, ESecurityEvent::Malformed_Packet,
+														TEXT("Packet failed PacketHandler processing."));
+
 			return;
 		}
 	}
@@ -913,13 +923,13 @@ void UNetConnection::ReceivedPacket( FBitReader& Reader )
 				OutLagPacketId[Index] = -1;		// Only use the ack once
 
 #if !UE_BUILD_SHIPPING
-				if ( CVarPingDisplayServerTime.GetValueOnGameThread() > 0 )
+				if ( CVarPingDisplayServerTime.GetValueOnAnyThread() > 0 )
 				{
 					UE_LOG( LogNetTraffic, Warning, TEXT( "ServerFrameTime: %2.2f" ), ServerFrameTime * 1000.0f );
 				}
 
 				const float GameTime	= ServerFrameTime + FrameTime;
-				const float RTT			= ( FPlatformTime::Seconds() - OutLagTime[Index] ) - ( CVarPingExcludeFrameTime.GetValueOnGameThread() ? GameTime : 0.0f );
+				const float RTT			= ( FPlatformTime::Seconds() - OutLagTime[Index] ) - ( CVarPingExcludeFrameTime.GetValueOnAnyThread() ? GameTime : 0.0f );
 				const float NewLag		= FMath::Max( RTT, 0.0f );
 #else
 				const float NewLag		= FPlatformTime::Seconds() - OutLagTime[Index];
@@ -1155,7 +1165,7 @@ void UNetConnection::ReceivedPacket( FBitReader& Reader )
 				if( !Driver->Notify->NotifyAcceptingChannel( Channel ) )
 				{
 					// Channel refused, so close it, flush it, and delete it.
-					UE_LOG(LogNetDormancy, Verbose, TEXT("      NotifyAcceptingChannel Failed! Channel: %s"), *Channel->Describe() );
+					UE_LOG(LogNet, Verbose, TEXT("      NotifyAcceptingChannel Failed! Channel: %s"), *Channel->Describe() );
 
 					FOutBunch CloseBunch( Channel, 1 );
 					check(!CloseBunch.IsError());
@@ -1510,7 +1520,9 @@ float UNetConnection::GetTimeoutValue()
 #if !UE_BUILD_SHIPPING
 	if (Driver->bNoTimeouts)
 	{
-		return MAX_FLT;
+		// APlayerController depends on this timeout to destroy itself and free up
+		// its resources, so we have to handle this case here as well
+		return bPendingDestroy ? 2.f : MAX_FLT;
 	}
 #endif
 
@@ -1657,7 +1669,7 @@ void UNetConnection::Tick()
 		checkf(ChannelsToTick.Num() <= OpenChannels.Num(), TEXT("More ticking channels (%d) than open channels (%d) for net connection!"), ChannelsToTick.Num(), OpenChannels.Num())
 
 		// Tick the channels.
-		if (CVarTickAllOpenChannels.GetValueOnGameThread() == 0)
+		if (CVarTickAllOpenChannels.GetValueOnAnyThread() == 0)
 		{
 			for( int32 i=ChannelsToTick.Num()-1; i>=0; i-- )
 			{
@@ -2025,7 +2037,7 @@ void UNetConnection::FlushDormancy(class AActor* Actor)
 void UNetConnection::FlushDormancyForObject( UObject* Object )
 {
 	static const auto ValidateCVar = IConsoleManager::Get().FindTConsoleVariableDataInt( TEXT( "net.DormancyValidate" ) );
-	const bool ValidateProperties = ( ValidateCVar && ValidateCVar->GetValueOnGameThread() == 1 );
+	const bool ValidateProperties = ( ValidateCVar && ValidateCVar->GetValueOnAnyThread() == 1 );
 
 	TSharedRef< FObjectReplicator > * Replicator = DormantReplicatorMap.Find( Object );
 
diff --git a/Engine/Source/Runtime/Engine/Private/NetworkDriver.cpp b/Engine/Source/Runtime/Engine/Private/NetworkDriver.cpp
--- a/Engine/Source/Runtime/Engine/Private/NetworkDriver.cpp
+++ b/Engine/Source/Runtime/Engine/Private/NetworkDriver.cpp
@@ -216,7 +216,7 @@ void UNetDriver::AssertValid()
 
 /*static*/ bool UNetDriver::IsAdaptiveNetUpdateFrequencyEnabled()
 {
-	const bool bUseAdapativeNetFrequency = CVarUseAdaptiveNetUpdateFrequency.GetValueOnGameThread() > 0;
+	const bool bUseAdapativeNetFrequency = CVarUseAdaptiveNetUpdateFrequency.GetValueOnAnyThread() > 0;
 	return bUseAdapativeNetFrequency;
 }
 
@@ -673,12 +673,12 @@ void UNetDriver::TickFlush(float DeltaSeconds)
 		FlushHandler();
 	}
 
-	if (CVarNetDormancyDraw.GetValueOnGameThread() > 0)
+	if (CVarNetDormancyDraw.GetValueOnAnyThread() > 0)
 	{
 		DrawNetDriverDebug();
 	}
 
-	if ( CVarOptimizedRemapping.GetValueOnGameThread() && GuidCache.IsValid() )
+	if ( CVarOptimizedRemapping.GetValueOnAnyThread() && GuidCache.IsValid() )
 	{
 		SCOPE_CYCLE_COUNTER( STAT_NetUpdateUnmappedObjectsTime );
 
@@ -1323,7 +1323,7 @@ void UNetDriver::InternalProcessRemoteFunction
 #if !(UE_BUILD_SHIPPING || UE_BUILD_TEST)
 	extern TAutoConsoleVariable< int32 > CVarNetReliableDebug;
 
-	if ( CVarNetReliableDebug.GetValueOnGameThread() > 0 )
+	if ( CVarNetReliableDebug.GetValueOnAnyThread() > 0 )
 	{
 		Bunch.DebugString = FString::Printf( TEXT( "%.2f RPC: %s - %s" ), Connection->Driver->Time, *Actor->GetName(), *Function->GetName() );
 	}
@@ -1404,7 +1404,7 @@ void UNetDriver::InternalProcessRemoteFunction
 	}
 
 	static const auto CVar = IConsoleManager::Get().FindTConsoleVariableDataInt(TEXT("net.RPC.Debug"));
-	bool LogAsWarning = (CVar && CVar->GetValueOnGameThread() == 1);
+	bool LogAsWarning = (CVar && CVar->GetValueOnAnyThread() == 1);
 
 	// Send the bunch.
 	if( Bunch.IsError() )
@@ -2037,7 +2037,7 @@ void UNetDriver::FlushActorDormancy(AActor* Actor)
 	// way too, since we dont have to check every dormant actor in ::ServerReplicateActor to see if it needs to go out of dormancy
 
 #if WITH_SERVER_CODE
-	if (CVarSetNetDormancyEnabled.GetValueOnGameThread() == 0)
+	if (CVarSetNetDormancyEnabled.GetValueOnAnyThread() == 0)
 		return;
 
 	check(Actor);
@@ -2570,7 +2570,7 @@ static FORCEINLINE_DEBUGGABLE bool IsActorDormant( const AActor* Actor, const UN
 	{
 		// net.DormancyValidate can be set to 2 to validate dormant actor properties on every replicate
 		// (this could be moved to be done every tick instead of every net update if necessary, but seems excessive)
-		if ( CVarNetDormancyValidate.GetValueOnGameThread() == 2 )
+		if ( CVarNetDormancyValidate.GetValueOnAnyThread() == 2 )
 		{
 			const TSharedRef<FObjectReplicator>* Replicator = Connection->DormantReplicatorMap.Find( Actor );
 
@@ -3221,7 +3221,7 @@ void UNetDriver::DrawNetDriverDebug()
 		return;
 	}
 
-	const float CullDistSqr = FMath::Square(CVarNetDormancyDrawCullDistance.GetValueOnGameThread());
+	const float CullDistSqr = FMath::Square(CVarNetDormancyDrawCullDistance.GetValueOnAnyThread());
 
 	for (FActorIterator It(LocalWorld); It; ++It)
 	{
diff --git a/Engine/Source/Runtime/Engine/Private/ObjectLibrary.cpp b/Engine/Source/Runtime/Engine/Private/ObjectLibrary.cpp
--- a/Engine/Source/Runtime/Engine/Private/ObjectLibrary.cpp
+++ b/Engine/Source/Runtime/Engine/Private/ObjectLibrary.cpp
@@ -11,6 +11,7 @@ UObjectLibrary::UObjectLibrary(const FObjectInitializer& ObjectInitializer)
 	bIsFullyLoaded = false;
 	bUseWeakReferences = false;
 	bIncludeOnlyOnDiskAssets = true;
+	bRecursivePaths = true;
 
 #if WITH_EDITOR
 	if ( !HasAnyFlags(RF_ClassDefaultObject) )
@@ -352,7 +353,7 @@ int32 UObjectLibrary::LoadAssetDataFromPaths(const TArray<FString>& Paths, bool
 		ARFilter.PackagePaths.Add(FName(*Paths[PathIndex]));
 	}
 
-	ARFilter.bRecursivePaths = true;
+	ARFilter.bRecursivePaths = bRecursivePaths;
 	ARFilter.bIncludeOnlyOnDiskAssets = bIncludeOnlyOnDiskAssets;
 
 	AssetDataList.Empty();
@@ -412,7 +413,7 @@ int32 UObjectLibrary::LoadBlueprintAssetDataFromPaths(const TArray<FString>& Pat
 		ARFilter.PackagePaths.Add(FName(*Paths[PathIndex]));
 	}
 	
-	ARFilter.bRecursivePaths = true;
+	ARFilter.bRecursivePaths = bRecursivePaths;
 	ARFilter.bIncludeOnlyOnDiskAssets = bIncludeOnlyOnDiskAssets;
 
 	/* GetDerivedClassNames doesn't work yet
diff --git a/Engine/Source/Runtime/Engine/Private/PackageMapClient.cpp b/Engine/Source/Runtime/Engine/Private/PackageMapClient.cpp
--- a/Engine/Source/Runtime/Engine/Private/PackageMapClient.cpp
+++ b/Engine/Source/Runtime/Engine/Private/PackageMapClient.cpp
@@ -510,7 +510,7 @@ void UPackageMapClient::InternalWriteObject( FArchive & Ar, FNetworkGUID NetGUID
 
 	const bool bNoLoad = !CanClientLoadObject( Object, NetGUID );
 
-	if ( CVarAllowAsyncLoading.GetValueOnGameThread() > 0 && IsNetGUIDAuthority() && !GuidCache->IsExportingNetGUIDBunch && !bNoLoad )
+	if ( CVarAllowAsyncLoading.GetValueOnAnyThread() > 0 && IsNetGUIDAuthority() && !GuidCache->IsExportingNetGUIDBunch && !bNoLoad )
 	{
 		// These are guids that must exist on the client in a package
 		// The client needs to know about these so it can determine if it has finished loading them
@@ -798,7 +798,7 @@ FNetworkGUID UPackageMapClient::InternalLoadObject( FArchive & Ar, UObject *& Ob
 				return NetGUID;
 			}
 
-			if ( NetworkChecksum != 0 && GuidCache->NetworkChecksumMode == FNetGUIDCache::NETCHECKSUM_SaveAndUse && !CVarIgnoreNetworkChecksumMismatch.GetValueOnGameThread() )
+			if ( NetworkChecksum != 0 && GuidCache->NetworkChecksumMode == FNetGUIDCache::NETCHECKSUM_SaveAndUse && !CVarIgnoreNetworkChecksumMismatch.GetValueOnAnyThread() )
 			{
 				const uint32 CompareNetworkChecksum = GuidCache->GetNetworkChecksum( Object );
 
@@ -2314,7 +2314,7 @@ UObject* FNetGUIDCache::GetObjectFromNetGUID( const FNetworkGUID& NetGUID, const
 			//	3. We're actually suppose to load (levels don't load here for example)
 			//		(Refer to CanClientLoadObject, which is where we protect clients from trying to load levels)
 			
-			if ( CVarAllowAsyncLoading.GetValueOnGameThread() > 0 )
+			if ( CVarAllowAsyncLoading.GetValueOnAnyThread() > 0 )
 			{
 				if (!PendingAsyncPackages.Contains(CacheObjectPtr->PathName))
 				{
@@ -2347,7 +2347,7 @@ UObject* FNetGUIDCache::GetObjectFromNetGUID( const FNetworkGUID& NetGUID, const
 			//	2. Someone else started async loading the outer package, and it's not fully loaded yet
 			Object = StaticLoadObject( UObject::StaticClass(), ObjOuter, *CacheObjectPtr->PathName.ToString(), NULL, LOAD_NoWarn );
 
-			if ( CVarAllowAsyncLoading.GetValueOnGameThread() > 0 )
+			if ( CVarAllowAsyncLoading.GetValueOnAnyThread() > 0 )
 			{
 				UE_LOG( LogNetPackageMap, Error, TEXT( "GetObjectFromNetGUID: Forced blocking load. Path: %s, NetGUID: %s" ), *CacheObjectPtr->PathName.ToString(), *NetGUID.ToString() );
 			}
@@ -2380,7 +2380,7 @@ UObject* FNetGUIDCache::GetObjectFromNetGUID( const FNetworkGUID& NetGUID, const
 		if (!Package->IsFullyLoaded() 
 			&& !Package->HasAnyPackageFlags(EPackageFlags::PKG_CompiledIn)) //TODO: dependencies of CompiledIn could still be loaded asynchronously. Are they necessary at this point??
 		{
-			if ( CVarAllowAsyncLoading.GetValueOnGameThread() > 0 )
+			if ( CVarAllowAsyncLoading.GetValueOnAnyThread() > 0 )
 			{
 				if (Package->HasAnyInternalFlags(EInternalObjectFlags::AsyncLoading))
 				{
@@ -2401,13 +2401,18 @@ UObject* FNetGUIDCache::GetObjectFromNetGUID( const FNetworkGUID& NetGUID, const
 			else
 			{
 				// If package isn't fully loaded, flush async loading now
+				if (IsAsyncLoading())
+				{
+					UE_LOG(LogNetPackageMap, Log, TEXT("FNetGUIDCache::GetObjectFromNetGUID for package %s is flushing async loading"), *Package->GetPathName());
+				}
+
 				FlushAsyncLoading(); 
-				check(Package->IsFullyLoaded());
+				checkf(Package->IsFullyLoaded(), TEXT("Package %s did not become loaded after FlushAsyncLoading"), *Package->GetPathName());
 			}
 		}
 	}
 
-	if ( CacheObjectPtr->NetworkChecksum != 0 && !CVarIgnoreNetworkChecksumMismatch.GetValueOnGameThread() )
+	if ( CacheObjectPtr->NetworkChecksum != 0 && !CVarIgnoreNetworkChecksumMismatch.GetValueOnAnyThread() )
 	{
 		const uint32 NetworkChecksum = GetNetworkChecksum( Object );
 
diff --git a/Engine/Source/Runtime/Engine/Private/Particles/ParticleComponents.cpp b/Engine/Source/Runtime/Engine/Private/Particles/ParticleComponents.cpp
--- a/Engine/Source/Runtime/Engine/Private/Particles/ParticleComponents.cpp
+++ b/Engine/Source/Runtime/Engine/Private/Particles/ParticleComponents.cpp
@@ -2181,7 +2181,7 @@ void UParticleSystem::PostLoad()
 		if (Emitter == NULL)
 		{
 			// Empty emitter slots are ok with cooked content.
-			if( !FPlatformProperties::RequiresCookedData() )
+			if( !FPlatformProperties::RequiresCookedData() && !GIsServer)
 			{
 				UE_LOG(LogParticles, Warning, TEXT("ParticleSystem contains empty emitter slots - %s"), *GetFullName());
 			}
diff --git a/Engine/Source/Runtime/Engine/Private/Particles/ParticleModules.cpp b/Engine/Source/Runtime/Engine/Private/Particles/ParticleModules.cpp
--- a/Engine/Source/Runtime/Engine/Private/Particles/ParticleModules.cpp
+++ b/Engine/Source/Runtime/Engine/Private/Particles/ParticleModules.cpp
@@ -3135,14 +3135,14 @@ uint64 UParticleModuleLight::SpawnHQLight(const FLightParticlePayload& Payload,
 			PointLightComponent->SetupAttachment(RootComponent);
 		}			
 		PointLightComponent->CreationMethod = ParticleSystem->CreationMethod;
-		PointLightComponent->RegisterComponent();
-		Owner->HighQualityLights.Add(PointLightComponent);
-			
-			
+		PointLightComponent->LightingChannels = LightingChannels;
 		PointLightComponent->bUseInverseSquaredFalloff = bUseInverseSquaredFalloff;
 		PointLightComponent->bAffectTranslucentLighting = bAffectsTranslucency;
 		PointLightComponent->SetCastShadows(bShadowCastingLights);
 
+		PointLightComponent->RegisterComponent();
+		Owner->HighQualityLights.Add(PointLightComponent);
+
 		int32 ScreenAlignment;
 		FVector ComponentScale;
 		Owner->GetScreenAlignmentAndScale(ScreenAlignment, ComponentScale);
diff --git a/Engine/Source/Runtime/Engine/Private/PlayerController.cpp b/Engine/Source/Runtime/Engine/Private/PlayerController.cpp
--- a/Engine/Source/Runtime/Engine/Private/PlayerController.cpp
+++ b/Engine/Source/Runtime/Engine/Private/PlayerController.cpp
@@ -1195,6 +1195,11 @@ void APlayerController::ClientSetHUD_Implementation(TSubclassOf<AHUD> NewHUDClas
 
 void APlayerController::CleanupPlayerState()
 {
+	if (PlayerState)
+	{
+		// By default this destroys it, but games can override
+		PlayerState->OnDeactivated();
+	}
 	PlayerState = NULL;
 }
 
diff --git a/Engine/Source/Runtime/Engine/Private/RepLayout.cpp b/Engine/Source/Runtime/Engine/Private/RepLayout.cpp
--- a/Engine/Source/Runtime/Engine/Private/RepLayout.cpp
+++ b/Engine/Source/Runtime/Engine/Private/RepLayout.cpp
@@ -1234,7 +1234,7 @@ void FRepLayout::SendProperties(
 	TArray< uint16 > &			Changed ) const
 {
 #ifdef ENABLE_PROPERTY_CHECKSUMS
-	const bool bDoChecksum = CVarDoPropertyChecksum.GetValueOnGameThread() == 1;
+	const bool bDoChecksum = CVarDoPropertyChecksum.GetValueOnAnyThread() == 1;
 #else
 	const bool bDoChecksum = false;
 #endif
@@ -1509,7 +1509,7 @@ void FRepLayout::SendProperties_BackwardsCompatible(
 	FBitWriterMark Mark( Writer );
 
 #ifdef ENABLE_PROPERTY_CHECKSUMS
-	const bool bDoChecksum = CVarDoPropertyChecksum.GetValueOnGameThread() == 1;
+	const bool bDoChecksum = CVarDoPropertyChecksum.GetValueOnAnyThread() == 1;
 	Writer.WriteBit( bDoChecksum ? 1 : 0 );
 #else
 	const bool bDoChecksum = false;
@@ -3381,6 +3381,7 @@ void FRepLayout::RebuildConditionalProperties( FRepState * RESTRICT	RepState, co
 	RepState->ConditionMap[COND_InitialOrOwner]				= bIsInitial || bIsOwner;
 	RepState->ConditionMap[COND_ReplayOrOwner]				= bIsReplay || bIsOwner;
 	RepState->ConditionMap[COND_ReplayOnly]					= bIsReplay;
+	RepState->ConditionMap[COND_SkipReplay]					= !bIsReplay;
 
 	RepState->ConditionMap[COND_Custom]						= true;
 
diff --git a/Engine/Source/Runtime/Engine/Private/Scalability.cpp b/Engine/Source/Runtime/Engine/Private/Scalability.cpp
--- a/Engine/Source/Runtime/Engine/Private/Scalability.cpp
+++ b/Engine/Source/Runtime/Engine/Private/Scalability.cpp
@@ -267,6 +267,12 @@ static void SetGroupQualityLevel(const TCHAR* InGroupName, int32 InQualityLevel,
 	}	
 }
 
+float GetResolutionScreenPercentage()
+{
+	static IConsoleVariable* CVar = IConsoleManager::Get().FindConsoleVariable(TEXT("r.ScreenPercentage"));
+	return CVar->GetFloat();
+}
+
 static void SetResolutionQualityLevel(float InResolutionQualityLevel)
 {
 	InResolutionQualityLevel = FMath::Clamp(InResolutionQualityLevel, Scalability::MinResolutionScale, Scalability::MaxResolutionScale);
diff --git a/Engine/Source/Runtime/Engine/Private/SkeletalMeshComponentPhysics.cpp b/Engine/Source/Runtime/Engine/Private/SkeletalMeshComponentPhysics.cpp
--- a/Engine/Source/Runtime/Engine/Private/SkeletalMeshComponentPhysics.cpp
+++ b/Engine/Source/Runtime/Engine/Private/SkeletalMeshComponentPhysics.cpp
@@ -2614,7 +2614,7 @@ void USkeletalMeshComponent::GetWindForCloth_GameThread(FVector& WindDirection,
 			float WindSpeed;
 			float WindMinGust;
 			float WindMaxGust;
-			World->Scene->GetWindParameters(Position, WindDirection, WindSpeed, WindMinGust, WindMaxGust);
+			World->Scene->GetWindParameters_GameThread(Position, WindDirection, WindSpeed, WindMinGust, WindMaxGust);
 
 			WindDirection *= WindUnitAmout * WindSpeed;
 			WindAdaption = FMath::Rand() % 20 * 0.1f; // make range from 0 to 2
diff --git a/Engine/Source/Runtime/Engine/Private/SpectatorPawnMovement.cpp b/Engine/Source/Runtime/Engine/Private/SpectatorPawnMovement.cpp
--- a/Engine/Source/Runtime/Engine/Private/SpectatorPawnMovement.cpp
+++ b/Engine/Source/Runtime/Engine/Private/SpectatorPawnMovement.cpp
@@ -29,15 +29,7 @@ void USpectatorPawnMovement::TickComponent(float DeltaTime, enum ELevelTick Tick
 	}
 
 	// We might want to ignore time dilation
-	float AdjustedDeltaTime = DeltaTime;
-	if (bIgnoreTimeDilation)
-	{
-		const float WorldTimeDilation = PawnOwner->GetActorTimeDilation();
-		if (WorldTimeDilation > KINDA_SMALL_NUMBER)
-		{
-			AdjustedDeltaTime = DeltaTime / WorldTimeDilation;
-		}
-	}
+	const float AdjustedDeltaTime = bIgnoreTimeDilation ? FApp::GetDeltaTime() : DeltaTime;
 
 	Super::TickComponent(AdjustedDeltaTime, TickType, ThisTickFunction);
 };
diff --git a/Engine/Source/Runtime/Engine/Private/StreamableManager.cpp b/Engine/Source/Runtime/Engine/Private/StreamableManager.cpp
--- a/Engine/Source/Runtime/Engine/Private/StreamableManager.cpp
+++ b/Engine/Source/Runtime/Engine/Private/StreamableManager.cpp
@@ -123,7 +123,15 @@ UObject* FStreamableManager::SynchronousLoad(FStringAssetReference const& InTarg
 	FStreamable* Existing = StreamableItems.FindRef(TargetName);
 	while (Existing && Existing->bAsyncLoadRequestOutstanding)
 	{
-		UE_LOG(LogStreamableManager, Verbose, TEXT("     Flushing async load for %s"), *TargetName.ToString());
+		if (IsAsyncLoading())
+		{
+			UE_LOG(LogStreamableManager, Log, TEXT("     Flushing async load for %s"), *TargetName.ToString());
+		}
+		else
+		{
+			UE_LOG(LogStreamableManager, Verbose, TEXT("     Flushing async load for %s"), *TargetName.ToString());
+		}
+
 		check(!Existing->bAsyncUnloadRequestOutstanding); // we should not be both loading and unloading
 		FlushAsyncLoading(); 
 		// the async request might have found a redirect and retried
diff --git a/Engine/Source/Runtime/Engine/Private/Timeline.cpp b/Engine/Source/Runtime/Engine/Private/Timeline.cpp
--- a/Engine/Source/Runtime/Engine/Private/Timeline.cpp
+++ b/Engine/Source/Runtime/Engine/Private/Timeline.cpp
@@ -835,6 +835,11 @@ void UTimelineComponent::SetTimelineFinishedFunc(FOnTimelineEvent NewTimelineFin
 	TheTimeline.SetTimelineFinishedFunc(NewTimelineFinishedFunc);
 }
 
+void UTimelineComponent::SetTimelineFinishedFunc(FOnTimelineEventStatic NewTimelineFinishedFunc)
+{
+	TheTimeline.SetTimelineFinishedFunc(NewTimelineFinishedFunc);
+}
+
 ETimelineSigType UTimelineComponent::GetTimelineSignatureForFunction(const UFunction* Func)
 {
 	if(Func != NULL)
diff --git a/Engine/Source/Runtime/Engine/Private/UnrealEngine.cpp b/Engine/Source/Runtime/Engine/Private/UnrealEngine.cpp
--- a/Engine/Source/Runtime/Engine/Private/UnrealEngine.cpp
+++ b/Engine/Source/Runtime/Engine/Private/UnrealEngine.cpp
@@ -6994,7 +6994,7 @@ FGuid UEngine::GetPackageGuid(FName PackageName, bool bForPIE)
 {
 	FGuid Result(0,0,0,0);
 
-	BeginLoad();
+	BeginLoad(*PackageName.ToString());
 	uint32 LoadFlags = LOAD_NoWarn | LOAD_NoVerify;
 	if (bForPIE)
 	{
@@ -10300,6 +10300,9 @@ void UEngine::MovePendingLevel(FWorldContext &Context)
 
 		FLevelCollection& SourceLevels = Context.World()->FindOrAddCollectionByType(ELevelCollectionType::DynamicSourceLevels);
 		SourceLevels.SetNetDriver(NetDriver);
+			
+		FLevelCollection& StaticLevels = Context.World()->FindOrAddCollectionByType(ELevelCollectionType::StaticLevels);
+		StaticLevels.SetNetDriver(NetDriver);
 	}
 
 	// Attach the DemoNetDriver to the world if there is one
diff --git a/Engine/Source/Runtime/Engine/Private/UserActivityTracking.cpp b/Engine/Source/Runtime/Engine/Private/UserActivityTracking.cpp
--- a/Engine/Source/Runtime/Engine/Private/UserActivityTracking.cpp
+++ b/Engine/Source/Runtime/Engine/Private/UserActivityTracking.cpp
@@ -9,6 +9,7 @@ FUserActivity FUserActivityTracking::UserActivity;
 
 void FUserActivityTracking::SetActivity(const FUserActivity& InUserActivity)
 { 
+	QUICK_SCOPE_CYCLE_COUNTER(STAT_FUserActivityTracking_SetActivity);
 	UserActivity = InUserActivity;
 	OnActivityChanged.Broadcast(UserActivity);
 	FCoreDelegates::UserActivityStringChanged.Broadcast(UserActivity.ActionName);
diff --git a/Engine/Source/Runtime/Engine/Private/UserInterface/Canvas.cpp b/Engine/Source/Runtime/Engine/Private/UserInterface/Canvas.cpp
--- a/Engine/Source/Runtime/Engine/Private/UserInterface/Canvas.cpp
+++ b/Engine/Source/Runtime/Engine/Private/UserInterface/Canvas.cpp
@@ -1173,6 +1173,8 @@ UCanvas::UCanvas(const FObjectInitializer& ObjectInitializer)
 	// only call once on construction.  Expensive on some platforms (occulus).
 	// Init gets called every frame.	
 	UpdateSafeZoneData();
+
+	FCoreDelegates::OnSafeFrameChangedEvent.AddUObject(this, &UCanvas::UpdateSafeZoneData);
 }
 
 void UCanvas::Init(int32 InSizeX, int32 InSizeY, FSceneView* InSceneView)
@@ -1184,9 +1186,17 @@ void UCanvas::Init(int32 InSizeX, int32 InSizeY, FSceneView* InSceneView)
 	UnsafeSizeY = SizeY;
 	SceneView = InSceneView;		
 	
-	Update();	
+	Update();
+}
+
+
+void UCanvas::BeginDestroy()
+{
+	Super::BeginDestroy();
+	FCoreDelegates::OnSafeFrameChangedEvent.RemoveAll(this);
 }
 
+
 void UCanvas::ApplySafeZoneTransform()
 {
 	// if there is no required safezone padding, then we can bail on the whole operation.
diff --git a/Engine/Source/Runtime/Engine/Private/WindDirectionalSource.cpp b/Engine/Source/Runtime/Engine/Private/WindDirectionalSource.cpp
--- a/Engine/Source/Runtime/Engine/Private/WindDirectionalSource.cpp
+++ b/Engine/Source/Runtime/Engine/Private/WindDirectionalSource.cpp
@@ -53,12 +53,12 @@ AWindDirectionalSource::AWindDirectionalSource(const FObjectInitializer& ObjectI
 #endif // WITH_EDITORONLY_DATA
 }
 
-void FWindSourceSceneProxy::FWindData::PrepareForAccumulate()
+void FWindData::PrepareForAccumulate()
 {
 	FMemory::Memset(this, 0, sizeof(FWindData));
 }
 
-void FWindSourceSceneProxy::FWindData::AddWeighted(const FWindData& InWindData, float Weight)
+void FWindData::AddWeighted(const FWindData& InWindData, float Weight)
 {	
 	Speed += InWindData.Speed * Weight;
 	MinGustAmt += InWindData.MinGustAmt * Weight;
@@ -66,7 +66,7 @@ void FWindSourceSceneProxy::FWindData::AddWeighted(const FWindData& InWindData,
 	Direction += InWindData.Direction * Weight;
 }
 
-void FWindSourceSceneProxy::FWindData::NormalizeByTotalWeight(float TotalWeight)
+void FWindData::NormalizeByTotalWeight(float TotalWeight)
 {
 	if (TotalWeight > 0.0f)
 	{
@@ -269,6 +269,24 @@ FWindSourceSceneProxy* UWindDirectionalSourceComponent::CreateSceneProxy() const
 	}
 }
 
+bool UWindDirectionalSourceComponent::GetWindParameters(const FVector& EvaluatePosition, FWindData& OutData, float& Weight) const
+{
+	// Whether we found any wind at the requested position (perhaps there's only point wind in the scene, or no wind)
+	bool bFoundWind = false;
+
+	if(bPointWind)
+	{
+		FWindSourceSceneProxy LocalProxy = FWindSourceSceneProxy(ComponentToWorld.GetLocation(), Strength, Speed, MinGustAmount, MaxGustAmount, Radius);
+		bFoundWind = LocalProxy.GetWindParameters(EvaluatePosition, OutData, Weight);
+	}
+	else
+	{
+		FWindSourceSceneProxy LocalProxy = FWindSourceSceneProxy(ComponentToWorld.GetUnitAxis(EAxis::X), Strength, Speed, MinGustAmount, MaxGustAmount);
+		bFoundWind = LocalProxy.GetWindParameters(EvaluatePosition, OutData, Weight);
+	}
+
+	return bFoundWind;
+}
 
 /** Returns Component subobject **/
 UWindDirectionalSourceComponent* AWindDirectionalSource::GetComponent() const { return Component; }
diff --git a/Engine/Source/Runtime/Engine/Private/World.cpp b/Engine/Source/Runtime/Engine/Private/World.cpp
--- a/Engine/Source/Runtime/Engine/Private/World.cpp
+++ b/Engine/Source/Runtime/Engine/Private/World.cpp
@@ -4374,6 +4374,11 @@ bool UWorld::Listen( FURL& InURL )
 		{
 			SourceCollection->SetNetDriver(NetDriver);
 		}
+		FLevelCollection* const StaticCollection = FindCollectionByType(ELevelCollectionType::StaticLevels);
+		if (StaticCollection)
+		{
+			StaticCollection->SetNetDriver(NetDriver);
+		}
 	}
 
 	if (NetDriver == NULL)
@@ -4394,6 +4399,11 @@ bool UWorld::Listen( FURL& InURL )
 		{
 			SourceCollection->SetNetDriver(nullptr);
 		}
+		FLevelCollection* StaticCollection = FindCollectionByType(ELevelCollectionType::StaticLevels);
+		if (StaticCollection)
+		{
+			StaticCollection->SetNetDriver(nullptr);
+		}
 		return false;
 	}
 	static bool LanPlay = FParse::Param(FCommandLine::Get(),TEXT("lanplay"));
@@ -4704,6 +4714,11 @@ bool FSeamlessTravelHandler::StartTravel(UWorld* InCurrentWorld, const FURL& InU
 				bCancelledExisting = true;
 			}
 
+			if (CurrentWorld->DemoNetDriver && CurrentWorld->DemoNetDriver->IsRecording())
+			{
+				CurrentWorld->DemoNetDriver->PauseRecording(true);
+			}
+
 			checkSlow(LoadedPackage == NULL);
 			checkSlow(LoadedWorld == NULL);
 
@@ -4896,10 +4911,12 @@ void FSeamlessTravelHandler::StartLoadingDestination()
 void FSeamlessTravelHandler::CopyWorldData()
 {
 	FLevelCollection* const CurrentCollection = CurrentWorld->FindCollectionByType(ELevelCollectionType::DynamicSourceLevels);
+	FLevelCollection* const CurrentStaticCollection = CurrentWorld->FindCollectionByType(ELevelCollectionType::StaticLevels);
 	FLevelCollection* const LoadedCollection = LoadedWorld->FindCollectionByType(ELevelCollectionType::DynamicSourceLevels);
+	FLevelCollection* const LoadedStaticCollection = LoadedWorld->FindCollectionByType(ELevelCollectionType::StaticLevels);
 
 	// If we are doing seamless travel for replay playback, then make sure to transfer the replay driver over to the new world
-	if ( CurrentWorld->DemoNetDriver && CurrentWorld->DemoNetDriver->IsPlaying() )
+	if ( CurrentWorld->DemoNetDriver && ( CurrentWorld->DemoNetDriver->IsPlaying() || CurrentWorld->DemoNetDriver->bRecordMapChanges) )
 	{
 		UDemoNetDriver* OldDriver = CurrentWorld->DemoNetDriver;
 		CurrentWorld->DemoNetDriver = nullptr;
@@ -4930,6 +4947,11 @@ void FSeamlessTravelHandler::CopyWorldData()
 		LoadedCollection->SetNetDriver(NetDriver);
 		CurrentCollection->SetNetDriver(nullptr);
 	}
+	if (CurrentStaticCollection && LoadedStaticCollection)
+	{
+		LoadedStaticCollection->SetNetDriver(NetDriver);
+		CurrentStaticCollection->SetNetDriver(nullptr);
+	}
 
 	if (NetDriver != NULL)
 	{
@@ -5016,7 +5038,7 @@ UWorld* FSeamlessTravelHandler::Tick()
 			}
 			
 			// If it's not still playing, destroy the demo net driver before we start renaming actors.
-			if ( CurrentWorld->DemoNetDriver && !CurrentWorld->DemoNetDriver->IsPlaying() )
+			if ( CurrentWorld->DemoNetDriver && !CurrentWorld->DemoNetDriver->IsPlaying() && !CurrentWorld->DemoNetDriver->bRecordMapChanges)
 			{
 				CurrentWorld->DestroyDemoNetDriver();
 			}
diff --git a/Engine/Source/Runtime/Engine/Public/Animation/AnimInstanceProxy.h b/Engine/Source/Runtime/Engine/Public/Animation/AnimInstanceProxy.h
--- a/Engine/Source/Runtime/Engine/Public/Animation/AnimInstanceProxy.h
+++ b/Engine/Source/Runtime/Engine/Public/Animation/AnimInstanceProxy.h
@@ -14,6 +14,66 @@ struct FAnimNode_Base;
 struct FAnimNode_SaveCachedPose;
 struct FAnimNode_SubInput;
 
+// Disable debugging information for shipping and test builds.
+#define ENABLE_ANIM_DRAW_DEBUG (1 && !(UE_BUILD_SHIPPING || UE_BUILD_TEST))
+
+UENUM()
+namespace EDrawDebugItemType
+{
+	enum Type
+	{
+		DirectionalArrow,
+		Sphere,
+		Line,
+		OnScreenMessage,
+	};
+}
+
+USTRUCT()
+struct FQueuedDrawDebugItem 
+{
+	GENERATED_BODY()
+
+	UPROPERTY(Transient)
+	TEnumAsByte<EDrawDebugItemType::Type> ItemType;
+
+	UPROPERTY(Transient)
+	FVector StartLoc;
+
+	UPROPERTY(Transient)
+	FVector EndLoc;
+
+	UPROPERTY(Transient)
+	FVector Center;
+
+	UPROPERTY(Transient)
+	float Radius;
+
+	UPROPERTY(Transient)
+	float Size;
+
+	UPROPERTY(Transient)
+	int32 Segments;
+
+	UPROPERTY(Transient)
+	FColor Color;
+
+	UPROPERTY(Transient)
+	bool bPersistentLines;
+
+	UPROPERTY(Transient)
+	float LifeTime;
+
+	UPROPERTY(Transient)
+	float Thickness;
+
+	UPROPERTY(Transient)
+	FString Message;
+
+	UPROPERTY(Transient)
+	FVector2D TextScale;
+};
+
 /** Proxy object passed around during animation tree update in lieu of a UAnimInstance */
 USTRUCT(meta = (DisplayName = "Native Variables"))
 struct ENGINE_API FAnimInstanceProxy
@@ -224,16 +284,16 @@ public:
 	 * @param out_SourceWeight : The Source weight for this node. 
 	 * @param out_TotalNodeWeight : Total weight of this node
 	 */
-	void GetSlotWeight(FName const& SlotNodeName, float& out_SlotNodeWeight, float& out_SourceWeight, float& out_TotalNodeWeight) const;
+	void GetSlotWeight(const FName& SlotNodeName, float& out_SlotNodeWeight, float& out_SourceWeight, float& out_TotalNodeWeight) const;
 
 	/** Evaluate a pose for a named montage slot */
-	void SlotEvaluatePose(FName SlotNodeName, const FCompactPose& SourcePose, const FBlendedCurve& SourceCurve, float InSourceWeight, FCompactPose& BlendedPose, FBlendedCurve& BlendedCurve, float InBlendWeight, float InTotalNodeWeight);
+	void SlotEvaluatePose(const FName& SlotNodeName, const FCompactPose& SourcePose, const FBlendedCurve& SourceCurve, float InSourceWeight, FCompactPose& BlendedPose, FBlendedCurve& BlendedCurve, float InBlendWeight, float InTotalNodeWeight);
 	
 	// Allow slot nodes to store off their weight during ticking
-	void UpdateSlotNodeWeight(FName SlotNodeName, float InLocalMontageWeight, float InNodeGlobalWeight);
+	void UpdateSlotNodeWeight(const FName& SlotNodeName, float InLocalMontageWeight, float InNodeGlobalWeight);
 
 	/** Register a named slot */
-	void RegisterSlotNodeWithAnimInstance(FName SlotNodeName);
+	void RegisterSlotNodeWithAnimInstance(const FName& SlotNodeName);
 
 	/** Check whether we have a valid root node */
 	bool HasRootNode() const
@@ -250,6 +310,20 @@ public:
 	/** Gather debug data from this instance proxy and the blend tree for display */
 	void GatherDebugData(FNodeDebugData& DebugData);
 
+#if ENABLE_ANIM_DRAW_DEBUG
+	TArray<FQueuedDrawDebugItem> QueuedDrawDebugItems;
+
+	void AnimDrawDebugOnScreenMessage(const FString& DebugMessage, const FColor& Color, const FVector2D& TextScale = FVector2D::UnitVector);
+	void AnimDrawDebugLine(const FVector& StartLoc, const FVector& EndLoc, const FColor& Color, bool bPersistentLines = false, float LifeTime = -1.f, float Thickness = 0.f);
+	void AnimDrawDebugDirectionalArrow(const FVector& LineStart, const FVector& LineEnd, float ArrowSize, const FColor& Color, bool bPersistentLines = false, float LifeTime = -1.f, float Thickness = 0.f);
+	void AnimDrawDebugSphere(const FVector& Center, float Radius, int32 Segments, const FColor& Color, bool bPersistentLines = false, float LifeTime = -1.f, float Thickness = 0.f);
+#else
+	void AnimDrawDebugOnScreenMessage(const FString& DebugMessage, const FColor& Color, const FVector2D& TextScale = FVector2D::UnitVector) {}
+	void AnimDrawDebugLine(const FVector& StartLoc, const FVector& EndLoc, const FColor& Color, bool bPersistentLines = false, float LifeTime = -1.f, float Thickness = 0.f) {}
+	void AnimDrawDebugDirectionalArrow(const FVector& LineStart, const FVector& LineEnd, float ArrowSize, const FColor& Color, bool bPersistentLines = false, float LifeTime = -1.f, float Thickness = 0.f) {}
+	void AnimDrawDebugSphere(const FVector& Center, float Radius, int32 Segments, const FColor& Color, bool bPersistentLines = false, float LifeTime = -1.f, float Thickness = 0.f) {}
+#endif // ENABLE_ANIM_DRAW_DEBUG
+
 #if !NO_LOGGING
 	const FString& GetActorName() const 
 	{
@@ -282,7 +356,7 @@ public:
 	float GetRecordedStateWeight(const int32& InMachineClassIndex, const int32& InStateIndex) const;
 	void RecordStateWeight(const int32& InMachineClassIndex, const int32& InStateIndex, const float& InStateWeight);
 
-	bool IsSlotNodeRelevantForNotifies(FName SlotNodeName) const;
+	bool IsSlotNodeRelevantForNotifies(const FName& SlotNodeName) const;
 	/** Reset any dynamics running simulation-style updates (e.g. on teleport, time skip etc.) */
 	void ResetDynamics();
 
@@ -380,11 +454,21 @@ protected:
 
 	/** Get global weight in AnimGraph for this slot node. 
 	 * Note: this is the weight of the node, not the weight of any potential montage it is playing. */
-	float GetSlotNodeGlobalWeight(FName SlotNodeName) const;
+	float GetSlotNodeGlobalWeight(const FName& SlotNodeName) const;
 
 	/** Get Global weight of any montages this slot node is playing. 
 	 * If this slot is not currently playing a montage, it will return 0. */
-	float GetSlotMontageGlobalWeight(FName SlotNodeName) const;
+	float GetSlotMontageGlobalWeight(const FName& SlotNodeName) const;
+
+	/** Get local weight of any montages this slot node is playing.
+	* If this slot is not currently playing a montage, it will return 0. 
+	* This is double buffered, will return last frame data if called from Update or Evaluate. */
+	float GetSlotMontageLocalWeight(const FName& SlotNodeName) const;
+
+	/** Get local weight of any montages this slot is playing.
+	* If this slot is not current playing a montage, it will return 0.
+	* This will return up to date data if called during Update or Evaluate. */
+	float CalcSlotMontageLocalWeight(const FName& SlotNodeName) const;
 
 	/** 
 	 * Recalculate Required Bones [RequiredBones]
diff --git a/Engine/Source/Runtime/Engine/Public/AudioDerivedData.h b/Engine/Source/Runtime/Engine/Public/AudioDerivedData.h
--- a/Engine/Source/Runtime/Engine/Public/AudioDerivedData.h
+++ b/Engine/Source/Runtime/Engine/Public/AudioDerivedData.h
@@ -27,7 +27,7 @@ public:
 		// want to bump this version, generate a new guid using VS->Tools->Create GUID and
 		// return it here. Ex.
 		// return TEXT("855EE5B3574C43ABACC6700C4ADC62E6");
-		return TEXT("0003_0000");
+		return TEXT("0004_0000");
 	}
 
 	virtual FString GetPluginSpecificCacheKeySuffix() const override;
diff --git a/Engine/Source/Runtime/Engine/Public/Scalability.h b/Engine/Source/Runtime/Engine/Public/Scalability.h
--- a/Engine/Source/Runtime/Engine/Public/Scalability.h
+++ b/Engine/Source/Runtime/Engine/Public/Scalability.h
@@ -118,4 +118,7 @@ namespace Scalability
 
 	/** Maximum single axis scale for render resolution */
 	static const float MaxResolutionScale = 100.0f;
+
+	/** Returns the current screen percentage */
+	ENGINE_API float GetResolutionScreenPercentage();
 }
diff --git a/Engine/Source/Runtime/Engine/Public/SceneInterface.h b/Engine/Source/Runtime/Engine/Public/SceneInterface.h
--- a/Engine/Source/Runtime/Engine/Public/SceneInterface.h
+++ b/Engine/Source/Runtime/Engine/Public/SceneInterface.h
@@ -225,6 +225,9 @@ public:
 	/** Accesses wind parameters.  XYZ will contain wind direction * Strength, W contains wind speed. */
 	virtual void GetWindParameters(const FVector& Position, FVector& OutDirection, float& OutSpeed, float& OutMinGustAmt, float& OutMaxGustAmt) const = 0;
 
+	/** Accesses wind parameters safely for game thread applications */
+	virtual void GetWindParameters_GameThread(const FVector& Position, FVector& OutDirection, float& OutSpeed, float& OutMinGustAmt, float& OutMaxGustAmt) const = 0;
+
 	/** Same as GetWindParameters, but ignores point wind sources. */
 	virtual void GetDirectionalWindParameters(FVector& OutDirection, float& OutSpeed, float& OutMinGustAmt, float& OutMaxGustAmt) const = 0;
 
diff --git a/Engine/Source/Runtime/Engine/Public/SceneManagement.h b/Engine/Source/Runtime/Engine/Public/SceneManagement.h
--- a/Engine/Source/Runtime/Engine/Public/SceneManagement.h
+++ b/Engine/Source/Runtime/Engine/Public/SceneManagement.h
@@ -1287,31 +1287,32 @@ public:
 	void SetTransform(const FMatrix& InTransform);
 };
 
-/** Represents a wind source component to the scene manager in the rendering thread. */
-class ENGINE_API FWindSourceSceneProxy
+/** Calculated wind data with support for accumulating other weighted wind data */
+class ENGINE_API FWindData
 {
-public:	
-
-	class ENGINE_API FWindData
+public:
+	FWindData()
+		: Speed(0.0f)
+		, MinGustAmt(0.0f)
+		, MaxGustAmt(0.0f)
+		, Direction(1.0f, 0.0f, 0.0f)
 	{
-	public:
-		FWindData()
-			: Speed(0.0f)
-			, MinGustAmt(0.0f)
-			, MaxGustAmt(0.0f)
-			, Direction(1.0f, 0.0f, 0.0f)
-		{
-		}
+	}
 
-		void PrepareForAccumulate();
-		void AddWeighted(const FWindData& InWindData, float Weight);
-		void NormalizeByTotalWeight(float TotalWeight);
+	void PrepareForAccumulate();
+	void AddWeighted(const FWindData& InWindData, float Weight);
+	void NormalizeByTotalWeight(float TotalWeight);
 
-		float Speed;
-		float MinGustAmt;
-		float MaxGustAmt;
-		FVector Direction;
-	};
+	float Speed;
+	float MinGustAmt;
+	float MaxGustAmt;
+	FVector Direction;
+};
+
+/** Represents a wind source component to the scene manager in the rendering thread. */
+class ENGINE_API FWindSourceSceneProxy
+{
+public:	
 
 	/** Initialization constructor. */
 	FWindSourceSceneProxy(const FVector& InDirection, float InStrength, float InSpeed, float InMinGustAmt, float InMaxGustAmt) :
diff --git a/Engine/Source/Runtime/GameplayTags/Classes/GameplayTagsManager.h b/Engine/Source/Runtime/GameplayTags/Classes/GameplayTagsManager.h
--- a/Engine/Source/Runtime/GameplayTags/Classes/GameplayTagsManager.h
+++ b/Engine/Source/Runtime/GameplayTags/Classes/GameplayTagsManager.h
@@ -261,6 +261,8 @@ class GAMEPLAYTAGS_API UGameplayTagsManager : public UObject
 	/** Splits a tag such as x.y.z into an array of names {x,y,z} */
 	void SplitGameplayTagFName(const FGameplayTag& Tag, TArray<FName>& OutNames);
 
+	void RequestAllGameplayTags(FGameplayTagContainer& TagContainer, bool OnlyIncludeDictionaryTags) const;
+
 	/**
 	 * Checks if the passed in name is in the tag dictionary and can be created
 	 *
diff --git a/Engine/Source/Runtime/GameplayTags/Private/GameplayTagsManager.cpp b/Engine/Source/Runtime/GameplayTags/Private/GameplayTagsManager.cpp
--- a/Engine/Source/Runtime/GameplayTags/Private/GameplayTagsManager.cpp
+++ b/Engine/Source/Runtime/GameplayTags/Private/GameplayTagsManager.cpp
@@ -148,7 +148,13 @@ void UGameplayTagsManager::ConstructGameplayTagTree()
 					FoundSource->SourceTagList->LoadConfig(UGameplayTagsList::StaticClass(), *FileName);
 				}
 
-				// Add to dictionary
+#if WITH_EDITOR
+				if (GIsEditor || IsRunningCommandlet()) // Sort tags for UI Purposes but don't sort in -game scenerio since this would break compat with noneditor cooked builds
+				{
+					FoundSource->SourceTagList->SortTags();
+				}
+#endif
+
 				for (const FGameplayTagTableRow& TableRow : FoundSource->SourceTagList->GameplayTagList)
 				{
 					AddTagTableRow(TableRow, TagSource);
@@ -952,6 +958,26 @@ FGameplayTagContainer UGameplayTagsManager::RequestGameplayTagParents(const FGam
 	return FGameplayTagContainer();
 }
 
+void UGameplayTagsManager::RequestAllGameplayTags(FGameplayTagContainer& TagContainer, bool OnlyIncludeDictionaryTags) const
+{
+	TArray<TSharedPtr<FGameplayTagNode>> ValueArray;
+	GameplayTagNodeMap.GenerateValueArray(ValueArray);
+	for (const TSharedPtr<FGameplayTagNode>& TagNode : ValueArray)
+	{
+#if WITH_EDITOR
+		bool DictTag = IsDictionaryTag(TagNode->GetSimpleTagName());
+#else
+		bool DictTag = false;
+#endif 
+		if (!OnlyIncludeDictionaryTags || DictTag)
+		{
+			const FGameplayTag* Tag = GameplayTagNodeMap.FindKey(TagNode);
+			check(Tag);
+			TagContainer.AddTagFast(*Tag);
+		}
+	}
+}
+
 FGameplayTagContainer UGameplayTagsManager::RequestGameplayTagChildren(const FGameplayTag& GameplayTag) const
 {
 	FGameplayTagContainer TagContainer;
diff --git a/Engine/Source/Runtime/Launch/Private/LaunchEngineLoop.cpp b/Engine/Source/Runtime/Launch/Private/LaunchEngineLoop.cpp
--- a/Engine/Source/Runtime/Launch/Private/LaunchEngineLoop.cpp
+++ b/Engine/Source/Runtime/Launch/Private/LaunchEngineLoop.cpp
@@ -62,6 +62,7 @@
 	#include "Engine/GameInstance.h"
 	#include "Net/OnlineEngineInterface.h"
 	#include "Internationalization/EnginePackageLocalizationCache.h"
+	#include "Engine/DemoNetDriver.h"
 
 #if !UE_SERVER
 	#include "HeadMountedDisplay.h"
@@ -117,6 +118,50 @@
 	#define USE_LOCALIZED_PACKAGE_CACHE 0
 #endif
 
+static TAutoConsoleVariable<int32> CVarDoAsyncEndOfFrameTasksRandomize(
+	TEXT("tick.DoAsyncEndOfFrameTasks.Randomize"),
+	0,
+	TEXT("Used to add random sleeps to tick.DoAsyncEndOfFrameTasks to shake loose bugs on either thread. Also does random render thread flushes from the game thread.")
+	);
+
+static FAutoConsoleTaskPriority CPrio_AsyncEndOfFrameGameTasks(
+	TEXT("TaskGraph.TaskPriorities.AsyncEndOfFrameGameTasks"),
+	TEXT("Task and thread priority for the experiemntal async end of frame tasks."),
+	ENamedThreads::HighThreadPriority,
+	ENamedThreads::NormalTaskPriority,
+	ENamedThreads::HighTaskPriority
+	);
+
+/** Task that executes concurrently with Slate when tick.DoAsyncEndOfFrameTasks is true. */
+class FExecuteConcurrentWithSlateTickTask
+{
+	TFunctionRef<void()> TickWithSlate;
+
+public:
+
+	FExecuteConcurrentWithSlateTickTask(TFunctionRef<void()> InTickWithSlate)
+		: TickWithSlate(InTickWithSlate)
+	{
+	}
+
+	static FORCEINLINE TStatId GetStatId()
+	{
+		RETURN_QUICK_DECLARE_CYCLE_STAT(FExecuteConcurrentWithSlateTickTask, STATGROUP_TaskGraphTasks);
+	}
+
+	static FORCEINLINE ENamedThreads::Type GetDesiredThread()
+	{
+		return CPrio_AsyncEndOfFrameGameTasks.Get();
+	}
+
+	static FORCEINLINE ESubsequentsMode::Type GetSubsequentsMode() { return ESubsequentsMode::TrackSubsequents; }
+
+	void DoTask(ENamedThreads::Type CurrentThread, const FGraphEventRef& MyCompletionGraphEvent)
+	{
+		TickWithSlate();
+	}
+};
+
 // Pipe output to std output
 // This enables UBT to collect the output for it's own use
 class FOutputDeviceStdOutput : public FOutputDevice
@@ -2869,6 +2914,34 @@ void FEngineLoop::Tick()
 			GDistanceFieldAsyncQueue->ProcessAsyncTasks();
 		}
 
+#if WITH_ENGINE
+		FGraphEventRef ConcurrentTask;
+		const bool bDoConcurrentSlateTick = GEngine->ShouldDoAsyncEndOfFrameTasks();
+
+		if (bDoConcurrentSlateTick)
+		{
+			const float DeltaSeconds = FApp::GetDeltaTime();
+			UGameViewportClient* const GameViewport = GEngine->GameViewport;
+			ConcurrentTask = TGraphTask<FExecuteConcurrentWithSlateTickTask>::CreateTask(nullptr, ENamedThreads::GameThread).ConstructAndDispatchWhenReady(
+				[GameViewport, DeltaSeconds]()
+				{
+					if (CVarDoAsyncEndOfFrameTasksRandomize.GetValueOnAnyThread(true) > 0)
+					{
+						FPlatformProcess::Sleep(FMath::RandRange(0.0f, .003f)); // this shakes up the threading to find race conditions
+					}
+					if (GameViewport != nullptr)
+					{
+						UWorld* const World = GameViewport->GetWorld();
+						if (World && World->DemoNetDriver)
+						{
+							World->DemoNetDriver->TickFlushAsyncEndOfFrame(DeltaSeconds);
+						}
+					}
+				}
+			);
+		}
+#endif
+
 		if (FSlateApplication::IsInitialized() && !bIdleMode)
 		{
 			{
@@ -2883,6 +2956,15 @@ void FEngineLoop::Tick()
 			FSlateApplication::Get().Tick();
 		}
 
+#if WITH_ENGINE
+		if (ConcurrentTask.GetReference())
+		{
+			QUICK_SCOPE_CYCLE_COUNTER(STAT_ConcurrentWithSlateTickTasks_Wait);
+			FTaskGraphInterface::Get().WaitUntilTaskCompletes(ConcurrentTask);
+			ConcurrentTask = nullptr;
+		}
+#endif
+
 #if STATS
 		// Clear any stat group notifications we have pending just incase they weren't claimed during FSlateApplication::Get().Tick
 		extern CORE_API void ClearPendingStatGroups();
diff --git a/Engine/Source/Runtime/Navmesh/Public/Detour/DetourNavMesh.h b/Engine/Source/Runtime/Navmesh/Public/Detour/DetourNavMesh.h
--- a/Engine/Source/Runtime/Navmesh/Public/Detour/DetourNavMesh.h
+++ b/Engine/Source/Runtime/Navmesh/Public/Detour/DetourNavMesh.h
@@ -750,6 +750,33 @@ public:
 	}
 
 	void applyAreaCostOrder(unsigned char* costOrder);
+	
+	/// Returns neighbour tile count based on side of given tile.
+	int getNeighbourTilesCountAt(const int x, const int y, const int side) const;
+
+	bool getNeighbourCoords(const int x, const int y, const int side, int& outX, int& outY) const
+	{
+		outX = x;
+		outY = y;
+		switch (side)
+		{
+		case 0: ++outX; break;
+		case 1: ++outX; ++outY; break;
+		case 2: ++outY; break;
+		case 3: --outX; ++outY; break;
+		case 4: --outX; break;
+		case 5: --outX; --outY; break;
+		case 6: --outY; break;
+		case 7: ++outX; --outY; break;
+		};
+		// @todo we might want to do some validation
+		return true;
+	}
+
+	unsigned int getTileIndex(const dtMeshTile* tile) const
+	{
+		return (unsigned int)(tile - m_tiles);
+	}
 	//@UE4 END
 	
 private:
@@ -764,10 +791,6 @@ private:
 	/// Returns neighbour tile based on side.
 	int getNeighbourTilesAt(const int x, const int y, const int side,
 							dtMeshTile** tiles, const int maxTiles) const;
-// @UE4 BEGIN
-	/// Returns neighbour tile based on side.
-	int getNeighbourTilesCountAt(const int x, const int y, const int side) const;
-// @UE4 END
 	
 	/// Returns all polygons in neighbour tile based on portal defined by the segment.
 	int findConnectingPolys(const float* va, const float* vb,
diff --git a/Engine/Source/Runtime/NetworkReplayStreaming/HttpNetworkReplayStreaming/Private/HttpNetworkReplayStreaming.cpp b/Engine/Source/Runtime/NetworkReplayStreaming/HttpNetworkReplayStreaming/Private/HttpNetworkReplayStreaming.cpp
--- a/Engine/Source/Runtime/NetworkReplayStreaming/HttpNetworkReplayStreaming/Private/HttpNetworkReplayStreaming.cpp
+++ b/Engine/Source/Runtime/NetworkReplayStreaming/HttpNetworkReplayStreaming/Private/HttpNetworkReplayStreaming.cpp
@@ -325,10 +325,15 @@ void FHttpNetworkReplayStreamer::StartStreaming( const FString& CustomName, cons
 		AddRequestToQueue( EQueuedHttpRequestType::StartUploading, HttpRequest, 3, 2.0f );
 		
 		// We need to upload the header AFTER StartUploading is done (so we have session name)
-		AddRequestToQueue( EQueuedHttpRequestType::UploadHeader, nullptr );
+		RefreshHeader();
 	}
 }
 
+void FHttpNetworkReplayStreamer::RefreshHeader()
+{
+	AddRequestToQueue(EQueuedHttpRequestType::UploadHeader, nullptr);
+}
+
 void FHttpNetworkReplayStreamer::AddRequestToQueue( const EQueuedHttpRequestType::Type Type, TSharedPtr< class IHttpRequest > Request, const int32 InMaxRetries, const float InRetryDelay )
 {
 	UE_LOG( LogHttpReplay, Verbose, TEXT( "FHttpNetworkReplayStreamer::AddRequestToQueue. Type: %s" ), EQueuedHttpRequestType::ToString( Type ) );
@@ -526,7 +531,11 @@ void FHttpNetworkReplayStreamer::UploadHeader()
 		return;
 	}
 
-	check( StreamChunkIndex == 0 );
+	if ( !IsStreaming() )
+	{
+		UE_LOG(LogHttpReplay, Warning, TEXT("FHttpNetworkReplayStreamer::UploadHeader. Not currently streaming"));
+		return;
+	}
 
 	// First upload the header
 	UE_LOG( LogHttpReplay, Log, TEXT( "FHttpNetworkReplayStreamer::UploadHeader. Header. StreamChunkIndex: %i, Size: %i" ), StreamChunkIndex, HeaderArchive.Buffer.Num() );
@@ -1544,9 +1553,7 @@ void FHttpNetworkReplayStreamer::HttpHeaderUploadFinished( FHttpRequestPtr HttpR
 	TSharedPtr< FQueuedHttpRequest > SavedFlightHttpRequest = InFlightHttpRequest;
 
 	RequestFinished( EStreamerState::StreamingUp, EQueuedHttpRequestType::UploadingHeader, HttpRequest );
-
-	check( StartStreamingDelegate.IsBound() );
-
+	
 	if ( bSucceeded && HttpResponse->GetResponseCode() == EHttpResponseCodes::NoContent )
 	{
 		if ( HttpRequest.IsValid() )
@@ -2168,6 +2175,7 @@ bool FHttpNetworkReplayStreamer::ProcessNextHttpRequest()
 			check( IsStreaming() );
 			StreamerState = EStreamerState::Idle;
 			bStopStreamingCalled = false;
+			ensure( QueuedHttpRequests.Num() == 0 );
 			return ProcessNextHttpRequest();
 		}
 
diff --git a/Engine/Source/Runtime/NetworkReplayStreaming/HttpNetworkReplayStreaming/Public/HttpNetworkReplayStreaming.h b/Engine/Source/Runtime/NetworkReplayStreaming/HttpNetworkReplayStreaming/Public/HttpNetworkReplayStreaming.h
--- a/Engine/Source/Runtime/NetworkReplayStreaming/HttpNetworkReplayStreaming/Public/HttpNetworkReplayStreaming.h
+++ b/Engine/Source/Runtime/NetworkReplayStreaming/HttpNetworkReplayStreaming/Public/HttpNetworkReplayStreaming.h
@@ -228,7 +228,8 @@ public:
 	virtual ENetworkReplayError::Type GetLastError() const override;
 	virtual FString		GetReplayID() const override { return SessionName; }
 	virtual void		SetTimeBufferHintSeconds(const float InTimeBufferHintSeconds) override {}
-	
+	virtual void		RefreshHeader() override;
+
 	/** FHttpNetworkReplayStreamer */
 	void UploadHeader();
 	void FlushStream();
diff --git a/Engine/Source/Runtime/NetworkReplayStreaming/InMemoryNetworkReplayStreaming/Private/InMemoryNetworkReplayStreaming.cpp b/Engine/Source/Runtime/NetworkReplayStreaming/InMemoryNetworkReplayStreaming/Private/InMemoryNetworkReplayStreaming.cpp
--- a/Engine/Source/Runtime/NetworkReplayStreaming/InMemoryNetworkReplayStreaming/Private/InMemoryNetworkReplayStreaming.cpp
+++ b/Engine/Source/Runtime/NetworkReplayStreaming/InMemoryNetworkReplayStreaming/Private/InMemoryNetworkReplayStreaming.cpp
@@ -290,17 +290,9 @@ FArchive* FInMemoryNetworkReplayStreamer::GetCheckpointArchive()
 			FoundReplay->StreamChunks.RemoveAt(0, NumChunksToRemove, false);
 		}
 
-		const int32 NewCheckpointIndex = FoundReplay->Checkpoints.Add(FInMemoryReplay::FCheckpoint());
-		CheckpointAr.Reset(new FMemoryWriter(FoundReplay->Checkpoints[NewCheckpointIndex].Data));
-
-		// Start a new stream chunk for the new checkpoint
-		FInMemoryReplay::FStreamChunk NewChunk;
-		if (FoundReplay->StreamChunks.Num() > 0)
-		{
-			NewChunk.StartIndex = FoundReplay->StreamChunks.Last().StartIndex + FoundReplay->StreamChunks.Last().Data.Num();
-			NewChunk.TimeInMS = FoundReplay->StreamInfo.LengthInMS;
-		}
-		FoundReplay->StreamChunks.Add(NewChunk);
+		// Save to a temporary checkpoint that will moved onto the replay's checkpoint list in FlushCheckpoint().
+		CheckpointCurrentlyBeingSaved.Reset();
+		CheckpointAr.Reset(new FMemoryWriter(CheckpointCurrentlyBeingSaved.Data));
 	}
 
 	return CheckpointAr.Get();
@@ -311,14 +303,26 @@ void FInMemoryNetworkReplayStreamer::FlushCheckpoint(const uint32 TimeInMS)
 	UE_LOG(LogMemoryReplay, Log, TEXT("FInMemoryNetworkReplayStreamer::FlushCheckpoint. TimeInMS: %u"), TimeInMS);
 
 	check(FileAr.Get() != nullptr);
+	check(CheckpointCurrentlyBeingSaved.Data.Num() != 0);
 
 	// Finalize the checkpoint data.
 	CheckpointAr.Reset();
 	
-	FInMemoryReplay* FoundReplay = GetCurrentReplayChecked();
+	CheckpointCurrentlyBeingSaved.TimeInMS = TimeInMS;
+	CheckpointCurrentlyBeingSaved.StreamByteOffset = FileAr->Tell();
 
-	FoundReplay->Checkpoints.Last().TimeInMS = TimeInMS;
-	FoundReplay->Checkpoints.Last().StreamByteOffset = FileAr->Tell();
+	FInMemoryReplay* const FoundReplay = GetCurrentReplayChecked();
+
+	const int32 NewCheckpointIndex = FoundReplay->Checkpoints.Add(MoveTemp(CheckpointCurrentlyBeingSaved));
+
+	// Start a new stream chunk for the new checkpoint
+	FInMemoryReplay::FStreamChunk NewChunk;
+	if (FoundReplay->StreamChunks.Num() > 0)
+	{
+		NewChunk.StartIndex = FoundReplay->StreamChunks.Last().StartIndex + FoundReplay->StreamChunks.Last().Data.Num();
+		NewChunk.TimeInMS = FoundReplay->StreamInfo.LengthInMS;
+	}
+	FoundReplay->StreamChunks.Add(NewChunk);
 }
 
 void FInMemoryNetworkReplayStreamer::GotoCheckpointIndex(const int32 CheckpointIndex, const FOnCheckpointReadyDelegate& Delegate)
diff --git a/Engine/Source/Runtime/NetworkReplayStreaming/InMemoryNetworkReplayStreaming/Public/InMemoryNetworkReplayStreaming.h b/Engine/Source/Runtime/NetworkReplayStreaming/InMemoryNetworkReplayStreaming/Public/InMemoryNetworkReplayStreaming.h
--- a/Engine/Source/Runtime/NetworkReplayStreaming/InMemoryNetworkReplayStreaming/Public/InMemoryNetworkReplayStreaming.h
+++ b/Engine/Source/Runtime/NetworkReplayStreaming/InMemoryNetworkReplayStreaming/Public/InMemoryNetworkReplayStreaming.h
@@ -23,6 +23,13 @@ struct FInMemoryReplay
 		TArray<uint8> Data;
 		uint32 TimeInMS;
 		uint32 StreamByteOffset;
+		
+		void Reset()
+		{
+			Data.Reset();
+			TimeInMS = 0;
+			StreamByteOffset = 0;
+		}
 	};
 
 	/** Represents a chunk of replay stream data between two checkpoints. */
@@ -142,6 +149,7 @@ public:
 	virtual void KeepReplay( const FString& ReplayName, const bool bKeep ) override;
 	virtual FString	GetReplayID() const override { return TEXT( "" ); }
 	virtual void SetTimeBufferHintSeconds(const float InTimeBufferHintSeconds) override { TimeBufferHintSeconds = InTimeBufferHintSeconds; }
+	virtual void RefreshHeader() override {};
 
 	/** FTickableObjectBase implementation */
 	virtual void Tick(float DeltaSeconds) override;
@@ -193,6 +201,12 @@ private:
 	/* Handle to the archive that will read/write checkpoint files */
 	TUniquePtr<FArchive> CheckpointAr;
 
+	/**
+	 * Temporary checkpoint used during recording. Will be moved onto the
+	 * replay's checkpoint list in FlushCheckpoint to commit it.
+	 */
+	FInMemoryReplay::FCheckpoint CheckpointCurrentlyBeingSaved;
+
 	/** EStreamerState - Overall state of the streamer */
 	enum class EStreamerState
 	{
diff --git a/Engine/Source/Runtime/NetworkReplayStreaming/NetworkReplayStreaming/Public/NetworkReplayStreaming.h b/Engine/Source/Runtime/NetworkReplayStreaming/NetworkReplayStreaming/Public/NetworkReplayStreaming.h
--- a/Engine/Source/Runtime/NetworkReplayStreaming/NetworkReplayStreaming/Public/NetworkReplayStreaming.h
+++ b/Engine/Source/Runtime/NetworkReplayStreaming/NetworkReplayStreaming/Public/NetworkReplayStreaming.h
@@ -174,6 +174,8 @@ public:
 	virtual void RequestEventData( const FString& EventID, const FOnRequestEventDataComplete& RequestEventDataComplete ) = 0;
 	virtual void SearchEvents(const FString& EventGroup, const FOnEnumerateStreamsComplete& Delegate) = 0;
 	virtual void KeepReplay( const FString& ReplayName, const bool bKeep ) = 0;
+	virtual void RefreshHeader() = 0;
+
 	/** Returns true if the playing stream is currently in progress */
 	virtual bool IsLive() const = 0;
 	virtual FString	GetReplayID() const = 0;
diff --git a/Engine/Source/Runtime/NetworkReplayStreaming/NullNetworkReplayStreaming/Public/NullNetworkReplayStreaming.h b/Engine/Source/Runtime/NetworkReplayStreaming/NullNetworkReplayStreaming/Public/NullNetworkReplayStreaming.h
--- a/Engine/Source/Runtime/NetworkReplayStreaming/NullNetworkReplayStreaming/Public/NullNetworkReplayStreaming.h
+++ b/Engine/Source/Runtime/NetworkReplayStreaming/NullNetworkReplayStreaming/Public/NullNetworkReplayStreaming.h
@@ -74,6 +74,7 @@ public:
 	virtual void KeepReplay( const FString& ReplayName, const bool bKeep ) override {}
 	virtual FString	GetReplayID() const override { return TEXT( "" ); }
 	virtual void SetTimeBufferHintSeconds(const float InTimeBufferHintSeconds) override {}
+	virtual void RefreshHeader() override {};
 
 	/** FTickableObjectBase implementation */
 	virtual void Tick(float DeltaSeconds) override;
diff --git a/Engine/Source/Runtime/Online/XMPP/Private/XmppJingle/XmppChatJingle.cpp b/Engine/Source/Runtime/Online/XMPP/Private/XmppJingle/XmppChatJingle.cpp
--- a/Engine/Source/Runtime/Online/XMPP/Private/XmppJingle/XmppChatJingle.cpp
+++ b/Engine/Source/Runtime/Online/XMPP/Private/XmppJingle/XmppChatJingle.cpp
@@ -185,6 +185,10 @@ protected:
 
 		Result->AddAttr(buzz::QN_TO, ToJid.Str());
 		Result->AddAttr(buzz::QN_TYPE, ChatType);
+
+		// Add CorrelationID for tracking purposes
+		FXmppJingle::AddCorrIdToStanza(*Result);
+
 		buzz::XmlElement* Body = new buzz::XmlElement(buzz::QN_BODY);
 		Body->SetBodyText(Message.Body);
 		Result->AddElement(Body);
diff --git a/Engine/Source/Runtime/Online/XMPP/Private/XmppJingle/XmppJingle.cpp b/Engine/Source/Runtime/Online/XMPP/Private/XmppJingle/XmppJingle.cpp
--- a/Engine/Source/Runtime/Online/XMPP/Private/XmppJingle/XmppJingle.cpp
+++ b/Engine/Source/Runtime/Online/XMPP/Private/XmppJingle/XmppJingle.cpp
@@ -45,4 +45,20 @@ void FXmppJingle::ConvertFromJid(buzz::Jid& OutJid, const FXmppUserJid& InJid)
 		);
 }
 
+void FXmppJingle::AddCorrIdToStanza(buzz::XmlElement& Stanza, const TCHAR* const CorrId /*= nullptr*/)
+{
+	static const buzz::StaticQName CORRID = { "", "corr-id" };
+
+	if (CorrId == nullptr)
+	{
+		const FString GuidCorrId = FGuid::NewGuid().ToString();
+
+		Stanza.AddAttr(CORRID, TCHAR_TO_UTF8(*GuidCorrId));
+	}
+	else
+	{
+		Stanza.AddAttr(CORRID, TCHAR_TO_UTF8(CorrId));
+	}
+}
+
 #endif //WITH_XMPP_JINGLE
diff --git a/Engine/Source/Runtime/Online/XMPP/Private/XmppJingle/XmppJingle.h b/Engine/Source/Runtime/Online/XMPP/Private/XmppJingle/XmppJingle.h
--- a/Engine/Source/Runtime/Online/XMPP/Private/XmppJingle/XmppJingle.h
+++ b/Engine/Source/Runtime/Online/XMPP/Private/XmppJingle/XmppJingle.h
@@ -60,6 +60,9 @@ public:
 
 	static void ConvertToJid(FXmppUserJid& OutJid, const buzz::Jid& InJid);
 	static void ConvertFromJid(buzz::Jid& OutJid, const FXmppUserJid& InJid);
+
+	/** Adds a Correlation ID to a stanza.  If none is provided, one will be generated */
+	static void AddCorrIdToStanza(buzz::XmlElement& Stanza, const TCHAR* const CorrId = nullptr);
 };
 
 #endif //WITH_XMPP_JINGLE
diff --git a/Engine/Source/Runtime/Online/XMPP/Private/XmppJingle/XmppMessagesJingle.cpp b/Engine/Source/Runtime/Online/XMPP/Private/XmppJingle/XmppMessagesJingle.cpp
--- a/Engine/Source/Runtime/Online/XMPP/Private/XmppJingle/XmppMessagesJingle.cpp
+++ b/Engine/Source/Runtime/Online/XMPP/Private/XmppJingle/XmppMessagesJingle.cpp
@@ -156,6 +156,10 @@ protected:
 
 		Result->AddAttr(buzz::QN_TO, ToJid.Str());
 		Result->AddAttr(buzz::QN_ID, rtc::CreateRandomString(16));
+
+		// Add CorrelationID for tracking purposes
+		FXmppJingle::AddCorrIdToStanza(*Result);
+
 		buzz::XmlElement* Body = new buzz::XmlElement(buzz::QN_BODY);
 		Body->SetBodyText(Message.Body);
 		Result->AddElement(Body);
diff --git a/Engine/Source/Runtime/Online/XMPP/Private/XmppJingle/XmppMultiUserChatJingle.cpp b/Engine/Source/Runtime/Online/XMPP/Private/XmppJingle/XmppMultiUserChatJingle.cpp
--- a/Engine/Source/Runtime/Online/XMPP/Private/XmppJingle/XmppMultiUserChatJingle.cpp
+++ b/Engine/Source/Runtime/Online/XMPP/Private/XmppJingle/XmppMultiUserChatJingle.cpp
@@ -103,6 +103,10 @@ private:
 	static buzz::XmlElement* MakeRequest()
 	{
 		buzz::XmlElement* OwnerConfigQuery = new buzz::XmlElement(buzz::QN_MUC_OWNER_QUERY, true);
+
+		// Add CorrelationID for tracking purposes
+		FXmppJingle::AddCorrIdToStanza(*OwnerConfigQuery);
+
 		return OwnerConfigQuery;
 	}
 
@@ -148,6 +152,10 @@ private:
 	static buzz::XmlElement* MakeDefaultRequest()
 	{
 		buzz::XmlElement* OwnerQuery = new buzz::XmlElement(buzz::QN_MUC_OWNER_QUERY, true);
+
+		// Add CorrelationID for tracking purposes
+		FXmppJingle::AddCorrIdToStanza(*OwnerQuery);
+
 		buzz::XmlElement* XForm = new buzz::XmlElement(buzz::QN_XDATA_X, true);
 		XForm->SetAttr(buzz::QN_TYPE, buzz::STR_SUBMIT);
 		OwnerQuery->AddElement(XForm);
@@ -158,6 +166,10 @@ private:
 		const FRoomFeatureValuePairs& RoomFeatureValuePairs)
 	{
 		buzz::XmlElement* OwnerQuery = new buzz::XmlElement(buzz::QN_MUC_OWNER_QUERY, true);
+
+		// Add CorrelationID for tracking purposes
+		FXmppJingle::AddCorrIdToStanza(*OwnerQuery);
+
 		buzz::XmlElement* XForm = new buzz::XmlElement(buzz::QN_XDATA_X, true);
 		XForm->SetAttr(buzz::QN_TYPE, buzz::STR_SUBMIT);
 
@@ -1065,6 +1077,9 @@ public:
 		Message.AddAttr(buzz::QN_TO, XmppRoom->chatroom_jid().Str());
 		Message.AddAttr(buzz::QN_TYPE, ChatType);
 
+		// Add CorrelationID for tracking purposes
+		FXmppJingle::AddCorrIdToStanza(Message);
+
 		buzz::XmlElement* Body = new buzz::XmlElement(buzz::QN_BODY);
 		Body->SetBodyText(TCHAR_TO_UTF8(*MsgBody));
 		Message.AddElement(Body);
