Copying //UE4/Dev-Core to //UE4/Dev-Main (Source: //UE4/Dev-Core @ 3208226)

diff --git a/Engine/Source/Runtime/Engine/Private/BlueprintGeneratedClass.cpp b/Engine/Source/Runtime/Engine/Private/BlueprintGeneratedClass.cpp
--- a/Engine/Source/Runtime/Engine/Private/BlueprintGeneratedClass.cpp
+++ b/Engine/Source/Runtime/Engine/Private/BlueprintGeneratedClass.cpp
@@ -268,6 +268,8 @@ UObject* UBlueprintGeneratedClass::GetArchetypeForCDO() const
 
 void UBlueprintGeneratedClass::SerializeDefaultObject(UObject* Object, FArchive& Ar)
 {
+	FScopeLock SerializeAndPostLoadLock(&SerializeAndPostLoadCritical);
+
 	Super::SerializeDefaultObject(Object, Ar);
 
 	if (Ar.IsLoading() && !Ar.IsObjectReferenceCollector() && Object == ClassDefaultObject)
@@ -280,6 +282,8 @@ void UBlueprintGeneratedClass::SerializeDefaultObject(UObject* Object, FArchive&
 
 void UBlueprintGeneratedClass::PostLoadDefaultObject(UObject* Object)
 {
+	FScopeLock SerializeAndPostLoadLock(&SerializeAndPostLoadCritical);
+
 	Super::PostLoadDefaultObject(Object);
 
 	if (Object == ClassDefaultObject)
@@ -1031,22 +1035,29 @@ void UBlueprintGeneratedClass::DestroyPersistentUberGraphFrame(UObject* Obj, boo
 void UBlueprintGeneratedClass::GetPreloadDependencies(TArray<UObject*>& OutDeps)
 {
 	Super::GetPreloadDependencies(OutDeps);
-	for (UField* Field = Children; Field; Field = Field->Next)
-	{
-		OutDeps.Add(Field);
-	}
-	OutDeps.Add(GetSuperClass());
+	
+	// Super handles parent class and fields
 	OutDeps.Add(GetSuperClass()->GetDefaultObject());
 
-	OutDeps.Add(UberGraphFunction);
-	OutDeps.Add(GetInheritableComponentHandler());
+	if (UberGraphFunction)
+	{
+		OutDeps.Add(UberGraphFunction);
+	}
+	if (UInheritableComponentHandler* Handler = GetInheritableComponentHandler())
+	{
+		OutDeps.Add(Handler);
+	}
+	
 	UObject *CDO = GetDefaultObject();
 	if (CDO)
 	{
 		ForEachObjectWithOuter(CDO, [&OutDeps](UObject* SubObj)
 		{
-			OutDeps.Add(SubObj->GetClass());
-			OutDeps.Add(SubObj->GetArchetype());
+			if (SubObj->HasAllFlags(RF_DefaultSubObject))
+			{
+				OutDeps.Add(SubObj->GetClass());
+				OutDeps.Add(SubObj->GetArchetype());
+			}
 		});
 	}
 }
diff --git a/Engine/Source/Runtime/Engine/Private/CheatManager.cpp b/Engine/Source/Runtime/Engine/Private/CheatManager.cpp
--- a/Engine/Source/Runtime/Engine/Private/CheatManager.cpp
+++ b/Engine/Source/Runtime/Engine/Private/CheatManager.cpp
@@ -464,7 +464,7 @@ void UCheatManager::ViewPlayer( const FString& S )
 	AController* Controller = NULL;
 	for( FConstControllerIterator Iterator = GetWorld()->GetControllerIterator(); Iterator; ++Iterator )
 	{
-		Controller = *Iterator;
+		Controller = Iterator->Get();
 		if ( Controller->PlayerState && (FCString::Stricmp(*Controller->PlayerState->PlayerName, *S) == 0 ) )
 		{
 			break;
@@ -662,7 +662,7 @@ void UCheatManager::ServerToggleAILogging_Implementation()
 	{
 		for (FConstPlayerControllerIterator Iterator = World->GetPlayerControllerIterator(); Iterator; ++Iterator)
 		{
-			APlayerController* PC = *Iterator;
+			APlayerController* PC = Iterator->Get();
 			if (PC)
 			{
 				PC->OnServerStartedVisualLogger(bToggleAILogging);
diff --git a/Engine/Source/Runtime/Engine/Private/Components/ActorComponent.cpp b/Engine/Source/Runtime/Engine/Private/Components/ActorComponent.cpp
--- a/Engine/Source/Runtime/Engine/Private/Components/ActorComponent.cpp
+++ b/Engine/Source/Runtime/Engine/Private/Components/ActorComponent.cpp
@@ -365,7 +365,7 @@ ULevel* UActorComponent::GetComponentLevel() const
 {
 	// For model components Level is outer object
 	AActor* MyOwner = GetOwner();
-	return (MyOwner ? CastChecked<ULevel>(MyOwner->GetOuter()) : Cast<ULevel>( GetOuter() ) );
+	return (MyOwner ? Cast<ULevel>(MyOwner->GetOuter()) : Cast<ULevel>( GetOuter() ) );
 }
 
 bool UActorComponent::ComponentIsInLevel(const ULevel *TestLevel) const
diff --git a/Engine/Source/Runtime/Engine/Private/Components/ReflectionCaptureComponent.cpp b/Engine/Source/Runtime/Engine/Private/Components/ReflectionCaptureComponent.cpp
--- a/Engine/Source/Runtime/Engine/Private/Components/ReflectionCaptureComponent.cpp
+++ b/Engine/Source/Runtime/Engine/Private/Components/ReflectionCaptureComponent.cpp
@@ -11,6 +11,7 @@
 #include "EngineModule.h"
 #include "RendererInterface.h"
 #include "ShaderCompiler.h"
+#include "LoadTimesObjectVersion.h"
 #include "RenderingObjectVersion.h"
 #include "Engine/SphereReflectionCapture.h"
 #include "Components/SphereReflectionCaptureComponent.h"
@@ -63,6 +64,19 @@ int32 UReflectionCaptureComponent::GetReflectionCaptureSize_RenderThread()
 	return SanatizeReflectionCaptureSize(CVarReflectionCaptureSize.GetValueOnRenderThread());
 }
 
+
+void UReflectionCaptureComponent::ReleaseHDRData()
+{
+	ENQUEUE_UNIQUE_RENDER_COMMAND_ONEPARAMETER(
+		ReleaseHDRData,
+		FReflectionCaptureFullHDR*, FullHDRData, FullHDRData,
+		{
+			delete FullHDRData;
+		});
+
+	FullHDRData = nullptr;
+}
+
 void UWorld::UpdateAllReflectionCaptures()
 {
 	if ( FeatureLevel < ERHIFeatureLevel::SM4)
@@ -326,26 +340,30 @@ void FReflectionCaptureFullHDR::InitializeFromUncompressedData(const TArray<uint
 	INC_MEMORY_STAT_BY(STAT_ReflectionCaptureMemory, CompressedCapturedData.GetAllocatedSize());
 }
 
-void FReflectionCaptureFullHDR::GetUncompressedData(TArray<uint8>& UncompressedData) const
+TRefCountPtr<FReflectionCaptureUncompressedData> FReflectionCaptureFullHDR::GetUncompressedData() const
 {
-	FMemoryReader Ar(CompressedCapturedData);
-
-	// Note: change REFLECTIONCAPTURE_FULL_DERIVEDDATA_VER when modifying the serialization layout
-	int32 UncompressedSize;
-	Ar << UncompressedSize;
-
-	int32 CompressedSize;
-	Ar << CompressedSize;
+	// If we have serialized uncompressed data (from a cook), use it rather than uncompressing
+	if (StoredUncompressedData)
+	{
+		return StoredUncompressedData;
+	}
+	else
+	{
+		check(CompressedCapturedData.Num() > 0);
+		FMemoryReader Ar(CompressedCapturedData);
 
-	TArray<uint8> CompressedData;
-	CompressedData.Empty(CompressedSize);
-	CompressedData.AddUninitialized(CompressedSize);
-	Ar.Serialize(CompressedData.GetData(), CompressedSize);
+		// Note: change REFLECTIONCAPTURE_FULL_DERIVEDDATA_VER when modifying the serialization layout
+		int32 UncompressedSize;
+		Ar << UncompressedSize;
 
-	UncompressedData.Empty(UncompressedSize);
-	UncompressedData.AddUninitialized(UncompressedSize);
+		int32 CompressedSize;
+		Ar << CompressedSize;
 
-	verify(FCompression::UncompressMemory((ECompressionFlags)COMPRESS_ZLIB, UncompressedData.GetData(), UncompressedSize, CompressedData.GetData(), CompressedSize));
+		TRefCountPtr<FReflectionCaptureUncompressedData> UncompressedDataOut = new FReflectionCaptureUncompressedData(UncompressedSize);
+		const uint8* SourceData = &CompressedCapturedData[Ar.Tell()];
+		verify(FCompression::UncompressMemory((ECompressionFlags)COMPRESS_ZLIB, UncompressedDataOut->GetData(), UncompressedSize, SourceData, CompressedSize));
+		return UncompressedDataOut;
+	}
 }
 
 FColor RGBMEncode( FLinearColor Color )
@@ -522,14 +540,12 @@ void FReflectionCaptureEncodedHDRDerivedData::GenerateFromDerivedDataSource(cons
 {
 	const int32 NumMips = FMath::CeilLogTwo(FullHDRData.CubemapSize) + 1;
 
-	TArray<uint8> CubemapData;
-	FullHDRData.GetUncompressedData(CubemapData);
+	TRefCountPtr<FReflectionCaptureUncompressedData> SourceCubemapData = FullHDRData.GetUncompressedData();
 
 	int32 SourceMipBaseIndex = 0;
 	int32 DestMipBaseIndex = 0;
 
-	CapturedData.Empty(CubemapData.Num() * sizeof(FColor) / sizeof(FFloat16Color));
-	CapturedData.AddUninitialized(CubemapData.Num() * sizeof(FColor) / sizeof(FFloat16Color));
+	CapturedData = new FReflectionCaptureUncompressedData(SourceCubemapData->Size() * sizeof(FColor) / sizeof(FFloat16Color));
 
 	// Note: change REFLECTIONCAPTURE_ENCODED_DERIVEDDATA_VER when modifying the encoded data layout or contents
 
@@ -539,8 +555,8 @@ void FReflectionCaptureEncodedHDRDerivedData::GenerateFromDerivedDataSource(cons
 		const int32 SourceCubeFaceBytes = MipSize * MipSize * sizeof(FFloat16Color);
 		const int32 DestCubeFaceBytes = MipSize * MipSize * sizeof(FColor);
 
-		const FFloat16Color*	MipSrcData = (const FFloat16Color*)&CubemapData[ SourceMipBaseIndex ];
-		FColor*					MipDstData = (FColor*)&CapturedData[ DestMipBaseIndex ];
+		const FFloat16Color*	MipSrcData = (const FFloat16Color*)SourceCubemapData->GetData(SourceMipBaseIndex);
+		FColor*					MipDstData = (FColor*)CapturedData->GetData(DestMipBaseIndex);
 
 		// Fix cubemap seams by averaging colors across edges
 
@@ -623,8 +639,8 @@ void FReflectionCaptureEncodedHDRDerivedData::GenerateFromDerivedDataSource(cons
 		{
 			const int32 FaceSourceIndex = SourceMipBaseIndex + CubeFace * SourceCubeFaceBytes;
 			const int32 FaceDestIndex = DestMipBaseIndex + CubeFace * DestCubeFaceBytes;
-			const FFloat16Color* FaceSourceData = (const FFloat16Color*)&CubemapData[FaceSourceIndex];
-			FColor* FaceDestData = (FColor*)&CapturedData[FaceDestIndex];
+			const FFloat16Color* FaceSourceData = (const FFloat16Color*)SourceCubemapData->GetData(FaceSourceIndex);
+			FColor* FaceDestData = (FColor*)CapturedData->GetData(FaceDestIndex);
 
 			// Convert each texel from linear space FP16 to RGBM FColor
 			// Note: Brightness on the capture is baked into the encoded HDR data
@@ -656,7 +672,6 @@ FString FReflectionCaptureEncodedHDRDerivedData::GetDDCKeyString(const FGuid& St
 
 FReflectionCaptureEncodedHDRDerivedData::~FReflectionCaptureEncodedHDRDerivedData()
 {
-	DEC_MEMORY_STAT_BY(STAT_ReflectionCaptureMemory,CapturedData.GetAllocatedSize());
 }
 
 TRefCountPtr<FReflectionCaptureEncodedHDRDerivedData> FReflectionCaptureEncodedHDRDerivedData::GenerateEncodedHDRData(const FReflectionCaptureFullHDR& FullHDRData, const FGuid& StateId, float Brightness)
@@ -665,20 +680,20 @@ TRefCountPtr<FReflectionCaptureEncodedHDRDerivedData> FReflectionCaptureEncodedH
 	const FString KeyString = GetDDCKeyString(StateId, FullHDRData.CubemapSize);
 
 	COOK_STAT(auto Timer = ReflectionCaptureCookStats::UsageStats.TimeSyncWork());
-	bool DDCHit = GetDerivedDataCacheRef().GetSynchronous(*KeyString, EncodedHDRData->CapturedData);
+	bool DDCHit = GetDerivedDataCacheRef().GetSynchronous(*KeyString, EncodedHDRData->CapturedData->GetArray() );
 	if (!DDCHit)
 	{
 		EncodedHDRData->GenerateFromDerivedDataSource(FullHDRData, Brightness);
 
-		if (EncodedHDRData->CapturedData.Num() > 0)
+		if (EncodedHDRData->CapturedData->Size() > 0)
 		{
-			GetDerivedDataCacheRef().Put(*KeyString, EncodedHDRData->CapturedData);
+			GetDerivedDataCacheRef().Put(*KeyString, EncodedHDRData->CapturedData->GetArray() );
 		}
 	}
-	COOK_STAT(Timer.AddHitOrMiss(DDCHit ? FCookStats::CallStats::EHitOrMiss::Hit : FCookStats::CallStats::EHitOrMiss::Miss, EncodedHDRData->CapturedData.Num()));
+	EncodedHDRData->CapturedData->UpdateMemoryTracking();
+	COOK_STAT(Timer.AddHitOrMiss(DDCHit ? FCookStats::CallStats::EHitOrMiss::Hit : FCookStats::CallStats::EHitOrMiss::Miss, EncodedHDRData->CapturedData->Size()));
 
-	check(EncodedHDRData->CapturedData.Num() > 0);
-	INC_MEMORY_STAT_BY(STAT_ReflectionCaptureMemory,EncodedHDRData->CapturedData.GetAllocatedSize());
+	check(EncodedHDRData->CapturedData->Size() > 0);
 	return EncodedHDRData;
 }
 
@@ -693,11 +708,10 @@ public:
 	FReflectionTextureCubeResource() :
 		Size(0),
 		NumMips(0),
-		Format(PF_Unknown),
-		SourceData(NULL)
+		Format(PF_Unknown)
 	{}
 
-	void SetupParameters(int32 InSize, int32 InNumMips, EPixelFormat InFormat, TArray<uint8>* InSourceData)
+	void SetupParameters(int32 InSize, int32 InNumMips, EPixelFormat InFormat, TRefCountPtr<FReflectionCaptureUncompressedData> InSourceData)
 	{
 		Size = InSize;
 		NumMips = InNumMips;
@@ -713,7 +727,7 @@ public:
 
 		if (SourceData)
 		{
-			check(SourceData->Num() > 0);
+			check(SourceData->Size() > 0);
 
 			const int32 BlockBytes = GPixelFormats[Format].BlockBytes;
 			int32 MipBaseIndex = 0;
@@ -733,7 +747,7 @@ public:
 					{
 						uint8* DestPtr = ((uint8*)DestBuffer + Y * DestStride);
 						const int32 SourceIndex = MipBaseIndex + CubeFace * CubeFaceBytes + Y * MipSize * BlockBytes;
-						const uint8* SourcePtr = &(*SourceData)[SourceIndex];
+						const uint8* SourcePtr = SourceData->GetData(SourceIndex);
 						FMemory::Memcpy(DestPtr, SourcePtr, MipSize * BlockBytes);
 					}
 
@@ -747,8 +761,7 @@ public:
 			{
 				// Toss the source data now that we've created the cubemap
 				// Note: can't do this if we ever use this texture resource in the editor and want to save the data later
-				DEC_MEMORY_STAT_BY(STAT_ReflectionCaptureMemory,SourceData->GetAllocatedSize());
-				SourceData->Empty();
+				SourceData = nullptr;
 			}
 		}
 
@@ -793,7 +806,9 @@ private:
 	int32 NumMips;
 	EPixelFormat Format;
 	FTextureCubeRHIRef TextureCubeRHI;
-	TArray<uint8>* SourceData;
+
+	// Source data. Note that this is owned by the cubemap
+	TRefCountPtr<FReflectionCaptureUncompressedData> SourceData;
 };
 
 
@@ -967,6 +982,7 @@ void UReflectionCaptureComponent::Serialize(FArchive& Ar)
 	DECLARE_SCOPE_CYCLE_COUNTER(TEXT("UReflectionCaptureComponent::Serialize"), STAT_ReflectionCaptureComponent_Serialize, STATGROUP_LoadTime);
 
 	Ar.UsingCustomVersion(FRenderingObjectVersion::GUID);
+	Ar.UsingCustomVersion(FLoadTimesObjectVersion::GUID);
 
 	Super::Serialize(Ar);
 
@@ -1019,7 +1035,10 @@ void UReflectionCaptureComponent::Serialize(FArchive& Ar)
 					if (FullHDRData)
 					{
 						Ar << FullHDRData->CubemapSize;
-						Ar << FullHDRData->CompressedCapturedData;
+
+						// Raw data needs to be uncompressed on cooked platforms to avoid decompression hitches
+						TRefCountPtr<FReflectionCaptureUncompressedData> UncompressedData = FullHDRData->GetUncompressedData();
+						Ar << UncompressedData->GetArray();
 					}
 				}
 				else
@@ -1041,7 +1060,7 @@ void UReflectionCaptureComponent::Serialize(FArchive& Ar)
 
 					if (bValid)
 					{
-						Ar << EncodedHDRData->CapturedData;
+						Ar << EncodedHDRData->CapturedData->GetArray();
 					}
 					else if (!IsTemplate())
 					{
@@ -1075,13 +1094,23 @@ void UReflectionCaptureComponent::Serialize(FArchive& Ar)
 						FullHDRData = new FReflectionCaptureFullHDR();
 
 						Ar << FullHDRData->CubemapSize;
-						Ar << FullHDRData->CompressedCapturedData;
+						if (Ar.CustomVer(FLoadTimesObjectVersion::GUID) >= FLoadTimesObjectVersion::UncompressedReflectionCapturesForCookedBuilds)
+						{
+							// Raw data needs to be uncompressed on cooked platforms to avoid hitches
+							FullHDRData->StoredUncompressedData = new FReflectionCaptureUncompressedData();
+							Ar << FullHDRData->StoredUncompressedData->GetArray();
+							FullHDRData->StoredUncompressedData->UpdateMemoryTracking();
+						}
+						else
+						{
+							Ar << FullHDRData->CompressedCapturedData;
+						}
 					}
 					else 
 					{
 						check(CurrentFormat == EncodedHDR);
 						EncodedHDRDerivedData = new FReflectionCaptureEncodedHDRDerivedData();
-						Ar << EncodedHDRDerivedData->CapturedData;
+						Ar << EncodedHDRDerivedData->CapturedData->GetArray();
 					}
 				}
 				else if (CurrentFormat == EncodedHDR)
@@ -1140,7 +1169,7 @@ void UReflectionCaptureComponent::PostLoad()
 		if (GMaxRHIFeatureLevel == ERHIFeatureLevel::SM4)
 		{
 			SM4FullHDRCubemapTexture = new FReflectionTextureCubeResource();
-			SM4FullHDRCubemapTexture->SetupParameters(FullHDRData->CubemapSize, FMath::CeilLogTwo(FullHDRData->CubemapSize) + 1, PF_FloatRGBA, &FullHDRData->GetCapturedDataForSM4Load());
+			SM4FullHDRCubemapTexture->SetupParameters(FullHDRData->CubemapSize, FMath::CeilLogTwo(FullHDRData->CubemapSize) + 1, PF_FloatRGBA, FullHDRData->GetCapturedDataForSM4Load() );
 			BeginInitResource(SM4FullHDRCubemapTexture);
 		}
 
@@ -1152,18 +1181,13 @@ void UReflectionCaptureComponent::PostLoad()
 
 	if (EncodedHDRDerivedData && bEncodedDataRequired)
 	{
-		if (FPlatformProperties::RequiresCookedData())
-		{
-			INC_MEMORY_STAT_BY(STAT_ReflectionCaptureMemory, EncodedHDRDerivedData->CapturedData.GetAllocatedSize());
-		}
-
 		int32 EncodedCubemapSize = EncodedHDRDerivedData->CalculateCubemapDimension();
 
 		if (EncodedCubemapSize == ReflectionCaptureSize)
 		{
 			// Create a cubemap texture out of the encoded HDR data
 			EncodedHDRCubemapTexture = new FReflectionTextureCubeResource();
-			EncodedHDRCubemapTexture->SetupParameters(EncodedCubemapSize, FMath::CeilLogTwo(EncodedCubemapSize) + 1, PF_B8G8R8A8, &EncodedHDRDerivedData->CapturedData);
+			EncodedHDRCubemapTexture->SetupParameters(EncodedCubemapSize, FMath::CeilLogTwo(EncodedCubemapSize) + 1, PF_B8G8R8A8, EncodedHDRDerivedData->CapturedData);
 			BeginInitResource(EncodedHDRCubemapTexture);
 		}
 		else
@@ -1538,11 +1562,6 @@ void UReflectionCaptureComponent::PreFeatureLevelChange(ERHIFeatureLevel::Type P
 		if (FullHDRData)
 		{
 			EncodedHDRDerivedData = FReflectionCaptureEncodedHDRDerivedData::GenerateEncodedHDRData(*FullHDRData, StateId, Brightness);
-			if (FPlatformProperties::RequiresCookedData() && EncodedHDRDerivedData)
-			{
-				INC_MEMORY_STAT_BY(STAT_ReflectionCaptureMemory, EncodedHDRDerivedData->CapturedData.GetAllocatedSize());
-			}
-
 			if (EncodedHDRCubemapTexture == nullptr)
 			{
 				EncodedHDRCubemapTexture = new FReflectionTextureCubeResource();
@@ -1550,18 +1569,12 @@ void UReflectionCaptureComponent::PreFeatureLevelChange(ERHIFeatureLevel::Type P
 
 			int32 EncodedCubemapSize = EncodedHDRDerivedData->CalculateCubemapDimension();
 
-			EncodedHDRCubemapTexture->SetupParameters(EncodedCubemapSize, FMath::CeilLogTwo(EncodedCubemapSize) + 1, PF_B8G8R8A8, &EncodedHDRDerivedData->CapturedData);
+			EncodedHDRCubemapTexture->SetupParameters(EncodedCubemapSize, FMath::CeilLogTwo(EncodedCubemapSize) + 1, PF_B8G8R8A8, EncodedHDRDerivedData->CapturedData);
 			BeginInitResource(EncodedHDRCubemapTexture);
 		}
 	}
 	else
 	{
-		// Release ES2's unused textures & data.
-		if (FPlatformProperties::RequiresCookedData() && EncodedHDRDerivedData)
-		{
-			DEC_MEMORY_STAT_BY(STAT_ReflectionCaptureMemory, EncodedHDRDerivedData->CapturedData.GetAllocatedSize());
-		}
-
 		EncodedHDRDerivedData = nullptr;
 		if (EncodedHDRCubemapTexture)
 		{
diff --git a/Engine/Source/Runtime/Engine/Private/Components/SceneComponent.cpp b/Engine/Source/Runtime/Engine/Private/Components/SceneComponent.cpp
--- a/Engine/Source/Runtime/Engine/Private/Components/SceneComponent.cpp
+++ b/Engine/Source/Runtime/Engine/Private/Components/SceneComponent.cpp
@@ -2370,7 +2370,7 @@ void USceneComponent::UpdatePhysicsVolume( bool bTriggerNotifiers )
 				bool bAnyPotentialOverlap = false;
 				for (auto VolumeIter = MyWorld->GetNonDefaultPhysicsVolumeIterator(); VolumeIter && !bAnyPotentialOverlap; ++VolumeIter, ++VolumeIndex)
 				{
-					const APhysicsVolume* Volume = *VolumeIter;
+					const APhysicsVolume* Volume = VolumeIter->Get();
 					if (Volume != nullptr)
 					{
 						const USceneComponent* VolumeRoot = Volume->GetRootComponent();
diff --git a/Engine/Source/Runtime/Engine/Private/Components/StaticMeshComponent.cpp b/Engine/Source/Runtime/Engine/Private/Components/StaticMeshComponent.cpp
--- a/Engine/Source/Runtime/Engine/Private/Components/StaticMeshComponent.cpp
+++ b/Engine/Source/Runtime/Engine/Private/Components/StaticMeshComponent.cpp
@@ -512,7 +512,7 @@ const FMeshMapBuildData* UStaticMeshComponent::GetMeshMapBuildData(const FStatic
 
 	if (LODInfo.OverrideMapBuildData)
 	{
-		return LODInfo.OverrideMapBuildData.GetOwnedPointer();
+		return LODInfo.OverrideMapBuildData.Get();
 	}
 
 	AActor* Owner = GetOwner();
diff --git a/Engine/Source/Runtime/Engine/Private/CoreSettings.cpp b/Engine/Source/Runtime/Engine/Private/CoreSettings.cpp
--- a/Engine/Source/Runtime/Engine/Private/CoreSettings.cpp
+++ b/Engine/Source/Runtime/Engine/Private/CoreSettings.cpp
@@ -10,8 +10,9 @@ float GAsyncLoadingTimeLimit = 5.0f;
 int32 GAsyncLoadingUseFullTimeLimit = 1;
 float GPriorityAsyncLoadingExtraTime = 20.0f;
 float GLevelStreamingActorsUpdateTimeLimit = 5.0f;
+float GLevelStreamingUnregisterComponentsTimeLimit = 1.0f;
 int32 GLevelStreamingComponentsRegistrationGranularity = 10;
-
+int32 GLevelStreamingComponentsUnregistrationGranularity = 5;
 
 static FAutoConsoleVariableRef CVarUseBackgroundLevelStreaming(
 	TEXT("s.UseBackgroundLevelStreaming"),
@@ -48,6 +49,13 @@ static FAutoConsoleVariableRef CVarLevelStreamingActorsUpdateTimeLimit(
 	ECVF_Default
 	);
 
+static FAutoConsoleVariableRef CVarLevelStreamingUnregisterComponentsTimeLimit(
+	TEXT("s.UnregisterComponentsTimeLimit"),
+	GLevelStreamingUnregisterComponentsTimeLimit,
+	TEXT("Maximum allowed time to spend for actor unregistration steps during level streaming (ms per frame). If this is zero then we don't timeslice"),
+	ECVF_Default
+);
+
 static FAutoConsoleVariableRef CVarLevelStreamingComponentsRegistrationGranularity(
 	TEXT("s.LevelStreamingComponentsRegistrationGranularity"),
 	GLevelStreamingComponentsRegistrationGranularity,
@@ -55,6 +63,13 @@ static FAutoConsoleVariableRef CVarLevelStreamingComponentsRegistrationGranulari
 	ECVF_Default
 	);
 
+static FAutoConsoleVariableRef CVarLevelStreamingComponentsUnregistrationGranularity(
+	TEXT("s.LevelStreamingComponentsUnregistrationGranularity"),
+	GLevelStreamingComponentsUnregistrationGranularity,
+	TEXT("Batching granularity used to unregister actor components during level unstreaming."),
+	ECVF_Default
+	);
+
 UStreamingSettings::UStreamingSettings()
 	: Super()
 {
diff --git a/Engine/Source/Runtime/Engine/Private/Curves/CurveBase.cpp b/Engine/Source/Runtime/Engine/Private/Curves/CurveBase.cpp
--- a/Engine/Source/Runtime/Engine/Private/Curves/CurveBase.cpp
+++ b/Engine/Source/Runtime/Engine/Private/Curves/CurveBase.cpp
@@ -166,10 +166,9 @@ void UCurveBase::GetAssetRegistryTags(TArray<FAssetRegistryTag>& OutTags) const
 	Super::GetAssetRegistryTags(OutTags);
 }
 
-
 void UCurveBase::PostInitProperties()
 {
-	if (!HasAnyFlags(RF_ClassDefaultObject))
+	if (IsAsset())
 	{
 		AssetImportData = NewObject<UAssetImportData>(this, TEXT("AssetImportData"));
 	}
@@ -177,11 +176,16 @@ void UCurveBase::PostInitProperties()
 	Super::PostInitProperties();
 }
 
-
 void UCurveBase::PostLoad()
 {
 	Super::PostLoad();
 
+	if (!IsAsset() && AssetImportData)
+	{
+		// UCurves inside Blueprints previously created these sub objects incorrectly, so ones loaded off disk will still exist
+		AssetImportData = nullptr;
+	}
+
 	if (!ImportPath_DEPRECATED.IsEmpty() && AssetImportData)
 	{
 		FAssetImportInfo Info;
diff --git a/Engine/Source/Runtime/Engine/Private/DataTableJSON.cpp b/Engine/Source/Runtime/Engine/Private/DataTableJSON.cpp
--- a/Engine/Source/Runtime/Engine/Private/DataTableJSON.cpp
+++ b/Engine/Source/Runtime/Engine/Private/DataTableJSON.cpp
@@ -131,7 +131,12 @@ bool FDataTableExporterJSON::WriteStructEntry(const void* InRowData, const UProp
 {
 	const FString Identifier = DataTableUtils::GetPropertyExportName(InProperty, DTExportFlags);
 
-	if (const UNumericProperty *NumProp = Cast<const UNumericProperty>(InProperty))
+	if (const UEnumProperty* EnumProp = Cast<const UEnumProperty>(InProperty))
+	{
+		const FString PropertyValue = DataTableUtils::GetPropertyValueAsString(EnumProp, (uint8*)InRowData, DTExportFlags);
+		JsonWriter->WriteValue(Identifier, PropertyValue);
+	}
+	else if (const UNumericProperty *NumProp = Cast<const UNumericProperty>(InProperty))
 	{
 		if (NumProp->IsEnum())
 		{
@@ -192,7 +197,12 @@ bool FDataTableExporterJSON::WriteStructEntry(const void* InRowData, const UProp
 
 bool FDataTableExporterJSON::WriteArrayEntry(const UProperty* InProperty, const void* InPropertyData)
 {
-	if (const UNumericProperty *NumProp = Cast<const UNumericProperty>(InProperty))
+	if (const UEnumProperty* EnumProp = Cast<const UEnumProperty>(InProperty))
+	{
+		const FString PropertyValue = DataTableUtils::GetPropertyValueAsString(InProperty, (uint8*)InPropertyData, DTExportFlags);
+		JsonWriter->WriteValue(PropertyValue);
+	}
+	else if (const UNumericProperty *NumProp = Cast<const UNumericProperty>(InProperty))
 	{
 		if (NumProp->IsEnum())
 		{
@@ -403,7 +413,31 @@ bool FDataTableImporterJSON::ReadStructEntry(const TSharedRef<FJsonValue>& InPar
 {
 	const TCHAR* const ParsedPropertyType = JSONTypeToString(InParsedPropertyValue->Type);
 
-	if (UNumericProperty *NumProp = Cast<UNumericProperty>(InProperty))
+	if (UEnumProperty* EnumProp = Cast<UEnumProperty>(InProperty))
+	{
+		FString EnumValue;
+		if (InParsedPropertyValue->TryGetString(EnumValue))
+		{
+			FString Error = DataTableUtils::AssignStringToProperty(EnumValue, InProperty, (uint8*)InRowData);
+			if (!Error.IsEmpty())
+			{
+				ImportProblems.Add(FString::Printf(TEXT("Property '%s' on row '%s' has invalid enum value: %s."), *InColumnName, *InRowName.ToString(), *EnumValue));
+				return false;
+			}
+		}
+		else
+		{
+			int64 PropertyValue = 0;
+			if (!InParsedPropertyValue->TryGetNumber(PropertyValue))
+			{
+				ImportProblems.Add(FString::Printf(TEXT("Property '%s' on row '%s' is the incorrect type. Expected Integer, got %s."), *InColumnName, *InRowName.ToString(), ParsedPropertyType));
+				return false;
+			}
+
+			EnumProp->GetUnderlyingProperty()->SetIntPropertyValue(InPropertyData, PropertyValue);
+		}
+	}
+	else if (UNumericProperty *NumProp = Cast<UNumericProperty>(InProperty))
 	{
 		FString EnumValue;
 		if (NumProp->IsEnum() && InParsedPropertyValue->TryGetString(EnumValue))
@@ -518,7 +552,31 @@ bool FDataTableImporterJSON::ReadArrayEntry(const TSharedRef<FJsonValue>& InPars
 {
 	const TCHAR* const ParsedPropertyType = JSONTypeToString(InParsedPropertyValue->Type);
 
-	if (UNumericProperty *NumProp = Cast<UNumericProperty>(InProperty))
+	if (UEnumProperty* EnumProp = Cast<UEnumProperty>(InProperty))
+	{
+		FString EnumValue;
+		if (InParsedPropertyValue->TryGetString(EnumValue))
+		{
+			FString Error = DataTableUtils::AssignStringToProperty(EnumValue, InProperty, (uint8*)InPropertyData);
+			if (!Error.IsEmpty())
+			{
+				ImportProblems.Add(FString::Printf(TEXT("Entry %d on property '%s' on row '%s' has invalid enum value: %s."), InArrayEntryIndex, *InColumnName, *InRowName.ToString(), *EnumValue));
+				return false;
+			}
+		}
+		else
+		{
+			int64 PropertyValue = 0;
+			if (!InParsedPropertyValue->TryGetNumber(PropertyValue))
+			{
+				ImportProblems.Add(FString::Printf(TEXT("Entry %d on property '%s' on row '%s' is the incorrect type. Expected Integer, got %s."), InArrayEntryIndex, *InColumnName, *InRowName.ToString(), ParsedPropertyType));
+				return false;
+			}
+
+			EnumProp->GetUnderlyingProperty()->SetIntPropertyValue(InPropertyData, PropertyValue);
+		}
+	}
+	else if (UNumericProperty *NumProp = Cast<UNumericProperty>(InProperty))
 	{
 		FString EnumValue;
 		if (NumProp->IsEnum() && InParsedPropertyValue->TryGetString(EnumValue))
diff --git a/Engine/Source/Runtime/Engine/Private/DataTableUtils.cpp b/Engine/Source/Runtime/Engine/Private/DataTableUtils.cpp
--- a/Engine/Source/Runtime/Engine/Private/DataTableUtils.cpp
+++ b/Engine/Source/Runtime/Engine/Private/DataTableUtils.cpp
@@ -30,10 +30,22 @@ void AssignStringToProperty(const FString& InString, const UProperty* InProp, ui
 
 	bool bNeedsImport = true;
 
-	const UByteProperty* EnumProp = Cast<const UByteProperty>(InProp);
-	if (EnumProp && EnumProp->IsEnum())
+	UEnum* Enum = nullptr;
+
+	if (const UEnumProperty* EnumProp = Cast<const UEnumProperty>(InProp))
+	{
+		Enum = EnumProp->GetEnum();
+	}
+	else if (const UByteProperty* ByteProp = Cast<const UByteProperty>(InProp))
+	{
+		if (ByteProp->IsEnum())
+		{
+			Enum = ByteProp->GetIntPropertyEnum();
+		}
+	}
+
+	if (Enum)
 	{
-		UEnum* Enum = EnumProp->GetIntPropertyEnum();
 		// Enum properties may use the friendly name in their import data, however the UPropertyByte::ImportText function will only accept the internal enum entry name
 		// Detect if we're using a friendly name for an entry, and if so, try and map it to the correct internal name before performing the import
 		const int32 EnumIndex = Enum->FindEnumIndex(*InString);
@@ -69,15 +81,25 @@ void GetPropertyValueAsString(const UProperty* InProp, const uint8* InData, cons
 {
 	if (!!(InDTExportFlags & EDataTableExportFlags::UsePrettyEnumNames))
 	{
-		const UByteProperty* EnumProp = Cast<const UByteProperty>(InProp);
-		if (EnumProp && EnumProp->IsEnum())
+		UEnum* Enum = nullptr;
+		int64 Val = 0;
+		if (const UEnumProperty* EnumProp = Cast<UEnumProperty>(InProp))
 		{
-			if (UUserDefinedEnum* Enum = Cast<UUserDefinedEnum>(EnumProp->GetIntPropertyEnum()))
-			{
-				const uint8* ValuePtr = InProp->ContainerPtrToValuePtr<uint8>(InData, InIndex);
-				OutString.Append(GetSourceString(Enum->GetEnumTextByValue(*ValuePtr)));
-				return;
-			}
+			Enum = EnumProp->GetEnum();
+			const uint8* ValuePtr = InProp->ContainerPtrToValuePtr<uint8>(InData, InIndex);
+			Val = EnumProp->GetUnderlyingProperty()->GetSignedIntPropertyValue(ValuePtr);
+		}
+		else if (const UByteProperty* ByteProp = Cast<UByteProperty>(InProp))
+		{
+			Enum = ByteProp->GetIntPropertyEnum();
+			const uint8* ValuePtr = InProp->ContainerPtrToValuePtr<uint8>(InData, InIndex);
+			Val = *ValuePtr;
+		}
+
+		if (UUserDefinedEnum* UDEnum = Cast<UUserDefinedEnum>(Enum))
+		{
+			OutString.Append(GetSourceString(UDEnum->GetEnumTextByValue(Val)));
+			return;
 		}
 	}
 #if WITH_EDITOR
@@ -335,7 +357,8 @@ bool DataTableUtils::IsSupportedTableProperty(const UProperty* InProp)
 			InProp->IsA(UStructProperty::StaticClass()) ||
 			InProp->IsA(UByteProperty::StaticClass()) ||
 			InProp->IsA(UTextProperty::StaticClass()) ||
-			InProp->IsA(UArrayProperty::StaticClass())
+			InProp->IsA(UArrayProperty::StaticClass()) ||
+			InProp->IsA(UEnumProperty::StaticClass())
 			);
 }
 
diff --git a/Engine/Source/Runtime/Engine/Private/DemoNetDriver.cpp b/Engine/Source/Runtime/Engine/Private/DemoNetDriver.cpp
--- a/Engine/Source/Runtime/Engine/Private/DemoNetDriver.cpp
+++ b/Engine/Source/Runtime/Engine/Private/DemoNetDriver.cpp
@@ -1855,7 +1855,7 @@ void UDemoNetDriver::ProcessSeamlessTravel(int32 LevelIndex)
 	TArray<AController*> Controllers;
 	for (FConstControllerIterator Iterator = World->GetControllerIterator(); Iterator; ++Iterator)
 	{
-		Controllers.Add(*Iterator);
+		Controllers.Add(Iterator->Get());
 	}
 
 	for (int i = 0; i < Controllers.Num(); i++)
diff --git a/Engine/Source/Runtime/Engine/Private/DeveloperSettings.cpp b/Engine/Source/Runtime/Engine/Private/DeveloperSettings.cpp
--- a/Engine/Source/Runtime/Engine/Private/DeveloperSettings.cpp
+++ b/Engine/Source/Runtime/Engine/Private/DeveloperSettings.cpp
@@ -133,6 +133,12 @@ void UDeveloperSettings::ExportValuesToConsoleVariables(UProperty* PropertyThatC
 			{
 				CVar->Set(ByteProperty->GetPropertyValue_InContainer(this), ECVF_SetByProjectSetting);
 			}
+			else if (UEnumProperty* EnumProperty = Cast<UEnumProperty>(PropertyThatChanged))
+			{
+				UNumericProperty* UnderlyingProp = EnumProperty->GetUnderlyingProperty();
+				void* PropertyAddress = EnumProperty->ContainerPtrToValuePtr<void>(this);
+				CVar->Set((int32)UnderlyingProp->GetSignedIntPropertyValue(PropertyAddress), ECVF_SetByProjectSetting);
+			}
 			else if (UBoolProperty* BoolProperty = Cast<UBoolProperty>(PropertyThatChanged))
 			{
 				CVar->Set((int32)BoolProperty->GetPropertyValue_InContainer(this), ECVF_SetByProjectSetting);
diff --git a/Engine/Source/Runtime/Engine/Private/DistanceFieldAtlas.cpp b/Engine/Source/Runtime/Engine/Private/DistanceFieldAtlas.cpp
--- a/Engine/Source/Runtime/Engine/Private/DistanceFieldAtlas.cpp
+++ b/Engine/Source/Runtime/Engine/Private/DistanceFieldAtlas.cpp
@@ -8,6 +8,7 @@
 #include "StaticMeshResources.h"
 #include "DistanceFieldAtlas.h"
 #include "CookStats.h"
+#include "UniquePtr.h"
 
 #if WITH_EDITOR
 #include "DerivedDataCacheInterface.h"
@@ -338,7 +339,7 @@ public:
 		)
 		: NextThreadIndex(0)
 		, AsyncQueue(*InAsyncQueue)
-		, Thread(NULL)
+		, Thread(nullptr)
 		, bIsRunning(false)
 		, bForceFinish(false)
 	{}
@@ -359,7 +360,7 @@ public:
 		check(!bIsRunning);
 
 		bForceFinish = false;
-		Thread = FRunnableThread::Create(this, *FString::Printf(TEXT("BuildDistanceFieldThread%u"), NextThreadIndex), 0, TPri_Normal, FPlatformAffinity::GetPoolThreadMask());
+		Thread.Reset(FRunnableThread::Create(this, *FString::Printf(TEXT("BuildDistanceFieldThread%u"), NextThreadIndex), 0, TPri_Normal, FPlatformAffinity::GetPoolThreadMask()));
 		NextThreadIndex++;
 	}
 
@@ -372,9 +373,9 @@ private:
 	FDistanceFieldAsyncQueue& AsyncQueue;
 
 	/** The runnable thread */
-	TScopedPointer<FRunnableThread> Thread;
+	TUniquePtr<FRunnableThread> Thread;
 
-	TScopedPointer<FQueuedThreadPool> WorkerThreadPool;
+	TUniquePtr<FQueuedThreadPool> WorkerThreadPool;
 
 	volatile bool bIsRunning;
 	volatile bool bForceFinish;
@@ -401,7 +402,7 @@ uint32 FBuildDistanceFieldThreadRunnable::Run()
 		{
 			if (!WorkerThreadPool)
 			{
-				WorkerThreadPool = CreateWorkerThreadPool();
+				WorkerThreadPool.Reset(CreateWorkerThreadPool());
 			}
 
 			AsyncQueue.Build(Task, *WorkerThreadPool);
@@ -413,7 +414,7 @@ uint32 FBuildDistanceFieldThreadRunnable::Run()
 		}
 	}
 
-	WorkerThreadPool = NULL;
+	WorkerThreadPool = nullptr;
 
 	return 0;
 }
@@ -434,7 +435,7 @@ FDistanceFieldAsyncQueue::FDistanceFieldAsyncQueue()
 	MeshUtilities = NULL;
 #endif
 
-	ThreadRunnable = new FBuildDistanceFieldThreadRunnable(this);
+	ThreadRunnable = MakeUnique<FBuildDistanceFieldThreadRunnable>(this);
 }
 
 FDistanceFieldAsyncQueue::~FDistanceFieldAsyncQueue()
@@ -462,7 +463,7 @@ void FDistanceFieldAsyncQueue::AddTask(FAsyncDistanceFieldTask* Task)
 	}
 	else
 	{
-		TScopedPointer<FQueuedThreadPool> WorkerThreadPool(CreateWorkerThreadPool());
+		TUniquePtr<FQueuedThreadPool> WorkerThreadPool(CreateWorkerThreadPool());
 		Build(Task, *WorkerThreadPool);
 	}
 #else
diff --git a/Engine/Source/Runtime/Engine/Private/DrawDebugHelpers.cpp b/Engine/Source/Runtime/Engine/Private/DrawDebugHelpers.cpp
--- a/Engine/Source/Runtime/Engine/Private/DrawDebugHelpers.cpp
+++ b/Engine/Source/Runtime/Engine/Private/DrawDebugHelpers.cpp
@@ -713,7 +713,7 @@ void DrawDebugString(const UWorld* InWorld, FVector const& TextLocation, const F
 		// iterate through the player controller list
 		for( FConstPlayerControllerIterator Iterator = InWorld->GetPlayerControllerIterator(); Iterator; ++Iterator )
 		{
-			APlayerController* PlayerController = *Iterator;
+			APlayerController* PlayerController = Iterator->Get();
 			if (PlayerController->MyHUD)
 			{
 				PlayerController->MyHUD->AddDebugText(Text, BaseAct, Duration, TextLocation, TextLocation, TextColor, true, (TestBaseActor==NULL), false, nullptr, 1.0f, bDrawShadow);
@@ -728,7 +728,7 @@ void FlushDebugStrings( const UWorld* InWorld )
 	for( FConstPlayerControllerIterator Iterator = InWorld->GetPlayerControllerIterator(); Iterator; ++Iterator )
 	{
 		// if it's a player
-		APlayerController* PlayerController = *Iterator;
+		APlayerController* PlayerController = Iterator->Get();
 		if (PlayerController->MyHUD)
 		{
 			PlayerController->MyHUD->RemoveAllDebugStrings();
diff --git a/Engine/Source/Runtime/Engine/Private/EngineUtils.cpp b/Engine/Source/Runtime/Engine/Private/EngineUtils.cpp
--- a/Engine/Source/Runtime/Engine/Private/EngineUtils.cpp
+++ b/Engine/Source/Runtime/Engine/Private/EngineUtils.cpp
@@ -397,7 +397,7 @@ TArray<FSubLevelStatus> GetSubLevelsStatus( UWorld* World )
 
 	for( FConstPlayerControllerIterator Iterator = World->GetPlayerControllerIterator(); Iterator; ++Iterator )
 	{
-		APlayerController* PlayerController = *Iterator;
+		APlayerController* PlayerController = Iterator->Get();
 
 		if( PlayerController->GetPawn() != NULL )
 		{
diff --git a/Engine/Source/Runtime/Engine/Private/GameInstance.cpp b/Engine/Source/Runtime/Engine/Private/GameInstance.cpp
--- a/Engine/Source/Runtime/Engine/Private/GameInstance.cpp
+++ b/Engine/Source/Runtime/Engine/Private/GameInstance.cpp
@@ -646,7 +646,7 @@ APlayerController* UGameInstance::GetFirstLocalPlayerController(UWorld* World) c
 		{
 			if (*Iterator != nullptr && (*Iterator)->IsLocalController())
 			{
-				return *Iterator;
+				return Iterator->Get();
 			}
 		}
 	}
diff --git a/Engine/Source/Runtime/Engine/Private/GameMode.cpp b/Engine/Source/Runtime/Engine/Private/GameMode.cpp
--- a/Engine/Source/Runtime/Engine/Private/GameMode.cpp
+++ b/Engine/Source/Runtime/Engine/Private/GameMode.cpp
@@ -209,7 +209,7 @@ void AGameMode::HandleMatchHasStarted()
 	// start human players first
 	for( FConstPlayerControllerIterator Iterator = GetWorld()->GetPlayerControllerIterator(); Iterator; ++Iterator )
 	{
-		APlayerController* PlayerController = *Iterator;
+		APlayerController* PlayerController = Iterator->Get();
 		if ((PlayerController->GetPawn() == nullptr) && PlayerCanRestart(PlayerController))
 		{
 			RestartPlayer(PlayerController);
@@ -232,7 +232,7 @@ void AGameMode::HandleMatchHasStarted()
 	{
 		for( FConstPlayerControllerIterator Iterator = GetWorld()->GetPlayerControllerIterator(); Iterator; ++Iterator )
 		{
-			APlayerController* PlayerController = *Iterator;
+			APlayerController* PlayerController = Iterator->Get();
 			if( PlayerController->CheatManager != nullptr )
 			{
 				PlayerController->CheatManager->BugItGoString( BugLocString, BugRotString );
@@ -749,16 +749,15 @@ void AGameMode::PostSeamlessTravel()
 
 	// We have to make a copy of the controller list, since the code after this will destroy
 	// and create new controllers in the world's list
-	TArray<TAutoWeakObjectPtr<class AController> >	OldControllerList;
+	TArray<AController*> OldControllerList;
 	for (auto It = GetWorld()->GetControllerIterator(); It; ++It)
 	{
-		OldControllerList.Add(*It);
+		OldControllerList.Add(It->Get());
 	}
 
 	// handle players that are already loaded
-	for( FConstControllerIterator Iterator = OldControllerList.CreateConstIterator(); Iterator; ++Iterator )
+	for( AController* Controller : OldControllerList )
 	{
-		AController* Controller = *Iterator;
 		if (Controller->PlayerState)
 		{
 			APlayerController* PlayerController = Cast<APlayerController>(Controller);
diff --git a/Engine/Source/Runtime/Engine/Private/GameModeBase.cpp b/Engine/Source/Runtime/Engine/Private/GameModeBase.cpp
--- a/Engine/Source/Runtime/Engine/Private/GameModeBase.cpp
+++ b/Engine/Source/Runtime/Engine/Private/GameModeBase.cpp
@@ -120,7 +120,7 @@ int32 AGameModeBase::GetNumPlayers()
 	int32 PlayerCount = 0;
 	for (FConstPlayerControllerIterator Iterator = GetWorld()->GetPlayerControllerIterator(); Iterator; ++Iterator)
 	{
-		APlayerController* PlayerActor = *Iterator;
+		APlayerController* PlayerActor = Iterator->Get();
 		if (PlayerActor->PlayerState && !MustSpectate(PlayerActor))
 		{
 			PlayerCount++;
@@ -134,7 +134,7 @@ int32 AGameModeBase::GetNumSpectators()
 	int32 PlayerCount = 0;
 	for (FConstPlayerControllerIterator Iterator = GetWorld()->GetPlayerControllerIterator(); Iterator; ++Iterator)
 	{
-		APlayerController* PlayerActor = *Iterator;
+		APlayerController* PlayerActor = Iterator->Get();
 		if (PlayerActor->PlayerState && MustSpectate(PlayerActor))
 		{
 			PlayerCount++;
@@ -238,7 +238,7 @@ void AGameModeBase::ForceClearUnpauseDelegates(AActor* PauseActor)
 			// Try to find another player to be the worldsettings's Pauser
 			for (FConstPlayerControllerIterator Iterator = GetWorld()->GetPlayerControllerIterator(); Iterator; ++Iterator)
 			{
-				APlayerController* Player = *Iterator;
+				APlayerController* Player = Iterator->Get();
 				if (Player->PlayerState != nullptr
 					&&	Player->PlayerState != PC->PlayerState
 					&& !Player->IsPendingKillPending() && !Player->PlayerState->IsPendingKillPending())
@@ -280,7 +280,7 @@ void AGameModeBase::ResetLevel()
 	// Reset ALL controllers first
 	for (FConstControllerIterator Iterator = GetWorld()->GetControllerIterator(); Iterator; ++Iterator)
 	{
-		AController* Controller = *Iterator;
+		AController* Controller = Iterator->Get();
 		APlayerController* PlayerController = Cast<APlayerController>(Controller);
 		if (PlayerController)
 		{
@@ -324,7 +324,7 @@ APlayerController* AGameModeBase::ProcessClientTravel(FString& FURL, FGuid NextM
 	APlayerController* LocalPlayerController = nullptr;
 	for (FConstPlayerControllerIterator Iterator = GetWorld()->GetPlayerControllerIterator(); Iterator; ++Iterator)
 	{
-		APlayerController* PlayerController = *Iterator;
+		APlayerController* PlayerController = Iterator->Get();
 		if (Cast<UNetConnection>(PlayerController->Player) != nullptr)
 		{
 			// Remote player
@@ -541,16 +541,15 @@ void AGameModeBase::PostSeamlessTravel()
 
 	// We have to make a copy of the controller list, since the code after this will destroy
 	// and create new controllers in the world's list
-	TArray<TAutoWeakObjectPtr<AController> >	OldControllerList;
+	TArray<AController*> OldControllerList;
 	for (auto It = GetWorld()->GetControllerIterator(); It; ++It)
 	{
-		OldControllerList.Add(*It);
+		OldControllerList.Add(It->Get());
 	}
 
 	// Handle players that are already loaded
-	for (FConstControllerIterator Iterator = OldControllerList.CreateConstIterator(); Iterator; ++Iterator)
+	for (AController* Controller : OldControllerList)
 	{
-		AController* Controller = *Iterator;
 		if (Controller->PlayerState)
 		{
 			APlayerController* PlayerController = Cast<APlayerController>(Controller);
@@ -1291,7 +1290,7 @@ bool AGameModeBase::SpawnPlayerFromSimulate(const FVector& NewLocation, const FR
 			// Use the "auto-possess" pawn in the world, if there is one.
 			for (FConstPawnIterator Iterator = GetWorld()->GetPawnIterator(); Iterator; ++Iterator)
 			{
-				APawn* Pawn = *Iterator;
+				APawn* Pawn = Iterator->Get();
 				if (Pawn && Pawn->AutoPossessPlayer == EAutoReceiveInput::Player0)
 				{
 					if (Pawn->Controller == nullptr)
diff --git a/Engine/Source/Runtime/Engine/Private/GameSession.cpp b/Engine/Source/Runtime/Engine/Private/GameSession.cpp
--- a/Engine/Source/Runtime/Engine/Private/GameSession.cpp
+++ b/Engine/Source/Runtime/Engine/Private/GameSession.cpp
@@ -28,7 +28,7 @@ APlayerController* GetPlayerControllerFromNetId(UWorld* World, const FUniqueNetI
 		// Iterate through the controller list looking for the net id
 		for (FConstPlayerControllerIterator Iterator = World->GetPlayerControllerIterator(); Iterator; ++Iterator)
 		{
-			APlayerController* PlayerController = *Iterator;
+			APlayerController* PlayerController = Iterator->Get();
 			// Determine if this is a player with replication
 			if (PlayerController->PlayerState != NULL && PlayerController->PlayerState->UniqueId.IsValid())
 			{
@@ -61,7 +61,7 @@ void AGameSession::HandleMatchHasStarted()
 	{
 		for (FConstPlayerControllerIterator Iterator = World->GetPlayerControllerIterator(); Iterator; ++Iterator)
 		{
-			APlayerController* PlayerController = *Iterator;
+			APlayerController* PlayerController = Iterator->Get();
 			if (!PlayerController->IsLocalController())
 			{
 				PlayerController->ClientStartOnlineSession();
@@ -103,7 +103,7 @@ void AGameSession::HandleMatchHasEnded()
 	{
 		for (FConstPlayerControllerIterator Iterator = World->GetPlayerControllerIterator(); Iterator; ++Iterator)
 		{
-			APlayerController* PlayerController = *Iterator;
+			APlayerController* PlayerController = Iterator->Get();
 			if (!PlayerController->IsLocalController())
 			{
 				PlayerController->ClientEndOnlineSession();
@@ -339,7 +339,7 @@ void AGameSession::ReturnToMainMenuHost()
 	FConstPlayerControllerIterator Iterator = GetWorld()->GetPlayerControllerIterator();
 	for(; Iterator; ++Iterator)
 	{
-		Controller = *Iterator;
+		Controller = Iterator->Get();
 		if (Controller && !Controller->IsLocalPlayerController() && Controller->IsPrimaryPlayer())
 		{
 			// Clients
@@ -350,7 +350,7 @@ void AGameSession::ReturnToMainMenuHost()
 	Iterator.Reset();
 	for(; Iterator; ++Iterator)
 	{
-		Controller = *Iterator;
+		Controller = Iterator->Get();
 		if (Controller && Controller->IsLocalPlayerController() && Controller->IsPrimaryPlayer())
 		{
 			Controller->ClientReturnToMainMenu(LocalReturnReason);
diff --git a/Engine/Source/Runtime/Engine/Private/GameStateBase.cpp b/Engine/Source/Runtime/Engine/Private/GameStateBase.cpp
--- a/Engine/Source/Runtime/Engine/Private/GameStateBase.cpp
+++ b/Engine/Source/Runtime/Engine/Private/GameStateBase.cpp
@@ -79,7 +79,7 @@ void AGameStateBase::ReceivedGameModeClass()
 	// Tell each PlayerController that the Game class is here
 	for( FConstPlayerControllerIterator Iterator = GetWorld()->GetPlayerControllerIterator(); Iterator; ++Iterator )
 	{
-		APlayerController* const PlayerController = *Iterator;
+		APlayerController* const PlayerController = Iterator->Get();
 		if (PlayerController)
 		{
 			PlayerController->ReceivedGameModeClass(GameModeClass);
@@ -92,7 +92,7 @@ void AGameStateBase::ReceivedSpectatorClass()
 	// Tell each PlayerController that the Spectator class is here
 	for (FConstPlayerControllerIterator Iterator = GetWorld()->GetPlayerControllerIterator(); Iterator; ++Iterator)
 	{
-		APlayerController* const PlayerController = *Iterator;
+		APlayerController* const PlayerController = Iterator->Get();
 		if (PlayerController && PlayerController->IsLocalController())
 		{
 			PlayerController->ReceivedSpectatorClass(SpectatorClass);
diff --git a/Engine/Source/Runtime/Engine/Private/GlobalShader.cpp b/Engine/Source/Runtime/Engine/Private/GlobalShader.cpp
--- a/Engine/Source/Runtime/Engine/Private/GlobalShader.cpp
+++ b/Engine/Source/Runtime/Engine/Private/GlobalShader.cpp
@@ -261,12 +261,12 @@ void BackupGlobalShaderMap(FGlobalShaderBackupData& OutGlobalShaderBackup)
 		EShaderPlatform ShaderPlatform = GetFeatureLevelShaderPlatform((ERHIFeatureLevel::Type)i);
 		if (ShaderPlatform < EShaderPlatform::SP_NumPlatforms && GGlobalShaderMap[ShaderPlatform] != nullptr)
 		{
-			TArray<uint8>* ShaderData = new TArray<uint8>();
+			TUniquePtr<TArray<uint8>> ShaderData = MakeUnique<TArray<uint8>>();
 			FMemoryWriter Ar(*ShaderData);
 			GGlobalShaderMap[ShaderPlatform]->SerializeInline(Ar, true, true);
 			GGlobalShaderMap[ShaderPlatform]->RegisterSerializedShaders();
 			GGlobalShaderMap[ShaderPlatform]->Empty();
-			OutGlobalShaderBackup.FeatureLevelShaderData[i] = ShaderData;
+			OutGlobalShaderBackup.FeatureLevelShaderData[i] = MoveTemp(ShaderData);
 		}
 	}
 
diff --git a/Engine/Source/Runtime/Engine/Private/InstancedStaticMesh.h b/Engine/Source/Runtime/Engine/Private/InstancedStaticMesh.h
--- a/Engine/Source/Runtime/Engine/Private/InstancedStaticMesh.h
+++ b/Engine/Source/Runtime/Engine/Private/InstancedStaticMesh.h
@@ -670,7 +670,7 @@ private:
 		const bool bInstanced = GRHISupportsInstancing;
 
 		// Copy the parameters for LOD - all instances
-		UserData_AllInstances.MeshRenderData = InComponent->GetStaticMesh()->RenderData;
+		UserData_AllInstances.MeshRenderData = InComponent->GetStaticMesh()->RenderData.Get();
 		UserData_AllInstances.StartCullDistance = InComponent->InstanceStartCullDistance;
 		UserData_AllInstances.EndCullDistance = InComponent->InstanceEndCullDistance;
 		UserData_AllInstances.MinLOD = ClampedMinLOD;
diff --git a/Engine/Source/Runtime/Engine/Private/Interpolation.cpp b/Engine/Source/Runtime/Engine/Private/Interpolation.cpp
--- a/Engine/Source/Runtime/Engine/Private/Interpolation.cpp
+++ b/Engine/Source/Runtime/Engine/Private/Interpolation.cpp
@@ -887,7 +887,7 @@ PRAGMA_ENABLE_DEPRECATION_WARNINGS
 					// iterate through the controller list
 					for (FConstPlayerControllerIterator Iterator = GetWorld()->GetPlayerControllerIterator(); Iterator; ++Iterator)
 					{
-						APlayerController *PC = *Iterator;
+						APlayerController *PC = Iterator->Get();
 
 						// If it's a player and this sequence is compatible with the player...
 						if (IsMatineeCompatibleWithPlayer( PC ) )
@@ -1257,7 +1257,7 @@ void AMatineeActor::EnableCinematicMode(bool bEnable)
 	{
 		for (FConstPlayerControllerIterator Iterator = GetWorld()->GetPlayerControllerIterator(); Iterator; ++Iterator)
 		{
-			APlayerController *PC = *Iterator;
+			APlayerController *PC = Iterator->Get();
 			if (bReplicates || PC->IsLocalController())
 			{
 				PC->SetCinematicMode(bEnable, bHidePlayer, bHideHud, bDisableMovementInput, bDisableLookAtInput);
diff --git a/Engine/Source/Runtime/Engine/Private/KismetSystemLibrary.cpp b/Engine/Source/Runtime/Engine/Private/KismetSystemLibrary.cpp
--- a/Engine/Source/Runtime/Engine/Private/KismetSystemLibrary.cpp
+++ b/Engine/Source/Runtime/Engine/Private/KismetSystemLibrary.cpp
@@ -743,11 +743,16 @@ void UKismetSystemLibrary::SetBytePropertyByName(UObject* Object, FName Property
 {
 	if(Object != NULL)
 	{
-		UByteProperty* ByteProp = FindField<UByteProperty>(Object->GetClass(), PropertyName);
-		if(ByteProp != NULL)
+		if(UByteProperty* ByteProp = FindField<UByteProperty>(Object->GetClass(), PropertyName))
 		{
 			ByteProp->SetPropertyValue_InContainer(Object, Value);
-		}		
+		}
+		else if(UEnumProperty* EnumProp = FindField<UEnumProperty>(Object->GetClass(), PropertyName))
+		{
+			void* PropAddr = EnumProp->ContainerPtrToValuePtr<void>(Object);
+			UNumericProperty* UnderlyingProp = EnumProp->GetUnderlyingProperty();
+			UnderlyingProp->SetIntPropertyValue(PropAddr, (int64)Value);
+		}
 	}
 }
 
diff --git a/Engine/Source/Runtime/Engine/Private/Level.cpp b/Engine/Source/Runtime/Engine/Private/Level.cpp
--- a/Engine/Source/Runtime/Engine/Private/Level.cpp
+++ b/Engine/Source/Runtime/Engine/Private/Level.cpp
@@ -944,6 +944,42 @@ void ULevel::IncrementalUpdateComponents(int32 NumComponentsToUpdate, bool bReru
 	}
 }
 
+
+bool ULevel::IncrementalUnregisterComponents(int32 NumComponentsToUnregister)
+{
+	// A value of 0 means that we want to unregister all components.
+	if (NumComponentsToUnregister != 0)
+	{
+		// Only the game can use incremental update functionality.
+		checkf(OwningWorld->IsGameWorld(), TEXT("Cannot call IncrementalUnregisterComponents with non 0 argument in the Editor/ commandlets."));
+	}
+
+	// Find next valid actor to process components unregistration
+	int32 NumComponentsUnregistered = 0;
+	while (CurrentActorIndexForUnregisterComponents < Actors.Num())
+	{
+		AActor* Actor = Actors[CurrentActorIndexForUnregisterComponents];
+		if (Actor)
+		{
+			int32 NumComponents = Actor->GetComponents().Num();
+			NumComponentsUnregistered += NumComponents;
+			Actor->UnregisterAllComponents();
+		}
+		CurrentActorIndexForUnregisterComponents++;
+		if (NumComponentsUnregistered > NumComponentsToUnregister )
+		{
+			break;
+		}
+	}
+
+	if (CurrentActorIndexForUnregisterComponents == Actors.Num())
+	{
+		CurrentActorIndexForUnregisterComponents = 0;
+		return true;
+	}
+	return false;
+}
+
 #if WITH_EDITOR
 
 void ULevel::MarkLevelComponentsRenderStateDirty()
@@ -1249,11 +1285,11 @@ void ULevel::UpdateModelComponents()
 	}
 
 	// Initialize the model's index buffers.
-	for(TMap<UMaterialInterface*,TScopedPointer<FRawIndexBuffer16or32> >::TIterator IndexBufferIt(Model->MaterialIndexBuffers);
+	for(TMap<UMaterialInterface*,TUniquePtr<FRawIndexBuffer16or32> >::TIterator IndexBufferIt(Model->MaterialIndexBuffers);
 		IndexBufferIt;
 		++IndexBufferIt)
 	{
-		BeginInitResource(IndexBufferIt.Value());
+		BeginInitResource(IndexBufferIt->Value.Get());
 	}
 
 	// Can now release the model's vertex buffer, will have been used for collision
@@ -1444,11 +1480,11 @@ void ULevel::CommitModelSurfaces()
 		}
 
 		// Initialize the model's index buffers.
-		for(TMap<UMaterialInterface*,TScopedPointer<FRawIndexBuffer16or32> >::TIterator IndexBufferIt(Model->MaterialIndexBuffers);
+		for(TMap<UMaterialInterface*,TUniquePtr<FRawIndexBuffer16or32> >::TIterator IndexBufferIt(Model->MaterialIndexBuffers);
 			IndexBufferIt;
 			++IndexBufferIt)
 		{
-			BeginInitResource(IndexBufferIt.Value());
+			BeginInitResource(IndexBufferIt->Value.Get());
 		}
 
 		Model->bOnlyRebuildMaterialIndexBuffers = false;
@@ -1987,7 +2023,7 @@ void ULevel::PushPendingAutoReceiveInput(APlayerController* InPlayerController)
 	int32 Index = 0;
 	for( FConstPlayerControllerIterator Iterator = InPlayerController->GetWorld()->GetPlayerControllerIterator(); Iterator; ++Iterator )
 	{
-		APlayerController* PlayerController = *Iterator;
+		APlayerController* PlayerController = Iterator->Get();
 		if (InPlayerController == PlayerController)
 		{
 			PlayerIndex = Index;
@@ -2069,3 +2105,10 @@ bool ULevel::HasVisibilityRequestPending() const
 {
 	return (OwningWorld && this == OwningWorld->CurrentLevelPendingVisibility);
 }
+
+bool ULevel::HasVisibilityChangeRequestPending() const
+{
+	return (OwningWorld && ( this == OwningWorld->CurrentLevelPendingVisibility || this == OwningWorld->CurrentLevelPendingInvisibility ) );
+}
+
+
diff --git a/Engine/Source/Runtime/Engine/Private/LevelTick.cpp b/Engine/Source/Runtime/Engine/Private/LevelTick.cpp
--- a/Engine/Source/Runtime/Engine/Private/LevelTick.cpp
+++ b/Engine/Source/Runtime/Engine/Private/LevelTick.cpp
@@ -570,7 +570,7 @@ void UWorld::ProcessLevelStreamingVolumes(FVector* OverrideViewLocation)
 	bool bStreamingVolumesAreRelevant = false;
 	for( FConstPlayerControllerIterator Iterator = GetPlayerControllerIterator(); Iterator; ++Iterator )
 	{
-		APlayerController* PlayerActor = *Iterator;
+		APlayerController* PlayerActor = Iterator->Get();
 		if (PlayerActor->bIsUsingStreamingVolumes)
 		{
 			bStreamingVolumesAreRelevant = true;
@@ -704,7 +704,7 @@ void UWorld::ProcessLevelStreamingVolumes(FVector* OverrideViewLocation)
 				{
 					for( FConstPlayerControllerIterator Iterator = GetPlayerControllerIterator(); Iterator; ++Iterator )
 					{
-						APlayerController* PlayerController = *Iterator;
+						APlayerController* PlayerController = Iterator->Get();
 						PlayerController->LevelStreamingStatusChanged( 
 								LevelStreamingObject, 
 								LevelStreamingObject->bShouldBeLoaded, 
@@ -1403,7 +1403,7 @@ void UWorld::Tick( ELevelTick TickType, float DeltaSeconds )
 				// Update cameras last. This needs to be done before NetUpdates, and after all actors have been ticked.
 				for( FConstPlayerControllerIterator Iterator = GetPlayerControllerIterator(); Iterator; ++Iterator )
 				{
-					APlayerController* PlayerController = *Iterator;			
+					APlayerController* PlayerController = Iterator->Get();
 					PlayerController->UpdateCameraManager(DeltaSeconds);
 				}
 
diff --git a/Engine/Source/Runtime/Engine/Private/Materials/Material.cpp b/Engine/Source/Runtime/Engine/Private/Materials/Material.cpp
--- a/Engine/Source/Runtime/Engine/Private/Materials/Material.cpp
+++ b/Engine/Source/Runtime/Engine/Private/Materials/Material.cpp
@@ -3554,7 +3554,7 @@ void UMaterial::UpdateMaterialShaders(TArray<FShaderType*>& ShaderTypesToFlush,
 	FMaterial::UpdateEditorLoadedMaterialResources(ShaderPlatform);
 }
 
-void UMaterial::BackupMaterialShadersToMemory(TMap<FMaterialShaderMap*, TScopedPointer<TArray<uint8> > >& ShaderMapToSerializedShaderData)
+void UMaterial::BackupMaterialShadersToMemory(TMap<FMaterialShaderMap*, TUniquePtr<TArray<uint8> > >& ShaderMapToSerializedShaderData)
 {
 	// Process FMaterialShaderMap's referenced by UObjects (UMaterial, UMaterialInstance)
 	for (TObjectIterator<UMaterialInterface> It; It; ++It)
@@ -3605,7 +3605,7 @@ void UMaterial::BackupMaterialShadersToMemory(TMap<FMaterialShaderMap*, TScopedP
 	FMaterial::BackupEditorLoadedMaterialShadersToMemory(ShaderMapToSerializedShaderData);
 }
 
-void UMaterial::RestoreMaterialShadersFromMemory(const TMap<FMaterialShaderMap*, TScopedPointer<TArray<uint8> > >& ShaderMapToSerializedShaderData)
+void UMaterial::RestoreMaterialShadersFromMemory(const TMap<FMaterialShaderMap*, TUniquePtr<TArray<uint8> > >& ShaderMapToSerializedShaderData)
 {
 	// Process FMaterialShaderMap's referenced by UObjects (UMaterial, UMaterialInstance)
 	for (TObjectIterator<UMaterialInterface> It; It; ++It)
@@ -3627,7 +3627,7 @@ void UMaterial::RestoreMaterialShadersFromMemory(const TMap<FMaterialShaderMap*,
 
 					if (ShaderMap)
 					{
-						const TScopedPointer<TArray<uint8> >* ShaderData = ShaderMapToSerializedShaderData.Find(ShaderMap);
+						const TUniquePtr<TArray<uint8> >* ShaderData = ShaderMapToSerializedShaderData.Find(ShaderMap);
 
 						if (ShaderData)
 						{
@@ -3648,7 +3648,7 @@ void UMaterial::RestoreMaterialShadersFromMemory(const TMap<FMaterialShaderMap*,
 
 					if (ShaderMap)
 					{
-						const TScopedPointer<TArray<uint8> >* ShaderData = ShaderMapToSerializedShaderData.Find(ShaderMap);
+						const TUniquePtr<TArray<uint8> >* ShaderData = ShaderMapToSerializedShaderData.Find(ShaderMap);
 
 						if (ShaderData)
 						{
diff --git a/Engine/Source/Runtime/Engine/Private/Materials/MaterialExpressions.cpp b/Engine/Source/Runtime/Engine/Private/Materials/MaterialExpressions.cpp
--- a/Engine/Source/Runtime/Engine/Private/Materials/MaterialExpressions.cpp
+++ b/Engine/Source/Runtime/Engine/Private/Materials/MaterialExpressions.cpp
@@ -6345,8 +6345,7 @@ FString UMaterialExpressionSceneDepth::GetInputName(int32 InputIndex) const
 	if(InputIndex == 0)
 	{
 		// Display the current InputMode enum's display name.
-		UByteProperty* InputModeProperty = NULL;
-		InputModeProperty = FindField<UByteProperty>( UMaterialExpressionSceneDepth::StaticClass(), "InputMode" );
+		UByteProperty* InputModeProperty = FindField<UByteProperty>( UMaterialExpressionSceneDepth::StaticClass(), "InputMode" );
 		return InputModeProperty->Enum->GetEnumName((int32)InputMode.GetValue());
 	}
 	return TEXT("");
@@ -6610,8 +6609,7 @@ FString UMaterialExpressionSceneColor::GetInputName(int32 InputIndex) const
 	if(InputIndex == 0)
 	{
 		// Display the current InputMode enum's display name.
-		UByteProperty* InputModeProperty = NULL;
-		InputModeProperty = FindField<UByteProperty>( UMaterialExpressionSceneColor::StaticClass(), "InputMode" );
+		UByteProperty* InputModeProperty = FindField<UByteProperty>( UMaterialExpressionSceneColor::StaticClass(), "InputMode" );
 		return InputModeProperty->Enum->GetEnumName((int32)InputMode.GetValue());
 	}
 	return TEXT("");
diff --git a/Engine/Source/Runtime/Engine/Private/Materials/MaterialShared.cpp b/Engine/Source/Runtime/Engine/Private/Materials/MaterialShared.cpp
--- a/Engine/Source/Runtime/Engine/Private/Materials/MaterialShared.cpp
+++ b/Engine/Source/Runtime/Engine/Private/Materials/MaterialShared.cpp
@@ -2317,7 +2317,7 @@ void FMaterial::UpdateEditorLoadedMaterialResources(EShaderPlatform InShaderPlat
 	}
 }
 
-void FMaterial::BackupEditorLoadedMaterialShadersToMemory(TMap<FMaterialShaderMap*, TScopedPointer<TArray<uint8> > >& ShaderMapToSerializedShaderData)
+void FMaterial::BackupEditorLoadedMaterialShadersToMemory(TMap<FMaterialShaderMap*, TUniquePtr<TArray<uint8> > >& ShaderMapToSerializedShaderData)
 {
 	for (TSet<FMaterial*>::TIterator It(EditorLoadedMaterialResources); It; ++It)
 	{
@@ -2332,7 +2332,7 @@ void FMaterial::BackupEditorLoadedMaterialShadersToMemory(TMap<FMaterialShaderMa
 	}
 }
 
-void FMaterial::RestoreEditorLoadedMaterialShadersFromMemory(const TMap<FMaterialShaderMap*, TScopedPointer<TArray<uint8> > >& ShaderMapToSerializedShaderData)
+void FMaterial::RestoreEditorLoadedMaterialShadersFromMemory(const TMap<FMaterialShaderMap*, TUniquePtr<TArray<uint8> > >& ShaderMapToSerializedShaderData)
 {
 	for (TSet<FMaterial*>::TIterator It(EditorLoadedMaterialResources); It; ++It)
 	{
@@ -2341,7 +2341,7 @@ void FMaterial::RestoreEditorLoadedMaterialShadersFromMemory(const TMap<FMateria
 
 		if (ShaderMap)
 		{
-			const TScopedPointer<TArray<uint8> >* ShaderData = ShaderMapToSerializedShaderData.Find(ShaderMap);
+			const TUniquePtr<TArray<uint8> >* ShaderData = ShaderMapToSerializedShaderData.Find(ShaderMap);
 
 			if (ShaderData)
 			{
@@ -2359,11 +2359,11 @@ FMaterialUpdateContext::FMaterialUpdateContext(uint32 Options, EShaderPlatform I
 	bSyncWithRenderingThread = (Options & EOptions::SyncWithRenderingThread) != 0;
 	if (bReregisterComponents)
 	{
-		ComponentReregisterContext = new FGlobalComponentReregisterContext();
+		ComponentReregisterContext = MakeUnique<FGlobalComponentReregisterContext>();
 	}
 	else if (bRecreateRenderStates)
 	{
-		ComponentRecreateRenderStateContext = new FGlobalComponentRecreateRenderStateContext();
+		ComponentRecreateRenderStateContext = MakeUnique<FGlobalComponentRecreateRenderStateContext>();
 	}
 	if (bSyncWithRenderingThread)
 	{
@@ -2421,7 +2421,7 @@ FMaterialUpdateContext::~FMaterialUpdateContext()
 	TArray<const FMaterial*> MaterialResourcesToUpdate;
 	TArray<UMaterialInstance*> InstancesToUpdate;
 
-	bool bUpdateStaticDrawLists = !ComponentReregisterContext.IsValid() && !ComponentRecreateRenderStateContext.IsValid();
+	bool bUpdateStaticDrawLists = !ComponentReregisterContext && !ComponentRecreateRenderStateContext;
 
 	// If static draw lists must be updated, gather material resources from all updated materials.
 	if (bUpdateStaticDrawLists)
@@ -2510,11 +2510,11 @@ FMaterialUpdateContext::~FMaterialUpdateContext()
 		// safe, e.g. while a component is being registered.
 		GetRendererModule().UpdateStaticDrawListsForMaterials(MaterialResourcesToUpdate);
 	}
-	else if (ComponentReregisterContext.IsValid())
+	else if (ComponentReregisterContext)
 	{
 		ComponentReregisterContext.Reset();
 	}
-	else if (ComponentRecreateRenderStateContext.IsValid())
+	else if (ComponentRecreateRenderStateContext)
 	{
 		ComponentRecreateRenderStateContext.Reset();
 	}
diff --git a/Engine/Source/Runtime/Engine/Private/Materials/MaterialUniformExpressions.cpp b/Engine/Source/Runtime/Engine/Private/Materials/MaterialUniformExpressions.cpp
--- a/Engine/Source/Runtime/Engine/Private/Materials/MaterialUniformExpressions.cpp
+++ b/Engine/Source/Runtime/Engine/Private/Materials/MaterialUniformExpressions.cpp
@@ -303,7 +303,7 @@ void FUniformExpressionSet::CreateBufferStruct()
 	NextMemberOffset += 8;
 
 	const uint32 StructSize = Align(NextMemberOffset,UNIFORM_BUFFER_STRUCT_ALIGNMENT);
-	UniformBufferStruct = new FUniformBufferStruct(
+	UniformBufferStruct.Emplace(
 		MaterialLayoutName,
 		TEXT("MaterialUniforms"),
 		TEXT("Material"),
@@ -316,8 +316,7 @@ void FUniformExpressionSet::CreateBufferStruct()
 
 const FUniformBufferStruct& FUniformExpressionSet::GetUniformBufferStruct() const
 {
-	check(UniformBufferStruct);
-	return *UniformBufferStruct;
+	return UniformBufferStruct.GetValue();
 }
 
 FUniformBufferRHIRef FUniformExpressionSet::CreateUniformBuffer(const FMaterialRenderContext& MaterialRenderContext, FRHICommandList* CommandListIfLocalMode, struct FLocalUniformBuffer* OutLocalUniformBuffer) const
diff --git a/Engine/Source/Runtime/Engine/Private/Materials/ParameterCollection.cpp b/Engine/Source/Runtime/Engine/Private/Materials/ParameterCollection.cpp
--- a/Engine/Source/Runtime/Engine/Private/Materials/ParameterCollection.cpp
+++ b/Engine/Source/Runtime/Engine/Private/Materials/ParameterCollection.cpp
@@ -411,7 +411,7 @@ void UMaterialParameterCollection::CreateBufferStruct()
 	NextMemberOffset += VectorArraySize;
 	static FName LayoutName(TEXT("MaterialCollection"));
 	const uint32 StructSize = Align(NextMemberOffset,UNIFORM_BUFFER_STRUCT_ALIGNMENT);
-	UniformBufferStruct = new FUniformBufferStruct(
+	UniformBufferStruct = MakeUnique<FUniformBufferStruct>(
 		LayoutName,
 		TEXT("MaterialCollection"),
 		TEXT("MaterialCollection"),
diff --git a/Engine/Source/Runtime/Engine/Private/Model.cpp b/Engine/Source/Runtime/Engine/Private/Model.cpp
--- a/Engine/Source/Runtime/Engine/Private/Model.cpp
+++ b/Engine/Source/Runtime/Engine/Private/Model.cpp
@@ -503,10 +503,9 @@ void UModel::GetResourceSizeEx(FResourceSizeEx& CumulativeResourceSize)
 	
 	// I'm adding extra stuff that haven't been covered by Serialize 
 	// I don't have to include VertexFactories (based on Sam Z)
-	for(TMap<UMaterialInterface*,TScopedPointer<FRawIndexBuffer16or32> >::TConstIterator IndexBufferIt(MaterialIndexBuffers);IndexBufferIt;++IndexBufferIt)
+	for(TMap<UMaterialInterface*,TUniquePtr<FRawIndexBuffer16or32> >::TConstIterator IndexBufferIt(MaterialIndexBuffers);IndexBufferIt;++IndexBufferIt)
 	{
-		const TScopedPointer<FRawIndexBuffer16or32> &IndexBuffer = IndexBufferIt.Value();
-		CumulativeResourceSize.AddUnknownMemoryBytes(IndexBuffer->Indices.Num() * sizeof(uint32));
+		CumulativeResourceSize.AddUnknownMemoryBytes(IndexBufferIt->Value->Indices.Num() * sizeof(uint32));
 	}
 }
 
@@ -704,9 +703,9 @@ void UModel::ShrinkModel()
 void UModel::BeginReleaseResources()
 {
 	// Release the index buffers.
-	for(TMap<UMaterialInterface*,TScopedPointer<FRawIndexBuffer16or32> >::TIterator IndexBufferIt(MaterialIndexBuffers);IndexBufferIt;++IndexBufferIt)
+	for(TMap<UMaterialInterface*,TUniquePtr<FRawIndexBuffer16or32> >::TIterator IndexBufferIt(MaterialIndexBuffers);IndexBufferIt;++IndexBufferIt)
 	{
-		BeginReleaseResource(IndexBufferIt.Value());
+		BeginReleaseResource(IndexBufferIt->Value.Get());
 	}
 
 	// Release the vertex buffer and factory.
@@ -882,11 +881,10 @@ int32 UModel::BuildVertexBuffers()
 */
 void UModel::ClearLocalMaterialIndexBuffersData()
 {
-	TMap<UMaterialInterface*,TScopedPointer<FRawIndexBuffer16or32> >::TIterator MaterialIterator(MaterialIndexBuffers);
+	TMap<UMaterialInterface*,TUniquePtr<FRawIndexBuffer16or32> >::TIterator MaterialIterator(MaterialIndexBuffers);
 	for(; MaterialIterator; ++MaterialIterator)
 	{
-		FRawIndexBuffer16or32& IndexBuffer = *MaterialIterator.Value();
-		IndexBuffer.Indices.Empty();
+		MaterialIterator->Value->Indices.Empty();
 	}
 }
 
diff --git a/Engine/Source/Runtime/Engine/Private/ModelLight.cpp b/Engine/Source/Runtime/Engine/Private/ModelLight.cpp
--- a/Engine/Source/Runtime/Engine/Private/ModelLight.cpp
+++ b/Engine/Source/Runtime/Engine/Private/ModelLight.cpp
@@ -244,9 +244,9 @@ void UModelComponent::ApplyTempElements(bool bLightingWasSuccessful)
 		for (int32 ModelIndex = 0; ModelIndex < UpdatedModels.Num(); ModelIndex++)
 		{
 			UModel* Model = UpdatedModels[ModelIndex];
-			for(TMap<UMaterialInterface*,TScopedPointer<FRawIndexBuffer16or32> >::TIterator IndexBufferIt(Model->MaterialIndexBuffers); IndexBufferIt; ++IndexBufferIt)
+			for(TMap<UMaterialInterface*,TUniquePtr<FRawIndexBuffer16or32> >::TIterator IndexBufferIt(Model->MaterialIndexBuffers); IndexBufferIt; ++IndexBufferIt)
 			{
-				BeginReleaseResource(IndexBufferIt.Value());
+				BeginReleaseResource(IndexBufferIt->Value.Get());
 			}
 		}
 
@@ -266,9 +266,9 @@ void UModelComponent::ApplyTempElements(bool bLightingWasSuccessful)
 		for (int32 ModelIndex = 0; ModelIndex < UpdatedModels.Num(); ModelIndex++)
 		{
 			UModel* Model = UpdatedModels[ModelIndex];
-			for(TMap<UMaterialInterface*,TScopedPointer<FRawIndexBuffer16or32> >::TIterator IndexBufferIt(Model->MaterialIndexBuffers); IndexBufferIt; ++IndexBufferIt)
+			for(TMap<UMaterialInterface*,TUniquePtr<FRawIndexBuffer16or32> >::TIterator IndexBufferIt(Model->MaterialIndexBuffers); IndexBufferIt; ++IndexBufferIt)
 			{
-				BeginInitResource(IndexBufferIt.Value());
+				BeginInitResource(IndexBufferIt->Value.Get());
 			}
 
 			// Mark the model's package as dirty.
diff --git a/Engine/Source/Runtime/Engine/Private/ModelRender.cpp b/Engine/Source/Runtime/Engine/Private/ModelRender.cpp
--- a/Engine/Source/Runtime/Engine/Private/ModelRender.cpp
+++ b/Engine/Source/Runtime/Engine/Private/ModelRender.cpp
@@ -82,12 +82,12 @@ void UModelComponent::BuildRenderData()
 		FModelElement& Element = Elements[ElementIndex];
 
 		// Find the index buffer for the element's material.
-		TScopedPointer<FRawIndexBuffer16or32>* IndexBufferRef = TheModel->MaterialIndexBuffers.Find(Element.Material);
+		TUniquePtr<FRawIndexBuffer16or32>* IndexBufferRef = TheModel->MaterialIndexBuffers.Find(Element.Material);
 		if(!IndexBufferRef)
 		{
 			IndexBufferRef = &TheModel->MaterialIndexBuffers.Emplace(Element.Material,new FRawIndexBuffer16or32());
 		}
-		FRawIndexBuffer16or32* const IndexBuffer = *IndexBufferRef;
+		FRawIndexBuffer16or32* const IndexBuffer = IndexBufferRef->Get();
 		check(IndexBuffer);
 
 		Element.IndexBuffer = IndexBuffer;
diff --git a/Engine/Source/Runtime/Engine/Private/NetworkDriver.cpp b/Engine/Source/Runtime/Engine/Private/NetworkDriver.cpp
--- a/Engine/Source/Runtime/Engine/Private/NetworkDriver.cpp
+++ b/Engine/Source/Runtime/Engine/Private/NetworkDriver.cpp
@@ -2433,9 +2433,9 @@ void UNetDriver::ServerReplicateActors_BuildConsiderList( TArray<FNetworkObjectI
 			continue;
 		}
 
-		// Don't send actors that may still be streaming in
+		// Don't send actors that may still be streaming in or out
 		ULevel* Level = Actor->GetLevel();
-		if ( Level->HasVisibilityRequestPending() || Level->bIsAssociatingLevel )
+		if ( Level->HasVisibilityChangeRequestPending() || Level->bIsAssociatingLevel )
 		{
 			continue;
 		}
diff --git a/Engine/Source/Runtime/Engine/Private/Particles/ParticleComponents.cpp b/Engine/Source/Runtime/Engine/Private/Particles/ParticleComponents.cpp
--- a/Engine/Source/Runtime/Engine/Private/Particles/ParticleComponents.cpp
+++ b/Engine/Source/Runtime/Engine/Private/Particles/ParticleComponents.cpp
@@ -4521,7 +4521,7 @@ void UParticleSystemComponent::TickComponent(float DeltaTime, enum ELevelTick Ti
 	{
 		for (FConstPlayerControllerIterator Iterator = World->GetPlayerControllerIterator(); Iterator; ++Iterator)
 		{
-			APlayerController* PlayerController = *Iterator;
+			APlayerController* PlayerController = Iterator->Get();
 			if (PlayerController->IsLocalPlayerController())
 			{
 				FVector POVLoc;
@@ -6199,7 +6199,7 @@ int32 UParticleSystemComponent::DetermineLODLevelForLocation(const FVector& Effe
 		{
 			for( FConstPlayerControllerIterator Iterator = World->GetPlayerControllerIterator(); Iterator; ++Iterator )
 			{
-				APlayerController* PlayerController = *Iterator;
+				APlayerController* PlayerController = Iterator->Get();
 				if(PlayerController->IsLocalPlayerController())
 				{
 					FVector* POVLoc = new(PlayerViewLocations) FVector;
@@ -6863,7 +6863,7 @@ bool UParticleSystemComponent::ShouldComputeLODFromGameThread()
 
 		for (FConstPlayerControllerIterator Iterator = World->GetPlayerControllerIterator(); Iterator; ++Iterator)
 		{
-			APlayerController* PlayerController = *Iterator;
+			APlayerController* PlayerController = Iterator->Get();
 			if (PlayerController->IsLocalPlayerController())
 			{
 				bUseGameThread = true;
diff --git a/Engine/Source/Runtime/Engine/Private/Player.cpp b/Engine/Source/Runtime/Engine/Private/Player.cpp
--- a/Engine/Source/Runtime/Engine/Private/Player.cpp
+++ b/Engine/Source/Runtime/Engine/Private/Player.cpp
@@ -83,7 +83,7 @@ APlayerController* UPlayer::GetPlayerController(UWorld* InWorld) const
 	{
 		if ( (*Iterator)->GetLocalPlayer() == this )
 		{
-			return *Iterator;
+			return Iterator->Get();
 		}
 	}
 
diff --git a/Engine/Source/Runtime/Engine/Private/PlayerController.cpp b/Engine/Source/Runtime/Engine/Private/PlayerController.cpp
--- a/Engine/Source/Runtime/Engine/Private/PlayerController.cpp
+++ b/Engine/Source/Runtime/Engine/Private/PlayerController.cpp
@@ -477,7 +477,7 @@ ACameraActor* APlayerController::GetAutoActivateCameraForPlayer() const
 	int32 PlayerIndex = INDEX_NONE;
 	for( FConstPlayerControllerIterator Iterator = CurWorld->GetPlayerControllerIterator(); Iterator; ++Iterator, ++IterIndex )
 	{
-		const APlayerController* PlayerController = *Iterator;
+		const APlayerController* PlayerController = Iterator->Get();
 		if (PlayerController == this)
 		{
 			PlayerIndex = IterIndex;
@@ -490,7 +490,7 @@ ACameraActor* APlayerController::GetAutoActivateCameraForPlayer() const
 		// Find the matching camera
 		for( /*CameraIterater initialized above*/; CameraIterator; ++CameraIterator)
 		{
-			ACameraActor* CameraActor = *CameraIterator;
+			ACameraActor* CameraActor = CameraIterator->Get();
 			if (CameraActor && CameraActor->GetAutoActivatePlayerIndex() == PlayerIndex)
 			{
 				return CameraActor;
diff --git a/Engine/Source/Runtime/Engine/Private/PlayerState.cpp b/Engine/Source/Runtime/Engine/Private/PlayerState.cpp
--- a/Engine/Source/Runtime/Engine/Private/PlayerState.cpp
+++ b/Engine/Source/Runtime/Engine/Private/PlayerState.cpp
@@ -174,7 +174,7 @@ void APlayerState::OnRep_PlayerName()
 		{
 			for( FConstPlayerControllerIterator Iterator = GetWorld()->GetPlayerControllerIterator(); Iterator; ++Iterator )
 			{
-				APlayerController* PlayerController = *Iterator;
+				APlayerController* PlayerController = Iterator->Get();
 				if( PlayerController )
 				{
 					PlayerController->ClientReceiveLocalizedMessage( EngineMessageClass, 2, this );
@@ -191,7 +191,7 @@ void APlayerState::OnRep_PlayerName()
 		{
 			for( FConstPlayerControllerIterator Iterator = GetWorld()->GetPlayerControllerIterator(); Iterator; ++Iterator )
 			{
-				APlayerController* PlayerController = *Iterator;
+				APlayerController* PlayerController = Iterator->Get();
 				if( PlayerController )
 				{
 					PlayerController->ClientReceiveLocalizedMessage( EngineMessageClass, WelcomeMessageNum, this );
@@ -229,7 +229,7 @@ void APlayerState::Destroyed()
 	{
 		for (FConstPlayerControllerIterator Iterator = World->GetPlayerControllerIterator(); Iterator; ++Iterator)
 		{
-			APlayerController* PlayerController = *Iterator;
+			APlayerController* PlayerController = Iterator->Get();
 			if( PlayerController )
 			{
 				PlayerController->ClientReceiveLocalizedMessage( EngineMessageClass, 4, this);
diff --git a/Engine/Source/Runtime/Engine/Private/RendererSupport.cpp b/Engine/Source/Runtime/Engine/Private/RendererSupport.cpp
--- a/Engine/Source/Runtime/Engine/Private/RendererSupport.cpp
+++ b/Engine/Source/Runtime/Engine/Private/RendererSupport.cpp
@@ -14,11 +14,12 @@
 #include "RendererInterface.h"
 #include "HotReloadInterface.h"
 #include "ComponentReregisterContext.h"
+#include "UniquePtr.h"
 
 /** Clears and optionally backs up all references to renderer module classes in other modules, particularly engine. */
 static void ClearReferencesToRendererModuleClasses(
 	TMap<UWorld*, bool>& WorldsToUpdate, 
-	TMap<FMaterialShaderMap*, TScopedPointer<TArray<uint8> > >& ShaderMapToSerializedShaderData,
+	TMap<FMaterialShaderMap*, TUniquePtr<TArray<uint8> > >& ShaderMapToSerializedShaderData,
 	FGlobalShaderBackupData& GlobalShaderBackup,
 	TMap<FShaderType*, FString>& ShaderTypeNames,
 	TMap<const FShaderPipelineType*, FString>& ShaderPipelineTypeNames,
@@ -122,9 +123,9 @@ void RecompileRendererModule()
 }
 
 /** Restores systems that need references to classes in the renderer module. */
-void RestoreReferencesToRendererModuleClasses(
+static void RestoreReferencesToRendererModuleClasses(
 	const TMap<UWorld*, bool>& WorldsToUpdate, 
-	const TMap<FMaterialShaderMap*, TScopedPointer<TArray<uint8> > >& ShaderMapToSerializedShaderData,
+	const TMap<FMaterialShaderMap*, TUniquePtr<TArray<uint8> > >& ShaderMapToSerializedShaderData,
 	const FGlobalShaderBackupData& GlobalShaderBackup,
 	const TMap<FShaderType*, FString>& ShaderTypeNames,
 	const TMap<const FShaderPipelineType*, FString>& ShaderPipelineTypeNames,
@@ -207,7 +208,7 @@ void RecompileRenderer(const TArray<FString>& Args)
 		SCOPED_SUSPEND_RENDERING_THREAD(true);
 
 		TMap<UWorld*, bool> WorldsToUpdate;
-		TMap<FMaterialShaderMap*, TScopedPointer<TArray<uint8> > > ShaderMapToSerializedShaderData;
+		TMap<FMaterialShaderMap*, TUniquePtr<TArray<uint8> > > ShaderMapToSerializedShaderData;
 		FGlobalShaderBackupData GlobalShaderBackup;
 		TMap<FShaderType*, FString> ShaderTypeNames;
 		TMap<const FShaderPipelineType*, FString> ShaderPipelineTypeNames;
diff --git a/Engine/Source/Runtime/Engine/Private/RepLayout.cpp b/Engine/Source/Runtime/Engine/Private/RepLayout.cpp
--- a/Engine/Source/Runtime/Engine/Private/RepLayout.cpp
+++ b/Engine/Source/Runtime/Engine/Private/RepLayout.cpp
@@ -2663,10 +2663,16 @@ uint32 FRepLayout::AddPropertyCmd( UProperty * Property, int32 Offset, int32 Rel
 	Cmd.CompatibleChecksum	= FCrc::StrCrc32( *Property->GetCPPType( nullptr, 0 ).ToLower(), Cmd.CompatibleChecksum );		// Evolve by property type
 	Cmd.CompatibleChecksum	= FCrc::StrCrc32( *FString::Printf( TEXT( "%i" ), StaticArrayIndex ), Cmd.CompatibleChecksum );	// Evolve by StaticArrayIndex (to make all unrolled static array elements unique)
 
+	UProperty * UnderlyingProperty = Property;
+	if ( UEnumProperty * EnumProperty = Cast< UEnumProperty >( Property ) )
+	{
+		UnderlyingProperty = EnumProperty->GetUnderlyingProperty();
+	}
+
 	// Try to special case to custom types we know about
-	if ( Property->IsA( UStructProperty::StaticClass() ) )
+	if ( UnderlyingProperty->IsA( UStructProperty::StaticClass() ) )
 	{
-		UStructProperty * StructProp = Cast< UStructProperty >( Property );
+		UStructProperty * StructProp = Cast< UStructProperty >( UnderlyingProperty );
 		UScriptStruct * Struct = StructProp->Struct;
 		if ( Struct->GetFName() == NAME_Vector )
 		{
@@ -2709,39 +2715,39 @@ uint32 FRepLayout::AddPropertyCmd( UProperty * Property, int32 Offset, int32 Rel
 			UE_LOG( LogRep, VeryVerbose, TEXT( "AddPropertyCmd: Falling back to default type for property [%s]" ), *Cmd.Property->GetFullName() );
 		}
 	}
-	else if ( Property->IsA( UBoolProperty::StaticClass() ) )
+	else if ( UnderlyingProperty->IsA( UBoolProperty::StaticClass() ) )
 	{
 		Cmd.Type = REPCMD_PropertyBool;
 	}
-	else if ( Property->IsA( UFloatProperty::StaticClass() ) )
+	else if ( UnderlyingProperty->IsA( UFloatProperty::StaticClass() ) )
 	{
 		Cmd.Type = REPCMD_PropertyFloat;
 	}
-	else if ( Property->IsA( UIntProperty::StaticClass() ) )
+	else if ( UnderlyingProperty->IsA( UIntProperty::StaticClass() ) )
 	{
 		Cmd.Type = REPCMD_PropertyInt;
 	}
-	else if ( Property->IsA( UByteProperty::StaticClass() ) )
+	else if ( UnderlyingProperty->IsA( UByteProperty::StaticClass() ) )
 	{
 		Cmd.Type = REPCMD_PropertyByte;
 	}
-	else if ( Property->IsA( UObjectPropertyBase::StaticClass() ) )
+	else if ( UnderlyingProperty->IsA( UObjectPropertyBase::StaticClass() ) )
 	{
 		Cmd.Type = REPCMD_PropertyObject;
 	}
-	else if ( Property->IsA( UNameProperty::StaticClass() ) )
+	else if ( UnderlyingProperty->IsA( UNameProperty::StaticClass() ) )
 	{
 		Cmd.Type = REPCMD_PropertyName;
 	}
-	else if ( Property->IsA( UUInt32Property::StaticClass() ) )
+	else if ( UnderlyingProperty->IsA( UUInt32Property::StaticClass() ) )
 	{
 		Cmd.Type = REPCMD_PropertyUInt32;
 	}
-	else if ( Property->IsA( UUInt64Property::StaticClass() ) )
+	else if ( UnderlyingProperty->IsA( UUInt64Property::StaticClass() ) )
 	{
 		Cmd.Type = REPCMD_PropertyUInt64;
 	}
-	else if ( Property->IsA( UStrProperty::StaticClass() ) )
+	else if ( UnderlyingProperty->IsA( UStrProperty::StaticClass() ) )
 	{
 		Cmd.Type = REPCMD_PropertyString;
 	}
diff --git a/Engine/Source/Runtime/Engine/Private/ShaderCompiler/ShaderCompiler.cpp b/Engine/Source/Runtime/Engine/Private/ShaderCompiler/ShaderCompiler.cpp
--- a/Engine/Source/Runtime/Engine/Private/ShaderCompiler/ShaderCompiler.cpp
+++ b/Engine/Source/Runtime/Engine/Private/ShaderCompiler/ShaderCompiler.cpp
@@ -1390,12 +1390,12 @@ FShaderCompilingManager::FShaderCompilingManager() :
 	if (FShaderCompileXGEThreadRunnable::IsSupported())
 	{
 		UE_LOG(LogShaderCompilers, Display, TEXT("Using XGE Shader Compiler."));
-		Thread = new FShaderCompileXGEThreadRunnable(this);
+		Thread = MakeUnique<FShaderCompileXGEThreadRunnable>(this);
 	}
 	else
 	{
 		UE_LOG(LogShaderCompilers, Display, TEXT("Using Local Shader Compiler."));
-		Thread = new FShaderCompileThreadRunnable(this);
+		Thread = MakeUnique<FShaderCompileThreadRunnable>(this);
 	}
 	Thread->StartThread();
 }
diff --git a/Engine/Source/Runtime/Engine/Private/StaticMesh.cpp b/Engine/Source/Runtime/Engine/Private/StaticMesh.cpp
--- a/Engine/Source/Runtime/Engine/Private/StaticMesh.cpp
+++ b/Engine/Source/Runtime/Engine/Private/StaticMesh.cpp
@@ -2176,7 +2176,7 @@ static FStaticMeshRenderData& GetPlatformStaticMeshRenderData(UStaticMesh* Mesh,
 	check(Mesh && Mesh->RenderData);
 	const FStaticMeshLODSettings& PlatformLODSettings = Platform->GetStaticMeshLODSettings();
 	FString PlatformDerivedDataKey = BuildStaticMeshDerivedDataKey(Mesh, PlatformLODSettings.GetLODGroup(Mesh->LODGroup));
-	FStaticMeshRenderData* PlatformRenderData = Mesh->RenderData;
+	FStaticMeshRenderData* PlatformRenderData = Mesh->RenderData.Get();
 
 	if (Mesh->GetOutermost()->HasAnyPackageFlags(PKG_FilterEditorOnly))
 	{
@@ -2186,7 +2186,7 @@ static FStaticMeshRenderData& GetPlatformStaticMeshRenderData(UStaticMesh* Mesh,
 
 	while (PlatformRenderData && PlatformRenderData->DerivedDataKey != PlatformDerivedDataKey)
 	{
-		PlatformRenderData = PlatformRenderData->NextCachedRenderData;
+		PlatformRenderData = PlatformRenderData->NextCachedRenderData.Get();
 	}
 	if (PlatformRenderData == NULL)
 	{
@@ -2194,8 +2194,8 @@ static FStaticMeshRenderData& GetPlatformStaticMeshRenderData(UStaticMesh* Mesh,
 		PlatformRenderData = new FStaticMeshRenderData();
 		PlatformRenderData->Cache(Mesh, PlatformLODSettings);
 		check(PlatformRenderData->DerivedDataKey == PlatformDerivedDataKey);
-		PlatformRenderData->NextCachedRenderData.Swap(Mesh->RenderData->NextCachedRenderData);
-		Mesh->RenderData->NextCachedRenderData = PlatformRenderData;
+		Swap(PlatformRenderData->NextCachedRenderData, Mesh->RenderData->NextCachedRenderData);
+		Mesh->RenderData->NextCachedRenderData = TUniquePtr<FStaticMeshRenderData>(PlatformRenderData);
 	}
 	check(PlatformRenderData);
 	return *PlatformRenderData;
@@ -2227,7 +2227,7 @@ void UStaticMesh::CacheDerivedData()
 		}
 	}
 
-	RenderData = new FStaticMeshRenderData();
+	RenderData = MakeUnique<FStaticMeshRenderData>();
 	RenderData->Cache(this, LODSettings);
 
 	// Additionally cache derived data for any other platforms we care about.
@@ -2398,7 +2398,7 @@ void UStaticMesh::Serialize(FArchive& Ar)
 	{	
 		if (Ar.IsLoading())
 		{
-			RenderData = new FStaticMeshRenderData();
+			RenderData = MakeUnique<FStaticMeshRenderData>();
 			RenderData->Serialize(Ar, this, bCooked);
 		}
 
diff --git a/Engine/Source/Runtime/Engine/Private/StaticMeshRender.cpp b/Engine/Source/Runtime/Engine/Private/StaticMeshRender.cpp
--- a/Engine/Source/Runtime/Engine/Private/StaticMeshRender.cpp
+++ b/Engine/Source/Runtime/Engine/Private/StaticMeshRender.cpp
@@ -90,7 +90,7 @@ FStaticMeshSceneProxy::FStaticMeshSceneProxy(UStaticMeshComponent* InComponent):
 	, Owner(InComponent->GetOwner())
 	, StaticMesh(InComponent->GetStaticMesh())
 	, BodySetup(InComponent->GetBodySetup())
-	, RenderData(InComponent->GetStaticMesh()->RenderData)
+	, RenderData(InComponent->GetStaticMesh()->RenderData.Get())
 	, ForcedLodModel(InComponent->ForcedLodModel)
 	, bCastShadow(InComponent->CastShadow)
 	, CollisionTraceFlag(ECollisionTraceFlag::CTF_UseSimpleAndComplex)
@@ -1345,7 +1345,7 @@ FStaticMeshSceneProxy::FLODInfo::FLODInfo(const UStaticMeshComponent* InComponen
 {
 	const auto FeatureLevel = InComponent->GetWorld()->FeatureLevel;
 
-	FStaticMeshRenderData* MeshRenderData = InComponent->GetStaticMesh()->RenderData;
+	FStaticMeshRenderData* MeshRenderData = InComponent->GetStaticMesh()->RenderData.Get();
 	FStaticMeshLODResources& LODModel = MeshRenderData->LODResources[LODIndex];
 	if (LODIndex < InComponent->LODData.Num())
 	{
diff --git a/Engine/Source/Runtime/Engine/Private/Texture2D.cpp b/Engine/Source/Runtime/Engine/Private/Texture2D.cpp
--- a/Engine/Source/Runtime/Engine/Private/Texture2D.cpp
+++ b/Engine/Source/Runtime/Engine/Private/Texture2D.cpp
@@ -1355,7 +1355,7 @@ void FTexture2DResource::ReleaseRHI()
 
 	// It should be safe to release the texture.
 	checkf(Owner->PendingMipChangeRequestStatus.GetValue() <= TexState_ReadyFor_Requests, TEXT("PendingMipChangeRequestStatus = %d"), Owner->PendingMipChangeRequestStatus.GetValue());
-	check(AsyncCreateTextureTask == NULL);
+	check(AsyncCreateTextureTask == nullptr);
 
 	if ( (Texture2DRHI->GetFlags() & TexCreate_Virtual) != TexCreate_Virtual )
 	{
@@ -2131,7 +2131,7 @@ void FTexture2DResource::UploadMipData()
 
 		if (!bDerivedDataStreamRequestFailed && !Owner->bHasCancelationPending)
 		{
-			check(AsyncCreateTextureTask == NULL);
+			check(AsyncCreateTextureTask == nullptr);
 			FCreateTextureTask::FArguments TaskArgs = {0};
 			TaskArgs.SizeX = OwnerMips[PendingFirstMip].SizeX;
 			TaskArgs.SizeY = OwnerMips[PendingFirstMip].SizeY;
@@ -2143,7 +2143,7 @@ void FTexture2DResource::UploadMipData()
 			TaskArgs.NumNewMips = Owner->RequestedMips - Owner->ResidentMips;
 			TaskArgs.TextureRefPtr = &IntermediateTextureRHI;
 			TaskArgs.ThreadSafeCounter = &Owner->PendingMipChangeRequestStatus;
-			AsyncCreateTextureTask = new FAsyncCreateTextureTask(TaskArgs);
+			AsyncCreateTextureTask = MakeUnique<FAsyncCreateTextureTask>(TaskArgs);
 			AsyncCreateTextureTask->StartBackgroundTask();
 		}
 		else
diff --git a/Engine/Source/Runtime/Engine/Private/Timeline.cpp b/Engine/Source/Runtime/Engine/Private/Timeline.cpp
--- a/Engine/Source/Runtime/Engine/Private/Timeline.cpp
+++ b/Engine/Source/Runtime/Engine/Private/Timeline.cpp
@@ -312,10 +312,15 @@ void FTimeline::SetPlaybackPosition(float NewPosition, bool bFireEvents, bool bF
 	{
 		if (PropSetObject)
 		{
-			if (DirectionProperty == NULL)
+			if (DirectionProperty == nullptr)
 			{
 				DirectionProperty = FindField<UByteProperty>(PropSetObject->GetClass(), DirectionPropertyName);
-				if (DirectionProperty == NULL)
+				if (DirectionProperty == nullptr)
+				{
+					DirectionProperty = FindField<UEnumProperty>(PropSetObject->GetClass(), DirectionPropertyName);
+				}
+
+				if (DirectionProperty == nullptr)
 				{
 					UE_LOG(LogTimeline, Log, TEXT("SetPlaybackPosition: No direction property '%s' in '%s'"), *DirectionPropertyName.ToString(), *PropSetObject->GetName());
 				}
@@ -324,7 +329,17 @@ void FTimeline::SetPlaybackPosition(float NewPosition, bool bFireEvents, bool bF
 			{
 				const ETimelineDirection::Type CurrentDirection = bReversePlayback ? ETimelineDirection::Backward : ETimelineDirection::Forward;
 				TEnumAsByte<ETimelineDirection::Type> ValueAsByte(CurrentDirection);
-				DirectionProperty->SetPropertyValue_InContainer(PropSetObject, ValueAsByte);
+				if (UByteProperty* ByteDirection = Cast<UByteProperty>(DirectionProperty))
+				{
+					ByteDirection->SetPropertyValue_InContainer(PropSetObject, ValueAsByte);
+				}
+				else
+				{
+					UEnumProperty* EnumProp = CastChecked<UEnumProperty>(DirectionProperty);
+					void* PropAddr = EnumProp->ContainerPtrToValuePtr<void>(PropSetObject);
+					UNumericProperty* UnderlyingProp = EnumProp->GetUnderlyingProperty();
+					UnderlyingProp->SetIntPropertyValue(PropAddr, (int64)ValueAsByte);
+				}
 			}
 		}
 	}
diff --git a/Engine/Source/Runtime/Engine/Private/UnrealClient.cpp b/Engine/Source/Runtime/Engine/Private/UnrealClient.cpp
--- a/Engine/Source/Runtime/Engine/Private/UnrealClient.cpp
+++ b/Engine/Source/Runtime/Engine/Private/UnrealClient.cpp
@@ -21,6 +21,7 @@
 #include "SceneUtils.h"
 #include "NotificationManager.h"
 #include "Performance/EnginePerformanceTargets.h"
+#include "UniquePtr.h"
 
 DEFINE_LOG_CATEGORY_STATIC(LogClient, Log, All);
 
@@ -1077,7 +1078,7 @@ bool GCaptureCompositionNextFrame = false;
 void FViewport::Draw( bool bShouldPresent /*= true */)
 {
 	UWorld* World = GetClient()->GetWorld();
-	static TScopedPointer<FSuspendRenderingThread> GRenderingThreadSuspension;
+	static TUniquePtr<FSuspendRenderingThread> GRenderingThreadSuspension;
 
 	// Ignore reentrant draw calls, since we can only redraw one viewport at a time.
 	static bool bReentrant = false;
@@ -1094,7 +1095,7 @@ void FViewport::Draw( bool bShouldPresent /*= true */)
 		{
 			// To capture the CompositionGraph we go into single threaded for one frame
 			// so that the Slate UI gets the data on the game thread.
-			GRenderingThreadSuspension.Reset(new FSuspendRenderingThread(true));
+			GRenderingThreadSuspension = MakeUnique<FSuspendRenderingThread>(true);
 		}
 
 		// if this is a game viewport, and game rendering is disabled, then we don't want to actually draw anything
@@ -1184,7 +1185,7 @@ void FViewport::Draw( bool bShouldPresent /*= true */)
 		{
 			for( FConstPlayerControllerIterator Iterator = World->GetPlayerControllerIterator(); Iterator; ++Iterator )
 			{
-				APlayerController* PlayerController = *Iterator;
+				APlayerController* PlayerController = Iterator->Get();
 				if (PlayerController && PlayerController->PlayerCameraManager)
 				{
 					PlayerController->PlayerCameraManager->bGameCameraCutThisFrame = false;
diff --git a/Engine/Source/Runtime/Engine/Private/UnrealEngine.cpp b/Engine/Source/Runtime/Engine/Private/UnrealEngine.cpp
--- a/Engine/Source/Runtime/Engine/Private/UnrealEngine.cpp
+++ b/Engine/Source/Runtime/Engine/Private/UnrealEngine.cpp
@@ -4683,7 +4683,7 @@ bool UEngine::HandleMergeMeshCommand( const TCHAR* Cmd, FOutputDevice& Ar, UWorl
 	USkeletalMesh * PlayerMesh = NULL;
 	for (FConstPlayerControllerIterator Iterator = InWorld->GetPlayerControllerIterator(); Iterator; ++Iterator)
 	{
-		APlayerController* PlayerController = *Iterator;
+		APlayerController* PlayerController = Iterator->Get();
 		if (PlayerController->GetCharacter() != NULL && PlayerController->GetCharacter()->GetMesh())
 		{
 			PlayerPawn = PlayerController->GetCharacter();
@@ -4705,7 +4705,7 @@ bool UEngine::HandleMergeMeshCommand( const TCHAR* Cmd, FOutputDevice& Ar, UWorl
 		// we don't have a pawn (because we couldn't find a mesh), use any pawn as a spawn point
 		for (FConstPlayerControllerIterator Iterator = InWorld->GetPlayerControllerIterator(); Iterator; ++Iterator)
 		{
-			APlayerController* PlayerController = *Iterator;
+			APlayerController* PlayerController = Iterator->Get();
 			if (PlayerController->GetPawn() != NULL)
 			{
 				PlayerPawn = PlayerController->GetPawn();
@@ -12985,7 +12985,7 @@ int32 UEngine::RenderStatAI(UWorld* World, FViewport* Viewport, FCanvas* Canvas,
 	int32 NumAIRendered = 0;
 	for (FConstControllerIterator Iterator = World->GetControllerIterator(); Iterator; ++Iterator)
 	{
-		AController* Controller = *Iterator;
+		AController* Controller = Iterator->Get();
 		if (!Cast<APlayerController>(Controller))
 		{
 			++NumAI;
diff --git a/Engine/Source/Runtime/Engine/Private/UserDefinedEnum.cpp b/Engine/Source/Runtime/Engine/Private/UserDefinedEnum.cpp
--- a/Engine/Source/Runtime/Engine/Private/UserDefinedEnum.cpp
+++ b/Engine/Source/Runtime/Engine/Private/UserDefinedEnum.cpp
@@ -96,7 +96,7 @@ FString UUserDefinedEnum::GenerateNewEnumeratorName()
 
 #endif	// WITH_EDITOR
 
-int32 UUserDefinedEnum::ResolveEnumerator(FArchive& Ar, int32 EnumeratorValue) const
+int64 UUserDefinedEnum::ResolveEnumerator(FArchive& Ar, int64 EnumeratorValue) const
 {
 #if WITH_EDITOR
 	return FEnumEditorUtils::ResolveEnumerator(this, Ar, EnumeratorValue);
@@ -116,7 +116,7 @@ FText UUserDefinedEnum::GetEnumText(int32 InIndex) const
 	return Super::GetEnumText(InIndex);
 }
 
-bool UUserDefinedEnum::SetEnums(TArray<TPair<FName, uint8>>& InNames, ECppForm InCppForm, bool bAddMaxKeyIfMissing)
+bool UUserDefinedEnum::SetEnums(TArray<TPair<FName, int64>>& InNames, ECppForm InCppForm, bool bAddMaxKeyIfMissing)
 {
 	ensure(bAddMaxKeyIfMissing);
 	if (Names.Num() > 0)
@@ -134,11 +134,11 @@ bool UUserDefinedEnum::SetEnums(TArray<TPair<FName, uint8>>& InNames, ECppForm I
 	{
 		const FString EnumPrefix = (TryNum == 0) ? BaseEnumPrefix : FString::Printf(TEXT("%s_%d"), *BaseEnumPrefix, TryNum - 1);
 		const FName MaxEnumItem = *GenerateFullEnumName(*(EnumPrefix + TEXT("_MAX")));
-		const int32 MaxEnumItemIndex = GetValueByName(MaxEnumItem);
+		const int64 MaxEnumItemIndex = GetValueByName(MaxEnumItem);
 		if ((MaxEnumItemIndex == INDEX_NONE) && (LookupEnumName(MaxEnumItem) == INDEX_NONE))
 		{
-			int MaxEnumValue = (InNames.Num() == 0)? 0 : GetMaxEnumValue() + 1;
-			Names.Add(TPairInitializer<FName, uint8>(MaxEnumItem, MaxEnumValue));
+			int64 MaxEnumValue = (InNames.Num() == 0)? 0 : GetMaxEnumValue() + 1;
+			Names.Add(TPairInitializer<FName, int64>(MaxEnumItem, MaxEnumValue));
 			AddNamesToMasterList();
 			return true;
 		}
