Fixed issue where any material using WorldAlignedTexture etc Material Functions shows default material on ES2
Treating the excluding shader offset modes of the World Position node as an error on ES2 is too heavy handed given its use in material functions, will now sliently fall back to the non-excluding modes.

diff --git a/Engine/Source/Runtime/Engine/Private/Materials/HLSLMaterialTranslator.h b/Engine/Source/Runtime/Engine/Private/Materials/HLSLMaterialTranslator.h
--- a/Engine/Source/Runtime/Engine/Private/Materials/HLSLMaterialTranslator.h
+++ b/Engine/Source/Runtime/Engine/Private/Materials/HLSLMaterialTranslator.h
@@ -2692,11 +2692,14 @@ protected:
 			{
 				if (FeatureLevel < ERHIFeatureLevel::ES3_1)
 				{
-					Errorf(TEXT("World position excluding shader offsets is not available on ES2"));
-					return INDEX_NONE;
+					// World position excluding shader offsets is not available on ES2
+					FunctionNamePattern = TEXT("Get<PREV>WorldPosition");
+				}
+				else
+				{
+					bNeedsWorldPositionExcludingShaderOffsets = true;
+					FunctionNamePattern = TEXT("Get<PREV>WorldPosition<NO_MATERIAL_OFFSETS>");
 				}
-				bNeedsWorldPositionExcludingShaderOffsets = true;
-				FunctionNamePattern = TEXT("Get<PREV>WorldPosition<NO_MATERIAL_OFFSETS>");
 				break;
 			}
 
@@ -2710,11 +2713,14 @@ protected:
 			{
 				if (FeatureLevel < ERHIFeatureLevel::ES3_1)
 				{
-					Errorf(TEXT("World position excluding shader offsets is not available on ES2"));
-					return INDEX_NONE;
+					// World position excluding shader offsets is not available on ES2
+					FunctionNamePattern = TEXT("Get<PREV>TranslatedWorldPosition");
+				}
+				else
+				{
+					bNeedsWorldPositionExcludingShaderOffsets = true;
+					FunctionNamePattern = TEXT("Get<PREV>TranslatedWorldPosition<NO_MATERIAL_OFFSETS>");
 				}
-				bNeedsWorldPositionExcludingShaderOffsets = true;
-				FunctionNamePattern = TEXT("Get<PREV>TranslatedWorldPosition<NO_MATERIAL_OFFSETS>");
 				break;
 			}
 
