Copying //UE4/Dev-Core to //UE4/Dev-Main (Source: //UE4/Dev-Core @ 3228803)

diff --git a/Engine/Source/Runtime/Engine/Private/Texture2D.cpp b/Engine/Source/Runtime/Engine/Private/Texture2D.cpp
--- a/Engine/Source/Runtime/Engine/Private/Texture2D.cpp
+++ b/Engine/Source/Runtime/Engine/Private/Texture2D.cpp
@@ -24,10 +24,7 @@
 #include "DerivedDataCacheInterface.h"
 #include "Engine/TextureStreamingTypes.h"
 #include "Streaming/TextureStreamingHelpers.h"
-
-#if USE_NEW_ASYNC_IO
 #include "AsyncFileHandle.h"
-#endif
 
 UTexture2D::UTexture2D(const FObjectInitializer& ObjectInitializer)
 	: Super(ObjectInitializer)
@@ -1118,10 +1115,8 @@ FTexture2DResource::FTexture2DResource( UTexture2D* InOwner, int32 InitialMipCou
 :	Owner( InOwner )
 ,	ResourceMem( InOwner->ResourceMem )
 ,	AsyncCreateTextureTask(nullptr)
-#if USE_NEW_ASYNC_IO
 ,	IORequestHandle(nullptr)
 ,	IORequestOffsetOffset(0)
-#endif
 ,	IORequestCount( 0 )
 ,	bUsingAsyncCreation(false)
 ,	bPrioritizedIORequest(false)
@@ -1173,10 +1168,12 @@ FTexture2DResource::FTexture2DResource( UTexture2D* InOwner, int32 InitialMipCou
  */
 FTexture2DResource::~FTexture2DResource()
 {
-#if USE_NEW_ASYNC_IO
-	check(!IORequests.Num()); // if these aren't gone by now, that seems problematic
-	verify(BlockTillAllRequestsFinished());
-#endif
+	if (GNewAsyncIO)
+	{
+		check(!IORequests.Num()); // if these aren't gone by now, that seems problematic
+		verify(BlockTillAllRequestsFinished());
+	}
+
 	// free resource memory that was preallocated
 	// The deletion needs to happen in the rendering thread.
 	ENQUEUE_UNIQUE_RENDER_COMMAND_ONEPARAMETER(
@@ -1753,27 +1750,26 @@ void FTexture2DResource::UpdateMipCount()
 	}
 }
 
-#if !defined(USE_NEW_ASYNC_IO) || !defined(SPLIT_COOKED_FILES) || !defined(USE_EVENT_DRIVEN_ASYNC_LOAD)
-#error "USE_NEW_ASYNC_IO and SPLIT_COOKED_FILES and USE_EVENT_DRIVEN_ASYNC_LOAD must be defined."
-#endif
-
-#if USE_NEW_ASYNC_IO
 void FTexture2DResource::AsyncPrep(const FByteBulkData& BulkData)
 {
+	check(GNewAsyncIO);
+
 	FString Filename = BulkData.GetFilename();
 	check(Filename.Len());
 	if (!IORequestHandle)
 	{
-#if SPLIT_COOKED_FILES && USE_EVENT_DRIVEN_ASYNC_LOAD
-		IORequestOffsetOffset = 0;
-		if (Filename.EndsWith(TEXT(".uasset")) || Filename.EndsWith(TEXT(".umap")))
+		if (GEventDrivenLoaderEnabled)
 		{
-			IORequestOffsetOffset = -IFileManager::Get().FileSize(*Filename);
-			check(IORequestOffsetOffset < 0);
-			Filename = FPaths::GetBaseFilename(Filename, false) + TEXT(".uexp");
-			UE_LOG(LogTexture, Error, TEXT("Streaming from the .uexp file '%s' this MUST be in a ubulk instead for best performance."), *Filename);
+			IORequestOffsetOffset = 0;
+			if (Filename.EndsWith(TEXT(".uasset")) || Filename.EndsWith(TEXT(".umap")))
+			{
+				IORequestOffsetOffset = -IFileManager::Get().FileSize(*Filename);
+				check(IORequestOffsetOffset < 0);
+				Filename = FPaths::GetBaseFilename(Filename, false) + TEXT(".uexp");
+				UE_LOG(LogTexture, Error, TEXT("Streaming from the .uexp file '%s' this MUST be in a ubulk instead for best performance."), *Filename);
+			}
 		}
-#endif
+
 		IORequestFilename = MoveTemp(Filename);
 		IORequestHandle = FPlatformFileManager::Get().GetPlatformFile().OpenAsyncRead(*IORequestFilename);
 		check(IORequestHandle); // this generally cannot fail because it is async
@@ -1792,7 +1788,6 @@ void FTexture2DResource::AsyncPrep(const FByteBulkData& BulkData)
 		check(FPaths::GetBaseFilename(IORequestFilename, false) == FPaths::GetBaseFilename(Filename, false)); // all of the streaming mips must be stored in the same file (could relax that, but this code makes this assumption)
 	}
 }
-#endif
 
 /**
  * Called from the rendering thread to start async I/O to load in new mips.
@@ -1804,34 +1799,35 @@ void FTexture2DResource::LoadMipData()
 	check(Owner->bIsStreamable);
 	check(Owner->PendingMipChangeRequestStatus.GetValue() == TexState_InProgress_Loading);
 
-	if ( (Texture2DRHI->GetFlags() & TexCreate_Virtual) == TexCreate_Virtual )
+	if ((Texture2DRHI->GetFlags() & TexCreate_Virtual) == TexCreate_Virtual)
 	{
 
-#if USE_NEW_ASYNC_IO
-		verify(BlockTillAllRequestsFinished());
-#endif
+		if (GNewAsyncIO)
+		{
+			verify(BlockTillAllRequestsFinished());
+		}
 		IORequestCount = 0;
-		if ( !Owner->bHasCancelationPending )
+		if (!Owner->bHasCancelationPending)
 		{
 			// Had this texture previously failed to reallocate?
-			if ( NumFailedReallocs > 0 )
+			if (NumFailedReallocs > 0)
 			{
-				DEC_DWORD_STAT( STAT_FailedReallocations );
+				DEC_DWORD_STAT(STAT_FailedReallocations);
 			}
 			NumFailedReallocs = 0;
 			bDerivedDataStreamRequestFailed = false;
 
 			// Read into new miplevels, if any.
 			const TIndirectArray<FTexture2DMipMap>& OwnerMips = Owner->GetPlatformMips();
-			int32 FirstSharedMip	= (Owner->RequestedMips - FMath::Min(Owner->ResidentMips, Owner->RequestedMips));
-			for( int32 MipIndex=0; MipIndex<FirstSharedMip; MipIndex++ )
+			int32 FirstSharedMip = (Owner->RequestedMips - FMath::Min(Owner->ResidentMips, Owner->RequestedMips));
+			for (int32 MipIndex = 0; MipIndex < FirstSharedMip; MipIndex++)
 			{
 				int32 ActualMipIndex = MipIndex + PendingFirstMip;
 				const FTexture2DMipMap& MipMap = OwnerMips[ActualMipIndex];
 
 				void* TheMipData = nullptr;
 				uint32 DestPitch;
-				TheMipData = RHILockTexture2D( Texture2DRHI, ActualMipIndex, RLM_WriteOnly, DestPitch, false );
+				TheMipData = RHILockTexture2D(Texture2DRHI, ActualMipIndex, RLM_WriteOnly, DestPitch, false);
 
 				// Pass the request on to the async io manager after increasing the request count. The request count 
 				// has been pre-incremented before fielding the update request so we don't have to worry about file
@@ -1840,25 +1836,21 @@ void FTexture2DResource::LoadMipData()
 				Owner->PendingMipChangeRequestStatus.Increment();
 
 				EAsyncIOPriority AsyncIOPriority = bPrioritizedIORequest ? AIOP_BelowNormal : AIOP_Low;
-#if USE_NEW_ASYNC_IO
-				if (MipMap.BulkData.IsStoredCompressedOnDisk())
-				{
-					check(!"USE_NEW_ASYNC_IO does not support compression at the package level.");
-				}
-				else
+				if (GNewAsyncIO)
 				{
+					UE_CLOG(MipMap.BulkData.IsStoredCompressedOnDisk(), LogTexture, Fatal, TEXT("GNewAsyncIO does not support compression at the package level."));
+
 					AsyncPrep(MipMap.BulkData);
 					IORequests.Add(IORequestHandle->ReadRequest(MipMap.BulkData.GetBulkDataOffsetInFile() + IORequestOffsetOffset, MipMap.BulkData.GetBulkDataSize(), AsyncIOPriority, &AsyncFileCallBack, (uint8*)TheMipData));
 					IORequestCount++; // no clue what this actually means, but it is important
 				}
-
-#else
-				// Load and decompress async.
+				else
+					// Load and decompress async.
 				{
 					check(MipMap.BulkData.GetFilename().Len());
-					if( MipMap.BulkData.IsStoredCompressedOnDisk() )
+					if (MipMap.BulkData.IsStoredCompressedOnDisk())
 					{
-						IORequestIndices[IORequestCount++] = FIOSystem::Get().LoadCompressedData( 
+						IORequestIndices[IORequestCount++] = FIOSystem::Get().LoadCompressedData(
 							MipMap.BulkData.GetFilename(),						// filename
 							MipMap.BulkData.GetBulkDataOffsetInFile(),			// offset
 							MipMap.BulkData.GetBulkDataSizeOnDisk(),			// compressed size
@@ -1867,34 +1859,33 @@ void FTexture2DResource::LoadMipData()
 							MipMap.BulkData.GetDecompressionFlags(),			// compressed data format
 							&Owner->PendingMipChangeRequestStatus,				// counter to decrement
 							AsyncIOPriority										// priority
-							);
+						);
 					}
 					// Load async.
 					else
 					{
-						IORequestIndices[IORequestCount++] = FIOSystem::Get().LoadData( 
+						IORequestIndices[IORequestCount++] = FIOSystem::Get().LoadData(
 							MipMap.BulkData.GetFilename(),						// filename
 							MipMap.BulkData.GetBulkDataOffsetInFile(),			// offset
 							MipMap.BulkData.GetBulkDataSize(),					// size
 							TheMipData,											// dest pointer
 							&Owner->PendingMipChangeRequestStatus,				// counter to decrement
 							AsyncIOPriority										// priority
-							);
+						);
 					}
 					check(IORequestIndices[MipIndex]);
 				}
-#endif
 
 				// For consistency with other code paths, track the pointer to the locked buffer.
 				MipData[ActualMipIndex] = TheMipData;
 			}
 
 			// Are we reducing the mip-count?
-			if ( Owner->RequestedMips < Owner->ResidentMips )
+			if (Owner->RequestedMips < Owner->ResidentMips)
 			{
 				// Set up MipBiasFade to start fading out mip-levels (start at 0, increase mip-bias over time).
 				EMipFadeSettings MipFadeSetting = (Owner->LODGroup == TEXTUREGROUP_Lightmap || Owner->LODGroup == TEXTUREGROUP_Shadowmap) ? MipFade_Slow : MipFade_Normal;
-				MipBiasFade.SetNewMipCount( Owner->ResidentMips, Owner->RequestedMips, LastRenderTime, MipFadeSetting );
+				MipBiasFade.SetNewMipCount(Owner->ResidentMips, Owner->RequestedMips, LastRenderTime, MipFadeSetting);
 			}
 		}
 
@@ -1904,9 +1895,11 @@ void FTexture2DResource::LoadMipData()
 		return;
 	}
 
-#if USE_NEW_ASYNC_IO
-	verify(BlockTillAllRequestsFinished());
-#endif
+	if (GNewAsyncIO)
+	{
+		verify(BlockTillAllRequestsFinished());
+	}
+
 	IORequestCount = 0;
 
 	if ( (bUsingAsyncCreation || IsValidRef(IntermediateTextureRHI)) && !Owner->bHasCancelationPending )
@@ -1933,7 +1926,7 @@ void FTexture2DResource::LoadMipData()
 		// Read into new miplevels, if any.
 		const TIndirectArray<FTexture2DMipMap>& OwnerMips = Owner->GetPlatformMips();
 		int32 FirstSharedMip	= (Owner->RequestedMips - FMath::Min(Owner->ResidentMips,Owner->RequestedMips));
-		for( int32 MipIndex=0; MipIndex<FirstSharedMip; MipIndex++ )
+		for (int32 MipIndex = 0; MipIndex < FirstSharedMip; MipIndex++)
 		{
 			int32 ActualMipIndex = MipIndex + PendingFirstMip;
 			const FTexture2DMipMap& MipMap = OwnerMips[ActualMipIndex];
@@ -1948,7 +1941,7 @@ void FTexture2DResource::LoadMipData()
 			{
 				// Lock the new texture.
 				uint32 DestPitch;
-				MipData[ActualMipIndex] = RHILockTexture2D( IntermediateTextureRHI, MipIndex, RLM_WriteOnly, DestPitch, false, CVarFlushRHIThreadOnSTreamingTextureLocks.GetValueOnAnyThread() > 0 );
+				MipData[ActualMipIndex] = RHILockTexture2D(IntermediateTextureRHI, MipIndex, RLM_WriteOnly, DestPitch, false, CVarFlushRHIThreadOnSTreamingTextureLocks.GetValueOnAnyThread() > 0);
 			}
 
 			// Pass the request on to the async io manager after increasing the request count. The request count 
@@ -1961,14 +1954,14 @@ void FTexture2DResource::LoadMipData()
 
 			// Load and decompress async.
 #if WITH_EDITORONLY_DATA
-			if( MipMap.DerivedDataKey.IsEmpty() == false )
+			if (MipMap.DerivedDataKey.IsEmpty() == false)
 			{
 				FAsyncStreamDerivedMipTask* Task = new(PendingAsyncStreamDerivedMipTasks) FAsyncStreamDerivedMipTask(
 					MipMap.DerivedDataKey,
 					&MipData[ActualMipIndex],
 					MipSize,
 					&Owner->PendingMipChangeRequestStatus
-					);
+				);
 				Task->StartBackgroundTask();
 			}
 			else
@@ -1979,46 +1972,44 @@ void FTexture2DResource::LoadMipData()
 					QUICK_SCOPE_CYCLE_COUNTER(STAT_FTexture2DResource_LoadMipData_Malloc);
 					MipData[ActualMipIndex] = FMemory::Malloc(MipSize);
 				}
-#if USE_NEW_ASYNC_IO
-				if (MipMap.BulkData.IsStoredCompressedOnDisk())
-				{
-					check(!"USE_NEW_ASYNC_IO does not support compression at the package level.");
-				}
-				else
+				if (GNewAsyncIO)
 				{
+					UE_CLOG(MipMap.BulkData.IsStoredCompressedOnDisk(), LogTexture, Fatal, TEXT("GNewAsyncIO does not support compression at the package level."));
+
 					AsyncPrep(MipMap.BulkData);
 					IORequests.Add(IORequestHandle->ReadRequest(MipMap.BulkData.GetBulkDataOffsetInFile() + IORequestOffsetOffset, MipMap.BulkData.GetBulkDataSize(), AsyncIOPriority, &AsyncFileCallBack, (uint8*)MipData[ActualMipIndex]));
 					IORequestCount++; // no clue what this actually means, but it is important
 				}
-#else
-				check(MipMap.BulkData.GetFilename().Len());
-				if( MipMap.BulkData.IsStoredCompressedOnDisk() )
-				{
-					IORequestIndices[IORequestCount++] = FIOSystem::Get().LoadCompressedData( 
-						MipMap.BulkData.GetFilename(),						// filename
-						MipMap.BulkData.GetBulkDataOffsetInFile(),			// offset
-						MipMap.BulkData.GetBulkDataSizeOnDisk(),			// compressed size
-						MipMap.BulkData.GetBulkDataSize(),					// uncompressed size
-						MipData[ActualMipIndex],							// dest pointer
-						MipMap.BulkData.GetDecompressionFlags(),			// compressed data format
-						&Owner->PendingMipChangeRequestStatus,				// counter to decrement
-						AsyncIOPriority										// priority
-						);
-				}
-				// Load async.
 				else
 				{
-					IORequestIndices[IORequestCount++] = FIOSystem::Get().LoadData( 
-						MipMap.BulkData.GetFilename(),						// filename
-						MipMap.BulkData.GetBulkDataOffsetInFile(),			// offset
-						MipMap.BulkData.GetBulkDataSize(),					// size
-						MipData[ActualMipIndex],							// dest pointer
-						&Owner->PendingMipChangeRequestStatus,				// counter to decrement
-						AsyncIOPriority										// priority
+					check(MipMap.BulkData.GetFilename().Len());
+					if (MipMap.BulkData.IsStoredCompressedOnDisk())
+					{
+						IORequestIndices[IORequestCount++] = FIOSystem::Get().LoadCompressedData(
+							MipMap.BulkData.GetFilename(),						// filename
+							MipMap.BulkData.GetBulkDataOffsetInFile(),			// offset
+							MipMap.BulkData.GetBulkDataSizeOnDisk(),			// compressed size
+							MipMap.BulkData.GetBulkDataSize(),					// uncompressed size
+							MipData[ActualMipIndex],							// dest pointer
+							MipMap.BulkData.GetDecompressionFlags(),			// compressed data format
+							&Owner->PendingMipChangeRequestStatus,				// counter to decrement
+							AsyncIOPriority										// priority
 						);
+					}
+					// Load async.
+					else
+					{
+						IORequestIndices[IORequestCount++] = FIOSystem::Get().LoadData(
+							MipMap.BulkData.GetFilename(),						// filename
+							MipMap.BulkData.GetBulkDataOffsetInFile(),			// offset
+							MipMap.BulkData.GetBulkDataSize(),					// size
+							MipData[ActualMipIndex],							// dest pointer
+							&Owner->PendingMipChangeRequestStatus,				// counter to decrement
+							AsyncIOPriority										// priority
+						);
+					}
+					check(IORequestIndices[MipIndex]);
 				}
-				check(IORequestIndices[MipIndex]);
-#endif
 			}
 		}
 
@@ -2173,7 +2164,6 @@ void FTexture2DResource::UploadMipData()
 	}
 }
 
-#if USE_NEW_ASYNC_IO
 bool FTexture2DResource::BlockTillAllRequestsFinished(float TimeLimit)
 {
 	QUICK_SCOPE_CYCLE_COUNTER(FTexture2DResource_BlockTillAllRequestsFinished);
@@ -2216,29 +2206,28 @@ bool FTexture2DResource::BlockTillAllRequestsFinished(float TimeLimit)
 	IORequests.Empty();
 	return true;
 }
-#endif
 
 /**
 * Helper function for cleaning up bulk data files after streaming
 */
 void FTexture2DResource::HintDoneWithStreamedTextureFiles()
 {
-#if USE_NEW_ASYNC_IO
-	verify(BlockTillAllRequestsFinished());
-	if (IORequestHandle)
+	if (GNewAsyncIO)
 	{
-		delete IORequestHandle;
-		IORequestHandle = nullptr;
+		verify(BlockTillAllRequestsFinished());
+		if (IORequestHandle)
+		{
+			delete IORequestHandle;
+			IORequestHandle = nullptr;
+		}
+		check(!IORequestHandle);
 	}
-	check(!IORequestHandle);
-#else
-	if (FPlatformProperties::RequiresCookedData())
+	else if (FPlatformProperties::RequiresCookedData())
 	{
 		const TIndirectArray<FTexture2DMipMap>& OwnerMips = Owner->GetPlatformMips();
 		const FTexture2DMipMap& MipMap = OwnerMips[0];
 		FIOSystem::Get().HintDoneWithFile(MipMap.BulkData.GetFilename());
 	}
-#endif
 }
 
 /**
@@ -2360,22 +2349,25 @@ void FTexture2DResource::CancelUpdate()
 	// We only have anything worth cancellation if there are outstanding I/O requests.
 	if( IORequestCount )
 	{
-#if USE_NEW_ASYNC_IO
-		// Cancel requests. Some may already have been fulfilled or be in the process of being fulfilled.
-		for (IAsyncReadRequest* Item : IORequests)
+		if (GNewAsyncIO)
 		{
-			Item->Cancel();
+			// Cancel requests. Some may already have been fulfilled or be in the process of being fulfilled.
+			for (IAsyncReadRequest* Item : IORequests)
+			{
+				Item->Cancel();
+			}
+			for (IAsyncReadRequest* Item : IORequests)
+			{
+				Item->WaitCompletion();
+				delete Item;
+			}
+			IORequests.Empty();
 		}
-		for (IAsyncReadRequest* Item : IORequests)
+		else
 		{
-			Item->WaitCompletion();
-			delete Item;
+			// Cancel requests. This only cancels pending requests and not ones currently being fulfilled.
+			FIOSystem::Get().CancelRequests(IORequestIndices, IORequestCount);
 		}
-		IORequests.Empty();
-#else
-		// Cancel requests. This only cancels pending requests and not ones currently being fulfilled.
-		FIOSystem::Get().CancelRequests( IORequestIndices, IORequestCount );
-#endif
 	}
 
 	if ( !bUsingAsyncCreation && IsValidRef(IntermediateTextureRHI) )
diff --git a/Engine/Source/Runtime/Engine/Private/UnrealEngine.cpp b/Engine/Source/Runtime/Engine/Private/UnrealEngine.cpp
--- a/Engine/Source/Runtime/Engine/Private/UnrealEngine.cpp
+++ b/Engine/Source/Runtime/Engine/Private/UnrealEngine.cpp
@@ -2762,17 +2762,16 @@ bool UEngine::Exec( UWorld* InWorld, const TCHAR* Cmd, FOutputDevice& Ar )
 	{
 		return HandleFreezeAllCommand( Cmd, Ar, InWorld );
 	}
-#if !USE_NEW_ASYNC_IO
-	else if( FParse::Command(&Cmd, TEXT("FLUSHIOMANAGER")) )
+	
+	else if( !GNewAsyncIO && FParse::Command(&Cmd, TEXT("FLUSHIOMANAGER")) )
 	{
 		return HandleFlushIOManagerCommand( Cmd, Ar );
 	}
 	// This will list out the packages which are in the precache list and have not been "loaded" out.  (e.g. could be just there taking up memory!)
-	else if (FParse::Command(&Cmd, TEXT("ListPrecacheMapPackages")))
+	else if ( !GNewAsyncIO && FParse::Command(&Cmd, TEXT("ListPrecacheMapPackages")))
 	{
 		return HandleListPreCacheMapPackagesCommand(Cmd, Ar);
 	}
-#endif
 
 	else if( FParse::Command(&Cmd,TEXT("ToggleRenderingThread")) )
 	{
@@ -3406,14 +3405,15 @@ bool UEngine::HandleFreezeAllCommand( const TCHAR* Cmd, FOutputDevice& Ar, UWorl
 	return true;
 }
 
-#if !USE_NEW_ASYNC_IO
 bool UEngine::HandleFlushIOManagerCommand( const TCHAR* Cmd, FOutputDevice& Ar )
 {
+	check(!GNewAsyncIO);
 	FIOSystem::Get().BlockTillAllRequestsFinishedAndFlushHandles();
 	return true;
 }
 bool UEngine::HandleListPreCacheMapPackagesCommand(const TCHAR* Cmd, FOutputDevice& Ar)
 {
+	check(!GNewAsyncIO);
 	TArray<FString> Packages;
 	FLinkerLoad::GetListOfPackagesInPackagePrecacheMap(Packages);
 
@@ -3430,8 +3430,6 @@ bool UEngine::HandleListPreCacheMapPackagesCommand(const TCHAR* Cmd, FOutputDevi
 	return true;
 }
 
-#endif
-
 bool UEngine::HandleFreezeRenderingCommand( const TCHAR* Cmd, FOutputDevice& Ar, UWorld* InWorld  )
 {
 	ProcessToggleFreezeCommand( InWorld );
diff --git a/Engine/Source/Runtime/Engine/Public/AudioStreaming.h b/Engine/Source/Runtime/Engine/Public/AudioStreaming.h
--- a/Engine/Source/Runtime/Engine/Public/AudioStreaming.h
+++ b/Engine/Source/Runtime/Engine/Public/AudioStreaming.h
@@ -95,21 +95,15 @@ struct FWaveRequest
 struct FLoadedAudioChunk
 {
 	uint8*	Data;
-#if USE_NEW_ASYNC_IO
 	class IAsyncReadRequest* IORequest;
-#else
 	int32	MemorySize; 
-#endif
 	int32	DataSize;
 	uint32	Index;
 
 	FLoadedAudioChunk()
 		: Data(nullptr)
-#if USE_NEW_ASYNC_IO
 		, IORequest(nullptr)
-#else
 		, MemorySize(0)
-#endif
 		, DataSize(0)
 		, Index(0)
 	{
@@ -161,7 +155,6 @@ struct FStreamingWaveData final
 	 */
 	void BeginPendingRequests(const TArray<uint32>& IndicesToLoad, const TArray<uint32>& IndicesToFree);
 
-#if USE_NEW_ASYNC_IO
 	/**
 	* Blocks till all pending requests are fulfilled.
 	*
@@ -169,7 +162,6 @@ struct FStreamingWaveData final
 	* @return				Return true if there are no requests left in flight, false if the time limit was reached before they were finished.
 	*/
 	bool BlockTillAllRequestsFinished(float TimeLimit = 0.0f);
-#endif
 
 #if WITH_EDITORONLY_DATA
 	/**
@@ -198,13 +190,11 @@ public:
 	/* Contains pointers to Chunks of audio data that have been streamed in */
 	TArray<FLoadedAudioChunk> LoadedChunks;
 
-#if USE_NEW_ASYNC_IO
 	class IAsyncReadFileHandle* IORequestHandle;
 	FAsyncFileCallBack AsyncFileCallBack;
-#else
+
 	/** Potentially outstanding audio chunk I/O requests */
 	TArray<uint64>	IORequestIndices;
-#endif
 
 	/** Indices of chunks that are currently loaded */
 	TArray<uint32>	LoadedChunkIndices;
diff --git a/Engine/Source/Runtime/Engine/Public/TextureResource.h b/Engine/Source/Runtime/Engine/Public/TextureResource.h
--- a/Engine/Source/Runtime/Engine/Public/TextureResource.h
+++ b/Engine/Source/Runtime/Engine/Public/TextureResource.h
@@ -329,7 +329,6 @@ private:
 	/** Local copy/ cache of mip data between creation and first call to InitRHI.							*/
 	void*				MipData[MAX_TEXTURE_MIP_COUNT];
 
-#if USE_NEW_ASYNC_IO
 	/**  Async handle */
 	class IAsyncReadFileHandle* IORequestHandle;
 	FString IORequestFilename;
@@ -338,10 +337,10 @@ private:
 
 	/** Potentially outstanding texture I/O requests.														*/
 	TArray<class IAsyncReadRequest*> IORequests;
-#else
+
 	/** Potentially outstanding texture I/O requests.														*/
 	uint64				IORequestIndices[MAX_TEXTURE_MIP_COUNT];
-#endif
+
 	/** Number of file I/O requests for current request. The use of this is crazy confusing.				*/
 	int32					IORequestCount;
 
@@ -418,14 +417,12 @@ private:
 	*/
 	void HintDoneWithStreamedTextureFiles();
 
-#if USE_NEW_ASYNC_IO
 	/**
 	* Block until all requests are done.
 	* A time limit of zero means infinite
 	*/
 	bool BlockTillAllRequestsFinished(float TimeLimit = 0.0f);
 	void AsyncPrep(const FByteBulkData& BulkData);
-#endif
 
 };
 
diff --git a/Engine/Source/Runtime/Launch/Private/LaunchEngineLoop.cpp b/Engine/Source/Runtime/Launch/Private/LaunchEngineLoop.cpp
--- a/Engine/Source/Runtime/Launch/Private/LaunchEngineLoop.cpp
+++ b/Engine/Source/Runtime/Launch/Private/LaunchEngineLoop.cpp
@@ -1341,17 +1341,6 @@ int32 FEngineLoop::PreInit( const TCHAR* CmdLine )
 			}
 			verify(GThreadPool->Create(NumThreadsInThreadPool, 128 * 1024));
 		}
-#if USE_NEW_ASYNC_IO
-		{
-			GIOThreadPool = FQueuedThreadPool::Allocate();
-			int32 NumThreadsInThreadPool = FPlatformMisc::NumberOfIOWorkerThreadsToSpawn();
-			if (FPlatformProperties::IsServerOnly())
-			{
-				NumThreadsInThreadPool = 2;
-			}
-			verify(GIOThreadPool->Create(NumThreadsInThreadPool, 16 * 1024, TPri_AboveNormal));
-		}
-#endif // USE_NEW_ASYNC_IO
 
 #if WITH_EDITOR
 		// when we are in the editor we like to do things like build lighting and such
@@ -1374,6 +1363,25 @@ int32 FEngineLoop::PreInit( const TCHAR* CmdLine )
 		return 1;
 	}
 	
+#if WITH_COREUOBJECT
+	GNewAsyncIO = IsEventDrivenLoaderEnabled();
+	FPlatformFileManager::Get().InitializeNewAsyncIO();
+#endif
+
+	if (FPlatformProcess::SupportsMultithreading())
+	{
+		if (GNewAsyncIO)
+		{
+			GIOThreadPool = FQueuedThreadPool::Allocate();
+			int32 NumThreadsInThreadPool = FPlatformMisc::NumberOfIOWorkerThreadsToSpawn();
+			if (FPlatformProperties::IsServerOnly())
+			{
+				NumThreadsInThreadPool = 2;
+			}
+			verify(GIOThreadPool->Create(NumThreadsInThreadPool, 16 * 1024, TPri_AboveNormal));
+		}
+	}
+
 #if WITH_ENGINE
 	// Initialize system settings before anyone tries to use it...
 	GSystemSettings.Initialize( bHasEditorToken );
@@ -1522,9 +1530,10 @@ int32 FEngineLoop::PreInit( const TCHAR* CmdLine )
 		InitializeStdOutDevice();
 	}
 
-#if !USE_NEW_ASYNC_IO
-	FIOSystem::Get(); // force it to be created if it isn't already
-#endif
+	if (!GNewAsyncIO)
+	{
+		FIOSystem::Get(); // force it to be created if it isn't already
+	}
 
 	// allow the platform to start up any features it may need
 	IPlatformFeaturesModule::Get();
@@ -2622,9 +2631,10 @@ void FEngineLoop::Exit()
 
 	FTaskGraphInterface::Shutdown();
 	IStreamingManager::Shutdown();
-#if !USE_NEW_ASYNC_IO
-	FIOSystem::Shutdown();
-#endif
+	if (!GNewAsyncIO)
+	{
+		FIOSystem::Shutdown();
+	}
 }
 
 
@@ -3566,12 +3576,11 @@ void FEngineLoop::AppPreExit( )
 	{
 		GThreadPool->Destroy();
 	}
-#if USE_NEW_ASYNC_IO
+
 	if (GIOThreadPool != nullptr)
 	{
 		GIOThreadPool->Destroy();
 	}
-#endif // USE_NEW_ASYNC_IO
 
 #if WITH_ENGINE
 	if ( GShaderCompilingManager )
diff --git a/Engine/Source/Runtime/MovieScene/Private/Tests/MovieSceneTransformTests.cpp b/Engine/Source/Runtime/MovieScene/Private/Tests/MovieSceneTransformTests.cpp
--- a/Engine/Source/Runtime/MovieScene/Private/Tests/MovieSceneTransformTests.cpp
+++ b/Engine/Source/Runtime/MovieScene/Private/Tests/MovieSceneTransformTests.cpp
@@ -27,7 +27,7 @@ bool IsNearly(TRange<float> A, TRange<float> B)
 	return IsNearly(A.GetLowerBound(), B.GetLowerBound()) && IsNearly(A.GetUpperBound(), B.GetUpperBound());
 }
 
-namespace LexicalConversion
+namespace Lex
 {
 	FString ToString(const TRange<float>& InRange)
 	{
@@ -52,7 +52,7 @@ namespace LexicalConversion
 
 bool TestTransform(FAutomationTestBase& Test, FMovieSceneSequenceTransform Transform, TArrayView<TRange<float>> InSource, TArrayView<TRange<float>> InExpected, const TCHAR* TestName)
 {
-	using namespace LexicalConversion;
+	using namespace Lex;
 
 	check(InSource.Num() == InExpected.Num());
 
diff --git a/Engine/Source/Runtime/NetworkFile/Public/NetworkPlatformFile.h b/Engine/Source/Runtime/NetworkFile/Public/NetworkPlatformFile.h
--- a/Engine/Source/Runtime/NetworkFile/Public/NetworkPlatformFile.h
+++ b/Engine/Source/Runtime/NetworkFile/Public/NetworkPlatformFile.h
@@ -49,7 +49,10 @@ public:
 	{
 		return InnerPlatformFile;
 	}
-
+	virtual void SetLowerLevel(IPlatformFile* NewLowerLevel) override
+	{
+		InnerPlatformFile = NewLowerLevel;
+	}
 	virtual void GetTimeStampPair(const TCHAR* PathA, const TCHAR* PathB, FDateTime& OutTimeStampA, FDateTime& OutTimeStampB)
 	{
 		OutTimeStampA = GetTimeStamp(PathA);
diff --git a/Engine/Source/Runtime/Online/HTTP/Private/Curl/CurlHttpManager.cpp b/Engine/Source/Runtime/Online/HTTP/Private/Curl/CurlHttpManager.cpp
--- a/Engine/Source/Runtime/Online/HTTP/Private/Curl/CurlHttpManager.cpp
+++ b/Engine/Source/Runtime/Online/HTTP/Private/Curl/CurlHttpManager.cpp
@@ -197,11 +197,11 @@ bool IsValidIPv4Address(const FString& InString)
 		if (IsUnsignedInteger(AStr) && IsUnsignedInteger(BStr) && IsUnsignedInteger(CStr) && IsUnsignedInteger(DStr) && IsUnsignedInteger(PortStr))
 		{
 			uint32 A, B, C, D, Port;
-			LexicalConversion::FromString(A, *AStr);
-			LexicalConversion::FromString(B, *BStr);
-			LexicalConversion::FromString(C, *CStr);
-			LexicalConversion::FromString(D, *DStr);
-			LexicalConversion::FromString(Port, *PortStr);
+			Lex::FromString(A, *AStr);
+			Lex::FromString(B, *BStr);
+			Lex::FromString(C, *CStr);
+			Lex::FromString(D, *DStr);
+			Lex::FromString(Port, *PortStr);
 
 			if (A < 256 && B < 256 && C < 256 && D < 256 && Port < 65536)
 			{
diff --git a/Engine/Source/Runtime/Online/ICMP/Private/UDPPing.cpp b/Engine/Source/Runtime/Online/ICMP/Private/UDPPing.cpp
--- a/Engine/Source/Runtime/Online/ICMP/Private/UDPPing.cpp
+++ b/Engine/Source/Runtime/Online/ICMP/Private/UDPPing.cpp
@@ -76,7 +76,7 @@ FIcmpEchoResult UDPEchoImpl(ISocketSubsystem* SocketSub, const FString& TargetAd
 	}
 
 	int32 Port = 0;
-	LexicalConversion::FromString(Port, *PortStr);
+	Lex::FromString(Port, *PortStr);
 
 	//ISocketSubsystem* SocketSub = ISocketSubsystem::Get(PLATFORM_SOCKETSUBSYSTEM);
 	if (SocketSub)
diff --git a/Engine/Source/Runtime/PakFile/Private/IPlatformFilePak.cpp b/Engine/Source/Runtime/PakFile/Private/IPlatformFilePak.cpp
--- a/Engine/Source/Runtime/PakFile/Private/IPlatformFilePak.cpp
+++ b/Engine/Source/Runtime/PakFile/Private/IPlatformFilePak.cpp
@@ -15,7 +15,6 @@
 #include "HAL/FileManagerGeneric.h"
 #include "HAL/IPlatformFileModule.h"
 #include "SignedArchiveReader.h"
-#include "PublicKey.inl"
 #include "Misc/AES.h"
 #include "GenericPlatform/GenericPlatformChunkInstall.h"
 #include "Async/AsyncFileHandle.h"
@@ -39,7 +38,7 @@ TPakChunkHash ComputePakChunkHash(const void* InData, int64 InDataSizeInBytes)
 }
 
 #ifndef EXCLUDE_NONPAK_UE_EXTENSIONS
-#define EXCLUDE_NONPAK_UE_EXTENSIONS USE_NEW_ASYNC_IO	// Use .Build.cs file to disable this if the game relies on accessing loose files
+#define EXCLUDE_NONPAK_UE_EXTENSIONS 1	// Use .Build.cs file to disable this if the game relies on accessing loose files
 #endif
 
 FFilenameSecurityDelegate& FPakPlatformFile::GetFilenameSecurityDelegate()
@@ -48,17 +47,34 @@ FFilenameSecurityDelegate& FPakPlatformFile::GetFilenameSecurityDelegate()
 	return Delegate;
 }
 
-#if USE_NEW_ASYNC_IO && !IS_PROGRAM && !WITH_EDITOR
-	#define USE_PAK_PRECACHE (1) // you can turn this off to use the async IO stuff without the precache
-#else
-	#define USE_PAK_PRECACHE (0)
-#endif
+#define USE_PAK_PRECACHE (!IS_PROGRAM && !WITH_EDITOR) // you can turn this off to use the async IO stuff without the precache
 
 /**
 * Precaching
 */
 
-#ifdef AES_KEY
+const ANSICHAR* FPakPlatformFile::GetPakEncryptionKey()
+{
+	FCoreDelegates::FPakEncryptionKeyDelegate& Delegate = FCoreDelegates::GetPakEncryptionKeyDelegate();
+	if (Delegate.IsBound())
+	{
+		return Delegate.Execute();
+	}
+	else
+	{
+		return nullptr;
+	}
+}
+
+void FPakPlatformFile::GetPakSigningKeys(FString& OutExponent, FString& OutModulus)
+{
+	FCoreDelegates::FPakSigningKeysDelegate& Delegate = FCoreDelegates::GetPakSigningKeysDelegate();
+	if (Delegate.IsBound())
+	{
+		return Delegate.Execute(OutExponent, OutModulus);
+	}
+}
+
 DECLARE_DWORD_ACCUMULATOR_STAT(TEXT("PakCache Sync Decrypts (Uncompressed Path)"), STAT_PakCache_SyncDecrypts, STATGROUP_PakFile);
 DECLARE_FLOAT_ACCUMULATOR_STAT(TEXT("PakCache Decrypt Time"), STAT_PakCache_DecryptTime, STATGROUP_PakFile);
 DECLARE_DWORD_ACCUMULATOR_STAT(TEXT("PakCache Async Decrypts (Compressed Path)"), STAT_PakCache_CompressedDecrypts, STATGROUP_PakFile);
@@ -67,9 +83,10 @@ DECLARE_DWORD_ACCUMULATOR_STAT(TEXT("PakCache Async Decrypts (Uncompressed Path)
 inline void DecryptData(uint8* InData, uint32 InDataSize)
 {
 	SCOPE_SECONDS_ACCUMULATOR(STAT_PakCache_DecryptTime);
-	FAES::DecryptData(InData, InDataSize);
+	const ANSICHAR* Key = FPakPlatformFile::GetPakEncryptionKey();
+	check(Key);
+	FAES::DecryptData(InData, InDataSize, Key);
 }
-#endif
 
 #if USE_PAK_PRECACHE
 #include "TaskGraphInterfaces.h"
@@ -108,9 +125,6 @@ static FAutoConsoleVariableRef CVar_MaxRequestSizeToLowerLevelKB(
 	TEXT("Controls the maximum size (in KB) of IO requests submitted to the OS filesystem.")
 	);
 
-//@todoio this won't work with UT etc, need *_API or a real api to work with the pak precacher.
-bool GPakCache_AcceptPrecacheRequests = true;
-
 class FPakPrecacher;
 
 typedef uint64 FJoinedOffsetAndPakIndex;
@@ -2430,7 +2444,6 @@ public:
 	}
 };
 
-#ifdef AES_KEY
 class FPakEncryptedReadRequest : public FPakReadRequestBase
 {
 	int64 OriginalOffset;
@@ -2506,7 +2519,6 @@ public:
 		}
 	}
 };
-#endif
 
 class FPakSizeRequest : public IAsyncReadRequest
 {
@@ -2861,11 +2873,7 @@ public:
 
 			if (FileEntry->bEncrypted)
 			{
-#ifdef AES_KEY
 				return new FPakEncryptedReadRequest(PakFile, PakFileSize, CompleteCallback, OffsetInPak, Offset, BytesToRead, Priority, UserSuppliedMemory);
-#else
-				UE_LOG(LogPakFile, Fatal, TEXT("Encountered an encrypted file in pak file but the executable wasn't compiled with the decryption key!"));
-#endif
 			}
 			else
 			{
@@ -2962,13 +2970,11 @@ public:
 		FCachedAsyncBlock& Block = Blocks[BlockIndex];
 		check(Block.Raw && Block.RawSize && !Block.Processed);
 
-#ifdef AES_KEY
 		if (FileEntry->bEncrypted)
 		{
 			INC_DWORD_STAT(STAT_PakCache_CompressedDecrypts);
 			DecryptData(Block.Raw, Align(Block.RawSize, FAES::AESBlockSize));
 		}
-#endif
 
 		check(Block.ProcessedSize);
 		INC_MEMORY_STAT_BY(STAT_AsyncFileMemory, Block.ProcessedSize);
@@ -3174,14 +3180,12 @@ void FAsyncIOCPUWorkTask::DoTask(ENamedThreads::Type CurrentThread, const FGraph
 	Owner.DoProcessing(BlockIndex);
 }
 
-#endif
-
-
-#if USE_NEW_ASYNC_IO
+#endif  
 
 
 IAsyncReadFileHandle* FPakPlatformFile::OpenAsyncRead(const TCHAR* Filename)
 {
+	check(GConfig && GNewAsyncIO);
 #if USE_PAK_PRECACHE
 	if (FPlatformProcess::SupportsMultithreading() && GPakCache_Enable > 0)
 	{
@@ -3193,11 +3197,10 @@ IAsyncReadFileHandle* FPakPlatformFile::OpenAsyncRead(const TCHAR* Filename)
 		}
 	}
 #endif
+
 	return IPlatformFile::OpenAsyncRead(Filename);
 }
 
-#endif // USE_NEW_ASYNC_IO
-
 /**
  * Class to handle correctly reading from a compressed file within a compressed package
  */
@@ -3216,10 +3219,8 @@ public:
 
 	static FORCEINLINE void DecryptBlock(void* Data, int64 Size)
 	{
-#ifdef AES_KEY
 		INC_DWORD_STAT(STAT_PakCache_SyncDecrypts);
 		DecryptData((uint8*)Data, Size);
-#endif
 	}
 };
 
@@ -3836,16 +3837,20 @@ FPakPlatformFile::~FPakPlatformFile()
 	FCoreDelegates::OnUnmountPak.Unbind();
 
 #if USE_PAK_PRECACHE
-	FPakPrecacher::Shutdown();
-#endif
-#if !USE_NEW_ASYNC_IO
-	// We need to flush async IO... if it hasn't been shut down already.
-	if (FIOSystem::HasShutdown() == false)
+	if (GNewAsyncIO)
 	{
-		FIOSystem& IOSystem = FIOSystem::Get();
-		IOSystem.BlockTillAllRequestsFinishedAndFlushHandles();
+		FPakPrecacher::Shutdown();
 	}
 #endif
+	if (!GNewAsyncIO)
+	{
+		// We need to flush async IO... if it hasn't been shut down already.
+		if (FIOSystem::HasShutdown() == false)
+		{
+			FIOSystem& IOSystem = FIOSystem::Get();
+			IOSystem.BlockTillAllRequestsFinishedAndFlushHandles();
+		}
+	}
 
 	{
 		FScopeLock ScopedLock(&PakListCritical);
@@ -3970,8 +3975,11 @@ bool FPakPlatformFile::Initialize(IPlatformFile* Inner, const TCHAR* CmdLine)
 #endif
 
 	FEncryptionKey DecryptionKey;
-	DecryptionKey.Exponent.Parse(DECRYPTION_KEY_EXPONENT);
-	DecryptionKey.Modulus.Parse(DECRYPTION_KEY_MODULUS);
+	FString PakSigningKeyExponent, PakSigningKeyModulus;
+	GetPakSigningKeys(PakSigningKeyExponent, PakSigningKeyModulus);
+	DecryptionKey.Exponent.Parse(PakSigningKeyExponent);
+	DecryptionKey.Modulus.Parse(PakSigningKeyModulus);
+
 	bSigned = !DecryptionKey.Exponent.IsZero() && !DecryptionKey.Modulus.IsZero();
 	
 	bool bMountPaks = true;
@@ -4041,9 +4049,22 @@ bool FPakPlatformFile::Initialize(IPlatformFile* Inner, const TCHAR* CmdLine)
 
 	FCoreDelegates::OnMountPak.BindRaw(this, &FPakPlatformFile::HandleMountPakDelegate);
 	FCoreDelegates::OnUnmountPak.BindRaw(this, &FPakPlatformFile::HandleUnmountPakDelegate);
-#if USE_PAK_PRECACHE
-	if (!WITH_EDITOR && FPlatformProcess::SupportsMultithreading())
+
+
+	return !!LowerLevel;
+}
+
+void FPakPlatformFile::InitializeNewAsyncIO()
+{
+#if USE_PAK_PRECACHE 
+	if (!WITH_EDITOR && FPlatformProcess::SupportsMultithreading() && GNewAsyncIO)
 	{
+		FEncryptionKey DecryptionKey;
+		FString PakSigningKeyExponent, PakSigningKeyModulus;
+		GetPakSigningKeys(PakSigningKeyExponent, PakSigningKeyModulus);
+		DecryptionKey.Exponent.Parse(PakSigningKeyExponent);
+		DecryptionKey.Modulus.Parse(PakSigningKeyModulus);
+
 		FPakPrecacher::Init(LowerLevel, DecryptionKey);
 	}
 	else
@@ -4051,8 +4072,6 @@ bool FPakPlatformFile::Initialize(IPlatformFile* Inner, const TCHAR* CmdLine)
 		GPakCache_Enable = 0;
 	}
 #endif
-
-	return !!LowerLevel;
 }
 
 bool FPakPlatformFile::Mount(const TCHAR* InPakFilename, uint32 PakOrder, const TCHAR* InPath /*= NULL*/)
diff --git a/Engine/Source/Runtime/PakFile/Private/SignedArchiveReader.cpp b/Engine/Source/Runtime/PakFile/Private/SignedArchiveReader.cpp
--- a/Engine/Source/Runtime/PakFile/Private/SignedArchiveReader.cpp
+++ b/Engine/Source/Runtime/PakFile/Private/SignedArchiveReader.cpp
@@ -4,7 +4,6 @@
 #include "HAL/FileManager.h"
 #include "HAL/Event.h"
 #include "HAL/RunnableThread.h"
-#include "PublicKey.inl"
 
 DECLARE_FLOAT_ACCUMULATOR_STAT(TEXT("FChunkCacheWorker.ProcessQueue"), STAT_FChunkCacheWorker_ProcessQueue, STATGROUP_PakFile);
 DECLARE_FLOAT_ACCUMULATOR_STAT(TEXT("FChunkCacheWorker.CheckSignature"), STAT_FChunkCacheWorker_CheckSignature, STATGROUP_PakFile);
@@ -57,9 +56,12 @@ bool FChunkCacheWorker::Init()
 
 void FChunkCacheWorker::SetupDecryptionKey()
 {
-	DecryptionKey.Exponent.Parse(DECRYPTION_KEY_EXPONENT);
-	DecryptionKey.Modulus.Parse(DECRYPTION_KEY_MODULUS);
-	// Public key should never be zero at this point. Check PublicKey.inl for more details.
+	FString PakSigningKeyExponent, PakSigningKeyModulus;
+	FPakPlatformFile::GetPakSigningKeys(PakSigningKeyExponent, PakSigningKeyModulus);
+	DecryptionKey.Exponent.Parse(PakSigningKeyExponent);
+	DecryptionKey.Modulus.Parse(PakSigningKeyModulus);
+
+	// Public key should never be zero at this point.
 	UE_CLOG(DecryptionKey.Exponent.IsZero() || DecryptionKey.Modulus.IsZero(), LogPakFile, Fatal, TEXT("Invalid decryption key detected"));
 	// Public key should produce decrypted results - check for identity keys
 	static TEncryptionInt TestValues[] = 
diff --git a/Engine/Source/Runtime/PakFile/Public/IPlatformFilePak.h b/Engine/Source/Runtime/PakFile/Public/IPlatformFilePak.h
--- a/Engine/Source/Runtime/PakFile/Public/IPlatformFilePak.h
+++ b/Engine/Source/Runtime/PakFile/Public/IPlatformFilePak.h
@@ -969,6 +969,16 @@ public:
 	 */
 	static void GetPakFolders(const TCHAR* CmdLine, TArray<FString>& OutPakFolders);
 
+	/**
+	* Helper function for accessing pak encryption key
+	*/
+	static const ANSICHAR* GetPakEncryptionKey();
+
+	/**
+	* Helper function for accessing pak signing keys
+	*/
+	static void GetPakSigningKeys(FString& OutExponent, FString& OutModulus);
+
 	/**
 	 * Constructor.
 	 * 
@@ -983,11 +993,16 @@ public:
 
 	virtual bool ShouldBeUsed(IPlatformFile* Inner, const TCHAR* CmdLine) const override;
 	virtual bool Initialize(IPlatformFile* Inner, const TCHAR* CommandLineParam) override;
+	virtual void InitializeNewAsyncIO() override;
 
 	virtual IPlatformFile* GetLowerLevel() override
 	{
 		return LowerLevel;
 	}
+	virtual void SetLowerLevel(IPlatformFile* NewLowerLevel) override
+	{
+		LowerLevel = NewLowerLevel;
+	}
 
 	virtual const TCHAR* GetName() const override
 	{
@@ -1159,13 +1174,9 @@ public:
 		FDateTime Result = FDateTime::MinValue();
 		if (IsNonPakFilenameAllowed(Filename))
 		{
-#if USE_NEW_ASYNC_IO
-			double StartTime(UE_LOG_ACTIVE(LogPakFile, Verbose) ? FPlatformTime::Seconds() : 0.0);
-#endif
+			double StartTime = (GNewAsyncIO && UE_LOG_ACTIVE(LogPakFile, Verbose)) ? FPlatformTime::Seconds() : 0.0;
 			Result = LowerLevel->GetTimeStamp(Filename);
-#if USE_NEW_ASYNC_IO
-			UE_LOG(LogPakFile, Verbose, TEXT("GetTimeStamp on disk (!!) for %s took %6.2fms."), Filename, float(FPlatformTime::Seconds() - StartTime) * 1000.0f);
-#endif
+			UE_CLOG(GNewAsyncIO, LogPakFile, Verbose, TEXT("GetTimeStamp on disk (!!) for %s took %6.2fms."), Filename, float(FPlatformTime::Seconds() - StartTime) * 1000.0f);
 		}
 		return Result;
 	}
@@ -1588,9 +1599,7 @@ public:
 
 	virtual bool CopyFile(const TCHAR* To, const TCHAR* From, EPlatformFileRead ReadFlags = EPlatformFileRead::None, EPlatformFileWrite WriteFlags = EPlatformFileWrite::None) override;
 
-#if USE_NEW_ASYNC_IO
 	virtual IAsyncReadFileHandle* OpenAsyncRead(const TCHAR* Filename) override;
-#endif // USE_NEW_ASYNC_IO
 
 	/**
 	 * Converts a filename to a path inside pak file.
diff --git a/Engine/Source/Runtime/SandboxFile/Public/IPlatformFileSandboxWrapper.h b/Engine/Source/Runtime/SandboxFile/Public/IPlatformFileSandboxWrapper.h
--- a/Engine/Source/Runtime/SandboxFile/Public/IPlatformFileSandboxWrapper.h
+++ b/Engine/Source/Runtime/SandboxFile/Public/IPlatformFileSandboxWrapper.h
@@ -204,6 +204,10 @@ public:
 	{
 		return LowerLevel;
 	}
+	virtual void SetLowerLevel(IPlatformFile* NewLowerLevel) override
+	{
+		LowerLevel = NewLowerLevel;
+	}
 
 	virtual const TCHAR* GetName() const override
 	{
@@ -618,7 +622,7 @@ public:
 
 	virtual FString ConvertToAbsolutePathForExternalAppForRead( const TCHAR* Filename ) override;
 	virtual FString ConvertToAbsolutePathForExternalAppForWrite( const TCHAR* Filename ) override;
-#if USE_NEW_ASYNC_IO
+
 	virtual IAsyncReadFileHandle* OpenAsyncRead(const TCHAR* Filename) override
 	{
 		FString UserFilename(*ConvertToSandboxPath(Filename));
@@ -628,7 +632,7 @@ public:
 		}
 		return LowerLevel->OpenAsyncRead(Filename);
 	}
-#endif // USE_NEW_ASYNC_IO
+
 };
 
 
diff --git a/Engine/Source/Runtime/Slate/Public/Widgets/Input/NumericTypeInterface.h b/Engine/Source/Runtime/Slate/Public/Widgets/Input/NumericTypeInterface.h
--- a/Engine/Source/Runtime/Slate/Public/Widgets/Input/NumericTypeInterface.h
+++ b/Engine/Source/Runtime/Slate/Public/Widgets/Input/NumericTypeInterface.h
@@ -30,7 +30,7 @@ struct TDefaultNumericTypeInterface : INumericTypeInterface<NumericType>
 	/** Convert the type to/from a string */
 	virtual FString ToString(const NumericType& Value) const override
 	{
-		return LexicalConversion::ToSanitizedString(Value);
+		return Lex::ToSanitizedString(Value);
 	}
 	virtual TOptional<NumericType> FromString(const FString& InString, const NumericType& InExistingValue) override
 	{
diff --git a/Engine/Source/Runtime/Slate/Public/Widgets/Input/NumericUnitTypeInterface.inl b/Engine/Source/Runtime/Slate/Public/Widgets/Input/NumericUnitTypeInterface.inl
--- a/Engine/Source/Runtime/Slate/Public/Widgets/Input/NumericUnitTypeInterface.inl
+++ b/Engine/Source/Runtime/Slate/Public/Widgets/Input/NumericUnitTypeInterface.inl
@@ -24,7 +24,7 @@ FString TNumericUnitTypeInterface<NumericType>::ToString(const NumericType& Valu
 		return TDefaultNumericTypeInterface<NumericType>::ToString(Value);
 	}
 
-	using namespace LexicalConversion;
+	using namespace Lex;
 
 	FNumericUnit<NumericType> FinalValue(Value, UnderlyingUnits);
 
@@ -48,7 +48,7 @@ TOptional<NumericType> TNumericUnitTypeInterface<NumericType>::FromString(const
 		return TDefaultNumericTypeInterface<NumericType>::FromString(InString, InExistingValue);
 	}
 
-	using namespace LexicalConversion;
+	using namespace Lex;
 
 	EUnit DefaultUnits = FixedDisplayUnits.IsSet() ? FixedDisplayUnits.GetValue() : UnderlyingUnits;
 
diff --git a/Engine/Source/Runtime/Slate/Public/Widgets/Input/SNumericEntryBox.h b/Engine/Source/Runtime/Slate/Public/Widgets/Input/SNumericEntryBox.h
--- a/Engine/Source/Runtime/Slate/Public/Widgets/Input/SNumericEntryBox.h
+++ b/Engine/Source/Runtime/Slate/Public/Widgets/Input/SNumericEntryBox.h
@@ -438,7 +438,7 @@ private:
 		else
 		{
 			NumericType NumericValue;
-			if (LexicalConversion::TryParseString(NumericValue, *NewValue.ToString()))
+			if (Lex::TryParseString(NumericValue, *NewValue.ToString()))
 			{
 				OnValueChanged.ExecuteIfBound( NumericValue );
 			}
diff --git a/Engine/Source/Runtime/SlateCore/Public/Fonts/CompositeFont.h b/Engine/Source/Runtime/SlateCore/Public/Fonts/CompositeFont.h
--- a/Engine/Source/Runtime/SlateCore/Public/Fonts/CompositeFont.h
+++ b/Engine/Source/Runtime/SlateCore/Public/Fonts/CompositeFont.h
@@ -206,7 +206,7 @@ struct SLATECORE_API FTypefaceEntry
 		, Font(MoveTemp(InFontFilename), InBulkData, InHinting)
 	{
 	}
-	PRAGMA_ENABLE_DEPRECATION_WARNINGS; /** Do not remove this semi-colon; it exists to stop UHT failing to parse any properties that may follow this macro */
+	PRAGMA_ENABLE_DEPRECATION_WARNINGS
 
 	/** Name used to identify this font within its typeface */
 	UPROPERTY()
diff --git a/Engine/Source/Runtime/StreamingFile/Public/StreamingNetworkPlatformFile.h b/Engine/Source/Runtime/StreamingFile/Public/StreamingNetworkPlatformFile.h
--- a/Engine/Source/Runtime/StreamingFile/Public/StreamingNetworkPlatformFile.h
+++ b/Engine/Source/Runtime/StreamingFile/Public/StreamingNetworkPlatformFile.h
@@ -171,6 +171,10 @@ public:
 	{
 		return nullptr;
 	}
+	virtual void SetLowerLevel(IPlatformFile* NewLowerLevel) override
+	{
+		check(false);
+	}
 
 	virtual const TCHAR* GetName() const override
 	{
