Copying //UE4/Dev-Editor to //UE4/Dev-Main (Source: //UE4/Dev-Editor @ 3228984)

diff --git a/Engine/Source/Runtime/Engine/Classes/Engine/Scene.h b/Engine/Source/Runtime/Engine/Classes/Engine/Scene.h
--- a/Engine/Source/Runtime/Engine/Classes/Engine/Scene.h
+++ b/Engine/Source/Runtime/Engine/Classes/Engine/Scene.h
@@ -498,121 +498,121 @@ struct FPostProcessSettings
 
 	// -----------------------------------------------------------------------
 
-	UPROPERTY(interp, BlueprintReadWrite, Category=WhiteBalance, meta=(UIMin = "1500.0", UIMax = "15000.0", editcondition = "bOverride_WhiteTemp", DisplayName = "Temp"))
+	UPROPERTY(interp, BlueprintReadWrite, Category="Color Grading|WhiteBalance", meta=(UIMin = "1500.0", UIMax = "15000.0", editcondition = "bOverride_WhiteTemp", DisplayName = "Temp"))
 	float WhiteTemp;
-	UPROPERTY(interp, BlueprintReadWrite, Category=WhiteBalance, meta=(UIMin = "-1.0", UIMax = "1.0", editcondition = "bOverride_WhiteTint", DisplayName = "Tint"))
+	UPROPERTY(interp, BlueprintReadWrite, Category="Color Grading|WhiteBalance", meta=(UIMin = "-1.0", UIMax = "1.0", editcondition = "bOverride_WhiteTint", DisplayName = "Tint"))
 	float WhiteTint;
 
 	// Color Correction controls
-	UPROPERTY(interp, BlueprintReadWrite, Category = ColorGrading, meta = (UIMin = "0.0", UIMax = "2.0", editcondition = "bOverride_ColorSaturation", DisplayName = "Saturation"))
+	UPROPERTY(interp, BlueprintReadWrite, Category = "Color Grading|Global", meta = (UIMin = "0.0", UIMax = "2.0", Delta = "0.01", ColorGradingMode = "saturation", ShiftMouseMovePixelPerDelta = "10", editcondition = "bOverride_ColorSaturation", DisplayName = "Saturation"))
 	FVector4 ColorSaturation;
-	UPROPERTY(interp, BlueprintReadWrite, Category = ColorGrading, meta = (UIMin = "0.0", UIMax = "2.0", editcondition = "bOverride_ColorContrast", DisplayName = "Contrast"))
+	UPROPERTY(interp, BlueprintReadWrite, Category = "Color Grading|Global", meta = (UIMin = "0.0", UIMax = "2.0", Delta = "0.01", ColorGradingMode = "contrast", ShiftMouseMovePixelPerDelta = "10", editcondition = "bOverride_ColorContrast", DisplayName = "Contrast"))
 	FVector4 ColorContrast;
-	UPROPERTY(interp, BlueprintReadWrite, Category = ColorGrading, meta = (UIMin = "0.0", UIMax = "2.0", editcondition = "bOverride_ColorGamma", DisplayName = "Gamma"))
+	UPROPERTY(interp, BlueprintReadWrite, Category = "Color Grading|Global", meta = (UIMin = "0.0", UIMax = "2.0", Delta = "0.01", ColorGradingMode = "gamma", ShiftMouseMovePixelPerDelta = "10", editcondition = "bOverride_ColorGamma", DisplayName = "Gamma"))
 	FVector4 ColorGamma;
-	UPROPERTY(interp, BlueprintReadWrite, Category = ColorGrading, meta = (UIMin = "0.0", UIMax = "2.0", editcondition = "bOverride_ColorGain", DisplayName = "Gain"))
+	UPROPERTY(interp, BlueprintReadWrite, Category = "Color Grading|Global", meta = (UIMin = "0.0", UIMax = "2.0", Delta = "0.01", ColorGradingMode = "gain", ShiftMouseMovePixelPerDelta = "10", editcondition = "bOverride_ColorGain", DisplayName = "Gain"))
 	FVector4 ColorGain;
-	UPROPERTY(interp, BlueprintReadWrite, Category = ColorGrading, meta = (UIMin = "-1.0", UIMax = "1.0", editcondition = "bOverride_ColorOffset", DisplayName = "Offset"))
+	UPROPERTY(interp, BlueprintReadWrite, Category = "Color Grading|Global", meta = (UIMin = "-1.0", UIMax = "1.0", Delta = "0.001", ColorGradingMode = "offset", ShiftMouseMovePixelPerDelta = "20", editcondition = "bOverride_ColorOffset", DisplayName = "Offset"))
 	FVector4 ColorOffset;
 
-	UPROPERTY(interp, BlueprintReadWrite, Category = "ColorGrading|Shadows", meta = (UIMin = "0.0", UIMax = "2.0", editcondition = "bOverride_ColorSaturationShadows", DisplayName = "SaturationShadows"))
+	UPROPERTY(interp, BlueprintReadWrite, Category = "Color Grading|Shadows", meta = (UIMin = "0.0", UIMax = "2.0", Delta = "0.01", ColorGradingMode = "saturation", ShiftMouseMovePixelPerDelta = "10", editcondition = "bOverride_ColorSaturationShadows", DisplayName = "Saturation"))
 	FVector4 ColorSaturationShadows;
-	UPROPERTY(interp, BlueprintReadWrite, Category = "ColorGrading|Shadows", meta = (UIMin = "0.0", UIMax = "2.0", editcondition = "bOverride_ColorContrastShadows", DisplayName = "ContrastShadows"))
+	UPROPERTY(interp, BlueprintReadWrite, Category = "Color Grading|Shadows", meta = (UIMin = "0.0", UIMax = "2.0", Delta = "0.01", ColorGradingMode = "contrast", ShiftMouseMovePixelPerDelta = "10", editcondition = "bOverride_ColorContrastShadows", DisplayName = "Contrast"))
 	FVector4 ColorContrastShadows;
-	UPROPERTY(interp, BlueprintReadWrite, Category = "ColorGrading|Shadows", meta = (UIMin = "0.0", UIMax = "2.0", editcondition = "bOverride_ColorGammaShadows", DisplayName = "GammaShadows"))
+	UPROPERTY(interp, BlueprintReadWrite, Category = "Color Grading|Shadows", meta = (UIMin = "0.0", UIMax = "2.0", Delta = "0.01", ColorGradingMode = "gamma", ShiftMouseMovePixelPerDelta = "10", editcondition = "bOverride_ColorGammaShadows", DisplayName = "Gamma"))
 	FVector4 ColorGammaShadows;
-	UPROPERTY(interp, BlueprintReadWrite, Category = "ColorGrading|Shadows", meta = (UIMin = "0.0", UIMax = "2.0", editcondition = "bOverride_ColorGainShadows", DisplayName = "GainShadows"))
+	UPROPERTY(interp, BlueprintReadWrite, Category = "Color Grading|Shadows", meta = (UIMin = "0.0", UIMax = "2.0", Delta = "0.01", ColorGradingMode = "gain", ShiftMouseMovePixelPerDelta = "10", editcondition = "bOverride_ColorGainShadows", DisplayName = "Gain"))
 	FVector4 ColorGainShadows;
-	UPROPERTY(interp, BlueprintReadWrite, Category = "ColorGrading|Shadows", meta = (UIMin = "-1.0", UIMax = "1.0", editcondition = "bOverride_ColorOffsetShadows", DisplayName = "OffsetShadows"))
+	UPROPERTY(interp, BlueprintReadWrite, Category = "Color Grading|Shadows", meta = (UIMin = "-1.0", UIMax = "1.0", Delta = "0.001", ColorGradingMode = "offset", ShiftMouseMovePixelPerDelta = "20", editcondition = "bOverride_ColorOffsetShadows", DisplayName = "Offset"))
 	FVector4 ColorOffsetShadows;
-	UPROPERTY(interp, BlueprintReadWrite, Category = "ColorGrading|Shadows", meta = (UIMin = "-1.0", UIMax = "1.0", editcondition = "bOverride_ColorCorrectionShadowsMax", DisplayName = "ShadowsMax"))
+	UPROPERTY(interp, BlueprintReadWrite, Category = "Color Grading|Shadows", meta = (UIMin = "-1.0", UIMax = "1.0", editcondition = "bOverride_ColorCorrectionShadowsMax", DisplayName = "ShadowsMax"))
 	float ColorCorrectionShadowsMax;
 
-	UPROPERTY(interp, BlueprintReadWrite, Category = "ColorGrading|Midtones", meta = (UIMin = "0.0", UIMax = "2.0", editcondition = "bOverride_ColorSaturationMidtones", DisplayName = "SaturationMidtones"))
+	UPROPERTY(interp, BlueprintReadWrite, Category = "Color Grading|Midtones", meta = (UIMin = "0.0", UIMax = "2.0", Delta = "0.01", ColorGradingMode = "saturation", ShiftMouseMovePixelPerDelta = "10", editcondition = "bOverride_ColorSaturationMidtones", DisplayName = "Saturation"))
 	FVector4 ColorSaturationMidtones;
-	UPROPERTY(interp, BlueprintReadWrite, Category = "ColorGrading|Midtones", meta = (UIMin = "0.0", UIMax = "2.0", editcondition = "bOverride_ColorContrastMidtones", DisplayName = "ContrastMidtones"))
+	UPROPERTY(interp, BlueprintReadWrite, Category = "Color Grading|Midtones", meta = (UIMin = "0.0", UIMax = "2.0", Delta = "0.01", ColorGradingMode = "contrast", ShiftMouseMovePixelPerDelta = "10", editcondition = "bOverride_ColorContrastMidtones", DisplayName = "Contrast"))
 	FVector4 ColorContrastMidtones;
-	UPROPERTY(interp, BlueprintReadWrite, Category = "ColorGrading|Midtones", meta = (UIMin = "0.0", UIMax = "2.0", editcondition = "bOverride_ColorGammaMidtones", DisplayName = "GammaMidtones"))
+	UPROPERTY(interp, BlueprintReadWrite, Category = "Color Grading|Midtones", meta = (UIMin = "0.0", UIMax = "2.0", Delta = "0.01", ColorGradingMode = "gamma", ShiftMouseMovePixelPerDelta = "10", editcondition = "bOverride_ColorGammaMidtones", DisplayName = "Gamma"))
 	FVector4 ColorGammaMidtones;
-	UPROPERTY(interp, BlueprintReadWrite, Category = "ColorGrading|Midtones", meta = (UIMin = "0.0", UIMax = "2.0", editcondition = "bOverride_ColorGainMidtones", DisplayName = "GainMidtones"))
+	UPROPERTY(interp, BlueprintReadWrite, Category = "Color Grading|Midtones", meta = (UIMin = "0.0", UIMax = "2.0", Delta = "0.01", ColorGradingMode = "gain", ShiftMouseMovePixelPerDelta = "10", editcondition = "bOverride_ColorGainMidtones", DisplayName = "Gain"))
 	FVector4 ColorGainMidtones;
-	UPROPERTY(interp, BlueprintReadWrite, Category = "ColorGrading|Midtones", meta = (UIMin = "-1.0", UIMax = "1.0", editcondition = "bOverride_ColorOffsetMidtones", DisplayName = "OffsetMidtones"))
+	UPROPERTY(interp, BlueprintReadWrite, Category = "Color Grading|Midtones", meta = (UIMin = "-1.0", UIMax = "1.0", Delta = "0.001", ColorGradingMode = "offset", ShiftMouseMovePixelPerDelta = "20", editcondition = "bOverride_ColorOffsetMidtones", DisplayName = "Offset"))
 	FVector4 ColorOffsetMidtones;
 
-	UPROPERTY(interp, BlueprintReadWrite, Category = "ColorGrading|Highlights", meta = (UIMin = "0.0", UIMax = "2.0", editcondition = "bOverride_ColorSaturationHighlights", DisplayName = "SaturationHighlights"))
+	UPROPERTY(interp, BlueprintReadWrite, Category = "Color Grading|Highlights", meta = (UIMin = "0.0", UIMax = "2.0", Delta = "0.01", ColorGradingMode = "saturation", ShiftMouseMovePixelPerDelta = "10", editcondition = "bOverride_ColorSaturationHighlights", DisplayName = "Saturation"))
 	FVector4 ColorSaturationHighlights;
-	UPROPERTY(interp, BlueprintReadWrite, Category = "ColorGrading|Highlights", meta = (UIMin = "0.0", UIMax = "2.0", editcondition = "bOverride_ColorContrastHighlights", DisplayName = "ContrastHighlights"))
+	UPROPERTY(interp, BlueprintReadWrite, Category = "Color Grading|Highlights", meta = (UIMin = "0.0", UIMax = "2.0", Delta = "0.01", ColorGradingMode = "contrast", ShiftMouseMovePixelPerDelta = "10", editcondition = "bOverride_ColorContrastHighlights", DisplayName = "Contrast"))
 	FVector4 ColorContrastHighlights;
-	UPROPERTY(interp, BlueprintReadWrite, Category = "ColorGrading|Highlights", meta = (UIMin = "0.0", UIMax = "2.0", editcondition = "bOverride_ColorGammaHighlights", DisplayName = "GammaHighlights"))
+	UPROPERTY(interp, BlueprintReadWrite, Category = "Color Grading|Highlights", meta = (UIMin = "0.0", UIMax = "2.0", Delta = "0.01", ColorGradingMode = "gamma", ShiftMouseMovePixelPerDelta = "10", editcondition = "bOverride_ColorGammaHighlights", DisplayName = "Gamma"))
 	FVector4 ColorGammaHighlights;
-	UPROPERTY(interp, BlueprintReadWrite, Category = "ColorGrading|Highlights", meta = (UIMin = "0.0", UIMax = "2.0", editcondition = "bOverride_ColorGainHighlights", DisplayName = "GainHighlights"))
+	UPROPERTY(interp, BlueprintReadWrite, Category = "Color Grading|Highlights", meta = (UIMin = "0.0", UIMax = "2.0", Delta = "0.01", ColorGradingMode = "gain", ShiftMouseMovePixelPerDelta = "10", editcondition = "bOverride_ColorGainHighlights", DisplayName = "Gain"))
 	FVector4 ColorGainHighlights;
-	UPROPERTY(interp, BlueprintReadWrite, Category = "ColorGrading|Highlights", meta = (UIMin = "-1.0", UIMax = "1.0", editcondition = "bOverride_ColorOffsetHighlights", DisplayName = "OffsetHighlights"))
+	UPROPERTY(interp, BlueprintReadWrite, Category = "Color Grading|Highlights", meta = (UIMin = "-1.0", UIMax = "1.0", Delta = "0.001", ColorGradingMode = "offset", ShiftMouseMovePixelPerDelta = "20", editcondition = "bOverride_ColorOffsetHighlights", DisplayName = "Offset"))
 	FVector4 ColorOffsetHighlights;
-	UPROPERTY(interp, BlueprintReadWrite, Category = "ColorGrading|Highlights", meta = (UIMin = "-1.0", UIMax = "1.0", editcondition = "bOverride_ColorCorrectionHighlightsMin", DisplayName = "HighlightsMin"))
+	UPROPERTY(interp, BlueprintReadWrite, Category = "Color Grading|Highlights", meta = (UIMin = "-1.0", UIMax = "1.0", editcondition = "bOverride_ColorCorrectionHighlightsMin", DisplayName = "HighlightsMin"))
 	float ColorCorrectionHighlightsMin;
 
-	UPROPERTY(interp, BlueprintReadWrite, Category=Film, meta=(editcondition = "bOverride_FilmWhitePoint", DisplayName = "Tint", HideAlphaChannel))
+	UPROPERTY(interp, BlueprintReadWrite, Category="Tonemapper", meta=(editcondition = "bOverride_FilmWhitePoint", DisplayName = "Tint", HideAlphaChannel, LegacyTonemapper))
 	FLinearColor FilmWhitePoint;
-	UPROPERTY(interp, BlueprintReadWrite, Category=Film, AdvancedDisplay, meta=(editcondition = "bOverride_FilmShadowTint", DisplayName = "Tint Shadow", HideAlphaChannel))
+	UPROPERTY(interp, BlueprintReadWrite, Category="Tonemapper", AdvancedDisplay, meta=(editcondition = "bOverride_FilmShadowTint", DisplayName = "Tint Shadow", HideAlphaChannel, LegacyTonemapper))
 	FLinearColor FilmShadowTint;
-	UPROPERTY(interp, BlueprintReadWrite, Category=Film, AdvancedDisplay, meta=(UIMin = "0.0", UIMax = "1.0", editcondition = "bOverride_FilmShadowTintBlend", DisplayName = "Tint Shadow Blend"))
+	UPROPERTY(interp, BlueprintReadWrite, Category="Tonemapper", AdvancedDisplay, meta=(UIMin = "0.0", UIMax = "1.0", editcondition = "bOverride_FilmShadowTintBlend", DisplayName = "Tint Shadow Blend", LegacyTonemapper))
 	float FilmShadowTintBlend;
-	UPROPERTY(interp, BlueprintReadWrite, Category=Film, AdvancedDisplay, meta=(UIMin = "0.0", UIMax = "1.0", editcondition = "bOverride_FilmShadowTintAmount", DisplayName = "Tint Shadow Amount"))
+	UPROPERTY(interp, BlueprintReadWrite, Category="Tonemapper", AdvancedDisplay, meta=(UIMin = "0.0", UIMax = "1.0", editcondition = "bOverride_FilmShadowTintAmount", DisplayName = "Tint Shadow Amount", LegacyTonemapper))
 	float FilmShadowTintAmount;
 
-	UPROPERTY(interp, BlueprintReadWrite, Category=Film, meta=(UIMin = "0.0", UIMax = "2.0", editcondition = "bOverride_FilmSaturation", DisplayName = "Saturation"))
+	UPROPERTY(interp, BlueprintReadWrite, Category="Tonemapper", meta=(UIMin = "0.0", UIMax = "2.0", editcondition = "bOverride_FilmSaturation", DisplayName = "Saturation", LegacyTonemapper))
 	float FilmSaturation;
-	UPROPERTY(interp, BlueprintReadWrite, Category=Film, AdvancedDisplay, meta=(editcondition = "bOverride_FilmChannelMixerRed", DisplayName = "Channel Mixer Red", HideAlphaChannel))
+	UPROPERTY(interp, BlueprintReadWrite, Category="Tonemapper", AdvancedDisplay, meta=(editcondition = "bOverride_FilmChannelMixerRed", DisplayName = "Channel Mixer Red", HideAlphaChannel, LegacyTonemapper))
 	FLinearColor FilmChannelMixerRed;
-	UPROPERTY(interp, BlueprintReadWrite, Category=Film, AdvancedDisplay, meta=(editcondition = "bOverride_FilmChannelMixerGreen", DisplayName = "Channel Mixer Green", HideAlphaChannel))
+	UPROPERTY(interp, BlueprintReadWrite, Category="Tonemapper", AdvancedDisplay, meta=(editcondition = "bOverride_FilmChannelMixerGreen", DisplayName = "Channel Mixer Green", HideAlphaChannel, LegacyTonemapper))
 	FLinearColor FilmChannelMixerGreen;
-	UPROPERTY(interp, BlueprintReadWrite, Category=Film, AdvancedDisplay, meta=(editcondition = "bOverride_FilmChannelMixerBlue", DisplayName = " Channel Mixer Blue", HideAlphaChannel))
+	UPROPERTY(interp, BlueprintReadWrite, Category="Tonemapper", AdvancedDisplay, meta=(editcondition = "bOverride_FilmChannelMixerBlue", DisplayName = " Channel Mixer Blue", HideAlphaChannel, LegacyTonemapper))
 	FLinearColor FilmChannelMixerBlue;
 
-	UPROPERTY(interp, BlueprintReadWrite, Category=Film, meta=(UIMin = "0.0", UIMax = "1.0", editcondition = "bOverride_FilmContrast", DisplayName = "Contrast"))
+	UPROPERTY(interp, BlueprintReadWrite, Category="Tonemapper", meta=(UIMin = "0.0", UIMax = "1.0", editcondition = "bOverride_FilmContrast", DisplayName = "Contrast", LegacyTonemapper))
 	float FilmContrast;
-	UPROPERTY(interp, BlueprintReadWrite, Category=Film, AdvancedDisplay, meta=(UIMin = "0.0", UIMax = "1.0", editcondition = "bOverride_FilmToeAmount", DisplayName = "Crush Shadows"))
+	UPROPERTY(interp, BlueprintReadWrite, Category="Tonemapper", AdvancedDisplay, meta=(UIMin = "0.0", UIMax = "1.0", editcondition = "bOverride_FilmToeAmount", DisplayName = "Crush Shadows", LegacyTonemapper))
 	float FilmToeAmount;
-	UPROPERTY(interp, BlueprintReadWrite, Category=Film, AdvancedDisplay, meta=(UIMin = "0.0", UIMax = "1.0", editcondition = "bOverride_FilmHealAmount", DisplayName = "Crush Highlights"))
+	UPROPERTY(interp, BlueprintReadWrite, Category="Tonemapper", AdvancedDisplay, meta=(UIMin = "0.0", UIMax = "1.0", editcondition = "bOverride_FilmHealAmount", DisplayName = "Crush Highlights", LegacyTonemapper))
 	float FilmHealAmount;
-	UPROPERTY(interp, BlueprintReadWrite, Category=Film, AdvancedDisplay, meta=(UIMin = "1.0", UIMax = "4.0", editcondition = "bOverride_FilmDynamicRange", DisplayName = "Dynamic Range"))
+	UPROPERTY(interp, BlueprintReadWrite, Category="Tonemapper", AdvancedDisplay, meta=(UIMin = "1.0", UIMax = "4.0", editcondition = "bOverride_FilmDynamicRange", DisplayName = "Dynamic Range", LegacyTonemapper))
 	float FilmDynamicRange;
 
-	UPROPERTY(interp, BlueprintReadWrite, Category=Film, AdvancedDisplay, meta=(UIMin = "0.0", UIMax = "1.0", editcondition = "bOverride_FilmSlope", DisplayName = "Slope"))
+	UPROPERTY(interp, BlueprintReadWrite, Category="Tonemapper", meta=(UIMin = "0.0", UIMax = "1.0", editcondition = "bOverride_FilmSlope", DisplayName = "Slope"))
 	float FilmSlope;
-	UPROPERTY(interp, BlueprintReadWrite, Category=Film, AdvancedDisplay, meta=(UIMin = "0.0", UIMax = "1.0", editcondition = "bOverride_FilmToe", DisplayName = "Toe"))
+	UPROPERTY(interp, BlueprintReadWrite, Category="Tonemapper", meta=(UIMin = "0.0", UIMax = "1.0", editcondition = "bOverride_FilmToe", DisplayName = "Toe"))
 	float FilmToe;
-	UPROPERTY(interp, BlueprintReadWrite, Category=Film, AdvancedDisplay, meta=(UIMin = "0.0", UIMax = "1.0", editcondition = "bOverride_FilmShoulder", DisplayName = "Shoulder"))
+	UPROPERTY(interp, BlueprintReadWrite, Category="Tonemapper", meta=(UIMin = "0.0", UIMax = "1.0", editcondition = "bOverride_FilmShoulder", DisplayName = "Shoulder"))
 	float FilmShoulder;
-	UPROPERTY(interp, BlueprintReadWrite, Category=Film, AdvancedDisplay, meta=(UIMin = "0.0", UIMax = "1.0", editcondition = "bOverride_FilmBlackClip", DisplayName = "Black clip"))
+	UPROPERTY(interp, BlueprintReadWrite, Category="Tonemapper", meta=(UIMin = "0.0", UIMax = "1.0", editcondition = "bOverride_FilmBlackClip", DisplayName = "Black clip"))
 	float FilmBlackClip;
-	UPROPERTY(interp, BlueprintReadWrite, Category=Film, AdvancedDisplay, meta=(UIMin = "0.0", UIMax = "1.0", editcondition = "bOverride_FilmWhiteClip", DisplayName = "White clip"))
+	UPROPERTY(interp, BlueprintReadWrite, Category="Tonemapper", meta=(UIMin = "0.0", UIMax = "1.0", editcondition = "bOverride_FilmWhiteClip", DisplayName = "White clip"))
 	float FilmWhiteClip;
 	
 	/** Scene tint color */
-	UPROPERTY(interp, BlueprintReadWrite, Category=SceneColor, AdvancedDisplay, meta=(editcondition = "bOverride_SceneColorTint", HideAlphaChannel))
+	UPROPERTY(interp, BlueprintReadWrite, Category="Color Grading|Global", meta=(editcondition = "bOverride_SceneColorTint", HideAlphaChannel))
 	FLinearColor SceneColorTint;
 	
 	/** in percent, Scene chromatic aberration / color fringe (camera imperfection) to simulate an artifact that happens in real-world lens, mostly visible in the image corners. */
-	UPROPERTY(interp, BlueprintReadWrite, Category=SceneColor, meta=(UIMin = "0.0", UIMax = "5.0", editcondition = "bOverride_SceneFringeIntensity", DisplayName = "Fringe Intensity"))
+	UPROPERTY(interp, BlueprintReadWrite, Category="Lens|Image Effects", meta=(UIMin = "0.0", UIMax = "5.0", editcondition = "bOverride_SceneFringeIntensity", DisplayName = "Chromatic Aberation"))
 	float SceneFringeIntensity;
 
 	/** Multiplier for all bloom contributions >=0: off, 1(default), >1 brighter */
-	UPROPERTY(interp, BlueprintReadWrite, Category=Bloom, meta=(ClampMin = "0.0", UIMax = "8.0", editcondition = "bOverride_BloomIntensity", DisplayName = "Intensity"))
+	UPROPERTY(interp, BlueprintReadWrite, Category="Lens|Bloom", meta=(ClampMin = "0.0", UIMax = "8.0", editcondition = "bOverride_BloomIntensity", DisplayName = "Intensity"))
 	float BloomIntensity;
 
 	/**
 	 * minimum brightness the bloom starts having effect
 	 * -1:all pixels affect bloom equally (physically correct, faster as a threshold pass is omitted), 0:all pixels affect bloom brights more, 1(default), >1 brighter
 	 */
-	UPROPERTY(interp, BlueprintReadWrite, Category=Bloom, AdvancedDisplay, meta=(ClampMin = "-1.0", UIMax = "8.0", editcondition = "bOverride_BloomThreshold", DisplayName = "Threshold"))
+	UPROPERTY(interp, BlueprintReadWrite, Category="Lens|Bloom", meta=(ClampMin = "-1.0", UIMax = "8.0", editcondition = "bOverride_BloomThreshold", DisplayName = "Threshold"))
 	float BloomThreshold;
 
 	/**
 	 * Scale for all bloom sizes
 	 */
-	UPROPERTY(interp, BlueprintReadWrite, Category=Bloom, AdvancedDisplay, meta=(ClampMin = "0.0", UIMax = "64.0", editcondition = "bOverride_BloomSizeScale", DisplayName = "Size scale"))
+	UPROPERTY(interp, BlueprintReadWrite, Category="Lens|Bloom", AdvancedDisplay, meta=(ClampMin = "0.0", UIMax = "64.0", editcondition = "bOverride_BloomSizeScale", DisplayName = "Size scale"))
 	float BloomSizeScale;
 
 	/**
@@ -620,144 +620,94 @@ struct FPostProcessSettings
 	 * (is done in 1/2 resolution, larger values cost more performance, good for high frequency details)
 	 * >=0: can be clamped because of shader limitations
 	 */
-	UPROPERTY(interp, BlueprintReadWrite, Category=Bloom, AdvancedDisplay, meta=(ClampMin = "0.0", UIMax = "4.0", editcondition = "bOverride_Bloom1Size", DisplayName = "#1 Size"))
+	UPROPERTY(interp, BlueprintReadWrite, Category="Lens|Bloom", AdvancedDisplay, meta=(ClampMin = "0.0", UIMax = "4.0", editcondition = "bOverride_Bloom1Size", DisplayName = "#1 Size"))
 	float Bloom1Size;
 	/**
 	 * Diameter size for Bloom2 in percent of the screen width
 	 * (is done in 1/4 resolution, larger values cost more performance)
 	 * >=0: can be clamped because of shader limitations
 	 */
-	UPROPERTY(interp, BlueprintReadWrite, Category=Bloom, AdvancedDisplay, meta=(ClampMin = "0.0", UIMax = "8.0", editcondition = "bOverride_Bloom2Size", DisplayName = "#2 Size"))
+	UPROPERTY(interp, BlueprintReadWrite, Category="Lens|Bloom", AdvancedDisplay, meta=(ClampMin = "0.0", UIMax = "8.0", editcondition = "bOverride_Bloom2Size", DisplayName = "#2 Size"))
 	float Bloom2Size;
 	/**
 	 * Diameter size for Bloom3 in percent of the screen width
 	 * (is done in 1/8 resolution, larger values cost more performance)
 	 * >=0: can be clamped because of shader limitations
 	 */
-	UPROPERTY(interp, BlueprintReadWrite, Category=Bloom, AdvancedDisplay, meta=(ClampMin = "0.0", UIMax = "16.0", editcondition = "bOverride_Bloom3Size", DisplayName = "#3 Size"))
+	UPROPERTY(interp, BlueprintReadWrite, Category="Lens|Bloom", AdvancedDisplay, meta=(ClampMin = "0.0", UIMax = "16.0", editcondition = "bOverride_Bloom3Size", DisplayName = "#3 Size"))
 	float Bloom3Size;
 	/**
 	 * Diameter size for Bloom4 in percent of the screen width
 	 * (is done in 1/16 resolution, larger values cost more performance, best for wide contributions)
 	 * >=0: can be clamped because of shader limitations
 	 */
-	UPROPERTY(interp, BlueprintReadWrite, Category=Bloom, AdvancedDisplay, meta=(ClampMin = "0.0", UIMax = "32.0", editcondition = "bOverride_Bloom4Size", DisplayName = "#4 Size"))
+	UPROPERTY(interp, BlueprintReadWrite, Category="Lens|Bloom", AdvancedDisplay, meta=(ClampMin = "0.0", UIMax = "32.0", editcondition = "bOverride_Bloom4Size", DisplayName = "#4 Size"))
 	float Bloom4Size;
 	/**
 	 * Diameter size for Bloom5 in percent of the screen width
 	 * (is done in 1/32 resolution, larger values cost more performance, best for wide contributions)
 	 * >=0: can be clamped because of shader limitations
 	 */
-	UPROPERTY(interp, BlueprintReadWrite, Category=Bloom, AdvancedDisplay, meta=(ClampMin = "0.0", UIMax = "64.0", editcondition = "bOverride_Bloom5Size", DisplayName = "#5 Size"))
+	UPROPERTY(interp, BlueprintReadWrite, Category="Lens|Bloom", AdvancedDisplay, meta=(ClampMin = "0.0", UIMax = "64.0", editcondition = "bOverride_Bloom5Size", DisplayName = "#5 Size"))
 	float Bloom5Size;
 	/**
 	 * Diameter size for Bloom6 in percent of the screen width
 	 * (is done in 1/64 resolution, larger values cost more performance, best for wide contributions)
 	 * >=0: can be clamped because of shader limitations
 	 */
-	UPROPERTY(interp, BlueprintReadWrite, Category=Bloom, AdvancedDisplay, meta=(ClampMin = "0.0", UIMax = "128.0", editcondition = "bOverride_Bloom6Size", DisplayName = "#6 Size"))
+	UPROPERTY(interp, BlueprintReadWrite, Category="Lens|Bloom", AdvancedDisplay, meta=(ClampMin = "0.0", UIMax = "128.0", editcondition = "bOverride_Bloom6Size", DisplayName = "#6 Size"))
 	float Bloom6Size;
 
 	/** Bloom1 tint color */
-	UPROPERTY(interp, BlueprintReadWrite, Category=Bloom, AdvancedDisplay, meta=(editcondition = "bOverride_Bloom1Tint", DisplayName = "#1 Tint", HideAlphaChannel))
+	UPROPERTY(interp, BlueprintReadWrite, Category="Lens|Bloom", AdvancedDisplay, meta=(editcondition = "bOverride_Bloom1Tint", DisplayName = "#1 Tint", HideAlphaChannel))
 	FLinearColor Bloom1Tint;
 	/** Bloom2 tint color */
-	UPROPERTY(interp, BlueprintReadWrite, Category=Bloom, AdvancedDisplay, meta=(editcondition = "bOverride_Bloom2Tint", DisplayName = "#2 Tint", HideAlphaChannel))
+	UPROPERTY(interp, BlueprintReadWrite, Category="Lens|Bloom", AdvancedDisplay, meta=(editcondition = "bOverride_Bloom2Tint", DisplayName = "#2 Tint", HideAlphaChannel))
 	FLinearColor Bloom2Tint;
 	/** Bloom3 tint color */
-	UPROPERTY(interp, BlueprintReadWrite, Category=Bloom, AdvancedDisplay, meta=(editcondition = "bOverride_Bloom3Tint", DisplayName = "#3 Tint", HideAlphaChannel))
+	UPROPERTY(interp, BlueprintReadWrite, Category="Lens|Bloom", AdvancedDisplay, meta=(editcondition = "bOverride_Bloom3Tint", DisplayName = "#3 Tint", HideAlphaChannel))
 	FLinearColor Bloom3Tint;
 	/** Bloom4 tint color */
-	UPROPERTY(interp, BlueprintReadWrite, Category=Bloom, AdvancedDisplay, meta=(editcondition = "bOverride_Bloom4Tint", DisplayName = "#4 Tint", HideAlphaChannel))
+	UPROPERTY(interp, BlueprintReadWrite, Category="Lens|Bloom", AdvancedDisplay, meta=(editcondition = "bOverride_Bloom4Tint", DisplayName = "#4 Tint", HideAlphaChannel))
 	FLinearColor Bloom4Tint;
 	/** Bloom5 tint color */
-	UPROPERTY(interp, BlueprintReadWrite, Category=Bloom, AdvancedDisplay, meta=(editcondition = "bOverride_Bloom5Tint", DisplayName = "#5 Tint", HideAlphaChannel))
+	UPROPERTY(interp, BlueprintReadWrite, Category="Lens|Bloom", AdvancedDisplay, meta=(editcondition = "bOverride_Bloom5Tint", DisplayName = "#5 Tint", HideAlphaChannel))
 	FLinearColor Bloom5Tint;
 	/** Bloom6 tint color */
-	UPROPERTY(interp, BlueprintReadWrite, Category=Bloom, AdvancedDisplay, meta=(editcondition = "bOverride_Bloom6Tint", DisplayName = "#6 Tint", HideAlphaChannel))
+	UPROPERTY(interp, BlueprintReadWrite, Category="Lens|Bloom", AdvancedDisplay, meta=(editcondition = "bOverride_Bloom6Tint", DisplayName = "#6 Tint", HideAlphaChannel))
 	FLinearColor Bloom6Tint;
 
 	/** BloomDirtMask intensity */
-	UPROPERTY(interp, BlueprintReadWrite, Category=Bloom, meta=(ClampMin = "0.0", UIMax = "8.0", editcondition = "bOverride_BloomDirtMaskIntensity", DisplayName = "Dirt Mask Intensity"))
+	UPROPERTY(interp, BlueprintReadWrite, Category="Lens|Bloom", meta=(ClampMin = "0.0", UIMax = "8.0", editcondition = "bOverride_BloomDirtMaskIntensity", DisplayName = "Dirt Mask Intensity"))
 	float BloomDirtMaskIntensity;
 
 	/** BloomDirtMask tint color */
-	UPROPERTY(interp, BlueprintReadWrite, Category=Bloom, AdvancedDisplay, meta=(editcondition = "bOverride_BloomDirtMaskTint", DisplayName = "Dirt Mask Tint", HideAlphaChannel))
+	UPROPERTY(interp, BlueprintReadWrite, Category="Lens|Dirt Mask", meta=(editcondition = "bOverride_BloomDirtMaskTint", DisplayName = "Dirt Mask Tint", HideAlphaChannel))
 	FLinearColor BloomDirtMaskTint;
 
 	/**
 	 * Texture that defines the dirt on the camera lens where the light of very bright objects is scattered.
 	 */
-	UPROPERTY(EditAnywhere, BlueprintReadWrite, Category=Bloom, meta=(editcondition = "bOverride_BloomDirtMask", DisplayName = "Dirt Mask"))
-	class UTexture* BloomDirtMask;	// The plan is to replace this texture with many small textures quads for better performance, more control and to animate the effect.
-
-	/** How strong the dynamic GI from the LPV should be. 0.0 is off, 1.0 is the "normal" value, but higher values can be used to boost the effect*/
-	UPROPERTY(interp, BlueprintReadWrite, Category=LightPropagationVolume, meta=(editcondition = "bOverride_LPVIntensity", UIMin = "0", UIMax = "20", DisplayName = "Intensity") )
-	float LPVIntensity;
-
-	/** Bias applied to light injected into the LPV in cell units. Increase to reduce bleeding through thin walls*/
-	UPROPERTY(interp, BlueprintReadWrite, Category=LightPropagationVolume, AdvancedDisplay, meta=(editcondition = "bOverride_LPVVplInjectionBias", UIMin = "0", UIMax = "2", DisplayName = "Light Injection Bias") )
-	float LPVVplInjectionBias;
-
-	/** The size of the LPV volume, in Unreal units*/
-	UPROPERTY(EditAnywhere, BlueprintReadWrite, Category=LightPropagationVolume, meta=(editcondition = "bOverride_LPVSize", UIMin = "100", UIMax = "20000", DisplayName = "Size") )
-	float LPVSize;
-
-	/** Secondary occlusion strength (bounce light shadows). Set to 0 to disable*/
-	UPROPERTY(interp, BlueprintReadWrite, Category=LightPropagationVolume, meta=(editcondition = "bOverride_LPVSecondaryOcclusionIntensity", UIMin = "0", UIMax = "1", DisplayName = "Secondary Occlusion Intensity") )
-	float LPVSecondaryOcclusionIntensity;
-
-	/** Secondary bounce light strength (bounce light shadows). Set to 0 to disable*/
-	UPROPERTY(interp, BlueprintReadWrite, Category=LightPropagationVolume, AdvancedDisplay, meta=(editcondition = "bOverride_LPVSecondaryBounceIntensity", UIMin = "0", UIMax = "1", DisplayName = "Secondary Bounce Intensity") )
-	float LPVSecondaryBounceIntensity;
-
-	/** Bias applied to the geometry volume in cell units. Increase to reduce darkening due to secondary occlusion */
-	UPROPERTY(interp, BlueprintReadWrite, Category=LightPropagationVolume, AdvancedDisplay, meta=(editcondition = "bOverride_LPVGeometryVolumeBias", UIMin = "0", UIMax = "2", DisplayName = "Geometry Volume Bias"))
-	float LPVGeometryVolumeBias;
-
-	UPROPERTY(interp, BlueprintReadWrite, Category=LightPropagationVolume, AdvancedDisplay, meta=(editcondition = "bOverride_LPVEmissiveInjectionIntensity", UIMin = "0", UIMax = "20", DisplayName = "Emissive Injection Intensity") )
-	float LPVEmissiveInjectionIntensity;
-
-	/** Controls the amount of directional occlusion. Requires LPV. Values very close to 1.0 are recommended */
-	UPROPERTY(interp, BlueprintReadWrite, Category=LightPropagationVolume, meta=(editcondition = "bOverride_LPVDirectionalOcclusionIntensity", UIMin = "0", UIMax = "1", DisplayName = "Occlusion Intensity") )
-	float LPVDirectionalOcclusionIntensity;
-
-	/** Occlusion Radius - 16 is recommended for most scenes */
-	UPROPERTY(interp, BlueprintReadWrite, Category=LightPropagationVolume, AdvancedDisplay, meta=(editcondition = "bOverride_LPVDirectionalOcclusionRadius", UIMin = "1", UIMax = "16", DisplayName = "Occlusion Radius") )
-	float LPVDirectionalOcclusionRadius;
-
-	/** Diffuse occlusion exponent - increase for more contrast. 1 to 2 is recommended */
-	UPROPERTY(interp, BlueprintReadWrite, Category=LightPropagationVolume, meta=(editcondition = "bOverride_LPVDiffuseOcclusionExponent", UIMin = "0.5", UIMax = "5", DisplayName = "Diffuse occlusion exponent") )
-	float LPVDiffuseOcclusionExponent;
-
-	/** Specular occlusion exponent - increase for more contrast. 6 to 9 is recommended */
-	UPROPERTY(interp, BlueprintReadWrite, Category=LightPropagationVolume, meta=(editcondition = "bOverride_LPVSpecularOcclusionExponent", UIMin = "1", UIMax = "16", DisplayName = "Specular occlusion exponent") )
-	float LPVSpecularOcclusionExponent;
-
-	/** Diffuse occlusion intensity - higher values provide increased diffuse occlusion.*/
-	UPROPERTY(interp, BlueprintReadWrite, Category=LightPropagationVolume, AdvancedDisplay, meta=(editcondition = "bOverride_LPVDiffuseOcclusionIntensity", UIMin = "0", UIMax = "4", DisplayName = "Diffuse occlusion intensity") )
-	float LPVDiffuseOcclusionIntensity;
-
-	/** Specular occlusion intensity - higher values provide increased specular occlusion.*/
-	UPROPERTY(interp, BlueprintReadWrite, Category=LightPropagationVolume, AdvancedDisplay, meta=(editcondition = "bOverride_LPVSpecularOcclusionIntensity", UIMin = "0", UIMax = "4", DisplayName = "Specular occlusion intensity") )
-	float LPVSpecularOcclusionIntensity;
+	UPROPERTY(EditAnywhere, BlueprintReadWrite, Category="Lens|Dirt Mask", meta=(editcondition = "bOverride_BloomDirtMask", DisplayName = "Dirt Mask Texture"))
+	class UTexture* BloomDirtMask;	
 	
 	/** AmbientCubemap tint color */
-	UPROPERTY(interp, BlueprintReadWrite, Category=AmbientCubemap, AdvancedDisplay, meta=(editcondition = "bOverride_AmbientCubemapTint", DisplayName = "Tint", HideAlphaChannel))
+	UPROPERTY(interp, BlueprintReadWrite, Category="Rendering Features|Ambient Cubemap", meta=(editcondition = "bOverride_AmbientCubemapTint", DisplayName = "Tint", HideAlphaChannel))
 	FLinearColor AmbientCubemapTint;
+
 	/**
 	 * To scale the Ambient cubemap brightness
 	 * >=0: off, 1(default), >1 brighter
 	 */
-	UPROPERTY(interp, BlueprintReadWrite, Category=AmbientCubemap, meta=(ClampMin = "0.0", UIMax = "4.0", editcondition = "bOverride_AmbientCubemapIntensity", DisplayName = "Intensity"))
+	UPROPERTY(interp, BlueprintReadWrite, Category="Rendering Features|Ambient Cubemap", meta=(ClampMin = "0.0", UIMax = "4.0", editcondition = "bOverride_AmbientCubemapIntensity", DisplayName = "Intensity"))
 	float AmbientCubemapIntensity;
+
 	/** The Ambient cubemap (Affects diffuse and specular shading), blends additively which if different from all other settings here */
-	UPROPERTY(EditAnywhere, BlueprintReadWrite, Category=AmbientCubemap, meta=(DisplayName = "Cubemap Texture"))
+	UPROPERTY(EditAnywhere, BlueprintReadWrite, Category="Rendering Features|Ambient Cubemap", meta=(DisplayName = "Cubemap Texture"))
 	class UTextureCube* AmbientCubemap;
 
-
 	/** Luminance computation method */
-	UPROPERTY(EditAnywhere, BlueprintReadWrite, Category=AutoExposure, meta=(editcondition = "bOverride_AutoExposureMethod", DisplayName = "Method"))
+	UPROPERTY(EditAnywhere, BlueprintReadWrite, Category="Lens|Auto Exposure", meta=(editcondition = "bOverride_AutoExposureMethod", DisplayName = "Method"))
     TEnumAsByte<enum EAutoExposureMethod> AutoExposureMethod;
 
 	/**
@@ -767,7 +717,7 @@ struct FPostProcessSettings
 	 * bright spots.
 	 * >0, <100, good values are in the range 70 .. 80
 	 */
-	UPROPERTY(interp, BlueprintReadWrite, Category=AutoExposure, AdvancedDisplay, meta=(ClampMin = "0.0", ClampMax = "100.0", editcondition = "bOverride_AutoExposureLowPercent", DisplayName = "Low Percent"))
+	UPROPERTY(interp, BlueprintReadWrite, Category="Lens|Auto Exposure", AdvancedDisplay, meta=(ClampMin = "0.0", ClampMax = "100.0", editcondition = "bOverride_AutoExposureLowPercent", DisplayName = "Low Percent"))
 	float AutoExposureLowPercent;
 
 	/**
@@ -777,7 +727,7 @@ struct FPostProcessSettings
 	 * bright spots.
 	 * >0, <100, good values are in the range 80 .. 95
 	 */
-	UPROPERTY(interp, BlueprintReadWrite, Category=AutoExposure, AdvancedDisplay, meta=(ClampMin = "0.0", ClampMax = "100.0", editcondition = "bOverride_AutoExposureHighPercent", DisplayName = "High Percent"))
+	UPROPERTY(interp, BlueprintReadWrite, Category="Lens|Auto Exposure", AdvancedDisplay, meta=(ClampMin = "0.0", ClampMax = "100.0", editcondition = "bOverride_AutoExposureHighPercent", DisplayName = "High Percent"))
 	float AutoExposureHighPercent;
 
 	/**
@@ -787,7 +737,7 @@ struct FPostProcessSettings
 	 * effect and defined the HDR range - you don't want to change that late in the project development.
 	 * Eye Adaptation is disabled if MinBrightness = MaxBrightness
 	 */
-	UPROPERTY(interp, BlueprintReadWrite, Category=AutoExposure, meta=(ClampMin = "0.0", UIMax = "10.0", editcondition = "bOverride_AutoExposureMinBrightness", DisplayName = "Min Brightness"))
+	UPROPERTY(interp, BlueprintReadWrite, Category="Lens|Auto Exposure", meta=(ClampMin = "0.0", UIMax = "10.0", editcondition = "bOverride_AutoExposureMinBrightness", DisplayName = "Min Brightness"))
 	float AutoExposureMinBrightness;
 
 	/**
@@ -797,90 +747,90 @@ struct FPostProcessSettings
 	 * effect and defined the HDR range - you don't want to change that late in the project development.
 	 * Eye Adaptation is disabled if MinBrightness = MaxBrightness
 	 */
-	UPROPERTY(interp, BlueprintReadWrite, Category=AutoExposure, meta=(ClampMin = "0.0", UIMax = "10.0", editcondition = "bOverride_AutoExposureMaxBrightness", DisplayName = "Max Brightness"))
+	UPROPERTY(interp, BlueprintReadWrite, Category="Lens|Auto Exposure", meta=(ClampMin = "0.0", UIMax = "10.0", editcondition = "bOverride_AutoExposureMaxBrightness", DisplayName = "Max Brightness"))
 	float AutoExposureMaxBrightness;
 
 	/** >0 */
-	UPROPERTY(interp, BlueprintReadWrite, Category=AutoExposure, AdvancedDisplay, meta=(ClampMin = "0.02", UIMax = "20.0", editcondition = "bOverride_AutoExposureSpeedUp", DisplayName = "Speed Up"))
+	UPROPERTY(interp, BlueprintReadWrite, Category="Lens|Auto Exposure", meta=(ClampMin = "0.02", UIMax = "20.0", editcondition = "bOverride_AutoExposureSpeedUp", DisplayName = "Speed Up"))
 	float AutoExposureSpeedUp;
 
 	/** >0 */
-	UPROPERTY(interp, BlueprintReadWrite, Category=AutoExposure, AdvancedDisplay, meta=(ClampMin = "0.02", UIMax = "20.0", editcondition = "bOverride_AutoExposureSpeedDown", DisplayName = "Speed Down"))
+	UPROPERTY(interp, BlueprintReadWrite, Category="Lens|Auto Exposure", meta=(ClampMin = "0.02", UIMax = "20.0", editcondition = "bOverride_AutoExposureSpeedDown", DisplayName = "Speed Down"))
 	float AutoExposureSpeedDown;
 
 	/**
 	 * Logarithmic adjustment for the exposure. Only used if a tonemapper is specified.
 	 * 0: no adjustment, -1:2x darker, -2:4x darker, 1:2x brighter, 2:4x brighter, ...
 	 */
-	UPROPERTY(interp, BlueprintReadWrite, Category = AutoExposure, meta = (UIMin = "-8.0", UIMax = "8.0", editcondition = "bOverride_AutoExposureBias", DisplayName = "Exposure Bias"))
+	UPROPERTY(interp, BlueprintReadWrite, Category = "Lens|Auto Exposure", meta = (UIMin = "-8.0", UIMax = "8.0", editcondition = "bOverride_AutoExposureBias", DisplayName = "Exposure Bias"))
 	float AutoExposureBias;
 
 	/** temporary exposed until we found good values, -8: 1/256, -10: 1/1024 */
-	UPROPERTY(interp, BlueprintReadWrite, Category=AutoExposure, AdvancedDisplay, meta=(UIMin = "-16", UIMax = "0.0", editcondition = "bOverride_HistogramLogMin"))
+	UPROPERTY(interp, BlueprintReadWrite, Category="Lens|Auto Exposure", AdvancedDisplay, meta=(UIMin = "-16", UIMax = "0.0", editcondition = "bOverride_HistogramLogMin"))
 	float HistogramLogMin;
 
 	/** temporary exposed until we found good values 4: 16, 8: 256 */
-	UPROPERTY(interp, BlueprintReadWrite, Category=AutoExposure, AdvancedDisplay, meta=(UIMin = "0.0", UIMax = "16.0", editcondition = "bOverride_HistogramLogMax"))
+	UPROPERTY(interp, BlueprintReadWrite, Category="Lens|Auto Exposure", AdvancedDisplay, meta=(UIMin = "0.0", UIMax = "16.0", editcondition = "bOverride_HistogramLogMax"))
 	float HistogramLogMax;
 
 	/** Brightness scale of the image cased lens flares (linear) */
-	UPROPERTY(interp, BlueprintReadWrite, Category=LensFlares, meta=(UIMin = "0.0", UIMax = "16.0", editcondition = "bOverride_LensFlareIntensity", DisplayName = "Intensity"))
+	UPROPERTY(interp, BlueprintReadWrite, Category="Lens|Lens Flares", meta=(UIMin = "0.0", UIMax = "16.0", editcondition = "bOverride_LensFlareIntensity", DisplayName = "Intensity"))
 	float LensFlareIntensity;
 
 	/** Tint color for the image based lens flares. */
-	UPROPERTY(interp, BlueprintReadWrite, Category=LensFlares, AdvancedDisplay, meta=(editcondition = "bOverride_LensFlareTint", DisplayName = "Tint", HideAlphaChannel))
+	UPROPERTY(interp, BlueprintReadWrite, Category="Lens|Lens Flares", meta=(editcondition = "bOverride_LensFlareTint", DisplayName = "Tint", HideAlphaChannel))
 	FLinearColor LensFlareTint;
 
 	/** Size of the Lens Blur (in percent of the view width) that is done with the Bokeh texture (note: performance cost is radius*radius) */
-	UPROPERTY(interp, BlueprintReadWrite, Category=LensFlares, meta=(UIMin = "0.0", UIMax = "32.0", editcondition = "bOverride_LensFlareBokehSize", DisplayName = "BokehSize"))
+	UPROPERTY(interp, BlueprintReadWrite, Category="Lens|Lens Flares", meta=(UIMin = "0.0", UIMax = "32.0", editcondition = "bOverride_LensFlareBokehSize", DisplayName = "BokehSize"))
 	float LensFlareBokehSize;
 
 	/** Minimum brightness the lens flare starts having effect (this should be as high as possible to avoid the performance cost of blurring content that is too dark too see) */
-	UPROPERTY(interp, BlueprintReadWrite, Category=LensFlares, AdvancedDisplay, meta=(UIMin = "0.1", UIMax = "32.0", editcondition = "bOverride_LensFlareThreshold", DisplayName = "Threshold"))
+	UPROPERTY(interp, BlueprintReadWrite, Category="Lens|Lens Flares", meta=(UIMin = "0.1", UIMax = "32.0", editcondition = "bOverride_LensFlareThreshold", DisplayName = "Threshold"))
 	float LensFlareThreshold;
 
 	/** Defines the shape of the Bokeh when the image base lens flares are blurred, cannot be blended */
-	UPROPERTY(EditAnywhere, BlueprintReadWrite, Category=LensFlares, meta=(editcondition = "bOverride_LensFlareBokehShape", DisplayName = "BokehShape"))
+	UPROPERTY(EditAnywhere, BlueprintReadWrite, Category="Lens|Lens Flares", meta=(editcondition = "bOverride_LensFlareBokehShape", DisplayName = "BokehShape"))
 	class UTexture* LensFlareBokehShape;
 
 	/** RGB defines the lens flare color, A it's position. This is a temporary solution. */
-	UPROPERTY(EditAnywhere, Category=LensFlares, meta=(editcondition = "bOverride_LensFlareTints", DisplayName = "Tints"))
+	UPROPERTY(EditAnywhere, Category="Lens|Lens Flares", meta=(editcondition = "bOverride_LensFlareTints", DisplayName = "Tints"))
 	FLinearColor LensFlareTints[8];
 
 	/** 0..1 0=off/no vignette .. 1=strong vignette */
-	UPROPERTY(interp, BlueprintReadWrite, Category=SceneColor, meta=(UIMin = "0.0", UIMax = "1.0", editcondition = "bOverride_VignetteIntensity"))
+	UPROPERTY(interp, BlueprintReadWrite, Category="Lens|Image Effects", meta=(UIMin = "0.0", UIMax = "1.0", editcondition = "bOverride_VignetteIntensity"))
 	float VignetteIntensity;
 
 	/** 0..1 grain jitter */
-	UPROPERTY(interp, BlueprintReadWrite, Category = SceneColor, AdvancedDisplay, meta=(UIMin = "0.0", UIMax = "1.0", editcondition = "bOverride_GrainJitter"))
+	UPROPERTY(interp, BlueprintReadWrite, Category = "Lens|Image Effects", meta=(UIMin = "0.0", UIMax = "1.0", editcondition = "bOverride_GrainJitter"))
 	float GrainJitter;
 
 	/** 0..1 grain intensity */
-	UPROPERTY(interp, BlueprintReadWrite, Category=SceneColor, meta=(UIMin = "0.0", UIMax = "1.0", editcondition = "bOverride_GrainIntensity"))
+	UPROPERTY(interp, BlueprintReadWrite, Category="Lens|Image Effects", meta=(UIMin = "0.0", UIMax = "1.0", editcondition = "bOverride_GrainIntensity"))
 	float GrainIntensity;
 
 	/** 0..1 0=off/no ambient occlusion .. 1=strong ambient occlusion, defines how much it affects the non direct lighting after base pass */
-	UPROPERTY(interp, BlueprintReadWrite, Category=AmbientOcclusion, meta=(ClampMin = "0.0", ClampMax = "1.0", editcondition = "bOverride_AmbientOcclusionIntensity", DisplayName = "Intensity"))
+	UPROPERTY(interp, BlueprintReadWrite, Category="Rendering Features|Ambient Occlusion", meta=(ClampMin = "0.0", ClampMax = "1.0", editcondition = "bOverride_AmbientOcclusionIntensity", DisplayName = "Intensity"))
 	float AmbientOcclusionIntensity;
 
 	/** 0..1 0=no effect on static lighting .. 1=AO affects the stat lighting, 0 is free meaning no extra rendering pass */
-	UPROPERTY(interp, BlueprintReadWrite, Category=AmbientOcclusion, AdvancedDisplay, meta=(ClampMin = "0.0", ClampMax = "1.0", editcondition = "bOverride_AmbientOcclusionStaticFraction", DisplayName = "Static Fraction"))
+	UPROPERTY(interp, BlueprintReadWrite, Category="Rendering Features|Ambient Occlusion", AdvancedDisplay, meta=(ClampMin = "0.0", ClampMax = "1.0", editcondition = "bOverride_AmbientOcclusionStaticFraction", DisplayName = "Static Fraction"))
 	float AmbientOcclusionStaticFraction;
 
 	/** >0, in unreal units, bigger values means even distant surfaces affect the ambient occlusion */
-	UPROPERTY(interp, BlueprintReadWrite, Category=AmbientOcclusion, meta=(ClampMin = "0.1", UIMax = "500.0", editcondition = "bOverride_AmbientOcclusionRadius", DisplayName = "Radius"))
+	UPROPERTY(interp, BlueprintReadWrite, Category="Rendering Features|Ambient Occlusion", meta=(ClampMin = "0.1", UIMax = "500.0", editcondition = "bOverride_AmbientOcclusionRadius", DisplayName = "Radius"))
 	float AmbientOcclusionRadius;
 
 	/** true: AO radius is in world space units, false: AO radius is locked the view space in 400 units */
-	UPROPERTY(EditAnywhere, BlueprintReadWrite, Category=AmbientOcclusion, AdvancedDisplay, meta=(editcondition = "bOverride_AmbientOcclusionRadiusInWS", DisplayName = "Radius in WorldSpace"))
+	UPROPERTY(EditAnywhere, BlueprintReadWrite, Category="Rendering Features|Ambient Occlusion", AdvancedDisplay, meta=(editcondition = "bOverride_AmbientOcclusionRadiusInWS", DisplayName = "Radius in WorldSpace"))
 	uint32 AmbientOcclusionRadiusInWS:1;
 
 	/** >0, in unreal units, at what distance the AO effect disppears in the distance (avoding artifacts and AO effects on huge object) */
-	UPROPERTY(interp, BlueprintReadWrite, Category=AmbientOcclusion, AdvancedDisplay, meta=(ClampMin = "0.0", UIMax = "20000.0", editcondition = "bOverride_AmbientOcclusionFadeDistance", DisplayName = "Fade Out Distance"))
+	UPROPERTY(interp, BlueprintReadWrite, Category="Rendering Features|Ambient Occlusion", AdvancedDisplay, meta=(ClampMin = "0.0", UIMax = "20000.0", editcondition = "bOverride_AmbientOcclusionFadeDistance", DisplayName = "Fade Out Distance"))
 	float AmbientOcclusionFadeDistance;
 	
 	/** >0, in unreal units, how many units before AmbientOcclusionFadeOutDistance it starts fading out */
-	UPROPERTY(interp, BlueprintReadWrite, Category=AmbientOcclusion, AdvancedDisplay, meta=(ClampMin = "0.0", UIMax = "20000.0", editcondition = "bOverride_AmbientOcclusionFadeRadius", DisplayName = "Fade Out Radius"))
+	UPROPERTY(interp, BlueprintReadWrite, Category="Rendering Features|Ambient Occlusion", AdvancedDisplay, meta=(ClampMin = "0.0", UIMax = "20000.0", editcondition = "bOverride_AmbientOcclusionFadeRadius", DisplayName = "Fade Out Radius"))
 	float AmbientOcclusionFadeRadius;
 
 	/** >0, in unreal units, how wide the ambient occlusion effect should affect the geometry (in depth), will be removed - only used for non normal method which is not exposed */
@@ -888,159 +838,210 @@ struct FPostProcessSettings
 	float AmbientOcclusionDistance_DEPRECATED;
 
 	/** >0, in unreal units, bigger values means even distant surfaces affect the ambient occlusion */
-	UPROPERTY(interp, BlueprintReadWrite, Category=AmbientOcclusion, AdvancedDisplay, meta=(ClampMin = "0.1", UIMax = "8.0", editcondition = "bOverride_AmbientOcclusionPower", DisplayName = "Power"))
+	UPROPERTY(interp, BlueprintReadWrite, Category="Rendering Features|Ambient Occlusion", AdvancedDisplay, meta=(ClampMin = "0.1", UIMax = "8.0", editcondition = "bOverride_AmbientOcclusionPower", DisplayName = "Power"))
 	float AmbientOcclusionPower;
 
 	/** >0, in unreal units, default (3.0) works well for flat surfaces but can reduce details */
-	UPROPERTY(interp, BlueprintReadWrite, Category=AmbientOcclusion, AdvancedDisplay, meta=(ClampMin = "0.0", UIMax = "10.0", editcondition = "bOverride_AmbientOcclusionBias", DisplayName = "Bias"))
+	UPROPERTY(interp, BlueprintReadWrite, Category="Rendering Features|Ambient Occlusion", AdvancedDisplay, meta=(ClampMin = "0.0", UIMax = "10.0", editcondition = "bOverride_AmbientOcclusionBias", DisplayName = "Bias"))
 	float AmbientOcclusionBias;
 
 	/** 0=lowest quality..100=maximum quality, only a few quality levels are implemented, no soft transition */
-	UPROPERTY(interp, BlueprintReadWrite, Category=AmbientOcclusion, AdvancedDisplay, meta=(ClampMin = "0.0", UIMax = "100.0", editcondition = "bOverride_AmbientOcclusionQuality", DisplayName = "Quality"))
+	UPROPERTY(interp, BlueprintReadWrite, Category="Rendering Features|Ambient Occlusion", AdvancedDisplay, meta=(ClampMin = "0.0", UIMax = "100.0", editcondition = "bOverride_AmbientOcclusionQuality", DisplayName = "Quality"))
 	float AmbientOcclusionQuality;
 
 	/** Affects the blend over the multiple mips (lower resolution versions) , 0:fully use full resolution, 1::fully use low resolution, around 0.6 seems to be a good value */
-	UPROPERTY(interp, BlueprintReadWrite, Category=AmbientOcclusion, AdvancedDisplay, meta=(ClampMin = "0.1", UIMax = "1.0", editcondition = "bOverride_AmbientOcclusionMipBlend", DisplayName = "Mip Blend"))
+	UPROPERTY(interp, BlueprintReadWrite, Category="Rendering Features|Ambient Occlusion", AdvancedDisplay, meta=(ClampMin = "0.1", UIMax = "1.0", editcondition = "bOverride_AmbientOcclusionMipBlend", DisplayName = "Mip Blend"))
 	float AmbientOcclusionMipBlend;
 
 	/** Affects the radius AO radius scale over the multiple mips (lower resolution versions) */
-	UPROPERTY(interp, BlueprintReadWrite, Category=AmbientOcclusion, AdvancedDisplay, meta=(ClampMin = "0.5", UIMax = "4.0", editcondition = "bOverride_AmbientOcclusionMipScale", DisplayName = "Mip Scale"))
+	UPROPERTY(interp, BlueprintReadWrite, Category="Rendering Features|Ambient Occlusion", AdvancedDisplay, meta=(ClampMin = "0.5", UIMax = "4.0", editcondition = "bOverride_AmbientOcclusionMipScale", DisplayName = "Mip Scale"))
 	float AmbientOcclusionMipScale;
 
 	/** to tweak the bilateral upsampling when using multiple mips (lower resolution versions) */
-	UPROPERTY(interp, BlueprintReadWrite, Category=AmbientOcclusion, AdvancedDisplay, meta=(ClampMin = "0.0", UIMax = "0.1", editcondition = "bOverride_AmbientOcclusionMipThreshold", DisplayName = "Mip Threshold"))
+	UPROPERTY(interp, BlueprintReadWrite, Category="Rendering Features|Ambient Occlusion", AdvancedDisplay, meta=(ClampMin = "0.0", UIMax = "0.1", editcondition = "bOverride_AmbientOcclusionMipThreshold", DisplayName = "Mip Threshold"))
 	float AmbientOcclusionMipThreshold;
 
 	/** Adjusts indirect lighting color. (1,1,1) is default. (0,0,0) to disable GI. The show flag 'Global Illumination' must be enabled to use this property. */
-	UPROPERTY(interp, BlueprintReadWrite, Category=GlobalIllumination, meta=(editcondition = "bOverride_IndirectLightingColor", DisplayName = "Indirect Lighting Color", HideAlphaChannel))
+	UPROPERTY(interp, BlueprintReadWrite, Category="Rendering Features|Global Illumination", meta=(editcondition = "bOverride_IndirectLightingColor", DisplayName = "Indirect Lighting Color", HideAlphaChannel))
 	FLinearColor IndirectLightingColor;
 
 	/** Scales the indirect lighting contribution. A value of 0 disables GI. Default is 1. The show flag 'Global Illumination' must be enabled to use this property. */
-	UPROPERTY(interp, BlueprintReadWrite, Category=GlobalIllumination, meta=(ClampMin = "0", UIMax = "4.0", editcondition = "bOverride_IndirectLightingIntensity", DisplayName = "Indirect Lighting Intensity"))
+	UPROPERTY(interp, BlueprintReadWrite, Category="Rendering Features|Global Illumination", meta=(ClampMin = "0", UIMax = "4.0", editcondition = "bOverride_IndirectLightingIntensity", DisplayName = "Indirect Lighting Intensity"))
 	float IndirectLightingIntensity;
 
-	/** 0..1=full intensity */
-	UPROPERTY(interp, BlueprintReadWrite, Category=SceneColor, meta=(ClampMin = "0", ClampMax = "1.0", editcondition = "bOverride_ColorGradingIntensity", DisplayName = "Color Grading Intensity"))
+	/** Color grading lookup table intensity. 0 = no intensity, 1=full intensity */
+	UPROPERTY(interp, BlueprintReadWrite, Category="Color Grading|Global", meta=(ClampMin = "0", ClampMax = "1.0", editcondition = "bOverride_ColorGradingIntensity", DisplayName = "Color Grading LUT Intensity"))
 	float ColorGradingIntensity;
 
-	/** Name of the LUT texture e.g. MyPackage01.LUTNeutral, empty if not used */
-	UPROPERTY(EditAnywhere, BlueprintReadWrite, Category=SceneColor, meta=(editcondition = "bOverride_ColorGradingLUT", DisplayName = "Color Grading"))
+	/** Look up table texture to use or none of not used*/
+	UPROPERTY(EditAnywhere, BlueprintReadWrite, Category="Color Grading|Global", meta=(editcondition = "bOverride_ColorGradingLUT", DisplayName = "Color Grading LUT"))
 	class UTexture* ColorGradingLUT;
 
 	/** BokehDOF, Simple gaussian, ... Mobile supports Gaussian only. */
-	UPROPERTY(EditAnywhere, BlueprintReadWrite, Category=DepthOfField, meta=(editcondition = "bOverride_DepthOfFieldMethod", DisplayName = "Method"))
+	UPROPERTY(EditAnywhere, BlueprintReadWrite, Category="Lens|Depth of Field", meta=(editcondition = "bOverride_DepthOfFieldMethod", DisplayName = "Method"))
 	TEnumAsByte<enum EDepthOfFieldMethod> DepthOfFieldMethod;
 
 	/** Enable HQ Gaussian on high end mobile platforms. (ES3_1) */
-	UPROPERTY(EditAnywhere, BlueprintReadWrite, Category = DepthOfField, meta = (editcondition = "bOverride_MobileHQGaussian", DisplayName = "High Quality Gaussian DoF on Mobile"))
+	UPROPERTY(EditAnywhere, BlueprintReadWrite, Category = "Lens|Depth of Field", meta = (editcondition = "bOverride_MobileHQGaussian", DisplayName = "High Quality Gaussian DoF on Mobile"))
 	uint32 bMobileHQGaussian : 1;
 
 	/** CircleDOF only: Defines the opening of the camera lens, Aperture is 1/fstop, typical lens go down to f/1.2 (large opening), larger numbers reduce the DOF effect */
-	UPROPERTY(interp, BlueprintReadWrite, Category=DepthOfField, meta=(ClampMin = "1.0", ClampMax = "32.0", editcondition = "bOverride_DepthOfFieldFstop", DisplayName = "Aperture F-stop"))
+	UPROPERTY(interp, BlueprintReadWrite, Category="Lens|Depth of Field", meta=(ClampMin = "1.0", ClampMax = "32.0", editcondition = "bOverride_DepthOfFieldFstop", DisplayName = "Aperture F-stop"))
 	float DepthOfFieldFstop;
 
 	/** Width of the camera sensor to assume, in mm. */
-	UPROPERTY(BlueprintReadWrite, Category=DepthOfField, meta=(ForceUnits=mm, ClampMin = "0.1", UIMin="0.1", UIMax= "1000.0", editcondition = "bOverride_DepthOfFieldSensorWidth", DisplayName = "Sensor Width (mm)"))
+	UPROPERTY(BlueprintReadWrite, Category="Lens|Depth of Field", meta=(ForceUnits=mm, ClampMin = "0.1", UIMin="0.1", UIMax= "1000.0", editcondition = "bOverride_DepthOfFieldSensorWidth", DisplayName = "Sensor Width (mm)"))
 	float DepthOfFieldSensorWidth;
 
 	/** Distance in which the Depth of Field effect should be sharp, in unreal units (cm) */
-	UPROPERTY(interp, BlueprintReadWrite, Category=DepthOfField, meta=(ClampMin = "0.0", UIMin = "1.0", UIMax = "10000.0", editcondition = "bOverride_DepthOfFieldFocalDistance", DisplayName = "Focal Distance"))
+	UPROPERTY(interp, BlueprintReadWrite, Category="Lens|Depth of Field", meta=(ClampMin = "0.0", UIMin = "1.0", UIMax = "10000.0", editcondition = "bOverride_DepthOfFieldFocalDistance", DisplayName = "Focal Distance"))
 	float DepthOfFieldFocalDistance;
 
 	/** CircleDOF only: Depth blur km for 50% */
-	UPROPERTY(interp, BlueprintReadWrite, AdvancedDisplay, Category=DepthOfField, meta=(ClampMin = "0.000001", ClampMax = "100.0", editcondition = "bOverride_DepthOfFieldDepthBlurAmount", DisplayName = "Depth Blur km for 50%"))
+	UPROPERTY(interp, BlueprintReadWrite, Category="Lens|Depth of Field", meta=(ClampMin = "0.000001", ClampMax = "100.0", editcondition = "bOverride_DepthOfFieldDepthBlurAmount", DisplayName = "Depth Blur km for 50%"))
 	float DepthOfFieldDepthBlurAmount;
 
 	/** CircleDOF only: Depth blur radius in pixels at 1920x */
-	UPROPERTY(interp, BlueprintReadWrite, AdvancedDisplay, Category=DepthOfField, meta=(ClampMin = "0.0", UIMax = "4.0", editcondition = "bOverride_DepthOfFieldDepthBlurRadius", DisplayName = "Depth Blur Radius"))
+	UPROPERTY(interp, BlueprintReadWrite, Category="Lens|Depth of Field", meta=(ClampMin = "0.0", UIMax = "4.0", editcondition = "bOverride_DepthOfFieldDepthBlurRadius", DisplayName = "Depth Blur Radius"))
 	float DepthOfFieldDepthBlurRadius;
 
 	/** Artificial region where all content is in focus, starting after DepthOfFieldFocalDistance, in unreal units  (cm) */
-	UPROPERTY(interp, BlueprintReadWrite, Category=DepthOfField, meta=(UIMin = "0.0", UIMax = "10000.0", editcondition = "bOverride_DepthOfFieldFocalRegion", DisplayName = "Focal Region"))
+	UPROPERTY(interp, BlueprintReadWrite, Category="Lens|Depth of Field", meta=(UIMin = "0.0", UIMax = "10000.0", editcondition = "bOverride_DepthOfFieldFocalRegion", DisplayName = "Focal Region"))
 	float DepthOfFieldFocalRegion;
 
 	/** To define the width of the transition region next to the focal region on the near side (cm) */
-	UPROPERTY(interp, BlueprintReadWrite, Category=DepthOfField, meta=(UIMin = "0.0", UIMax = "10000.0", editcondition = "bOverride_DepthOfFieldNearTransitionRegion", DisplayName = "Near Transition Region"))
+	UPROPERTY(interp, BlueprintReadWrite, Category="Lens|Depth of Field", meta=(UIMin = "0.0", UIMax = "10000.0", editcondition = "bOverride_DepthOfFieldNearTransitionRegion", DisplayName = "Near Transition Region"))
 	float DepthOfFieldNearTransitionRegion;
 
 	/** To define the width of the transition region next to the focal region on the near side (cm) */
-	UPROPERTY(interp, BlueprintReadWrite, Category=DepthOfField, meta=(UIMin = "0.0", UIMax = "10000.0", editcondition = "bOverride_DepthOfFieldFarTransitionRegion", DisplayName = "Far Transition Region"))
+	UPROPERTY(interp, BlueprintReadWrite, Category="Lens|Depth of Field", meta=(UIMin = "0.0", UIMax = "10000.0", editcondition = "bOverride_DepthOfFieldFarTransitionRegion", DisplayName = "Far Transition Region"))
 	float DepthOfFieldFarTransitionRegion;
 
 	/** SM5: BokehDOF only: To amplify the depth of field effect (like aperture)  0=off 
 	    ES2: Used to blend DoF. 0=off
 	*/
-	UPROPERTY(interp, BlueprintReadWrite, Category=DepthOfField, meta=(ClampMin = "0.0", ClampMax = "2.0", editcondition = "bOverride_DepthOfFieldScale", DisplayName = "Scale"))
+	UPROPERTY(interp, BlueprintReadWrite, Category="Lens|Depth of Field", meta=(ClampMin = "0.0", ClampMax = "2.0", editcondition = "bOverride_DepthOfFieldScale", DisplayName = "Scale"))
 	float DepthOfFieldScale;
 
 	/** BokehDOF only: Maximum size of the Depth of Field blur (in percent of the view width) (note: performance cost scales with size*size) */
-	UPROPERTY(interp, BlueprintReadWrite, Category=DepthOfField, meta=(UIMin = "0.0", UIMax = "32.0", editcondition = "bOverride_DepthOfFieldMaxBokehSize", DisplayName = "Max Bokeh Size"))
+	UPROPERTY(interp, BlueprintReadWrite, Category="Lens|Depth of Field", meta=(UIMin = "0.0", UIMax = "32.0", editcondition = "bOverride_DepthOfFieldMaxBokehSize", DisplayName = "Max Bokeh Size"))
 	float DepthOfFieldMaxBokehSize;
 
 	/** Gaussian only: Maximum size of the Depth of Field blur (in percent of the view width) (note: performance cost scales with size) */
-	UPROPERTY(interp, BlueprintReadWrite, Category=DepthOfField, meta=(UIMin = "0.0", UIMax = "32.0", editcondition = "bOverride_DepthOfFieldNearBlurSize", DisplayName = "Near Blur Size"))
+	UPROPERTY(interp, BlueprintReadWrite, Category="Lens|Depth of Field", meta=(UIMin = "0.0", UIMax = "32.0", editcondition = "bOverride_DepthOfFieldNearBlurSize", DisplayName = "Near Blur Size"))
 	float DepthOfFieldNearBlurSize;
 
 	/** Gaussian only: Maximum size of the Depth of Field blur (in percent of the view width) (note: performance cost scales with size) */
-	UPROPERTY(interp, BlueprintReadWrite, Category=DepthOfField, meta=(UIMin = "0.0", UIMax = "32.0", editcondition = "bOverride_DepthOfFieldFarBlurSize", DisplayName = "Far Blur Size"))
+	UPROPERTY(interp, BlueprintReadWrite, Category="Lens|Depth of Field", meta=(UIMin = "0.0", UIMax = "32.0", editcondition = "bOverride_DepthOfFieldFarBlurSize", DisplayName = "Far Blur Size"))
 	float DepthOfFieldFarBlurSize;
 
 	/** Defines the shape of the Bokeh when object get out of focus, cannot be blended */
-	UPROPERTY(EditAnywhere, BlueprintReadWrite, AdvancedDisplay, Category=DepthOfField, meta=(editcondition = "bOverride_DepthOfFieldBokehShape", DisplayName = "Shape"))
+	UPROPERTY(EditAnywhere, BlueprintReadWrite, AdvancedDisplay, Category="Lens|Depth of Field", meta=(editcondition = "bOverride_DepthOfFieldBokehShape", DisplayName = "Shape"))
 	class UTexture* DepthOfFieldBokehShape;
 
 	/** Occlusion tweak factor 1 (0.18 to get natural occlusion, 0.4 to solve layer color leaking issues) */
-	UPROPERTY(interp, BlueprintReadWrite, Category=DepthOfField, AdvancedDisplay, meta=(ClampMin = "0.0", ClampMax = "1.0", editcondition = "bOverride_DepthOfFieldOcclusion", DisplayName = "Occlusion"))
+	UPROPERTY(interp, BlueprintReadWrite, Category="Lens|Depth of Field", AdvancedDisplay, meta=(ClampMin = "0.0", ClampMax = "1.0", editcondition = "bOverride_DepthOfFieldOcclusion", DisplayName = "Occlusion"))
 	float DepthOfFieldOcclusion;
 	
 	/** Color threshold to do full quality DOF (BokehDOF only) */
-	UPROPERTY(interp, BlueprintReadWrite, Category=DepthOfField, AdvancedDisplay, meta=(ClampMin = "0.0", ClampMax = "10.0", editcondition = "bOverride_DepthOfFieldColorThreshold", DisplayName = "Color Threshold"))
+	UPROPERTY(interp, BlueprintReadWrite, Category="Lens|Depth of Field", AdvancedDisplay, meta=(ClampMin = "0.0", ClampMax = "10.0", editcondition = "bOverride_DepthOfFieldColorThreshold", DisplayName = "Color Threshold"))
 	float DepthOfFieldColorThreshold;
 
 	/** Size threshold to do full quality DOF (BokehDOF only) */
-	UPROPERTY(interp, BlueprintReadWrite, Category=DepthOfField, AdvancedDisplay, meta=(ClampMin = "0.0", ClampMax = "1.0", editcondition = "bOverride_DepthOfFieldSizeThreshold", DisplayName = "Size Threshold"))
+	UPROPERTY(interp, BlueprintReadWrite, Category="Lens|Depth of Field", AdvancedDisplay, meta=(ClampMin = "0.0", ClampMax = "1.0", editcondition = "bOverride_DepthOfFieldSizeThreshold", DisplayName = "Size Threshold"))
 	float DepthOfFieldSizeThreshold;
 	
 	/** Artificial distance to allow the skybox to be in focus (e.g. 200000), <=0 to switch the feature off, only for GaussianDOF, can cost performance */
-	UPROPERTY(interp, BlueprintReadWrite, Category=DepthOfField, AdvancedDisplay, meta=(ClampMin = "0.0", ClampMax = "200000.0", editcondition = "bOverride_DepthOfFieldSkyFocusDistance", DisplayName = "Sky Distance"))
+	UPROPERTY(interp, BlueprintReadWrite, Category="Lens|Depth of Field", AdvancedDisplay, meta=(ClampMin = "0.0", ClampMax = "200000.0", editcondition = "bOverride_DepthOfFieldSkyFocusDistance", DisplayName = "Sky Distance"))
 	float DepthOfFieldSkyFocusDistance;
 
 	/** Artificial circular mask to (near) blur content outside the radius, only for GaussianDOF, diameter in percent of screen width, costs performance if the mask is used, keep Feather can Radius on default to keep it off */
-	UPROPERTY(interp, BlueprintReadWrite, Category=DepthOfField, AdvancedDisplay, meta=(UIMin = "0.0", UIMax = "100.0", editcondition = "bOverride_DepthOfFieldVignetteSize", DisplayName = "Vignette Size"))
+	UPROPERTY(interp, BlueprintReadWrite, Category="Lens|Depth of Field", AdvancedDisplay, meta=(UIMin = "0.0", UIMax = "100.0", editcondition = "bOverride_DepthOfFieldVignetteSize", DisplayName = "Vignette Size"))
 	float DepthOfFieldVignetteSize;
 
 	/** Strength of motion blur, 0:off, should be renamed to intensity */
-	UPROPERTY(interp, BlueprintReadWrite, Category=MotionBlur, meta=(ClampMin = "0.0", ClampMax = "1.0", editcondition = "bOverride_MotionBlurAmount", DisplayName = "Amount"))
+	UPROPERTY(interp, BlueprintReadWrite, Category="Rendering Features|Motion Blur", meta=(ClampMin = "0.0", ClampMax = "1.0", editcondition = "bOverride_MotionBlurAmount", DisplayName = "Amount"))
 	float MotionBlurAmount;
 	/** max distortion caused by motion blur, in percent of the screen width, 0:off */
-	UPROPERTY(interp, BlueprintReadWrite, Category=MotionBlur, meta=(ClampMin = "0.0", ClampMax = "100.0", editcondition = "bOverride_MotionBlurMax", DisplayName = "Max"))
+	UPROPERTY(interp, BlueprintReadWrite, Category="Rendering Features|Motion Blur", meta=(ClampMin = "0.0", ClampMax = "100.0", editcondition = "bOverride_MotionBlurMax", DisplayName = "Max"))
 	float MotionBlurMax;
 	/** The minimum projected screen radius for a primitive to be drawn in the velocity pass, percentage of screen width. smaller numbers cause more draw calls, default: 4% */
-	UPROPERTY(interp, BlueprintReadWrite, Category=MotionBlur, AdvancedDisplay, meta=(ClampMin = "0.0", UIMax = "100.0", editcondition = "bOverride_MotionBlurPerObjectSize", DisplayName = "Per Object Size"))
+	UPROPERTY(interp, BlueprintReadWrite, Category="Rendering Features|Motion Blur", meta=(ClampMin = "0.0", UIMax = "100.0", editcondition = "bOverride_MotionBlurPerObjectSize", DisplayName = "Per Object Size"))
 	float MotionBlurPerObjectSize;
 
-	/**
-	 * To render with lower or high resolution than it is presented, 
-	 * controlled by console variable, 
-	 * 100:off, needs to be <99 to get upsampling and lower to get performance,
-	 * >100 for super sampling (slower but higher quality), 
-	 * only applied in game
-	 */
-	UPROPERTY(interp, BlueprintReadWrite, Category=Misc, AdvancedDisplay, meta=(ClampMin = "0.0", ClampMax = "400.0", editcondition = "bOverride_ScreenPercentage"))
-	float ScreenPercentage;
+	/** How strong the dynamic GI from the LPV should be. 0.0 is off, 1.0 is the "normal" value, but higher values can be used to boost the effect*/
+	UPROPERTY(interp, BlueprintReadWrite, Category="Rendering Features|Light Propagation Volume", meta=(editcondition = "bOverride_LPVIntensity", UIMin = "0", UIMax = "20", DisplayName = "Intensity"))
+	float LPVIntensity;
+
+	/** Bias applied to light injected into the LPV in cell units. Increase to reduce bleeding through thin walls*/
+	UPROPERTY(interp, BlueprintReadWrite, Category="Rendering Features|Light Propagation Volume", AdvancedDisplay, meta=(editcondition = "bOverride_LPVVplInjectionBias", UIMin = "0", UIMax = "2", DisplayName = "Light Injection Bias"))
+	float LPVVplInjectionBias;
+
+	/** The size of the LPV volume, in Unreal units*/
+	UPROPERTY(EditAnywhere, BlueprintReadWrite, Category="Rendering Features|Light Propagation Volume", meta=(editcondition = "bOverride_LPVSize", UIMin = "100", UIMax = "20000", DisplayName = "Size"))
+	float LPVSize;
+
+	/** Secondary occlusion strength (bounce light shadows). Set to 0 to disable*/
+	UPROPERTY(interp, BlueprintReadWrite, Category="Rendering Features|Light Propagation Volume", meta=(editcondition = "bOverride_LPVSecondaryOcclusionIntensity", UIMin = "0", UIMax = "1", DisplayName = "Secondary Occlusion Intensity"))
+	float LPVSecondaryOcclusionIntensity;
+
+	/** Secondary bounce light strength (bounce light shadows). Set to 0 to disable*/
+	UPROPERTY(interp, BlueprintReadWrite, Category="Rendering Features|Light Propagation Volume", AdvancedDisplay, meta=(editcondition = "bOverride_LPVSecondaryBounceIntensity", UIMin = "0", UIMax = "1", DisplayName = "Secondary Bounce Intensity"))
+	float LPVSecondaryBounceIntensity;
+
+	/** Bias applied to the geometry volume in cell units. Increase to reduce darkening due to secondary occlusion */
+	UPROPERTY(interp, BlueprintReadWrite, Category="Rendering Features|Light Propagation Volume", AdvancedDisplay, meta=(editcondition = "bOverride_LPVGeometryVolumeBias", UIMin = "0", UIMax = "2", DisplayName = "Geometry Volume Bias"))
+	float LPVGeometryVolumeBias;
+
+	UPROPERTY(interp, BlueprintReadWrite, Category="Rendering Features|Light Propagation Volume", AdvancedDisplay, meta=(editcondition = "bOverride_LPVEmissiveInjectionIntensity", UIMin = "0", UIMax = "20", DisplayName = "Emissive Injection Intensity"))
+	float LPVEmissiveInjectionIntensity;
+
+	/** Controls the amount of directional occlusion. Requires LPV. Values very close to 1.0 are recommended */
+	UPROPERTY(interp, BlueprintReadWrite, Category="Rendering Features|Light Propagation Volume", meta=(editcondition = "bOverride_LPVDirectionalOcclusionIntensity", UIMin = "0", UIMax = "1", DisplayName = "Occlusion Intensity"))
+	float LPVDirectionalOcclusionIntensity;
+
+	/** Occlusion Radius - 16 is recommended for most scenes */
+	UPROPERTY(interp, BlueprintReadWrite, Category="Rendering Features|Light Propagation Volume", AdvancedDisplay, meta=(editcondition = "bOverride_LPVDirectionalOcclusionRadius", UIMin = "1", UIMax = "16", DisplayName = "Occlusion Radius"))
+	float LPVDirectionalOcclusionRadius;
+
+	/** Diffuse occlusion exponent - increase for more contrast. 1 to 2 is recommended */
+	UPROPERTY(interp, BlueprintReadWrite, Category="Rendering Features|Light Propagation Volume", meta=(editcondition = "bOverride_LPVDiffuseOcclusionExponent", UIMin = "0.5", UIMax = "5", DisplayName = "Diffuse occlusion exponent"))
+	float LPVDiffuseOcclusionExponent;
+
+	/** Specular occlusion exponent - increase for more contrast. 6 to 9 is recommended */
+	UPROPERTY(interp, BlueprintReadWrite, Category="Rendering Features|Light Propagation Volume", meta=(editcondition = "bOverride_LPVSpecularOcclusionExponent", UIMin = "1", UIMax = "16", DisplayName = "Specular occlusion exponent"))
+	float LPVSpecularOcclusionExponent;
+
+	/** Diffuse occlusion intensity - higher values provide increased diffuse occlusion.*/
+	UPROPERTY(interp, BlueprintReadWrite, Category="Rendering Features|Light Propagation Volume", AdvancedDisplay, meta=(editcondition = "bOverride_LPVDiffuseOcclusionIntensity", UIMin = "0", UIMax = "4", DisplayName = "Diffuse occlusion intensity"))
+	float LPVDiffuseOcclusionIntensity;
+
+	/** Specular occlusion intensity - higher values provide increased specular occlusion.*/
+	UPROPERTY(interp, BlueprintReadWrite, Category="Rendering Features|Light Propagation Volume", AdvancedDisplay, meta=(editcondition = "bOverride_LPVSpecularOcclusionIntensity", UIMin = "0", UIMax = "4", DisplayName = "Specular occlusion intensity"))
+	float LPVSpecularOcclusionIntensity;
 
 	/** Enable/Fade/disable the Screen Space Reflection feature, in percent, avoid numbers between 0 and 1 fo consistency */
-	UPROPERTY(interp, BlueprintReadWrite, Category=ScreenSpaceReflections, meta=(ClampMin = "0.0", ClampMax = "100.0", editcondition = "bOverride_ScreenSpaceReflectionIntensity", DisplayName = "Intensity"))
+	UPROPERTY(interp, BlueprintReadWrite, Category="Rendering Features|Screen Space Reflections", meta=(ClampMin = "0.0", ClampMax = "100.0", editcondition = "bOverride_ScreenSpaceReflectionIntensity", DisplayName = "Intensity"))
 	float ScreenSpaceReflectionIntensity;
 
 	/** 0=lowest quality..100=maximum quality, only a few quality levels are implemented, no soft transition, 50 is the default for better performance. */
-	UPROPERTY(interp, BlueprintReadWrite, Category=ScreenSpaceReflections, meta=(ClampMin = "0.0", UIMax = "100.0", editcondition = "bOverride_ScreenSpaceReflectionQuality", DisplayName = "Quality"))
+	UPROPERTY(interp, BlueprintReadWrite, Category="Rendering Features|Screen Space Reflections", meta=(ClampMin = "0.0", UIMax = "100.0", editcondition = "bOverride_ScreenSpaceReflectionQuality", DisplayName = "Quality"))
 	float ScreenSpaceReflectionQuality;
 
 	/** Until what roughness we fade the screen space reflections, 0.8 works well, smaller can run faster */
-	UPROPERTY(interp, BlueprintReadWrite, Category=ScreenSpaceReflections, meta=(ClampMin = "0.01", ClampMax = "1.0", editcondition = "bOverride_ScreenSpaceReflectionMaxRoughness", DisplayName = "Max Roughness"))
+	UPROPERTY(interp, BlueprintReadWrite, Category="Rendering Features|Screen Space Reflections", meta=(ClampMin = "0.01", ClampMax = "1.0", editcondition = "bOverride_ScreenSpaceReflectionMaxRoughness", DisplayName = "Max Roughness"))
 	float ScreenSpaceReflectionMaxRoughness;
 
+	/**
+	* To render with lower or high resolution than it is presented,
+	* controlled by console variable,
+	* 100:off, needs to be <99 to get upsampling and lower to get performance,
+	* >100 for super sampling (slower but higher quality),
+	* only applied in game
+	*/
+	UPROPERTY(interp, BlueprintReadWrite, Category="Rendering Features|Misc", meta=(ClampMin = "0.0", ClampMax = "400.0", editcondition = "bOverride_ScreenPercentage"))
+	float ScreenPercentage;
+
 	// Note: Adding properties before this line require also changes to the OverridePostProcessSettings() function and 
 	// FPostProcessSettings constructor and possibly the SetBaseValues() method.
 	// -----------------------------------------------------------------------
@@ -1049,7 +1050,7 @@ struct FPostProcessSettings
 	 * Allows custom post process materials to be defined, using a MaterialInstance with the same Material as its parent to allow blending.
 	 * For materials this needs to be the "PostProcess" domain type. This can be used for any UObject object implementing the IBlendableInterface (e.g. could be used to fade weather settings).
 	 */
-	UPROPERTY(EditAnywhere, Category="PostProcessSettings", meta=( Keywords="PostProcess", DisplayName = "Blendables" ))
+	UPROPERTY(EditAnywhere, Category="Rendering Features", meta=( Keywords="PostProcess", DisplayName = "Post Process Materials" ))
 	FWeightedBlendables WeightedBlendables;
 
 	// for backwards compatibility
diff --git a/Engine/Source/Runtime/Engine/Classes/Engine/UserDefinedEnum.h b/Engine/Source/Runtime/Engine/Classes/Engine/UserDefinedEnum.h
--- a/Engine/Source/Runtime/Engine/Classes/Engine/UserDefinedEnum.h
+++ b/Engine/Source/Runtime/Engine/Classes/Engine/UserDefinedEnum.h
@@ -28,13 +28,11 @@ class ENGINE_API UUserDefinedEnum : public UEnum
 #endif //WITH_EDITORONLY_DATA
 
 	/**
-	 * Names stored in "DisplayName" meta data. They are duplicated here, 
-	 * so functions like UKismetNodeHelperLibrary::GetEnumeratorUserFriendlyName can use them
-	 * outside the editor. (When meta data are not loaded).
-	 * To sync DisplayNames with meta-data use FEnumEditorUtils::EnsureAllDisplayNamesExist.
+	 * De-facto display names for enum entries mapped against their raw enum name (UEnum::GetEnumName).
+	 * To sync DisplayNameMap use FEnumEditorUtils::EnsureAllDisplayNamesExist.
 	 */
 	UPROPERTY()
-	TArray<FText> DisplayNames;
+	TMap<FName, FText> DisplayNameMap;
 
 public:
 	//~ Begin UObject Interface.
diff --git a/Engine/Source/Runtime/Engine/Classes/GameFramework/Actor.h b/Engine/Source/Runtime/Engine/Classes/GameFramework/Actor.h
--- a/Engine/Source/Runtime/Engine/Classes/GameFramework/Actor.h
+++ b/Engine/Source/Runtime/Engine/Classes/GameFramework/Actor.h
@@ -597,7 +597,7 @@ public:
 	UPROPERTY(EditAnywhere, AdvancedDisplay, Category=Actor)
 	uint8 bIgnoresOriginShifting:1;
 	
-	/** If true, and if World setting has bEnableHigerarhicalLOD is true, then it will generate LODActor from groups of clustered Actor */
+	/** If true, and if World setting has bEnableHierarchicalLOD equal to true, then it will generate LODActor from groups of clustered Actor */
 	UPROPERTY(EditAnywhere, AdvancedDisplay, Category=Actor)
 	uint8 bEnableAutoLODGeneration:1;
 
diff --git a/Engine/Source/Runtime/Engine/Classes/Materials/MaterialExpression.h b/Engine/Source/Runtime/Engine/Classes/Materials/MaterialExpression.h
--- a/Engine/Source/Runtime/Engine/Classes/Materials/MaterialExpression.h
+++ b/Engine/Source/Runtime/Engine/Classes/Materials/MaterialExpression.h
@@ -248,6 +248,9 @@ class ENGINE_API UMaterialExpression : public UObject
 #if WITH_EDITOR
 	virtual uint32 GetInputType(int32 InputIndex);
 	virtual uint32 GetOutputType(int32 OutputIndex);
+
+	virtual FString GetCreationDescription() const;
+	virtual FString GetCreationName() const;
 #endif
 
 	/**
@@ -392,6 +395,15 @@ class ENGINE_API UMaterialExpression : public UObject
 	 */
 	virtual void SetEditableName(const FString& NewName);
 
+	/** 
+	* Parameter Name functions, this is requires as multiple class have ParameterName
+	* but are not UMaterialExpressionParameter due to class hierarchy. */
+	virtual bool HasAParameterName() const { return false; }
+	virtual void ValidateParameterName();
+
+	virtual FName GetParameterName() const { return NAME_None; }
+	virtual void SetParameterName(const FName& Name) {}
+
 #endif // WITH_EDITOR
 
 	/** Checks whether any inputs to this expression create a loop */
diff --git a/Engine/Source/Runtime/Engine/Classes/Materials/MaterialExpressionCollectionParameter.h b/Engine/Source/Runtime/Engine/Classes/Materials/MaterialExpressionCollectionParameter.h
--- a/Engine/Source/Runtime/Engine/Classes/Materials/MaterialExpressionCollectionParameter.h
+++ b/Engine/Source/Runtime/Engine/Classes/Materials/MaterialExpressionCollectionParameter.h
@@ -43,6 +43,11 @@ class UMaterialExpressionCollectionParameter : public UMaterialExpression
 #if WITH_EDITOR
 	virtual int32 Compile(class FMaterialCompiler* Compiler, int32 OutputIndex) override;
 	virtual void GetCaption(TArray<FString>& OutCaptions) const override;
+
+	virtual bool HasAParameterName() const override { return true; }
+	virtual FName GetParameterName() const override { return ParameterName; }
+	virtual void SetParameterName(const FName& Name) override { ParameterName = Name; }
+
 #endif
 	virtual bool MatchesSearchQuery( const TCHAR* SearchQuery ) override;
 	//~ End UMaterialExpression Interface
diff --git a/Engine/Source/Runtime/Engine/Classes/Materials/MaterialExpressionFontSampleParameter.h b/Engine/Source/Runtime/Engine/Classes/Materials/MaterialExpressionFontSampleParameter.h
--- a/Engine/Source/Runtime/Engine/Classes/Materials/MaterialExpressionFontSampleParameter.h
+++ b/Engine/Source/Runtime/Engine/Classes/Materials/MaterialExpressionFontSampleParameter.h
@@ -38,6 +38,10 @@ class UMaterialExpressionFontSampleParameter : public UMaterialExpressionFontSam
 	virtual bool CanRenameNode() const override { return true; }
 	virtual FString GetEditableName() const override;
 	virtual void SetEditableName(const FString& NewName) override;
+
+	virtual bool HasAParameterName() const override { return true; }
+	virtual FName GetParameterName() const override { return ParameterName; }
+	virtual void SetParameterName(const FName& Name) override { ParameterName = Name; }
 #endif
 	//~ End UMaterialExpression Interface
 	
diff --git a/Engine/Source/Runtime/Engine/Classes/Materials/MaterialExpressionParameter.h b/Engine/Source/Runtime/Engine/Classes/Materials/MaterialExpressionParameter.h
--- a/Engine/Source/Runtime/Engine/Classes/Materials/MaterialExpressionParameter.h
+++ b/Engine/Source/Runtime/Engine/Classes/Materials/MaterialExpressionParameter.h
@@ -28,10 +28,16 @@ class UMaterialExpressionParameter : public UMaterialExpression
 
 	//~ Begin UMaterialExpression Interface
 	virtual bool MatchesSearchQuery( const TCHAR* SearchQuery ) override;
+
 #if WITH_EDITOR
 	virtual bool CanRenameNode() const override { return true; }
 	virtual FString GetEditableName() const override;
 	virtual void SetEditableName(const FString& NewName) override;
+
+	virtual bool HasAParameterName() const override { return true; }
+	virtual FName GetParameterName() const override { return ParameterName; }
+	virtual void SetParameterName(const FName& Name) override { ParameterName = Name; }
+
 #endif
 	//~ End UMaterialExpression Interface
 
diff --git a/Engine/Source/Runtime/Engine/Classes/Materials/MaterialExpressionReroute.h b/Engine/Source/Runtime/Engine/Classes/Materials/MaterialExpressionReroute.h
new file mode 100644
--- /dev/null
+++ b/Engine/Source/Runtime/Engine/Classes/Materials/MaterialExpressionReroute.h
@@ -0,0 +1,50 @@
+// Copyright 1998-2016 Epic Games, Inc. All Rights Reserved.
+
+/**
+ * Absolute value material expression for user-defined materials
+ *
+ */
+
+#pragma once
+#include "Materials/MaterialExpression.h"
+#include "MaterialExpressionReroute.generated.h"
+
+UCLASS(collapsecategories, hidecategories=Object, DisplayName = "Reroute")
+class ENGINE_API UMaterialExpressionReroute : public UMaterialExpression
+{
+	GENERATED_UCLASS_BODY()
+
+	/** Link to the input expression to be evaluated */
+	UPROPERTY()
+	FExpressionInput Input;
+
+	/**
+	 * Trace through the graph to find the first non Reroute node connected to this input. If there is a loop for some reason, we will bail out and return nullptr.
+	 *
+	 * @param OutputIndex The output index of the connection that was traced back.
+	 * @return The final traced material expression.
+	*/
+	UMaterialExpression* TraceInputsToRealExpression(int32& OutputIndex) const;
+
+	FExpressionInput TraceInputsToRealInput() const;
+
+	//~ Begin UMaterialExpression Interface
+#if WITH_EDITOR
+	virtual int32 Compile(class FMaterialCompiler* Compiler, int32 OutputIndex) override;
+	virtual void GetCaption(TArray<FString>& OutCaptions) const override;
+
+	virtual uint32 GetInputType(int32 InputIndex) override;
+	virtual uint32 GetOutputType(int32 OutputIndex) override;
+	virtual FString GetCreationDescription() const override;
+	virtual FString GetCreationName() const override;
+	virtual bool IsResultMaterialAttributes(int32 OutputIndex) override;
+#endif
+	//~ End UMaterialExpression Interface
+
+private:
+	FExpressionInput TraceInputsToRealExpressionInternal(TSet<FMaterialExpressionKey>& VisitedExpressions) const;
+
+};
+
+
+
diff --git a/Engine/Source/Runtime/Engine/Classes/Materials/MaterialExpressionTextureObjectParameter.h b/Engine/Source/Runtime/Engine/Classes/Materials/MaterialExpressionTextureObjectParameter.h
--- a/Engine/Source/Runtime/Engine/Classes/Materials/MaterialExpressionTextureObjectParameter.h
+++ b/Engine/Source/Runtime/Engine/Classes/Materials/MaterialExpressionTextureObjectParameter.h
@@ -12,7 +12,7 @@
 #include "Materials/MaterialExpressionTextureSampleParameter.h"
 #include "MaterialExpressionTextureObjectParameter.generated.h"
 
-UCLASS(collapsecategories, hidecategories=(Object, MaterialExpressionTextureSample))
+UCLASS(collapsecategories, hidecategories=(Object, MaterialExpressionTextureSample), MinimalAPI)
 class UMaterialExpressionTextureObjectParameter : public UMaterialExpressionTextureSampleParameter
 {
 	GENERATED_UCLASS_BODY()
diff --git a/Engine/Source/Runtime/Engine/Classes/Materials/MaterialExpressionTextureSampleParameter.h b/Engine/Source/Runtime/Engine/Classes/Materials/MaterialExpressionTextureSampleParameter.h
--- a/Engine/Source/Runtime/Engine/Classes/Materials/MaterialExpressionTextureSampleParameter.h
+++ b/Engine/Source/Runtime/Engine/Classes/Materials/MaterialExpressionTextureSampleParameter.h
@@ -37,6 +37,11 @@ class ENGINE_API UMaterialExpressionTextureSampleParameter : public UMaterialExp
 	virtual bool CanRenameNode() const override { return true; }
 	virtual FString GetEditableName() const override;
 	virtual void SetEditableName(const FString& NewName) override;
+
+	virtual bool HasAParameterName() const override { return true; }
+	virtual FName GetParameterName() const override { return ParameterName; }
+	virtual void SetParameterName(const FName& Name) override { ParameterName = Name; }
+
 #endif
 	//~ End UMaterialExpression Interface
 
diff --git a/Engine/Source/Runtime/Engine/Classes/Materials/MaterialExpressionVectorParameter.h b/Engine/Source/Runtime/Engine/Classes/Materials/MaterialExpressionVectorParameter.h
--- a/Engine/Source/Runtime/Engine/Classes/Materials/MaterialExpressionVectorParameter.h
+++ b/Engine/Source/Runtime/Engine/Classes/Materials/MaterialExpressionVectorParameter.h
@@ -18,6 +18,9 @@ class UMaterialExpressionVectorParameter : public UMaterialExpressionParameter
 	UPROPERTY(EditAnywhere, Category=MaterialExpressionVectorParameter)
 	FLinearColor DefaultValue;
 
+	UPROPERTY(EditAnywhere, Category = MaterialExpressionVectorParameter)
+	FColor DefaultValue2;
+
 	//~ Begin UMaterialExpression Interface
 #if WITH_EDITOR
 	virtual int32 Compile(class FMaterialCompiler* Compiler, int32 OutputIndex) override;
diff --git a/Engine/Source/Runtime/Engine/Classes/PhysicsEngine/BodySetup.h b/Engine/Source/Runtime/Engine/Classes/PhysicsEngine/BodySetup.h
--- a/Engine/Source/Runtime/Engine/Classes/PhysicsEngine/BodySetup.h
+++ b/Engine/Source/Runtime/Engine/Classes/PhysicsEngine/BodySetup.h
@@ -200,6 +200,7 @@ public:
 	virtual void PostLoad() override;
 	virtual void PostInitProperties() override;
 #if WITH_EDITOR
+	virtual void PostEditChangeProperty(struct FPropertyChangedEvent& PropertyChangedEvent) override;
 	virtual void PostEditUndo() override;
 #endif // WITH_EDITOR
 	virtual void GetResourceSizeEx(FResourceSizeEx& CumulativeResourceSize) override;
diff --git a/Engine/Source/Runtime/Engine/Classes/PhysicsEngine/PhysicsSettings.h b/Engine/Source/Runtime/Engine/Classes/PhysicsEngine/PhysicsSettings.h
--- a/Engine/Source/Runtime/Engine/Classes/PhysicsEngine/PhysicsSettings.h
+++ b/Engine/Source/Runtime/Engine/Classes/PhysicsEngine/PhysicsSettings.h
@@ -233,12 +233,12 @@ class ENGINE_API UPhysicsSettings : public UDeveloperSettings
 	float InitialAverageFrameRate;
 
 	// PhysicalMaterial Surface Types
-	UPROPERTY(config)
+	UPROPERTY(config, EditAnywhere, Category=PhysicalSurfaces)
 	TArray<FPhysicalSurfaceName> PhysicalSurfaces;
 
 public:
 
-	static UPhysicsSettings * Get() { return CastChecked<UPhysicsSettings>(UPhysicsSettings::StaticClass()->GetDefaultObject()); }
+	static UPhysicsSettings* Get() { return CastChecked<UPhysicsSettings>(UPhysicsSettings::StaticClass()->GetDefaultObject()); }
 
 	virtual void PostInitProperties() override;
 
diff --git a/Engine/Source/Runtime/Engine/Private/Animation/Skeleton.cpp b/Engine/Source/Runtime/Engine/Private/Animation/Skeleton.cpp
--- a/Engine/Source/Runtime/Engine/Private/Animation/Skeleton.cpp
+++ b/Engine/Source/Runtime/Engine/Private/Animation/Skeleton.cpp
@@ -355,7 +355,7 @@ bool USkeleton::IsCompatibleMesh(const USkeletalMesh* InSkelMesh) const
 			// still no match, return false, no parent to look for
 			if( SkeletonBoneIndex == INDEX_NONE )
 			{
-				UE_LOG(LogAnimation, Warning, TEXT("%s : Missing joint on skeleton.  Make sure to assign to the skeleeton."), *MeshBoneName.ToString());
+				UE_LOG(LogAnimation, Warning, TEXT("%s : Missing joint on skeleton.  Make sure to assign to the skeleton."), *MeshBoneName.ToString());
 				return false;
 			}
 
diff --git a/Engine/Source/Runtime/Engine/Private/Blueprint.cpp b/Engine/Source/Runtime/Engine/Private/Blueprint.cpp
--- a/Engine/Source/Runtime/Engine/Private/Blueprint.cpp
+++ b/Engine/Source/Runtime/Engine/Private/Blueprint.cpp
@@ -1507,6 +1507,19 @@ void UBlueprint::GatherDependencies(TSet<TWeakObjectPtr<UBlueprint>>& InDependen
 
 }
 
+void UBlueprint::ReplaceDeprecatedNodes()
+{
+	TArray<UEdGraph*> Graphs;
+	GetAllGraphs(Graphs);
+
+	for ( auto It = Graphs.CreateIterator(); It; ++It )
+	{
+		UEdGraph* const Graph = *It;
+		const UEdGraphSchema* Schema = Graph->GetSchema();
+		Schema->BackwardCompatibilityNodeConversion(Graph, true);
+	}
+}
+
 UInheritableComponentHandler* UBlueprint::GetInheritableComponentHandler(bool bCreateIfNecessary)
 {
 	static const FBoolConfigValueHelper EnableInheritableComponents(TEXT("Kismet"), TEXT("bEnableInheritableComponents"), GEngineIni);
diff --git a/Engine/Source/Runtime/Engine/Private/Components/StaticMeshComponent.cpp b/Engine/Source/Runtime/Engine/Private/Components/StaticMeshComponent.cpp
--- a/Engine/Source/Runtime/Engine/Private/Components/StaticMeshComponent.cpp
+++ b/Engine/Source/Runtime/Engine/Private/Components/StaticMeshComponent.cpp
@@ -188,6 +188,7 @@ UStaticMeshComponent::UStaticMeshComponent(const FObjectInitializer& ObjectIniti
 	SelectedEditorSection = INDEX_NONE;
 	SectionIndexPreview = INDEX_NONE;
 	StaticMeshImportVersion = BeforeImportStaticMeshVersionWasAdded;
+	bCustomOverrideVertexColorPerLOD = false;
 #endif
 }
 
@@ -1199,61 +1200,78 @@ void UStaticMeshComponent::PrivateFixupOverrideColors()
 
 	// Initialize override vertex colors on any new LODs which have just been created
 	SetLODDataCount(NumLODs, LODData.Num());
-
+	bool UpdateStaticMeshDeriveDataKey = false;
 	FStaticMeshComponentLODInfo& LOD0Info = LODData[0];
-	if (LOD0Info.OverrideVertexColors)
+	if (!bCustomOverrideVertexColorPerLOD && LOD0Info.OverrideVertexColors == nullptr)
 	{
-		for (uint32 LODIndex = 0; LODIndex < NumLODs; ++LODIndex)
-		{
-			FStaticMeshComponentLODInfo& LODInfo = LODData[LODIndex];
+		return;
+	}
 
-			if (LODInfo.OverrideVertexColors == nullptr)
-			{
-				LODInfo.OverrideVertexColors = new FColorVertexBuffer;
-			}
-			else
-			{
-				LODInfo.BeginReleaseOverrideVertexColors();
-				FlushRenderingCommands();
-			}
+	for (uint32 LODIndex = 0; LODIndex < NumLODs; ++LODIndex)
+	{
+		FStaticMeshComponentLODInfo& LODInfo = LODData[LODIndex];
+		if (LODInfo.OverrideVertexColors == nullptr)
+		{
+			if (bCustomOverrideVertexColorPerLOD) //No fixup required if the component is in custom LOD paint and there is no paint on a LOD
+				continue;
+			LODInfo.OverrideVertexColors = new FColorVertexBuffer;
+		}
+		else
+		{
+			LODInfo.BeginReleaseOverrideVertexColors();
+			FlushRenderingCommands();
+		}
 
-			FStaticMeshLODResources& CurRenderData = GetStaticMesh()->RenderData->LODResources[LODIndex];
 
-			TArray<FColor> NewOverrideColors;
+		FStaticMeshLODResources& CurRenderData = GetStaticMesh()->RenderData->LODResources[LODIndex];
+		TArray<FColor> NewOverrideColors;
+		if (bCustomOverrideVertexColorPerLOD)
+		{
+			//Since in custom we fix paint only if the component has some, the PaintedVertices should be allocate
+			check(LODInfo.PaintedVertices.Num() > 0);
+			//Use the existing LOD custom paint and remap it on the new mesh
+			RemapPaintedVertexColors(
+				LODInfo.PaintedVertices,
+				*LODInfo.OverrideVertexColors,
+				CurRenderData.PositionVertexBuffer,
+				&CurRenderData.VertexBuffer,
+				NewOverrideColors
+				);
+		}
+		else if(LOD0Info.PaintedVertices.Num() > 0)
+		{
+			RemapPaintedVertexColors(
+				LOD0Info.PaintedVertices,
+				*LOD0Info.OverrideVertexColors,
+				CurRenderData.PositionVertexBuffer,
+				&CurRenderData.VertexBuffer,
+				NewOverrideColors
+				);
+		}
+		if (NewOverrideColors.Num())
+		{
+			LODInfo.OverrideVertexColors->InitFromColorArray(NewOverrideColors);
 
-			if (LOD0Info.PaintedVertices.Num() > 0)
-			{
-				// Build override colors for LOD, based on LOD0
-				RemapPaintedVertexColors(
-					LOD0Info.PaintedVertices,
-					*LOD0Info.OverrideVertexColors,
-					CurRenderData.PositionVertexBuffer,
-					&CurRenderData.VertexBuffer,
-					NewOverrideColors
-					);
-			}
+			// Update the PaintedVertices array
+			const int32 NumVerts = CurRenderData.GetNumVertices();
+			check(NumVerts == NewOverrideColors.Num());
 
-			if (NewOverrideColors.Num())
+			LODInfo.PaintedVertices.Reserve(NumVerts);
+			for (int32 VertIndex = 0; VertIndex < NumVerts; ++VertIndex)
 			{
-				LODInfo.OverrideVertexColors->InitFromColorArray(NewOverrideColors);
-
-				// Update the PaintedVertices array
-				const int32 NumVerts = CurRenderData.GetNumVertices();
-				check(NumVerts == NewOverrideColors.Num());
-
-				LODInfo.PaintedVertices.Reserve(NumVerts);
-				for (int32 VertIndex = 0; VertIndex < NumVerts; ++VertIndex)
-				{
-					FPaintedVertex* Vertex = new(LODInfo.PaintedVertices) FPaintedVertex;
-					Vertex->Position = CurRenderData.PositionVertexBuffer.VertexPosition(VertIndex);
-					Vertex->Normal = CurRenderData.VertexBuffer.VertexTangentZ(VertIndex);
-					Vertex->Color = LODInfo.OverrideVertexColors->VertexColor(VertIndex);
-				}
+				FPaintedVertex* Vertex = new(LODInfo.PaintedVertices) FPaintedVertex;
+				Vertex->Position = CurRenderData.PositionVertexBuffer.VertexPosition(VertIndex);
+				Vertex->Normal = CurRenderData.VertexBuffer.VertexTangentZ(VertIndex);
+				Vertex->Color = LODInfo.OverrideVertexColors->VertexColor(VertIndex);
 			}
-
-			BeginInitResource(LODInfo.OverrideVertexColors);
 		}
 
+		BeginInitResource(LODInfo.OverrideVertexColors);
+		UpdateStaticMeshDeriveDataKey = true;
+	}
+
+	if (UpdateStaticMeshDeriveDataKey)
+	{
 		StaticMeshDerivedDataKey = GetStaticMesh()->RenderData->DerivedDataKey;
 	}
 
diff --git a/Engine/Source/Runtime/Engine/Private/CurveTable.cpp b/Engine/Source/Runtime/Engine/Private/CurveTable.cpp
--- a/Engine/Source/Runtime/Engine/Private/CurveTable.cpp
+++ b/Engine/Source/Runtime/Engine/Private/CurveTable.cpp
@@ -542,6 +542,16 @@ bool UCurveTable::IsValidCurve(FRichCurveEditInfo CurveInfo)
 	return false;
 }
 
+
+
+TArray<const UObject*> UCurveTable::GetOwners() const
+{
+	TArray<const UObject*> Owners;
+	Owners.Add(this);
+
+	return Owners;
+}
+
 //////////////////////////////////////////////////////////////////////////
 
 
diff --git a/Engine/Source/Runtime/Engine/Private/Curves/CurveBase.cpp b/Engine/Source/Runtime/Engine/Private/Curves/CurveBase.cpp
--- a/Engine/Source/Runtime/Engine/Private/Curves/CurveBase.cpp
+++ b/Engine/Source/Runtime/Engine/Private/Curves/CurveBase.cpp
@@ -62,6 +62,14 @@ void UCurveBase::ModifyOwner()
 	Modify(true);
 }
 
+TArray<const UObject*> UCurveBase::GetOwners() const
+{
+	TArray<const UObject*> Owners;
+	Owners.Add(this);		// CurveBase owns its own curve
+
+	return Owners;
+}
+
 
 void UCurveBase::MakeTransactional() 
 {
diff --git a/Engine/Source/Runtime/Engine/Private/DataTable.cpp b/Engine/Source/Runtime/Engine/Private/DataTable.cpp
--- a/Engine/Source/Runtime/Engine/Private/DataTable.cpp
+++ b/Engine/Source/Runtime/Engine/Private/DataTable.cpp
@@ -23,10 +23,7 @@ namespace
 		for (const auto& Pair : DataTable->RowMap)
 		{
 			const FString PathToRow = PathToObject + TEXT(".") + Pair.Key.ToString();
-			for (TFieldIterator<const UProperty> PropIt(DataTable->RowStruct, EFieldIteratorFlags::IncludeSuper, EFieldIteratorFlags::ExcludeDeprecated, EFieldIteratorFlags::IncludeInterfaces); PropIt; ++PropIt)
-			{
-				PropertyLocalizationDataGatherer.GatherLocalizationDataFromChildTextProperties(PathToRow, *PropIt, PropIt->ContainerPtrToValuePtr<void>(Pair.Value), GatherTextFlags);
-			}
+			PropertyLocalizationDataGatherer.GatherLocalizationDataFromStructFields(PathToRow, DataTable->RowStruct, Pair.Value, nullptr, GatherTextFlags);
 		}
 	}
 }
diff --git a/Engine/Source/Runtime/Engine/Private/EdGraph/EdGraph.cpp b/Engine/Source/Runtime/Engine/Private/EdGraph/EdGraph.cpp
--- a/Engine/Source/Runtime/Engine/Private/EdGraph/EdGraph.cpp
+++ b/Engine/Source/Runtime/Engine/Private/EdGraph/EdGraph.cpp
@@ -82,6 +82,60 @@ UEdGraph::UEdGraph(const FObjectInitializer& ObjectInitializer)
 	bAllowDeletion = true;
 }
 
+void UEdGraph::BuildSubobjectMapping(UObject* OtherObject, TMap<UObject*, UObject*>& ObjectMapping) const
+{
+#if WITH_EDITORONLY_DATA
+	UEdGraph* OtherGraph = CastChecked<UEdGraph>(OtherObject);
+
+	auto FindMatchingNode = [](UEdGraph* InGraphToSearch, UEdGraphNode* InNodeToFind) -> UEdGraphNode*
+	{
+		TArray<UEdGraphNode*, TInlineAllocator<8>> PotentialMatches;
+
+		for (UEdGraphNode* GraphNode : InGraphToSearch->Nodes)
+		{
+			if (GraphNode->GetClass() == InNodeToFind->GetClass())
+			{
+				// Ideally the node would match by GUID or name
+				if (GraphNode->NodeGuid == InNodeToFind->NodeGuid || GraphNode->GetFName() == InNodeToFind->GetFName())
+				{
+					return GraphNode;
+				}
+
+				// If it doesn't match by either, we may be searching an older graph without stable node IDs... consider the node title as significant in this case
+				{
+					const FText Title1 = GraphNode->GetNodeTitle(ENodeTitleType::FullTitle);
+					const FText Title2 = InNodeToFind->GetNodeTitle(ENodeTitleType::FullTitle);
+
+					if (Title1.ToString().Equals(Title2.ToString(), ESearchCase::CaseSensitive))
+					{
+						PotentialMatches.Emplace(GraphNode);
+					}
+				}
+			}
+		}
+
+		// Only use the title-based resolution if we found a single match!
+		if (PotentialMatches.Num() == 1)
+		{
+			return PotentialMatches[0];
+		}
+
+		return nullptr;
+	};
+
+	for (UEdGraphNode* GraphNode : Nodes)
+	{
+		if (!ObjectMapping.Contains(GraphNode))
+		{
+			UEdGraphNode* OtherGraphNode = FindMatchingNode(OtherGraph, GraphNode);
+			ObjectMapping.Emplace(GraphNode, OtherGraphNode);
+		}
+	}
+#endif // WITH_EDITORONLY_DATA
+
+	Super::BuildSubobjectMapping(OtherObject, ObjectMapping);
+}
+
 #if WITH_EDITORONLY_DATA
 void UEdGraph::PostInitProperties()
 {
diff --git a/Engine/Source/Runtime/Engine/Private/EdGraph/EdGraphNode.cpp b/Engine/Source/Runtime/Engine/Private/EdGraph/EdGraphNode.cpp
--- a/Engine/Source/Runtime/Engine/Private/EdGraph/EdGraphNode.cpp
+++ b/Engine/Source/Runtime/Engine/Private/EdGraph/EdGraphNode.cpp
@@ -509,6 +509,23 @@ int32 UEdGraphNode::GetPinIndex(UEdGraphPin* Pin) const
 	return Pins.Find(Pin);
 }
 
+bool UEdGraphNode::ShouldDrawNodeAsControlPointOnly(int32& OutInputPinIndex, int32& OutOutputPinIndex) const
+{
+	OutInputPinIndex = -1;
+	OutOutputPinIndex = -1; 
+	return false;
+}
+
+
+UEdGraphPin* UEdGraphNode::GetPinAt(int32 index) const
+{
+	if (Pins.Num() > index)
+	{
+		return Pins[index];
+	}
+	return nullptr;
+}
+
 void UEdGraphNode::AddSearchMetaDataInfo(TArray<struct FSearchTagDataPair>& OutTaggedMetaData) const
 {
 	// Searchable - Primary label for the item in the search results
diff --git a/Engine/Source/Runtime/Engine/Private/Font.cpp b/Engine/Source/Runtime/Engine/Private/Font.cpp
--- a/Engine/Source/Runtime/Engine/Private/Font.cpp
+++ b/Engine/Source/Runtime/Engine/Private/Font.cpp
@@ -375,14 +375,12 @@ void UFont::GetResourceSizeEx(FResourceSizeEx& CumulativeResourceSize)
 			{
 				for (const FTypefaceEntry& TypefaceEntry : Typeface.Fonts)
 				{
-#if WITH_EDITORONLY_DATA
 					const UFontFace* FontFace = Cast<const UFontFace>(TypefaceEntry.Font.GetFontFaceAsset());
 					if (FontFace)
 					{
 						const_cast<UFontFace*>(FontFace)->GetResourceSizeEx(CumulativeResourceSize);
 					}
-					else if (TypefaceEntry.Font.GetLoadingPolicy() == EFontLoadingPolicy::PreLoad)
-#endif // WITH_EDITORONLY_DATA
+					else if (TypefaceEntry.Font.GetLoadingPolicy() == EFontLoadingPolicy::LazyLoad)
 					{
 						const int64 FileSize = IFileManager::Get().FileSize(*TypefaceEntry.Font.GetFontFilename());
 						if (FileSize > 0)
