Copying //UE4/Dev-Editor to //UE4/Dev-Main (Source: //UE4/Dev-Editor @ 3228984)

diff --git a/Engine/Source/Runtime/Engine/Private/FontFace.cpp b/Engine/Source/Runtime/Engine/Private/FontFace.cpp
--- a/Engine/Source/Runtime/Engine/Private/FontFace.cpp
+++ b/Engine/Source/Runtime/Engine/Private/FontFace.cpp
@@ -8,23 +8,84 @@
 #include "Misc/FileHelper.h"
 #include "UObject/Package.h"
 #include "Misc/PackageName.h"
+#include "EditorObjectVersion.h"
 
 UFontFace::UFontFace()
+	: FontFaceData(FFontFaceData::MakeFontFaceData())
 {
 }
 
+void UFontFace::Serialize(FArchive& Ar)
+{
+	Ar.UsingCustomVersion(FEditorObjectVersion::GUID);
+
+	const FString LocalSourceFilename = SourceFilename;
+
+	if (Ar.IsCooking())
+	{
+		// Set the cooked filename prior to serializing the properties
+		SourceFilename = GetCookedFilename();
+	}
+
+	Super::Serialize(Ar);
+
+	if (Ar.IsCooking())
+	{
+		// Restore the original filename after serializing the properties
+		SourceFilename = LocalSourceFilename;
+	}
+
+	if (Ar.IsLoading())
+	{
+		if (Ar.CustomVer(FEditorObjectVersion::GUID) < FEditorObjectVersion::AddedInlineFontFaceAssets)
+		{
+#if WITH_EDITORONLY_DATA
+			// Port the old property data into the shared instance
+			FontFaceData->SetData(MoveTemp(FontFaceData_DEPRECATED));
+#endif // WITH_EDITORONLY_DATA
+		}
+		else
+		{
+			bool bLoadInlineData = false;
+			Ar << bLoadInlineData;
+
+			if (bLoadInlineData)
+			{
+				if (FontFaceData->HasData())
+				{
+					// If we already have data, make a new instance in-case the existing one is being referenced by the font cache
+					FontFaceData = FFontFaceData::MakeFontFaceData();
+				}
+				FontFaceData->Serialize(Ar);
+			}
+		}
+	}
+	else
+	{
+		// Only save the inline data in a cooked build if we're using the inline loading policy
+		bool bSaveInlineData = LoadingPolicy == EFontLoadingPolicy::Inline || !Ar.IsCooking();
+		Ar << bSaveInlineData;
+
+		if (bSaveInlineData)
+		{
+			FontFaceData->Serialize(Ar);
+		}
+	}
+}
+
 void UFontFace::GetResourceSizeEx(FResourceSizeEx& CumulativeResourceSize)
 {
 	Super::GetResourceSizeEx(CumulativeResourceSize);
 
-	// Only count the memory size for fonts that will be pre-loaded
-	if (LoadingPolicy == EFontLoadingPolicy::PreLoad)
+	// Only count the memory size for fonts that will be loaded
+	if (LoadingPolicy != EFontLoadingPolicy::Stream)
 	{
-#if WITH_EDITORONLY_DATA
+		const bool bCountInlineData = WITH_EDITORONLY_DATA || LoadingPolicy == EFontLoadingPolicy::Inline;
+		if (bCountInlineData)
 		{
-			CumulativeResourceSize.AddDedicatedSystemMemoryBytes(FontFaceData.Num());
+			CumulativeResourceSize.AddDedicatedSystemMemoryBytes(FontFaceData->GetData().Num());
 		}
-#else  // WITH_EDITORONLY_DATA
+		else
 		{
 			const int64 FileSize = IFileManager::Get().FileSize(*SourceFilename);
 			if (FileSize > 0)
@@ -32,7 +93,6 @@ void UFontFace::GetResourceSizeEx(FResourceSizeEx& CumulativeResourceSize)
 				CumulativeResourceSize.AddDedicatedSystemMemoryBytes(FileSize);
 			}
 		}
-#endif // WITH_EDITORONLY_DATA
 	}
 }
 
@@ -64,22 +124,28 @@ void UFontFace::CookAdditionalFiles(const TCHAR* PackageFilename, const ITargetP
 {
 	Super::CookAdditionalFiles(PackageFilename, TargetPlatform);
 
-	// We replace the package name with the cooked font face name
-	// Note: This must match the replacement logic in UFontFace::GetCookedFilename
-	const FString CookedFontFilename = FPaths::GetPath(PackageFilename) / GetName() + TEXT(".ufont");
-	FFileHelper::SaveArrayToFile(FontFaceData, *CookedFontFilename);
+	if (LoadingPolicy != EFontLoadingPolicy::Inline)
+	{
+		// We replace the package name with the cooked font face name
+		// Note: This must match the replacement logic in UFontFace::GetCookedFilename
+		const FString CookedFontFilename = FPaths::GetPath(PackageFilename) / GetName() + TEXT(".ufont");
+		FFileHelper::SaveArrayToFile(FontFaceData->GetData(), *CookedFontFilename);
+	}
 }
 #endif // WITH_EDITOR
 
 #if WITH_EDITORONLY_DATA
 void UFontFace::InitializeFromBulkData(const FString& InFilename, const EFontHinting InHinting, const void* InBulkDataPtr, const int32 InBulkDataSizeBytes)
 {
-	check(InBulkDataPtr && InBulkDataSizeBytes > 0 && FontFaceData.Num() == 0);
+	check(InBulkDataPtr && InBulkDataSizeBytes > 0 && !FontFaceData->HasData());
 
 	SourceFilename = InFilename;
 	Hinting = InHinting;
-	LoadingPolicy = EFontLoadingPolicy::PreLoad;
-	FontFaceData.Append(static_cast<const uint8*>(InBulkDataPtr), InBulkDataSizeBytes);
+	LoadingPolicy = EFontLoadingPolicy::LazyLoad;
+	
+	TArray<uint8> FontData;
+	FontData.Append(static_cast<const uint8*>(InBulkDataPtr), InBulkDataSizeBytes);
+	FontFaceData->SetData(MoveTemp(FontData));
 }
 #endif // WITH_EDITORONLY_DATA
 
@@ -98,12 +164,10 @@ EFontLoadingPolicy UFontFace::GetLoadingPolicy() const
 	return LoadingPolicy;
 }
 
-#if WITH_EDITORONLY_DATA
-const TArray<uint8>& UFontFace::GetFontFaceData() const
+FFontFaceDataConstRef UFontFace::GetFontFaceData() const
 {
 	return FontFaceData;
 }
-#endif // WITH_EDITORONLY_DATA
 
 FString UFontFace::GetCookedFilename() const
 {
diff --git a/Engine/Source/Runtime/Engine/Private/Materials/Material.cpp b/Engine/Source/Runtime/Engine/Private/Materials/Material.cpp
--- a/Engine/Source/Runtime/Engine/Private/Materials/Material.cpp
+++ b/Engine/Source/Runtime/Engine/Private/Materials/Material.cpp
@@ -2168,14 +2168,17 @@ void UMaterial::FlushResourceShaderMaps()
 {
 	FPlatformMisc::CreateGuid(StateId);
 
-	UMaterialInterface::IterateOverActiveFeatureLevels([&](ERHIFeatureLevel::Type InFeatureLevel)
+	if(FApp::CanEverRender())
 	{
-		for (int32 QualityLevelIndex = 0; QualityLevelIndex < EMaterialQualityLevel::Num; QualityLevelIndex++)
+		UMaterialInterface::IterateOverActiveFeatureLevels([&](ERHIFeatureLevel::Type InFeatureLevel)
 		{
-			FMaterialResource* CurrentResource = MaterialResources[QualityLevelIndex][InFeatureLevel];
-			CurrentResource->ReleaseShaderMap();
-		}
-	});
+			for(int32 QualityLevelIndex = 0; QualityLevelIndex < EMaterialQualityLevel::Num; QualityLevelIndex++)
+			{
+				FMaterialResource* CurrentResource = MaterialResources[QualityLevelIndex][InFeatureLevel];
+				CurrentResource->ReleaseShaderMap();
+			}
+		});
+	}
 }
 
 void UMaterial::RebuildMaterialFunctionInfo()
@@ -3019,16 +3022,17 @@ void UMaterial::PostEditChangeProperty(FPropertyChangedEvent& PropertyChangedEve
 
 	TranslucencyDirectionalLightingIntensity = FMath::Clamp(TranslucencyDirectionalLightingIntensity, .1f, 10.0f);
 
-	// Don't want to recompile after a duplicate because it's just been done by PostLoad
-	if( PropertyChangedEvent.ChangeType == EPropertyChangeType::Duplicate )
+	// Don't want to recompile after a duplicate because it's just been done by PostLoad, nor during interactive changes to prevent constant recompliation while spinning properties.
+	if( PropertyChangedEvent.ChangeType == EPropertyChangeType::Duplicate || PropertyChangedEvent.ChangeType == EPropertyChangeType::Interactive )
 	{
 		bRequiresCompilation = false;
 	}
-
-	// Prevent constant recompliation while spinning properties
-	if (bRequiresCompilation && PropertyChangedEvent.ChangeType != EPropertyChangeType::Interactive )
+	
+	if (bRequiresCompilation)
 	{
-		CacheResourceShadersForRendering(true);
+		// When redirecting an object pointer, we trust that the DDC hash will detect the change and that we don't need to force a recompile.
+		const bool bRegenerateId = PropertyChangedEvent.ChangeType != EPropertyChangeType::Redirected;
+		CacheResourceShadersForRendering(bRegenerateId);
 		RecacheMaterialInstanceUniformExpressions(this);
 
 		// Ensure that the ReferencedTextureGuids array is up to date.
diff --git a/Engine/Source/Runtime/Engine/Private/Materials/MaterialExpressions.cpp b/Engine/Source/Runtime/Engine/Private/Materials/MaterialExpressions.cpp
--- a/Engine/Source/Runtime/Engine/Private/Materials/MaterialExpressions.cpp
+++ b/Engine/Source/Runtime/Engine/Private/Materials/MaterialExpressions.cpp
@@ -112,6 +112,7 @@
 #include "Materials/MaterialExpressionPanner.h"
 #include "Materials/MaterialExpressionParameter.h"
 #include "Materials/MaterialExpressionPreviousFrameSwitch.h"
+#include "Materials/MaterialExpressionReroute.h"
 #include "Materials/MaterialExpressionScalarParameter.h"
 #include "Materials/MaterialExpressionSetMaterialAttributes.h"
 #include "Materials/MaterialExpressionStaticBoolParameter.h"
@@ -543,17 +544,19 @@ void UMaterialExpression::PostEditChangeProperty(FPropertyChangedEvent& Property
 
 	if (!GIsImportingT3D && !GIsTransacting && PropertyChangedEvent.ChangeType != EPropertyChangeType::Interactive )
 	{
+		FPropertyChangedEvent SubPropertyChangedEvent(nullptr, PropertyChangedEvent.ChangeType);
+
 		// Don't recompile the outer material if we are in the middle of a transaction or interactively changing properties
 		// as there may be many expressions in the transaction buffer and we would just be recompiling over and over again.
 		if (Material && !Material->bIsPreviewMaterial)
 		{
-			Material->PreEditChange(NULL);
-			Material->PostEditChange();
+			Material->PreEditChange(nullptr);
+			Material->PostEditChangeProperty(SubPropertyChangedEvent);
 		}
 		else if (Function)
 		{
-			Function->PreEditChange(NULL);
-			Function->PostEditChange();
+			Function->PreEditChange(nullptr);
+			Function->PostEditChangeProperty(SubPropertyChangedEvent);
 		}
 	}
 
@@ -622,7 +625,7 @@ bool UMaterialExpression::CanEditChange(const UProperty* InProperty) const
 				{
 					const FExpressionInput* Input = StructProp->ContainerPtrToValuePtr<FExpressionInput>(this);
 
-					if (Input->Expression != nullptr)
+					if (Input->Expression != nullptr && Input->GetTracedInput().Expression != nullptr)
 					{
 						bIsEditable = false;
 					}
@@ -713,6 +716,18 @@ FString UMaterialExpression::GetInputName(int32 InputIndex) const
 	return TEXT("");
 }
 
+#if WITH_EDITOR
+FString UMaterialExpression::GetCreationDescription() const
+{
+	return FString();
+}
+
+FString UMaterialExpression::GetCreationName() const
+{
+	return FString();
+}
+#endif
+
 bool UMaterialExpression::IsInputConnectionRequired(int32 InputIndex) const
 {
 #if WITH_EDITOR
@@ -1036,8 +1051,49 @@ void UMaterialExpression::SetEditableName(const FString& NewName)
 	check(false);
 }
 
+void UMaterialExpression::ValidateParameterName()
+{
+	if (Material != nullptr)
+	{
+		int32 NameIndex = 1;
+		bool FoundValidName = false;
+		FName PotentialName;
+
+		// Find an available unique name
+		while (!FoundValidName)
+		{
+			PotentialName = GetParameterName();
+
+			if (NameIndex != 1)
+			{
+				PotentialName.SetNumber(NameIndex);
+			}
+
+			FoundValidName = true;
+
+			for (UMaterialExpression* Expression : Material->Expressions)
+			{
+				if (Expression != nullptr && Expression->HasAParameterName())
+				{
+					// Name are unique per class type
+					if (Expression != this && Expression->GetClass() == GetClass() && Expression->GetParameterName() == PotentialName)
+					{
+						FoundValidName = false;
+						break;
+					}
+				}
+			}
+
+			++NameIndex;
+		}
+
+		SetParameterName(PotentialName);
+	}
+}
+
 #endif // WITH_EDITOR
 
+
 bool UMaterialExpression::ContainsInputLoop(const bool bStopOnFunctionCall /*= true*/)
 {
 	TArray<FMaterialExpressionKey> ExpressionStack;
@@ -1225,12 +1281,12 @@ bool UMaterialExpressionTextureSample::CanEditChange(const UProperty* InProperty
 		}
 		else if (PropertyFName == GET_MEMBER_NAME_CHECKED(UMaterialExpressionTextureSample, ConstCoordinate))
 		{
-			bIsEditable = !Coordinates.Expression;
+			bIsEditable = !Coordinates.GetTracedInput().Expression;
 		}
 		else if (PropertyFName == GET_MEMBER_NAME_CHECKED(UMaterialExpressionTextureSample, Texture))
 		{
 			// The Texture property is overridden by a connection to TextureObject
-			bIsEditable = TextureObject.Expression == NULL;
+			bIsEditable = TextureObject.GetTracedInput().Expression == NULL;
 		}
 	}
 
@@ -1399,7 +1455,7 @@ static bool VerifySamplerType(
 #if WITH_EDITOR
 int32 UMaterialExpressionTextureSample::Compile(class FMaterialCompiler* Compiler, int32 OutputIndex)
 {
-	if (Texture || TextureObject.Expression)
+	if (Texture || TextureObject.Expression) // We deal with reroute textures later on in this function..
 	{
 		int32 TextureReferenceIndex = INDEX_NONE;
 		int32 TextureCodeIndex = TextureObject.Expression ? TextureObject.Compile(Compiler) : Compiler->Texture(Texture, TextureReferenceIndex, SamplerSource, MipValueMode);
@@ -1430,6 +1486,34 @@ int32 UMaterialExpressionTextureSample::Compile(class FMaterialCompiler* Compile
 				InputExpression = NestedFunctionInput->GetEffectivePreviewExpression();
 			}
 
+			// If we are referencing a texture input through a reroute node, we'll need to backtrack
+			// to get to the real texture.
+			UMaterialExpressionReroute* RerouteInput = Cast<UMaterialExpressionReroute>(InputExpression);
+			if (RerouteInput)
+			{
+				int32 ExpressionOutputIndex = -1;
+				InputExpression  = RerouteInput->TraceInputsToRealExpression(ExpressionOutputIndex);
+				if (InputExpression == nullptr)
+				{
+					if (Desc.Len() > 0)
+					{
+						return Compiler->Errorf(TEXT("%s> Missing rerouted input texture"), *Desc);
+					}
+					else
+					{
+						return Compiler->Errorf(TEXT("TextureSample> Missing rerouted input texture"));
+					}
+				}
+				else if (OutputIndex >= 0)
+				{
+					uint32 OutputType = InputExpression->GetOutputType(OutputIndex);
+					if (OutputType != MCT_Texture2D && OutputType != MCT_TextureCube)
+					{
+						return Compiler->Errorf(TEXT("TextureSample> Reroute not bound to proper texture type!"));
+					}
+				}
+			}
+			
 			UMaterialExpressionTextureObject* TextureObjectExpression = Cast<UMaterialExpressionTextureObject>(InputExpression);
 			UMaterialExpressionTextureObjectParameter* TextureObjectParameter = Cast<UMaterialExpressionTextureObjectParameter>(InputExpression);
 			if (TextureObjectExpression)
@@ -1450,7 +1534,7 @@ int32 UMaterialExpressionTextureSample::Compile(class FMaterialCompiler* Compile
 		{
 			return Compiler->TextureSample(
 				TextureCodeIndex,
-				Coordinates.Expression ? Coordinates.Compile(Compiler) : Compiler->TextureCoordinate(ConstCoordinate, false, false),
+				Coordinates.GetTracedInput().Expression ? Coordinates.Compile(Compiler) : Compiler->TextureCoordinate(ConstCoordinate, false, false),
 				(EMaterialSamplerType)EffectiveSamplerType,
 				CompileMipValue0(Compiler),
 				CompileMipValue1(Compiler),
@@ -1531,12 +1615,12 @@ int32 UMaterialExpressionTextureSample::CompileMipValue0(class FMaterialCompiler
 {
 	if (MipValueMode == TMVM_Derivative)
 	{
-		if (CoordinatesDX.IsConnected())
+		if (CoordinatesDX.GetTracedInput().IsConnected())
 		{
 			return CoordinatesDX.Compile(Compiler);
 		}
 	}
-	else if (MipValue.IsConnected())
+	else if (MipValue.GetTracedInput().IsConnected())
 	{
 		return MipValue.Compile(Compiler);
 	}
@@ -1550,7 +1634,7 @@ int32 UMaterialExpressionTextureSample::CompileMipValue0(class FMaterialCompiler
 
 int32 UMaterialExpressionTextureSample::CompileMipValue1(class FMaterialCompiler* Compiler)
 {
-	if (MipValueMode == TMVM_Derivative && CoordinatesDY.IsConnected())
+	if (MipValueMode == TMVM_Derivative && CoordinatesDY.GetTracedInput().IsConnected())
 	{
 		return CoordinatesDY.Compile(Compiler);
 	}
@@ -1639,7 +1723,7 @@ int32 UMaterialExpressionTextureSampleParameter::Compile(class FMaterialCompiler
 
 	return Compiler->TextureSample(
 					Compiler->TextureParameter(ParameterName, Texture, TextureReferenceIndex, SamplerSource),
-					Coordinates.Expression ? Coordinates.Compile(Compiler) : Compiler->TextureCoordinate(ConstCoordinate, false, false),
+					Coordinates.GetTracedInput().Expression ? Coordinates.Compile(Compiler) : Compiler->TextureCoordinate(ConstCoordinate, false, false),
 					(EMaterialSamplerType)SamplerType,
 					MipValue0Index,
 					MipValue1Index,
@@ -1889,7 +1973,7 @@ UMaterialExpressionTextureProperty::UMaterialExpressionTextureProperty(const FOb
 #if WITH_EDITOR
 int32 UMaterialExpressionTextureProperty::Compile(class FMaterialCompiler* Compiler, int32 OutputIndex)
 {	
-	if (!TextureObject.Expression)
+	if (!TextureObject.GetTracedInput().Expression)
 	{
 		return Compiler->Errorf(TEXT("TextureSample> Missing input texture"));
 	}
@@ -1905,7 +1989,8 @@ int32 UMaterialExpressionTextureProperty::Compile(class FMaterialCompiler* Compi
 
 void UMaterialExpressionTextureProperty::GetTexturesForceMaterialRecompile(TArray<UTexture *> &Textures) const
 {
-	UMaterialExpression *TextureObjectExpression = TextureObject.Expression;
+	UMaterialExpression *TextureObjectExpression = TextureObject.GetTracedInput().Expression;
+
 	if (TextureObjectExpression && TextureObjectExpression->IsA(UMaterialExpressionTextureBase::StaticClass()))
 	{
 		UMaterialExpressionTextureBase *TextureExpressionBase = Cast<UMaterialExpressionTextureBase>(TextureObjectExpression);
@@ -2067,7 +2152,7 @@ UMaterialExpressionTextureSampleParameterCube::UMaterialExpressionTextureSampleP
 #if WITH_EDITOR
 int32 UMaterialExpressionTextureSampleParameterCube::Compile(class FMaterialCompiler* Compiler, int32 OutputIndex)
 {
-	if (!Coordinates.Expression)
+	if (!Coordinates.GetTracedInput().Expression)
 	{
 		return CompilerError(Compiler, TEXT("Cube sample needs UV input"));
 	}
@@ -2182,9 +2267,9 @@ const TCHAR* UMaterialExpressionTextureSampleParameterSubUV::GetRequirements()
 int32 UMaterialExpressionAdd::Compile(class FMaterialCompiler* Compiler, int32 OutputIndex)
 {
 	// if the input is hooked up, use it, otherwise use the internal constant
-	int32 Arg1 = A.Expression ? A.Compile(Compiler) : Compiler->Constant(ConstA);
+	int32 Arg1 = A.GetTracedInput().Expression ? A.Compile(Compiler) : Compiler->Constant(ConstA);
 	// if the input is hooked up, use it, otherwise use the internal constant
-	int32 Arg2 = B.Expression ? B.Compile(Compiler) : Compiler->Constant(ConstB);
+	int32 Arg2 = B.GetTracedInput().Expression ? B.Compile(Compiler) : Compiler->Constant(ConstB);
 
 	return Compiler->Add(Arg1, Arg2);
 }
@@ -2194,11 +2279,13 @@ void UMaterialExpressionAdd::GetCaption(TArray<FString>& OutCaptions) const
 {
 	FString ret = TEXT("Add");
 
-	if(!A.Expression || !B.Expression)
+	FExpressionInput ATraced = A.GetTracedInput();
+	FExpressionInput BTraced = B.GetTracedInput();
+	if(!ATraced.Expression || !BTraced.Expression)
 	{
 		ret += TEXT("(");
-		ret += A.Expression ? TEXT(",") : FString::Printf( TEXT("%.4g,"), ConstA);
-		ret += B.Expression ? TEXT(")") : FString::Printf( TEXT("%.4g)"), ConstB);
+		ret += ATraced.Expression ? TEXT(",") : FString::Printf( TEXT("%.4g,"), ConstA);
+		ret += BTraced.Expression ? TEXT(")") : FString::Printf( TEXT("%.4g)"), ConstB);
 	}
 
 	OutCaptions.Add(ret);
@@ -2234,9 +2321,9 @@ UMaterialExpressionMultiply::UMaterialExpressionMultiply(const FObjectInitialize
 int32 UMaterialExpressionMultiply::Compile(class FMaterialCompiler* Compiler, int32 OutputIndex)
 {
 	// if the input is hooked up, use it, otherwise use the internal constant
-	int32 Arg1 = A.Expression ? A.Compile(Compiler) : Compiler->Constant(ConstA);
+	int32 Arg1 = A.GetTracedInput().Expression ? A.Compile(Compiler) : Compiler->Constant(ConstA);
 	// if the input is hooked up, use it, otherwise use the internal constant
-	int32 Arg2 = B.Expression ? B.Compile(Compiler) : Compiler->Constant(ConstB);
+	int32 Arg2 = B.GetTracedInput().Expression ? B.Compile(Compiler) : Compiler->Constant(ConstB);
 
 	return Compiler->Mul(Arg1, Arg2);
 }
@@ -2245,11 +2332,14 @@ void UMaterialExpressionMultiply::GetCaption(TArray<FString>& OutCaptions) const
 {
 	FString ret = TEXT("Multiply");
 
-	if(!A.Expression || !B.Expression)
+	FExpressionInput ATraced = A.GetTracedInput();
+	FExpressionInput BTraced = B.GetTracedInput();
+
+	if(!ATraced.Expression || !BTraced.Expression)
 	{
 		ret += TEXT("(");
-		ret += A.Expression ? TEXT(",") : FString::Printf( TEXT("%.4g,"), ConstA);
-		ret += B.Expression ? TEXT(")") : FString::Printf( TEXT("%.4g)"), ConstB);
+		ret += ATraced.Expression ? TEXT(",") : FString::Printf( TEXT("%.4g,"), ConstA);
+		ret += BTraced.Expression ? TEXT(")") : FString::Printf( TEXT("%.4g)"), ConstB);
 	}
 
 	OutCaptions.Add(ret);
@@ -2282,9 +2372,9 @@ UMaterialExpressionDivide::UMaterialExpressionDivide(const FObjectInitializer& O
 int32 UMaterialExpressionDivide::Compile(class FMaterialCompiler* Compiler, int32 OutputIndex)
 {
 	// if the input is hooked up, use it, otherwise use the internal constant
-	int32 Arg1 = A.Expression ? A.Compile(Compiler) : Compiler->Constant(ConstA);
+	int32 Arg1 = A.GetTracedInput().Expression ? A.Compile(Compiler) : Compiler->Constant(ConstA);
 	// if the input is hooked up, use it, otherwise use the internal constant
-	int32 Arg2 = B.Expression ? B.Compile(Compiler) : Compiler->Constant(ConstB);
+	int32 Arg2 = B.GetTracedInput().Expression ? B.Compile(Compiler) : Compiler->Constant(ConstB);
 
 	return Compiler->Div(Arg1, Arg2);
 }
@@ -2293,11 +2383,14 @@ void UMaterialExpressionDivide::GetCaption(TArray<FString>& OutCaptions) const
 {
 	FString ret = TEXT("Divide");
 
-	if(!A.Expression || !B.Expression)
+	FExpressionInput ATraced = A.GetTracedInput();
+	FExpressionInput BTraced = B.GetTracedInput();
+
+	if(!ATraced.Expression || !BTraced.Expression)
 	{
 		ret += TEXT("(");
-		ret += A.Expression ? TEXT(",") : FString::Printf( TEXT("%.4g,"), ConstA);
-		ret += B.Expression ? TEXT(")") : FString::Printf( TEXT("%.4g)"), ConstB);
+		ret += ATraced.Expression ? TEXT(",") : FString::Printf( TEXT("%.4g,"), ConstA);
+		ret += BTraced.Expression ? TEXT(")") : FString::Printf( TEXT("%.4g)"), ConstB);
 	}
 
 	OutCaptions.Add(ret);
@@ -2333,9 +2426,9 @@ UMaterialExpressionSubtract::UMaterialExpressionSubtract(const FObjectInitialize
 int32 UMaterialExpressionSubtract::Compile(class FMaterialCompiler* Compiler, int32 OutputIndex)
 {
 	// if the input is hooked up, use it, otherwise use the internal constant
-	int32 Arg1 = A.Expression ? A.Compile(Compiler) : Compiler->Constant(ConstA);
+	int32 Arg1 = A.GetTracedInput().Expression ? A.Compile(Compiler) : Compiler->Constant(ConstA);
 	// if the input is hooked up, use it, otherwise use the internal constant
-	int32 Arg2 = B.Expression ? B.Compile(Compiler) : Compiler->Constant(ConstB);
+	int32 Arg2 = B.GetTracedInput().Expression ? B.Compile(Compiler) : Compiler->Constant(ConstB);
 
 	return Compiler->Sub(Arg1, Arg2);
 }
@@ -2344,11 +2437,13 @@ void UMaterialExpressionSubtract::GetCaption(TArray<FString>& OutCaptions) const
 {
 	FString ret = TEXT("Subtract");
 
-	if(!A.Expression || !B.Expression)
+	FExpressionInput ATraced = A.GetTracedInput();
+	FExpressionInput BTraced = B.GetTracedInput();
+	if(!ATraced.Expression || !BTraced.Expression)
 	{
 		ret += TEXT("(");
-		ret += A.Expression ? TEXT(",") : FString::Printf( TEXT("%.4g,"), ConstA);
-		ret += B.Expression ? TEXT(")") : FString::Printf( TEXT("%.4g)"), ConstB);
+		ret += ATraced.Expression ? TEXT(",") : FString::Printf( TEXT("%.4g,"), ConstA);
+		ret += BTraced.Expression ? TEXT(")") : FString::Printf( TEXT("%.4g)"), ConstB);
 	}
 
 	OutCaptions.Add(ret);
@@ -2389,11 +2484,11 @@ UMaterialExpressionLinearInterpolate::UMaterialExpressionLinearInterpolate(const
 int32 UMaterialExpressionLinearInterpolate::Compile(class FMaterialCompiler* Compiler, int32 OutputIndex)
 {
 	// if the input is hooked up, use it, otherwise use the internal constant
-	int32 Arg1 = A.Expression ? A.Compile(Compiler) : Compiler->Constant(ConstA);
+	int32 Arg1 = A.GetTracedInput().Expression ? A.Compile(Compiler) : Compiler->Constant(ConstA);
 	// if the input is hooked up, use it, otherwise use the internal constant
-	int32 Arg2 = B.Expression ? B.Compile(Compiler) : Compiler->Constant(ConstB);
+	int32 Arg2 = B.GetTracedInput().Expression ? B.Compile(Compiler) : Compiler->Constant(ConstB);
 	// if the input is hooked up, use it, otherwise use the internal constant
-	int32 Arg3 = Alpha.Expression ? Alpha.Compile(Compiler) : Compiler->Constant(ConstAlpha);
+	int32 Arg3 = Alpha.GetTracedInput().Expression ? Alpha.Compile(Compiler) : Compiler->Constant(ConstAlpha);
 
 	return Compiler->Lerp(Arg1, Arg2, Arg3);
 }
@@ -2402,12 +2497,16 @@ void UMaterialExpressionLinearInterpolate::GetCaption(TArray<FString>& OutCaptio
 {
 	FString ret = TEXT("Lerp");
 
-	if(!A.Expression || !B.Expression || !Alpha.Expression)
+	FExpressionInput ATraced = A.GetTracedInput();
+	FExpressionInput BTraced = B.GetTracedInput();
+	FExpressionInput AlphaTraced = Alpha.GetTracedInput();
+
+	if(!ATraced.Expression || !BTraced.Expression || !AlphaTraced.Expression)
 	{
 		ret += TEXT("(");
-		ret += A.Expression ? TEXT(",") : FString::Printf( TEXT("%.4g,"), ConstA);
-		ret += B.Expression ? TEXT(",") : FString::Printf( TEXT("%.4g,"), ConstB);
-		ret += Alpha.Expression ? TEXT(")") : FString::Printf( TEXT("%.4g)"), ConstAlpha);
+		ret += ATraced.Expression ? TEXT(",") : FString::Printf( TEXT("%.4g,"), ConstA);
+		ret += BTraced.Expression ? TEXT(",") : FString::Printf( TEXT("%.4g,"), ConstB);
+		ret += AlphaTraced.Expression ? TEXT(")") : FString::Printf( TEXT("%.4g)"), ConstAlpha);
 	}
 
 	OutCaptions.Add(ret);
@@ -2635,14 +2734,14 @@ void UMaterialExpressionClamp::Serialize(FArchive& Ar)
 #if WITH_EDITOR
 int32 UMaterialExpressionClamp::Compile(class FMaterialCompiler* Compiler, int32 OutputIndex)
 {
-	if(!Input.Expression)
+	if(!Input.GetTracedInput().Expression)
 	{
 		return Compiler->Errorf(TEXT("Missing Clamp input"));
 	}
 	else
 	{
-		const int32 MinIndex = Min.Expression ? Min.Compile(Compiler) : Compiler->Constant(MinDefault);
-		const int32 MaxIndex = Max.Expression ? Max.Compile(Compiler) : Compiler->Constant(MaxDefault);
+		const int32 MinIndex = Min.GetTracedInput().Expression ? Min.Compile(Compiler) : Compiler->Constant(MinDefault);
+		const int32 MaxIndex = Max.GetTracedInput().Expression ? Max.Compile(Compiler) : Compiler->Constant(MaxDefault);
 
 		if (ClampMode == CMODE_Clamp)
 		{
@@ -2666,11 +2765,11 @@ void UMaterialExpressionClamp::GetCaption(TArray<FString>& OutCaptions) const
 
 	if (ClampMode == CMODE_ClampMin || ClampMode == CMODE_Clamp)
 	{
-		NewCaption += Min.Expression ? TEXT(" (Min)") : FString::Printf(TEXT(" (Min=%.4g)"), MinDefault);
+		NewCaption += Min.GetTracedInput().Expression ? TEXT(" (Min)") : FString::Printf(TEXT(" (Min=%.4g)"), MinDefault);
 	}
 	if (ClampMode == CMODE_ClampMax || ClampMode == CMODE_Clamp)
 	{
-		NewCaption += Max.Expression ? TEXT(" (Max)") : FString::Printf(TEXT(" (Max=%.4g)"), MaxDefault);
+		NewCaption += Max.GetTracedInput().Expression ? TEXT(" (Max)") : FString::Printf(TEXT(" (Max=%.4g)"), MaxDefault);
 	}
 	OutCaptions.Add(NewCaption);
 }
@@ -2701,7 +2800,7 @@ UMaterialExpressionSaturate::UMaterialExpressionSaturate(const FObjectInitialize
 #if WITH_EDITOR
 int32 UMaterialExpressionSaturate::Compile(class FMaterialCompiler* Compiler, int32 OutputIndex)
 {
-	if(!Input.Expression)
+	if(!Input.GetTracedInput().Expression)
 	{
 		return Compiler->Errorf(TEXT("Missing Saturate input"));
 	}
@@ -2745,9 +2844,9 @@ UMaterialExpressionMin::UMaterialExpressionMin(const FObjectInitializer& ObjectI
 int32 UMaterialExpressionMin::Compile(class FMaterialCompiler* Compiler, int32 OutputIndex)
 {
 	// if the input is hooked up, use it, otherwise use the internal constant
-	int32 Arg1 = A.Expression ? A.Compile(Compiler) : Compiler->Constant(ConstA);
+	int32 Arg1 = A.GetTracedInput().Expression ? A.Compile(Compiler) : Compiler->Constant(ConstA);
 	// if the input is hooked up, use it, otherwise use the internal constant
-	int32 Arg2 = B.Expression ? B.Compile(Compiler) : Compiler->Constant(ConstB);
+	int32 Arg2 = B.GetTracedInput().Expression ? B.Compile(Compiler) : Compiler->Constant(ConstB);
 
 	return Compiler->Min(Arg1, Arg2);
 }
@@ -2756,11 +2855,14 @@ void UMaterialExpressionMin::GetCaption(TArray<FString>& OutCaptions) const
 {
 	FString ret = TEXT("Min");
 
-	if(!A.Expression || !B.Expression)
+	FExpressionInput ATraced = A.GetTracedInput();
+	FExpressionInput BTraced = B.GetTracedInput();
+
+	if(!ATraced.Expression || !BTraced.Expression)
 	{
 		ret += TEXT("(");
-		ret += A.Expression ? TEXT(",") : FString::Printf( TEXT("%.4g,"), ConstA);
-		ret += B.Expression ? TEXT(")") : FString::Printf( TEXT("%.4g)"), ConstB);
+		ret += ATraced.Expression ? TEXT(",") : FString::Printf( TEXT("%.4g,"), ConstA);
+		ret += BTraced.Expression ? TEXT(")") : FString::Printf( TEXT("%.4g)"), ConstB);
 	}
 
 	OutCaptions.Add(ret);
@@ -2797,9 +2899,9 @@ UMaterialExpressionMax::UMaterialExpressionMax(const FObjectInitializer& ObjectI
 int32 UMaterialExpressionMax::Compile(class FMaterialCompiler* Compiler, int32 OutputIndex)
 {
 	// if the input is hooked up, use it, otherwise use the internal constant
-	int32 Arg1 = A.Expression ? A.Compile(Compiler) : Compiler->Constant(ConstA);
+	int32 Arg1 = A.GetTracedInput().Expression ? A.Compile(Compiler) : Compiler->Constant(ConstA);
 	// if the input is hooked up, use it, otherwise use the internal constant
-	int32 Arg2 = B.Expression ? B.Compile(Compiler) : Compiler->Constant(ConstB);
+	int32 Arg2 = B.GetTracedInput().Expression ? B.Compile(Compiler) : Compiler->Constant(ConstB);
 
 	return Compiler->Max(Arg1, Arg2);
 }
@@ -2808,11 +2910,14 @@ void UMaterialExpressionMax::GetCaption(TArray<FString>& OutCaptions) const
 {
 	FString ret = TEXT("Max");
 
-	if(!A.Expression || !B.Expression)
+	FExpressionInput ATraced = A.GetTracedInput();
+	FExpressionInput BTraced = B.GetTracedInput();
+
+	if(!ATraced.Expression || !BTraced.Expression)
 	{
 		ret += TEXT("(");
-		ret += A.Expression ? TEXT(",") : FString::Printf( TEXT("%.4g,"), ConstA);
-		ret += B.Expression ? TEXT(")") : FString::Printf( TEXT("%.4g)"), ConstB);
+		ret += ATraced.Expression ? TEXT(",") : FString::Printf( TEXT("%.4g,"), ConstA);
+		ret += BTraced.Expression ? TEXT(")") : FString::Printf( TEXT("%.4g)"), ConstB);
 	}
 
 	OutCaptions.Add(ret);
@@ -2894,11 +2999,11 @@ UMaterialExpressionDotProduct::UMaterialExpressionDotProduct(const FObjectInitia
 #if WITH_EDITOR
 int32 UMaterialExpressionDotProduct::Compile(class FMaterialCompiler* Compiler, int32 OutputIndex)
 {
-	if(!A.Expression)
+	if(!A.GetTracedInput().Expression)
 	{
 		return Compiler->Errorf(TEXT("Missing DotProduct input A"));
 	}
-	else if(!B.Expression)
+	else if(!B.GetTracedInput().Expression)
 	{
 		return Compiler->Errorf(TEXT("Missing DotProduct input B"));
 	}
@@ -2944,11 +3049,11 @@ UMaterialExpressionCrossProduct::UMaterialExpressionCrossProduct(const FObjectIn
 #if WITH_EDITOR
 int32 UMaterialExpressionCrossProduct::Compile(class FMaterialCompiler* Compiler, int32 OutputIndex)
 {
-	if(!A.Expression)
+	if(!A.GetTracedInput().Expression)
 	{
 		return Compiler->Errorf(TEXT("Missing CrossProduct input A"));
 	}
-	else if(!B.Expression)
+	else if(!B.GetTracedInput().Expression)
 	{
 		return Compiler->Errorf(TEXT("Missing CrossProduct input B"));
 	}
@@ -2994,7 +3099,7 @@ UMaterialExpressionComponentMask::UMaterialExpressionComponentMask(const FObject
 #if WITH_EDITOR
 int32 UMaterialExpressionComponentMask::Compile(class FMaterialCompiler* Compiler, int32 OutputIndex)
 {
-	if(!Input.Expression)
+	if(!Input.GetTracedInput().Expression)
 	{
 		return Compiler->Errorf(TEXT("Missing ComponentMask input"));
 	}
@@ -3045,7 +3150,7 @@ UMaterialExpressionStaticComponentMaskParameter::UMaterialExpressionStaticCompon
 #if WITH_EDITOR
 int32 UMaterialExpressionStaticComponentMaskParameter::Compile(class FMaterialCompiler* Compiler, int32 OutputIndex)
 {
-	if(!Input.Expression)
+	if(!Input.GetTracedInput().Expression)
 	{
 		return Compiler->Errorf(TEXT("Missing ComponentMaskParameter input"));
 	}
@@ -3235,8 +3340,9 @@ UMaterialExpressionReflectionVectorWS::UMaterialExpressionReflectionVectorWS(con
 int32 UMaterialExpressionReflectionVectorWS::Compile(class FMaterialCompiler* Compiler, int32 OutputIndex)
 {
 	int32 Result = CustomWorldNormal.Compile(Compiler);
-	if (CustomWorldNormal.Expression)
+	if (CustomWorldNormal.Expression) 
 	{
+		// Don't do anything special here in regards to if the Expression is a Reroute node, the compiler will handle properly internally and return INDEX_NONE if rerouted to nowhere.
 		return Compiler->ReflectionAboutCustomWorldNormal(Result, bNormalizeCustomWorldNormal); 
 	}
 	else
@@ -3279,10 +3385,11 @@ UMaterialExpressionPanner::UMaterialExpressionPanner(const FObjectInitializer& O
 #if WITH_EDITOR
 int32 UMaterialExpressionPanner::Compile(class FMaterialCompiler* Compiler, int32 OutputIndex)
 {
-	int32 TimeArg = Time.Expression ? Time.Compile(Compiler) : Compiler->GameTime(false, 0.0f);
-	int32 SpeedVectorArg = Speed.Expression ? Speed.Compile(Compiler) : INDEX_NONE;
-	int32 SpeedXArg = Speed.Expression ? Compiler->ComponentMask(SpeedVectorArg, true, false, false, false) : Compiler->Constant(SpeedX);
-	int32 SpeedYArg = Speed.Expression ? Compiler->ComponentMask(SpeedVectorArg, false, true, false, false) : Compiler->Constant(SpeedY);
+	int32 TimeArg = Time.GetTracedInput().Expression ? Time.Compile(Compiler) : Compiler->GameTime(false, 0.0f);
+	bool bIsSpeedExpressionValid = Speed.GetTracedInput().Expression != nullptr;
+	int32 SpeedVectorArg = bIsSpeedExpressionValid ? Speed.Compile(Compiler) : INDEX_NONE;
+	int32 SpeedXArg = bIsSpeedExpressionValid ? Compiler->ComponentMask(SpeedVectorArg, true, false, false, false) : Compiler->Constant(SpeedX);
+	int32 SpeedYArg = bIsSpeedExpressionValid ? Compiler->ComponentMask(SpeedVectorArg, false, true, false, false) : Compiler->Constant(SpeedY);
 	int32 Arg1;
 	int32 Arg2;
 	if (bFractionalPart)
@@ -3298,7 +3405,7 @@ int32 UMaterialExpressionPanner::Compile(class FMaterialCompiler* Compiler, int3
 		Arg2 = Compiler->PeriodicHint(Compiler->Mul(TimeArg, SpeedYArg));
 	}
 
-	int32 Arg3 = Coordinate.Expression ? Coordinate.Compile(Compiler) : Compiler->TextureCoordinate(ConstCoordinate, false, false);
+	int32 Arg3 = Coordinate.GetTracedInput().Expression ? Coordinate.Compile(Compiler) : Compiler->TextureCoordinate(ConstCoordinate, false, false);
 	return Compiler->Add(
 			Compiler->AppendVector(
 				Arg1,
@@ -3346,12 +3453,12 @@ UMaterialExpressionRotator::UMaterialExpressionRotator(const FObjectInitializer&
 #if WITH_EDITOR
 int32 UMaterialExpressionRotator::Compile(class FMaterialCompiler* Compiler, int32 OutputIndex)
 {
-	int32	Cosine = Compiler->Cosine(Compiler->Mul(Time.Expression ? Time.Compile(Compiler) : Compiler->GameTime(false, 0.0f),Compiler->Constant(Speed))),
-		Sine = Compiler->Sine(Compiler->Mul(Time.Expression ? Time.Compile(Compiler) : Compiler->GameTime(false, 0.0f),Compiler->Constant(Speed))),
+	int32	Cosine = Compiler->Cosine(Compiler->Mul(Time.GetTracedInput().Expression ? Time.Compile(Compiler) : Compiler->GameTime(false, 0.0f),Compiler->Constant(Speed))),
+		Sine = Compiler->Sine(Compiler->Mul(Time.GetTracedInput().Expression ? Time.Compile(Compiler) : Compiler->GameTime(false, 0.0f),Compiler->Constant(Speed))),
 		RowX = Compiler->AppendVector(Cosine,Compiler->Mul(Compiler->Constant(-1.0f),Sine)),
 		RowY = Compiler->AppendVector(Sine,Cosine),
 		Origin = Compiler->Constant2(CenterX,CenterY),
-		BaseCoordinate = Coordinate.Expression ? Coordinate.Compile(Compiler) : Compiler->TextureCoordinate(ConstCoordinate, false, false);
+		BaseCoordinate = Coordinate.GetTracedInput().Expression ? Coordinate.Compile(Compiler) : Compiler->TextureCoordinate(ConstCoordinate, false, false);
 
 	const int32 Arg1 = Compiler->Dot(RowX,Compiler->Sub(Compiler->ComponentMask(BaseCoordinate,1,1,0,0),Origin));
 	const int32 Arg2 = Compiler->Dot(RowY,Compiler->Sub(Compiler->ComponentMask(BaseCoordinate,1,1,0,0),Origin));
@@ -3415,7 +3522,7 @@ UMaterialExpressionSine::UMaterialExpressionSine(const FObjectInitializer& Objec
 #if WITH_EDITOR
 int32 UMaterialExpressionSine::Compile(class FMaterialCompiler* Compiler, int32 OutputIndex)
 {
-	if(!Input.Expression)
+	if(!Input.GetTracedInput().Expression)
 	{
 		return Compiler->Errorf(TEXT("Missing Sine input"));
 	}
@@ -3456,7 +3563,7 @@ UMaterialExpressionCosine::UMaterialExpressionCosine(const FObjectInitializer& O
 #if WITH_EDITOR
 int32 UMaterialExpressionCosine::Compile(class FMaterialCompiler* Compiler, int32 OutputIndex)
 {
-	if(!Input.Expression)
+	if(!Input.GetTracedInput().Expression)
 	{
 		return Compiler->Errorf(TEXT("Missing Cosine input"));
 	}
@@ -3497,7 +3604,7 @@ UMaterialExpressionTangent::UMaterialExpressionTangent(const FObjectInitializer&
 #if WITH_EDITOR
 int32 UMaterialExpressionTangent::Compile(class FMaterialCompiler* Compiler, int32 OutputIndex)
 {
-	if(!Input.Expression)
+	if(!Input.GetTracedInput().Expression)
 	{
 		return Compiler->Errorf(TEXT("Missing Tangent input"));
 	}
@@ -3536,7 +3643,7 @@ UMaterialExpressionArcsine::UMaterialExpressionArcsine(const FObjectInitializer&
 #if WITH_EDITOR
 int32 UMaterialExpressionArcsine::Compile(class FMaterialCompiler* Compiler, int32 OutputIndex)
 {
-	if(!Input.Expression)
+	if(!Input.GetTracedInput().Expression)
 	{
 		return Compiler->Errorf(TEXT("Missing Arcsine input"));
 	}
@@ -3575,7 +3682,7 @@ UMaterialExpressionArcsineFast::UMaterialExpressionArcsineFast(const FObjectInit
 #if WITH_EDITOR
 int32 UMaterialExpressionArcsineFast::Compile(class FMaterialCompiler* Compiler, int32 OutputIndex)
 {
-	if(!Input.Expression)
+	if(!Input.GetTracedInput().Expression)
 	{
 		return Compiler->Errorf(TEXT("Missing ArcsineFast input"));
 	}
@@ -3614,7 +3721,7 @@ UMaterialExpressionArccosine::UMaterialExpressionArccosine(const FObjectInitiali
 #if WITH_EDITOR
 int32 UMaterialExpressionArccosine::Compile(class FMaterialCompiler* Compiler, int32 OutputIndex)
 {
-	if(!Input.Expression)
+	if(!Input.GetTracedInput().Expression)
 	{
 		return Compiler->Errorf(TEXT("Missing Arccosine input"));
 	}
@@ -3653,7 +3760,7 @@ UMaterialExpressionArccosineFast::UMaterialExpressionArccosineFast(const FObject
 #if WITH_EDITOR
 int32 UMaterialExpressionArccosineFast::Compile(class FMaterialCompiler* Compiler, int32 OutputIndex)
 {
-	if(!Input.Expression)
+	if(!Input.GetTracedInput().Expression)
 	{
 		return Compiler->Errorf(TEXT("Missing ArccosineFast input"));
 	}
@@ -3692,7 +3799,7 @@ UMaterialExpressionArctangent::UMaterialExpressionArctangent(const FObjectInitia
 #if WITH_EDITOR
 int32 UMaterialExpressionArctangent::Compile(class FMaterialCompiler* Compiler, int32 OutputIndex)
 {
-	if(!Input.Expression)
+	if(!Input.GetTracedInput().Expression)
 	{
 		return Compiler->Errorf(TEXT("Missing Arctangent input"));
 	}
@@ -3731,7 +3838,7 @@ UMaterialExpressionArctangentFast::UMaterialExpressionArctangentFast(const FObje
 #if WITH_EDITOR
 int32 UMaterialExpressionArctangentFast::Compile(class FMaterialCompiler* Compiler, int32 OutputIndex)
 {
-	if(!Input.Expression)
+	if(!Input.GetTracedInput().Expression)
 	{
 		return Compiler->Errorf(TEXT("Missing ArctangentFast input"));
 	}
@@ -3770,7 +3877,7 @@ UMaterialExpressionArctangent2::UMaterialExpressionArctangent2(const FObjectInit
 #if WITH_EDITOR
 int32 UMaterialExpressionArctangent2::Compile(class FMaterialCompiler* Compiler, int32 OutputIndex)
 {
-	if(!Y.Expression || !X.Expression)
+	if(!Y.GetTracedInput().Expression || !X.GetTracedInput().Expression)
 	{
 		return Compiler->Errorf(TEXT("Missing Arctangent2 input"));
 	}
@@ -3811,7 +3918,7 @@ UMaterialExpressionArctangent2Fast::UMaterialExpressionArctangent2Fast(const FOb
 #if WITH_EDITOR
 int32 UMaterialExpressionArctangent2Fast::Compile(class FMaterialCompiler* Compiler, int32 OutputIndex)
 {
-	if(!Y.Expression || !X.Expression)
+	if(!Y.GetTracedInput().Expression || !X.GetTracedInput().Expression)
 	{
 		return Compiler->Errorf(TEXT("Missing Arctangent2Fast input"));
 	}
@@ -3854,7 +3961,7 @@ UMaterialExpressionBumpOffset::UMaterialExpressionBumpOffset(const FObjectInitia
 #if WITH_EDITOR
 int32 UMaterialExpressionBumpOffset::Compile(class FMaterialCompiler* Compiler, int32 OutputIndex)
 {
-	if(!Height.Expression)
+	if(!Height.GetTracedInput().Expression)
 	{
 		return Compiler->Errorf(TEXT("Missing Height input"));
 	}
@@ -3864,13 +3971,13 @@ int32 UMaterialExpressionBumpOffset::Compile(class FMaterialCompiler* Compiler,
 				Compiler->ComponentMask(Compiler->TransformVector(MCB_World, MCB_Tangent, Compiler->CameraVector()),1,1,0,0),
 				Compiler->Add(
 					Compiler->Mul(
-						HeightRatioInput.Expression ? Compiler->ForceCast(HeightRatioInput.Compile(Compiler),MCT_Float1) : Compiler->Constant(HeightRatio),
+						HeightRatioInput.GetTracedInput().Expression ? Compiler->ForceCast(HeightRatioInput.Compile(Compiler),MCT_Float1) : Compiler->Constant(HeightRatio),
 						Compiler->ForceCast(Height.Compile(Compiler),MCT_Float1)
 						),
-					HeightRatioInput.Expression ? Compiler->Mul(Compiler->Constant(-ReferencePlane), Compiler->ForceCast(HeightRatioInput.Compile(Compiler),MCT_Float1)) : Compiler->Constant(-ReferencePlane * HeightRatio)
+					HeightRatioInput.GetTracedInput().Expression ? Compiler->Mul(Compiler->Constant(-ReferencePlane), Compiler->ForceCast(HeightRatioInput.Compile(Compiler),MCT_Float1)) : Compiler->Constant(-ReferencePlane * HeightRatio)
 					)
 				),
-			Coordinate.Expression ? Coordinate.Compile(Compiler) : Compiler->TextureCoordinate(ConstCoordinate, false, false)
+			Coordinate.GetTracedInput().Expression ? Coordinate.Compile(Compiler) : Compiler->TextureCoordinate(ConstCoordinate, false, false)
 			);
 }
 
@@ -3905,11 +4012,11 @@ UMaterialExpressionAppendVector::UMaterialExpressionAppendVector(const FObjectIn
 #if WITH_EDITOR
 int32 UMaterialExpressionAppendVector::Compile(class FMaterialCompiler* Compiler, int32 OutputIndex)
 {
-	if(!A.Expression)
+	if(!A.GetTracedInput().Expression)
 	{
 		return Compiler->Errorf(TEXT("Missing AppendVector input A"));
 	}
-	else if(!B.Expression)
+	else if(!B.GetTracedInput().Expression)
 	{
 		return Compiler->Errorf(TEXT("Missing AppendVector input B"));
 	}
@@ -3964,6 +4071,7 @@ int32 UMaterialExpressionMakeMaterialAttributes::Compile(class FMaterialCompiler
 		");
 
 	EMaterialProperty Property = FMaterialAttributeDefinitionMap::GetProperty(Compiler->GetMaterialAttribute());
+	// We don't worry about reroute nodes in the switch, as we have a test for their validity afterwards.
 	switch (Property)
 	{
 	case MP_BaseColor: Ret = BaseColor.Compile(Compiler); Expression = BaseColor.Expression; break;
@@ -3990,7 +4098,7 @@ int32 UMaterialExpressionMakeMaterialAttributes::Compile(class FMaterialCompiler
 		Ret = CustomizedUVs[Property - MP_CustomizedUVs0].Compile(Compiler); Expression = CustomizedUVs[Property - MP_CustomizedUVs0].Expression;
 	}
 
-	//If we've connected an expression but its still returned INDEX_NONE, flag the error.
+	//If we've connected an expression but its still returned INDEX_NONE, flag the error. This also catches reroute nodes to nowhere.
 	if (Expression && INDEX_NONE == Ret)
 	{
 		Compiler->Errorf(TEXT("Error on property %s"), *FMaterialAttributeDefinitionMap::GetDisplayName(Property));
@@ -4665,7 +4773,7 @@ UMaterialExpressionFloor::UMaterialExpressionFloor(const FObjectInitializer& Obj
 #if WITH_EDITOR
 int32 UMaterialExpressionFloor::Compile(class FMaterialCompiler* Compiler, int32 OutputIndex)
 {
-	if(!Input.Expression)
+	if(!Input.GetTracedInput().Expression)
 	{
 		return Compiler->Errorf(TEXT("Missing Floor input"));
 	}
@@ -4701,7 +4809,7 @@ UMaterialExpressionCeil::UMaterialExpressionCeil(const FObjectInitializer& Objec
 #if WITH_EDITOR
 int32 UMaterialExpressionCeil::Compile(class FMaterialCompiler* Compiler, int32 OutputIndex)
 {
-	if(!Input.Expression)
+	if(!Input.GetTracedInput().Expression)
 	{
 		return Compiler->Errorf(TEXT("Missing Ceil input"));
 	}
@@ -4740,7 +4848,7 @@ UMaterialExpressionRound::UMaterialExpressionRound(const FObjectInitializer& Obj
 #if WITH_EDITOR
 int32 UMaterialExpressionRound::Compile(class FMaterialCompiler* Compiler, int32 OutputIndex)
 {
-	if(!Input.Expression)
+	if(!Input.GetTracedInput().Expression)
 	{
 		return Compiler->Errorf(TEXT("Missing Round input"));
 	}
@@ -4778,7 +4886,7 @@ UMaterialExpressionTruncate::UMaterialExpressionTruncate(const FObjectInitialize
 #if WITH_EDITOR
 int32 UMaterialExpressionTruncate::Compile(class FMaterialCompiler* Compiler, int32 OutputIndex)
 {
-	if(!Input.Expression)
+	if(!Input.GetTracedInput().Expression)
 	{
 		return Compiler->Errorf(TEXT("Missing Truncate input"));
 	}
@@ -4817,11 +4925,11 @@ UMaterialExpressionFmod::UMaterialExpressionFmod(const FObjectInitializer& Objec
 #if WITH_EDITOR
 int32 UMaterialExpressionFmod::Compile(class FMaterialCompiler* Compiler, int32 OutputIndex)
 {
-	if (!A.Expression)
+	if (!A.GetTracedInput().Expression)
 	{
 		return Compiler->Errorf(TEXT("Missing Fmod input A"));
 	}
-	if (!B.Expression)
+	if (!B.GetTracedInput().Expression)
 	{
 		return Compiler->Errorf(TEXT("Missing Fmod input B"));
 	}
@@ -4856,7 +4964,7 @@ UMaterialExpressionFrac::UMaterialExpressionFrac(const FObjectInitializer& Objec
 #if WITH_EDITOR
 int32 UMaterialExpressionFrac::Compile(class FMaterialCompiler* Compiler, int32 OutputIndex)
 {
-	if(!Input.Expression)
+	if(!Input.GetTracedInput().Expression)
 	{
 		return Compiler->Errorf(TEXT("Missing Frac input"));
 	}
@@ -4897,13 +5005,13 @@ UMaterialExpressionDesaturation::UMaterialExpressionDesaturation(const FObjectIn
 #if WITH_EDITOR
 int32 UMaterialExpressionDesaturation::Compile(class FMaterialCompiler* Compiler, int32 OutputIndex)
 {
-	if(!Input.Expression)
+	if(!Input.GetTracedInput().Expression)
 		return Compiler->Errorf(TEXT("Missing Desaturation input"));
 
 	int32 Color = Compiler->ForceCast(Input.Compile(Compiler), MCT_Float3, MFCF_ExactMatch|MFCF_ReplicateValue),
 		Grey = Compiler->Dot(Color,Compiler->Constant3(LuminanceFactors.R,LuminanceFactors.G,LuminanceFactors.B));
 
-	if(Fraction.Expression)
+	if(Fraction.GetTracedInput().Expression)
 		return Compiler->Lerp(Color,Grey,Fraction.Compile(Compiler));
 	else
 		return Grey;
@@ -5111,8 +5219,11 @@ UMaterialExpressionStaticSwitchParameter::UMaterialExpressionStaticSwitchParamet
 bool UMaterialExpressionStaticSwitchParameter::IsResultMaterialAttributes(int32 OutputIndex)
 {
 	check(OutputIndex == 0);
-	if( (A.Expression && !A.Expression->ContainsInputLoop() && A.Expression->IsResultMaterialAttributes(A.OutputIndex)) ||
-		(B.Expression && !B.Expression->ContainsInputLoop() && B.Expression->IsResultMaterialAttributes(B.OutputIndex)))
+	// This one is a little tricky. Since we are treating a dangling reroute as an empty expression, this
+	// should early out, whereas IsResultMaterialAttributes on a reroute node will return false as the 
+	// reroute node's input is dangling and therefore its type is unknown.
+	if ((A.GetTracedInput().Expression && !A.Expression->ContainsInputLoop() && A.Expression->IsResultMaterialAttributes(A.OutputIndex)) ||
+		(B.GetTracedInput().Expression && !B.Expression->ContainsInputLoop() && B.Expression->IsResultMaterialAttributes(B.OutputIndex)))
 	{
 		return true;
 	}
@@ -5128,12 +5239,12 @@ int32 UMaterialExpressionStaticSwitchParameter::Compile(class FMaterialCompiler*
 	const bool bValue = Compiler->GetStaticBoolValue(Compiler->StaticBoolParameter(ParameterName,DefaultValue), bSucceeded);
 	
 	//Both A and B must be connected in a parameter. 
-	if( !A.IsConnected() )
+	if( !A.GetTracedInput().IsConnected() )
 	{
 		Compiler->Errorf(TEXT("Missing A input"));
 		bSucceeded = false;
 	}
-	if( !B.IsConnected() )
+	if( !B.GetTracedInput().IsConnected() )
 	{
 		Compiler->Errorf(TEXT("Missing B input"));
 		bSucceeded = false;
@@ -5281,10 +5392,13 @@ UMaterialExpressionStaticSwitch::UMaterialExpressionStaticSwitch(const FObjectIn
 #if WITH_EDITOR
 bool UMaterialExpressionStaticSwitch::IsResultMaterialAttributes(int32 OutputIndex)
 {
-	// If there is a loop anywhere in this expression's inputs then we can't risk checking them
+	// If there is a loop anywhere in this expression's inputs then we can't risk checking them. 
+	// This one is a little tricky with respect to Reroute nodes. Since we are treating a dangling reroute as an empty expression, this
+	// should early out, whereas IsResultMaterialAttributes on a reroute node will return false as the 
+	// reroute node's input is dangling and therefore its type is unknown.
 	check(OutputIndex == 0);
-	if( (A.Expression && !A.Expression->ContainsInputLoop() && A.Expression->IsResultMaterialAttributes(A.OutputIndex)) ||
-		(B.Expression && !B.Expression->ContainsInputLoop() && B.Expression->IsResultMaterialAttributes(B.OutputIndex)))
+	if ((A.GetTracedInput().Expression && !A.Expression->ContainsInputLoop() && A.Expression->IsResultMaterialAttributes(A.OutputIndex)) ||
+		(B.GetTracedInput().Expression && !B.Expression->ContainsInputLoop() && B.Expression->IsResultMaterialAttributes(B.OutputIndex)))
 	{
 		return true;
 	}
@@ -5298,7 +5412,7 @@ int32 UMaterialExpressionStaticSwitch::Compile(class FMaterialCompiler* Compiler
 {
 	bool bValue = DefaultValue;
 
-	if (Value.Expression)
+	if (Value.GetTracedInput().Expression)
 	{
 		bool bSucceeded;
 		bValue = Compiler->GetStaticBoolValue(Value.Compile(Compiler), bSucceeded);
@@ -5452,9 +5566,10 @@ int32 UMaterialExpressionQualitySwitch::Compile(class FMaterialCompiler* Compile
 {
 	const EMaterialQualityLevel::Type QualityLevelToCompile = Compiler->GetQualityLevel();
 	check(QualityLevelToCompile < ARRAY_COUNT(Inputs));
-	FExpressionInput& QualityInput = Inputs[QualityLevelToCompile];
+	FExpressionInput QualityInput = Inputs[QualityLevelToCompile].GetTracedInput();
+	FExpressionInput DefaultTraced = Default.GetTracedInput();
 
-	if (!Default.Expression)
+	if (!DefaultTraced.Expression)
 	{
 		return Compiler->Errorf(TEXT("Quality switch missing default input"));
 	}
@@ -5464,7 +5579,7 @@ int32 UMaterialExpressionQualitySwitch::Compile(class FMaterialCompiler* Compile
 		return QualityInput.Compile(Compiler);
 	}
 
-	return Default.Compile(Compiler);
+	return DefaultTraced.Compile(Compiler);
 }
 
 void UMaterialExpressionQualitySwitch::GetCaption(TArray<FString>& OutCaptions) const
@@ -5568,12 +5683,12 @@ int32 UMaterialExpressionFeatureLevelSwitch::Compile(class FMaterialCompiler* Co
 	check(FeatureLevelToCompile < ARRAY_COUNT(Inputs));
 	FExpressionInput& FeatureInput = Inputs[FeatureLevelToCompile];
 
-	if (!Default.Expression)
+	if (!Default.GetTracedInput().Expression)
 	{
 		return Compiler->Errorf(TEXT("Feature Level switch missing default input"));
 	}
 
-	if (FeatureInput.Expression)
+	if (FeatureInput.GetTracedInput().Expression)
 	{
 		return FeatureInput.Compile(Compiler);
 	}
@@ -5637,7 +5752,7 @@ bool UMaterialExpressionFeatureLevelSwitch::IsResultMaterialAttributes(int32 Out
 	for (FExpressionInput* ExpressionInput : ExpressionInputs)
 	{
 		// If there is a loop anywhere in this expression's inputs then we can't risk checking them
-		if (ExpressionInput->Expression && !ExpressionInput->Expression->ContainsInputLoop() && ExpressionInput->Expression->IsResultMaterialAttributes(ExpressionInput->OutputIndex))
+		if (ExpressionInput->GetTracedInput().Expression && !ExpressionInput->Expression->ContainsInputLoop() && ExpressionInput->Expression->IsResultMaterialAttributes(ExpressionInput->OutputIndex))
 		{
 			return true;
 		}
@@ -5690,7 +5805,7 @@ UMaterialExpressionNormalize::UMaterialExpressionNormalize(const FObjectInitiali
 #if WITH_EDITOR
 int32 UMaterialExpressionNormalize::Compile(class FMaterialCompiler* Compiler, int32 OutputIndex)
 {
-	if(!VectorInput.Expression)
+	if(!VectorInput.GetTracedInput().Expression)
 		return Compiler->Errorf(TEXT("Missing Normalize input"));
 
 	int32	V = VectorInput.Compile(Compiler);
@@ -6303,7 +6418,7 @@ UMaterialExpressionSquareRoot::UMaterialExpressionSquareRoot(const FObjectInitia
 #if WITH_EDITOR
 int32 UMaterialExpressionSquareRoot::Compile(class FMaterialCompiler* Compiler, int32 OutputIndex)
 {
-	if(!Input.Expression)
+	if(!Input.GetTracedInput().Expression)
 	{
 		return Compiler->Errorf(TEXT("Missing square root input"));
 	}
@@ -6405,7 +6520,7 @@ int32 UMaterialExpressionSceneDepth::Compile(class FMaterialCompiler* Compiler,
 
 	if(InputMode == EMaterialSceneAttributeInputMode::OffsetFraction)
 	{
-		if (Input.Expression)
+		if (Input.GetTracedInput().Expression)
 		{
 			OffsetIndex = Input.Compile(Compiler);
 		} 
@@ -6417,7 +6532,7 @@ int32 UMaterialExpressionSceneDepth::Compile(class FMaterialCompiler* Compiler,
 	}
 	else if(InputMode == EMaterialSceneAttributeInputMode::Coordinates)
 	{
-		if (Input.Expression)
+		if (Input.GetTracedInput().Expression)
 		{
 			CoordinateIndex = Input.Compile(Compiler);
 		} 
@@ -6486,7 +6601,7 @@ int32 UMaterialExpressionSceneTexture::Compile(class FMaterialCompiler* Compiler
 {    
 	int32 UV = INDEX_NONE;
 
-	if (Coordinates.Expression)
+	if (Coordinates.GetTracedInput().Expression)
 	{
 		UV = Coordinates.Compile(Compiler);
 	}
@@ -6582,7 +6697,7 @@ int32 UMaterialExpressionSceneColor::Compile(class FMaterialCompiler* Compiler,
 
 	if(InputMode == EMaterialSceneAttributeInputMode::OffsetFraction)
 	{
-		if (Input.Expression)
+		if (Input.GetTracedInput().Expression)
 		{
 			OffsetIndex = Input.Compile(Compiler);
 		}
@@ -6595,7 +6710,7 @@ int32 UMaterialExpressionSceneColor::Compile(class FMaterialCompiler* Compiler,
 	}
 	else if(InputMode == EMaterialSceneAttributeInputMode::Coordinates)
 	{
-		if (Input.Expression)
+		if (Input.GetTracedInput().Expression)
 		{
 			CoordinateIndex = Input.Compile(Compiler);
 		} 
@@ -6635,13 +6750,13 @@ UMaterialExpressionPower::UMaterialExpressionPower(const FObjectInitializer& Obj
 #if WITH_EDITOR
 int32 UMaterialExpressionPower::Compile(class FMaterialCompiler* Compiler, int32 OutputIndex)
 {
-	if(!Base.Expression)
+	if(!Base.GetTracedInput().Expression)
 	{
 		return Compiler->Errorf(TEXT("Missing Power Base input"));
 	}
 
 	int32 Arg1 = Base.Compile(Compiler);
-	int32 Arg2 = Exponent.Expression ? Exponent.Compile(Compiler) : Compiler->Constant(ConstExponent);
+	int32 Arg2 = Exponent.GetTracedInput().Expression ? Exponent.Compile(Compiler) : Compiler->Constant(ConstExponent);
 	return Compiler->Power(
 		Arg1,
 		Arg2
@@ -6652,7 +6767,7 @@ void UMaterialExpressionPower::GetCaption(TArray<FString>& OutCaptions) const
 {
 	FString ret = TEXT("Power");
 
-	if (!Exponent.Expression)
+	if (!Exponent.GetTracedInput().Expression)
 	{
 		ret += FString::Printf( TEXT("(X, %.4g)"), ConstExponent);
 	}
@@ -6683,7 +6798,7 @@ UMaterialExpressionLogarithm2::UMaterialExpressionLogarithm2(const FObjectInitia
 #if WITH_EDITOR
 int32 UMaterialExpressionLogarithm2::Compile(class FMaterialCompiler* Compiler, int32 OutputIndex)
 {
-	if(!X.Expression)
+	if(!X.GetTracedInput().Expression)
 	{
 		return Compiler->Errorf(TEXT("Missing Log2 X input"));
 	}
@@ -6733,21 +6848,21 @@ UMaterialExpressionIf::UMaterialExpressionIf(const FObjectInitializer& ObjectIni
 #if WITH_EDITOR
 int32 UMaterialExpressionIf::Compile(class FMaterialCompiler* Compiler, int32 OutputIndex)
 {
-	if(!A.Expression)
+	if(!A.GetTracedInput().Expression)
 	{
 		return Compiler->Errorf(TEXT("Missing If A input"));
 	}
-	if(!AGreaterThanB.Expression)
+	if(!AGreaterThanB.GetTracedInput().Expression)
 	{
 		return Compiler->Errorf(TEXT("Missing If AGreaterThanB input"));
 	}
-	if(!ALessThanB.Expression)
+	if(!ALessThanB.GetTracedInput().Expression)
 	{
 		return Compiler->Errorf(TEXT("Missing If ALessThanB input"));
 	}
 
 	int32 CompiledA = A.Compile(Compiler);
-	int32 CompiledB = B.Expression ? B.Compile(Compiler) : Compiler->Constant(ConstB);
+	int32 CompiledB = B.GetTracedInput().Expression ? B.Compile(Compiler) : Compiler->Constant(ConstB);
 
 	if(Compiler->GetType(CompiledA) != MCT_Float)
 	{
@@ -6760,7 +6875,7 @@ int32 UMaterialExpressionIf::Compile(class FMaterialCompiler* Compiler, int32 Ou
 	}
 
 	int32 Arg3 = AGreaterThanB.Compile(Compiler);
-	int32 Arg4 = AEqualsB.Expression ? AEqualsB.Compile(Compiler) : INDEX_NONE;
+	int32 Arg4 = AEqualsB.GetTracedInput().Expression ? AEqualsB.Compile(Compiler) : INDEX_NONE;
 	int32 Arg5 = ALessThanB.Compile(Compiler);
 	int32 ThresholdArg = Compiler->Constant(EqualsThreshold);
 
@@ -6808,7 +6923,7 @@ UMaterialExpressionOneMinus::UMaterialExpressionOneMinus(const FObjectInitialize
 #if WITH_EDITOR
 int32 UMaterialExpressionOneMinus::Compile(class FMaterialCompiler* Compiler, int32 OutputIndex)
 {
-	if(!Input.Expression)
+	if(!Input.GetTracedInput().Expression)
 	{
 		return Compiler->Errorf(TEXT("Missing 1-x input"));
 	}
@@ -6845,7 +6960,7 @@ int32 UMaterialExpressionAbs::Compile( FMaterialCompiler* Compiler, int32 Output
 {
 	int32 Result=INDEX_NONE;
 
-	if( !Input.Expression )
+	if( !Input.GetTracedInput().Expression )
 	{
 		// an input expression must exist
 		Result = Compiler->Errorf( TEXT("Missing Abs input") );
@@ -6901,7 +7016,7 @@ int32 UMaterialExpressionTransform::Compile(class FMaterialCompiler* Compiler, i
 {
 	int32 Result = INDEX_NONE;
 
-	if (!Input.Expression)
+	if (!Input.GetTracedInput().Expression)
 	{
 		Result = Compiler->Errorf(TEXT("Missing Transform input vector"));
 	}
@@ -7003,7 +7118,7 @@ int32 UMaterialExpressionTransformPosition::Compile(class FMaterialCompiler* Com
 {
 	int32 Result=INDEX_NONE;
 	
-	if( !Input.Expression )
+	if( !Input.GetTracedInput().Expression )
 	{
 		Result = Compiler->Errorf(TEXT("Missing Transform Position input vector"));
 	}
@@ -7147,13 +7262,13 @@ int32 UMaterialExpressionFresnel::Compile(class FMaterialCompiler* Compiler, int
 {
 	// pow(1 - max(0,Normal dot Camera),Exponent) * (1 - BaseReflectFraction) + BaseReflectFraction
 	//
-	int32 NormalArg = Normal.Expression ? Normal.Compile(Compiler) : Compiler->PixelNormalWS();
+	int32 NormalArg = Normal.GetTracedInput().Expression ? Normal.Compile(Compiler) : Compiler->PixelNormalWS();
 	int32 DotArg = Compiler->Dot(NormalArg,Compiler->CameraVector());
 	int32 MaxArg = Compiler->Max(Compiler->Constant(0.f),DotArg);
 	int32 MinusArg = Compiler->Sub(Compiler->Constant(1.f),MaxArg);
-	int32 ExponentArg = ExponentIn.Expression ? ExponentIn.Compile(Compiler) : Compiler->Constant(Exponent);
+	int32 ExponentArg = ExponentIn.GetTracedInput().Expression ? ExponentIn.Compile(Compiler) : Compiler->Constant(Exponent);
 	int32 PowArg = Compiler->Power(MinusArg,ExponentArg);
-	int32 BaseReflectFractionArg = BaseReflectFractionIn.Expression ? BaseReflectFractionIn.Compile(Compiler) : Compiler->Constant(BaseReflectFraction);
+	int32 BaseReflectFractionArg = BaseReflectFractionIn.GetTracedInput().Expression ? BaseReflectFractionIn.Compile(Compiler) : Compiler->Constant(BaseReflectFraction);
 	int32 ScaleArg = Compiler->Mul(PowArg, Compiler->Sub(Compiler->Constant(1.f), BaseReflectFractionArg));
 	
 	return Compiler->Add(ScaleArg, BaseReflectFractionArg);
@@ -7723,7 +7838,7 @@ UMaterialExpressionDeriveNormalZ::UMaterialExpressionDeriveNormalZ(const FObject
 #if WITH_EDITOR
 int32 UMaterialExpressionDeriveNormalZ::Compile(class FMaterialCompiler* Compiler, int32 OutputIndex)
 {
-	if(!InXY.Expression)
+	if(!InXY.GetTracedInput().Expression)
 	{
 		return Compiler->Errorf(TEXT("Missing input normal xy vector whose z should be derived."));
 	}
@@ -7772,7 +7887,7 @@ UMaterialExpressionConstantBiasScale::UMaterialExpressionConstantBiasScale(const
 #if WITH_EDITOR
 int32 UMaterialExpressionConstantBiasScale::Compile(class FMaterialCompiler* Compiler, int32 OutputIndex)
 {
-	if (!Input.Expression)
+	if (!Input.GetTracedInput().Expression)
 	{
 		return Compiler->Errorf(TEXT("Missing ConstantBiasScale input"));
 	}
@@ -7833,7 +7948,7 @@ int32 UMaterialExpressionCustom::Compile(class FMaterialCompiler* Compiler, int3
 		}
 		else
 		{
-			if(!Inputs[i].Input.Expression)
+			if(!Inputs[i].Input.GetTracedInput().Expression)
 			{
 				return Compiler->Errorf(TEXT("Custom material %s missing input %d (%s)"), *Description, i+1, *Inputs[i].InputName);
 			}
@@ -8421,7 +8536,7 @@ static int32 FindOutputIndexByName(const FString& Name, const TArray<FFunctionEx
 int32 UMaterialFunction::Compile(class FMaterialCompiler* Compiler, const FFunctionExpressionOutput& Output)
 {
 	int32 ReturnValue = INDEX_NONE;
-	if (Output.ExpressionOutput->A.Expression)
+	if (Output.ExpressionOutput->A.GetTracedInput().Expression)
 	{
 		// Compile the given function output
 		ReturnValue = Output.ExpressionOutput->A.Compile(Compiler);
@@ -9250,7 +9365,7 @@ void UMaterialExpressionFunctionInput::GetExpressionToolTip(TArray<FString>& Out
 
 int32 UMaterialExpressionFunctionInput::CompilePreviewValue(FMaterialCompiler* Compiler)
 {
-	if (Preview.Expression)
+	if (Preview.GetTracedInput().Expression)
 	{
 		return Preview.Compile(Compiler);
 	}
@@ -9294,13 +9409,14 @@ int32 UMaterialExpressionFunctionInput::Compile(class FMaterialCompiler* Compile
 	check(InputType < FunctionInput_MAX);
 
 	// If we are being compiled as part of a material which calls this function
-	if (EffectivePreviewDuringCompile.Expression && !bCompilingFunctionPreview)
+	FExpressionInput EffectivePreviewDuringCompileTracedInput = EffectivePreviewDuringCompile.GetTracedInput();
+	if (EffectivePreviewDuringCompileTracedInput.Expression && !bCompilingFunctionPreview)
 	{
 		int32 ExpressionResult;
 
 		// Stay in this function if we are compiling an expression that is in the current function
 		// This can happen if bUsePreviewValueAsDefault is true and the calling material didn't override the input
-		if (bUsePreviewValueAsDefault && EffectivePreviewDuringCompile.Expression->GetOuter() == GetOuter())
+		if (bUsePreviewValueAsDefault && EffectivePreviewDuringCompileTracedInput.Expression->GetOuter() == GetOuter())
 		{
 			// Compile the function input
 			ExpressionResult = EffectivePreviewDuringCompile.Compile(Compiler);
@@ -9545,7 +9661,7 @@ uint32 UMaterialExpressionFunctionOutput::GetInputType(int32 InputIndex)
 
 int32 UMaterialExpressionFunctionOutput::Compile(class FMaterialCompiler* Compiler, int32 OutputIndex)
 {
-	if (!A.Expression)
+	if (!A.GetTracedInput().Expression)
 	{
 		return Compiler->Errorf(TEXT("Missing function output '%s'"), *OutputName);
 	}
@@ -9601,7 +9717,7 @@ void UMaterialExpressionFunctionOutput::ValidateName()
 bool UMaterialExpressionFunctionOutput::IsResultMaterialAttributes(int32 OutputIndex)
 {
 	// If there is a loop anywhere in this expression's inputs then we can't risk checking them
-	if( A.Expression && !A.Expression->ContainsInputLoop() )
+	if( A.GetTracedInput().Expression && !A.Expression->ContainsInputLoop() )
 	{
 		return A.Expression->IsResultMaterialAttributes(A.OutputIndex);
 	}
@@ -9854,11 +9970,11 @@ UMaterialExpressionLightmassReplace::UMaterialExpressionLightmassReplace(const F
 #if WITH_EDITOR
 int32 UMaterialExpressionLightmassReplace::Compile(class FMaterialCompiler* Compiler, int32 OutputIndex)
 {
-	if (!Realtime.Expression)
+	if (!Realtime.GetTracedInput().Expression)
 	{
 		return Compiler->Errorf(TEXT("Missing LightmassReplace input Realtime"));
 	}
-	else if (!Lightmass.Expression)
+	else if (!Lightmass.GetTracedInput().Expression)
 	{
 		return Compiler->Errorf(TEXT("Missing LightmassReplace input Lightmass"));
 	}
@@ -9902,11 +10018,11 @@ UMaterialExpressionMaterialProxyReplace::UMaterialExpressionMaterialProxyReplace
 #if WITH_EDITOR
 int32 UMaterialExpressionMaterialProxyReplace::Compile(class FMaterialCompiler* Compiler, int32 OutputIndex)
 {
-	if (!Realtime.Expression)
+	if (!Realtime.GetTracedInput().Expression)
 	{
 		return Compiler->Errorf(TEXT("Missing MaterialProxyReplace input Realtime"));
 	}
-	else if (!MaterialProxy.Expression)
+	else if (!MaterialProxy.GetTracedInput().Expression)
 	{
 		return Compiler->Errorf(TEXT("Missing MaterialProxyReplace input MaterialProxy"));
 	}
@@ -9949,10 +10065,10 @@ UMaterialExpressionGIReplace::UMaterialExpressionGIReplace(const FObjectInitiali
 #if WITH_EDITOR
 int32 UMaterialExpressionGIReplace::Compile(class FMaterialCompiler* Compiler, int32 OutputIndex)
 {
-	FExpressionInput& LocalStaticIndirect = StaticIndirect.Expression ? StaticIndirect : Default;
-	FExpressionInput& LocalDynamicIndirect = DynamicIndirect.Expression ? DynamicIndirect : Default;
+	FExpressionInput& LocalStaticIndirect = StaticIndirect.GetTracedInput().Expression ? StaticIndirect : Default;
+	FExpressionInput& LocalDynamicIndirect = DynamicIndirect.GetTracedInput().Expression ? DynamicIndirect : Default;
 
-	if(!Default.Expression)
+	if(!Default.GetTracedInput().Expression)
 	{
 		return Compiler->Errorf(TEXT("Missing GIReplace input 'Default'"));
 	}
@@ -10012,6 +10128,149 @@ void UMaterialExpressionObjectOrientation::GetCaption(TArray<FString>& OutCaptio
 }
 #endif // WITH_EDITOR
 
+UMaterialExpressionReroute::UMaterialExpressionReroute(const FObjectInitializer& ObjectInitializer)
+	: Super(ObjectInitializer)
+{
+	// Structure to hold one-time initialization
+	struct FConstructorStatics
+	{
+		FText NAME_Utility;
+		FConstructorStatics()
+			: NAME_Utility(LOCTEXT("Utility", "Utility"))
+		{
+		}
+	};
+	static FConstructorStatics ConstructorStatics;
+
+#if WITH_EDITORONLY_DATA
+	MenuCategories.Add(ConstructorStatics.NAME_Utility);
+#endif
+}
+
+
+UMaterialExpression* UMaterialExpressionReroute::TraceInputsToRealExpression(int32& OutputIndex) const
+{
+#if WITH_EDITORONLY_DATA
+	TSet<FMaterialExpressionKey> VisitedExpressions;
+	FExpressionInput RealInput = TraceInputsToRealExpressionInternal(VisitedExpressions);
+	OutputIndex = RealInput.OutputIndex;
+	return RealInput.Expression;
+#else
+	OutputIndex = 0;
+	return nullptr;
+#endif
+}
+
+FExpressionInput UMaterialExpressionReroute::TraceInputsToRealInput() const
+{
+	TSet<FMaterialExpressionKey> VisitedExpressions;
+	FExpressionInput RealInput = TraceInputsToRealExpressionInternal(VisitedExpressions);
+	return RealInput;
+}
+
+FExpressionInput UMaterialExpressionReroute::TraceInputsToRealExpressionInternal(TSet<FMaterialExpressionKey>& VisitedExpressions) const
+{
+#if WITH_EDITORONLY_DATA
+	// First check to see if this is a terminal node, if it is then we have a reroute to nowhere.
+	if (Input.Expression != nullptr)
+	{
+		// Now check to see if we're also connected to another reroute. If we are, then keep going unless we hit a loop condition.
+		UMaterialExpressionReroute* RerouteInput = Cast<UMaterialExpressionReroute>(Input.Expression);
+		if (RerouteInput != nullptr)
+		{
+			FMaterialExpressionKey InputExpressionKey(Input.Expression, Input.OutputIndex);
+			// prevent recurring visits to expressions we've already checked
+			if (VisitedExpressions.Contains(InputExpressionKey))
+			{
+				// We have a loop! This should result in not finding the value!
+				return FExpressionInput();
+			}
+			else 
+			{
+				VisitedExpressions.Add(InputExpressionKey);
+				FExpressionInput OutputExpressionInput = RerouteInput->TraceInputsToRealExpressionInternal(VisitedExpressions);
+				return OutputExpressionInput;
+			}
+		}
+		else
+		{
+			// We aren't connected to another Reroute, so we are good.
+			return Input;
+		}
+	}
+#endif // WITH_EDITORONLY_DATA
+	// We went to nowhere, so bail out.
+	return FExpressionInput();
+}
+
+#if WITH_EDITOR
+int32 UMaterialExpressionReroute::Compile(FMaterialCompiler* Compiler, int32 OutputIndex)
+{
+	// Because we don't want to generate *any* additional instructions, we just forward this request
+	// to the node that this input is connected to. If it isn't connected, then the compile will return INDEX_NONE.
+	int32 Result = Input.Compile(Compiler);
+	return Result;
+}
+
+void UMaterialExpressionReroute::GetCaption(TArray<FString>& OutCaptions) const
+{
+	OutCaptions.Add(TEXT("Reroute Node (reroutes wires)"));
+}
+
+
+FString UMaterialExpressionReroute::GetCreationDescription() const 
+{
+	return LOCTEXT("RerouteNodeCreationDesc", "This node looks like a single pin and can be used to tidy up your graph by adding a movable control point to the connection spline.").ToString();
+}
+
+FString UMaterialExpressionReroute::GetCreationName() const
+{
+	return LOCTEXT("RerouteNodeCreationName", "Add Reroute Node...").ToString();
+}
+
+
+uint32 UMaterialExpressionReroute::GetInputType(int32 InputIndex)
+{
+	// Our input type should match the node that we are ultimately connected to, no matter how many reroute nodes lie between us.
+	if (InputIndex == 0 && Input.IsConnected() && Input.Expression != nullptr )
+	{
+		int32 RealExpressionOutputIndex = -1;
+		UMaterialExpression* RealExpression = TraceInputsToRealExpression(RealExpressionOutputIndex);
+
+		// If we found a valid connection to a real output, then our type becomes that type.
+		if (RealExpression != nullptr && RealExpressionOutputIndex != -1 && RealExpression->Outputs.Num() > RealExpressionOutputIndex && RealExpressionOutputIndex >= 0)
+		{
+			return RealExpression->GetOutputType(RealExpressionOutputIndex);
+		}
+	}
+	return MCT_Unknown;
+}
+
+uint32 UMaterialExpressionReroute::GetOutputType(int32 OutputIndex)
+{
+	// Our node is a passthrough so input and output types must match.
+	return GetInputType(0);
+}
+
+bool UMaterialExpressionReroute::IsResultMaterialAttributes(int32 OutputIndex)
+{
+	// Most code checks to make sure that there aren't loops before going here. In our case, we rely on the fact that
+	// UMaterialExpressionReroute's implementation of TraceInputsToRealExpression is resistant to input loops.
+	if (Input.IsConnected() && Input.Expression != nullptr && OutputIndex == 0)
+	{
+		int32 RealExpressionOutputIndex = -1;
+		UMaterialExpression* RealExpression = TraceInputsToRealExpression(RealExpressionOutputIndex);
+		if (RealExpression != nullptr)
+		{
+			return RealExpression->IsResultMaterialAttributes(RealExpressionOutputIndex);
+		}
+	}
+
+	return false;
+}
+
+#endif // WITH_EDITOR
+
 UMaterialExpressionRotateAboutAxis::UMaterialExpressionRotateAboutAxis(const FObjectInitializer& ObjectInitializer)
 	: Super(ObjectInitializer)
 {
@@ -10036,19 +10295,19 @@ UMaterialExpressionRotateAboutAxis::UMaterialExpressionRotateAboutAxis(const FOb
 #if WITH_EDITOR
 int32 UMaterialExpressionRotateAboutAxis::Compile(class FMaterialCompiler* Compiler, int32 OutputIndex)
 {
-	if (!NormalizedRotationAxis.Expression)
+	if (!NormalizedRotationAxis.GetTracedInput().Expression)
 	{
 		return Compiler->Errorf(TEXT("Missing RotateAboutAxis input NormalizedRotationAxis"));
 	}
-	else if (!RotationAngle.Expression)
+	else if (!RotationAngle.GetTracedInput().Expression)
 	{
 		return Compiler->Errorf(TEXT("Missing RotateAboutAxis input RotationAngle"));
 	}
-	else if (!PivotPoint.Expression)
+	else if (!PivotPoint.GetTracedInput().Expression)
 	{
 		return Compiler->Errorf(TEXT("Missing RotateAboutAxis input PivotPoint"));
 	}
-	else if (!Position.Expression)
+	else if (!Position.GetTracedInput().Expression)
 	{
 		return Compiler->Errorf(TEXT("Missing RotateAboutAxis input Position"));
 	}
@@ -10125,11 +10384,11 @@ UMaterialExpressionSphereMask::UMaterialExpressionSphereMask(const FObjectInitia
 #if WITH_EDITOR
 int32 UMaterialExpressionSphereMask::Compile(class FMaterialCompiler* Compiler, int32 OutputIndex)
 {
-	if(!A.Expression)
+	if(!A.GetTracedInput().Expression)
 	{
 		return Compiler->Errorf(TEXT("Missing input A"));
 	}
-	else if(!B.Expression)
+	else if(!B.GetTracedInput().Expression)
 	{
 		return Compiler->Errorf(TEXT("Missing input B"));
 	}
@@ -10140,7 +10399,7 @@ int32 UMaterialExpressionSphereMask::Compile(class FMaterialCompiler* Compiler,
 		int32 Distance = CompileHelperLength(Compiler, Arg1, Arg2);
 
 		int32 ArgInvRadius;
-		if(Radius.Expression)
+		if(Radius.GetTracedInput().Expression)
 		{
 			// if the radius input is hooked up, use it
 			ArgInvRadius = Compiler->Div(Compiler->Constant(1.0f), Compiler->Max(Compiler->Constant(0.00001f), Radius.Compile(Compiler)));
@@ -10154,7 +10413,7 @@ int32 UMaterialExpressionSphereMask::Compile(class FMaterialCompiler* Compiler,
 		int32 NormalizeDistance = Compiler->Mul(Distance, ArgInvRadius);
 
 		int32 ArgInvHardness;
-		if(Hardness.Expression)
+		if(Hardness.GetTracedInput().Expression)
 		{
 			int32 Softness = Compiler->Sub(Compiler->Constant(1.0f), Hardness.Compile(Compiler));
 
@@ -10253,7 +10512,7 @@ int32 UMaterialExpressionNoise::Compile(class FMaterialCompiler* Compiler, int32
 {
 	int32 PositionInput;
 
-	if(Position.Expression)
+	if(Position.GetTracedInput().Expression)
 	{
 		PositionInput = Position.Compile(Compiler);
 	}
@@ -10264,7 +10523,7 @@ int32 UMaterialExpressionNoise::Compile(class FMaterialCompiler* Compiler, int32
 
 	int32 FilterWidthInput;
 
-	if(FilterWidth.Expression)
+	if(FilterWidth.GetTracedInput().Expression)
 	{
 		FilterWidthInput = FilterWidth.Compile(Compiler);
 	}
@@ -10337,7 +10596,7 @@ int32 UMaterialExpressionVectorNoise::Compile(class FMaterialCompiler* Compiler,
 {
 	int32 PositionInput;
 
-	if (Position.Expression)
+	if (Position.GetTracedInput().Expression)
 	{
 		PositionInput = Position.Compile(Compiler);
 	}
@@ -10382,7 +10641,7 @@ int32 UMaterialExpressionBlackBody::Compile(class FMaterialCompiler* Compiler, i
 {
 	int32 TempInput = INDEX_NONE;
 
-	if( Temp.Expression )
+	if( Temp.GetTracedInput().Expression )
 	{
 		TempInput = Temp.Compile(Compiler);
 	}
@@ -10428,7 +10687,7 @@ int32 UMaterialExpressionDistanceToNearestSurface::Compile(class FMaterialCompil
 {
 	int32 PositionArg = INDEX_NONE;
 
-	if (Position.Expression)
+	if (Position.GetTracedInput().Expression)
 	{
 		PositionArg = Position.Compile(Compiler);
 	}
@@ -10473,7 +10732,7 @@ int32 UMaterialExpressionDistanceFieldGradient::Compile(class FMaterialCompiler*
 {
 	int32 PositionArg = INDEX_NONE;
 
-	if (Position.Expression)
+	if (Position.GetTracedInput().Expression)
 	{
 		PositionArg = Position.Compile(Compiler);
 	}
@@ -10516,11 +10775,11 @@ UMaterialExpressionDistance::UMaterialExpressionDistance(const FObjectInitialize
 #if WITH_EDITOR
 int32 UMaterialExpressionDistance::Compile(class FMaterialCompiler* Compiler, int32 OutputIndex)
 {
-	if(!A.Expression)
+	if(!A.GetTracedInput().Expression)
 	{
 		return Compiler->Errorf(TEXT("Missing input A"));
 	}
-	else if(!B.Expression)
+	else if(!B.GetTracedInput().Expression)
 	{
 		return Compiler->Errorf(TEXT("Missing input B"));
 	}
@@ -10943,7 +11202,7 @@ int32 UMaterialExpressionDecalMipmapLevel::Compile(class FMaterialCompiler* Comp
 
 	int32 TextureSizeInput = INDEX_NONE;
 
-	if (TextureSize.Expression)
+	if (TextureSize.GetTracedInput().Expression)
 	{
 		TextureSizeInput = TextureSize.Compile(Compiler);
 	}
@@ -10998,8 +11257,8 @@ int32 UMaterialExpressionDepthFade::Compile(class FMaterialCompiler* Compiler, i
 {
 	// Scales Opacity by a Linear fade based on SceneDepth, from 0 at PixelDepth to 1 at FadeDistance
 	// Result = Opacity * saturate((SceneDepth - PixelDepth) / max(FadeDistance, DELTA))
-	const int32 OpacityIndex = InOpacity.Expression ? InOpacity.Compile(Compiler) : Compiler->Constant(OpacityDefault);
-	const int32 FadeDistanceIndex = Compiler->Max(FadeDistance.Expression ? FadeDistance.Compile(Compiler) : Compiler->Constant(FadeDistanceDefault), Compiler->Constant(DELTA));
+	const int32 OpacityIndex = InOpacity.GetTracedInput().Expression ? InOpacity.Compile(Compiler) : Compiler->Constant(OpacityDefault);
+	const int32 FadeDistanceIndex = Compiler->Max(FadeDistance.GetTracedInput().Expression ? FadeDistance.Compile(Compiler) : Compiler->Constant(FadeDistanceDefault), Compiler->Constant(DELTA));
 	const int32 FadeIndex = CompileHelperSaturate(Compiler, Compiler->Div(Compiler->Sub(Compiler->SceneDepth(INDEX_NONE, INDEX_NONE, false), Compiler->PixelDepth()), FadeDistanceIndex));
 	return Compiler->Mul(OpacityIndex, FadeIndex);
 }
@@ -11034,7 +11293,7 @@ UMaterialExpressionSphericalParticleOpacity::UMaterialExpressionSphericalParticl
 #if WITH_EDITOR
 int32 UMaterialExpressionSphericalParticleOpacity::Compile(class FMaterialCompiler* Compiler, int32 OutputIndex)
 {
-	const int32 DensityIndex = Density.Expression ? Density.Compile(Compiler) : Compiler->Constant(ConstantDensity);
+	const int32 DensityIndex = Density.GetTracedInput().Expression ? Density.Compile(Compiler) : Compiler->Constant(ConstantDensity);
 	return Compiler->SphericalParticleOpacity(DensityIndex);
 }
 #endif // WITH_EDITOR
@@ -11068,7 +11327,7 @@ int32 UMaterialExpressionDepthOfFieldFunction::Compile(class FMaterialCompiler*
 {
 	int32 DepthInput;
 
-	if(Depth.Expression)
+	if(Depth.GetTracedInput().Expression)
 	{
 		// using the input allows more custom behavior
 		DepthInput = Depth.Compile(Compiler);
@@ -11123,7 +11382,7 @@ int32 UMaterialExpressionDDX::Compile(class FMaterialCompiler* Compiler, int32 O
 {
 	int32 ValueInput = INDEX_NONE;
 
-	if(Value.Expression)
+	if(Value.GetTracedInput().Expression)
 	{
 		ValueInput = Value.Compile(Compiler);
 	}
@@ -11172,7 +11431,7 @@ int32 UMaterialExpressionDDY::Compile(class FMaterialCompiler* Compiler, int32 O
 {
 	int32 ValueInput = INDEX_NONE;
 
-	if(Value.Expression)
+	if(Value.GetTracedInput().Expression)
 	{
 		ValueInput = Value.Compile(Compiler);
 	}
@@ -11455,7 +11714,7 @@ int32 UMaterialExpressionAtmosphericFogColor::Compile(class FMaterialCompiler* C
 {
 	int32 WorldPositionInput = INDEX_NONE;
 
-	if( WorldPosition.Expression )
+	if( WorldPosition.GetTracedInput().Expression )
 	{
 		WorldPositionInput = WorldPosition.Compile(Compiler);
 	}
@@ -11635,7 +11894,7 @@ UMaterialExpressionTangentOutput::UMaterialExpressionTangentOutput(const FObject
 #if WITH_EDITOR
 int32 UMaterialExpressionTangentOutput::Compile(class FMaterialCompiler* Compiler, int32 OutputIndex)
 {
-	if( Input.Expression )
+	if( Input.GetTracedInput().Expression )
 	{
 		return Compiler->CustomOutput(this, OutputIndex, Input.Compile(Compiler));
 	}
@@ -11686,7 +11945,7 @@ UMaterialExpressionClearCoatNormalCustomOutput::UMaterialExpressionClearCoatNorm
 #if WITH_EDITOR
 int32  UMaterialExpressionClearCoatNormalCustomOutput::Compile(class FMaterialCompiler* Compiler, int32 OutputIndex)
 {
-	if (Input.Expression)
+	if (Input.GetTracedInput().Expression)
 	{
 		return Compiler->CustomOutput(this, OutputIndex, Input.Compile(Compiler));
 	}
diff --git a/Engine/Source/Runtime/Engine/Private/Materials/MaterialShared.cpp b/Engine/Source/Runtime/Engine/Private/Materials/MaterialShared.cpp
--- a/Engine/Source/Runtime/Engine/Private/Materials/MaterialShared.cpp
+++ b/Engine/Source/Runtime/Engine/Private/Materials/MaterialShared.cpp
@@ -18,6 +18,7 @@
 #include "UObject/UObjectIterator.h"
 #include "ComponentReregisterContext.h"
 #include "Materials/MaterialExpressionBreakMaterialAttributes.h"
+#include "Materials/MaterialExpressionReroute.h"
 #include "ShaderCompiler.h"
 #include "MaterialCompiler.h"
 #include "MeshMaterialShaderType.h"
@@ -118,6 +119,18 @@ void FExpressionInput::Connect( int32 InOutputIndex, class UMaterialExpression*
 }
 #endif // WITH_EDITOR
 
+FExpressionInput FExpressionInput::GetTracedInput() const
+{
+#if WITH_EDITORONLY_DATA
+	if (Expression != nullptr && Expression->IsA(UMaterialExpressionReroute::StaticClass()))
+	{
+		UMaterialExpressionReroute* Reroute = CastChecked<UMaterialExpressionReroute>(Expression);
+		return Reroute->TraceInputsToRealInput();
+	}
+#endif
+	return *this;
+}
+
 /** Native serialize for FMaterialExpression struct */
 static bool SerializeExpressionInput(FArchive& Ar, FExpressionInput& Input)
 {
@@ -751,6 +764,14 @@ const TArray<UTexture*>& FMaterialResource::GetReferencedTextures() const
 	return Material->ExpressionTextureReferences;
 }
 
+void FMaterialResource::AddReferencedObjects(FReferenceCollector& Collector)
+{
+	FMaterial::AddReferencedObjects(Collector);
+
+	Collector.AddReferencedObject(Material);
+	Collector.AddReferencedObject(MaterialInstance);
+}
+
 bool FMaterialResource::GetAllowDevelopmentShaderCompile()const
 {
 	return Material->bAllowDevelopmentShaderCompile;
@@ -1579,7 +1600,7 @@ bool FMaterial::CacheShaders(const FMaterialShaderMapId& ShaderMapId, EShaderPla
 	{
 		//FMaterialShaderMap::ShaderMapsBeingCompiled.Find(GameThreadShaderMap);
 #if DEBUG_INFINITESHADERCOMPILE
-		UE_LOG(LogTemp, Display, TEXT("Found exisitng compiling shader for material %s, linking to other GameThreadShaderMap 0x%08X%08X"), *GetFriendlyName(), (int)((int64)(GameThreadShaderMap.GetReference()) >> 32), (int)((int64)(GameThreadShaderMap.GetReference())) );
+		UE_LOG(LogTemp, Display, TEXT("Found existing compiling shader for material %s, linking to other GameThreadShaderMap 0x%08X%08X"), *GetFriendlyName(), (int)((int64)(GameThreadShaderMap.GetReference()) >> 32), (int)((int64)(GameThreadShaderMap.GetReference())) );
 #endif
 		OutstandingCompileShaderMapIds.AddUnique(GameThreadShaderMap->GetCompilingId());
 		// Reset the shader map so the default material will be used until the compile finishes.
diff --git a/Engine/Source/Runtime/Engine/Private/PhysicsEngine/BodySetup.cpp b/Engine/Source/Runtime/Engine/Private/PhysicsEngine/BodySetup.cpp
--- a/Engine/Source/Runtime/Engine/Private/PhysicsEngine/BodySetup.cpp
+++ b/Engine/Source/Runtime/Engine/Private/PhysicsEngine/BodySetup.cpp
@@ -1164,6 +1164,29 @@ void UBodySetup::PostInitProperties()
 }
 
 #if WITH_EDITOR
+
+void UBodySetup::PostEditChangeProperty(struct FPropertyChangedEvent& PropertyChangedEvent)
+{
+	if(PropertyChangedEvent.MemberProperty && PropertyChangedEvent.MemberProperty->GetFName() == GET_MEMBER_NAME_CHECKED(UBodySetup, AggGeom))
+	{
+		UStaticMesh* StaticMesh = GetTypedOuter<UStaticMesh>();
+		if(StaticMesh)
+		{
+			for(UStaticMeshComponent* StaticMeshComponent : TObjectRange<UStaticMeshComponent>())
+			{
+				if(StaticMeshComponent->GetStaticMesh() == StaticMesh)
+				{
+					// it needs to recreate IF it already has been created
+					if(StaticMeshComponent->IsPhysicsStateCreated())
+					{
+						StaticMeshComponent->RecreatePhysicsState();
+					}
+				}
+			}
+		}
+	}
+}
+
 void UBodySetup::PostEditUndo()
 {
 	Super::PostEditUndo();
diff --git a/Engine/Source/Runtime/Engine/Private/Slate/SGameLayerManager.cpp b/Engine/Source/Runtime/Engine/Private/Slate/SGameLayerManager.cpp
--- a/Engine/Source/Runtime/Engine/Private/Slate/SGameLayerManager.cpp
+++ b/Engine/Source/Runtime/Engine/Private/Slate/SGameLayerManager.cpp
@@ -190,7 +190,15 @@ float SGameLayerManager::GetGameViewportDPIScale() const
 	if ( const FSceneViewport* Viewport = SceneViewport.Get() )
 	{
 		FIntPoint ViewportSize = Viewport->GetSize();
-		return GetDefault<UUserInterfaceSettings>(UUserInterfaceSettings::StaticClass())->GetDPIScaleBasedOnSize(ViewportSize);
+		const float GameUIScale = GetDefault<UUserInterfaceSettings>(UUserInterfaceSettings::StaticClass())->GetDPIScaleBasedOnSize(ViewportSize);
+
+		// Remove the platform DPI scale from the incoming size.  Since the platform DPI is already
+		// attempt to normalize the UI for a high DPI, and the DPI scale curve is based on raw resolution
+		// for what a assumed platform scale of 1, extract that scale the calculated scale, since that will
+		// already be applied by slate.
+		const float FinalUIScale = GameUIScale / Viewport->GetCachedGeometry().Scale;
+
+		return FinalUIScale;
 	}
 
 	return 1;
diff --git a/Engine/Source/Runtime/Engine/Private/Slate/SceneViewport.cpp b/Engine/Source/Runtime/Engine/Private/Slate/SceneViewport.cpp
--- a/Engine/Source/Runtime/Engine/Private/Slate/SceneViewport.cpp
+++ b/Engine/Source/Runtime/Engine/Private/Slate/SceneViewport.cpp
@@ -115,8 +115,8 @@ void FSceneViewport::ShowCursor( bool bVisible )
 	{
 		if( bIsSoftwareCursorVisible )
 		{
-			const int32 ClampedMouseX = FMath::Clamp<int32>(SoftwareCursorPosition.X, 0, SizeX);
-			const int32 ClampedMouseY = FMath::Clamp<int32>(SoftwareCursorPosition.Y, 0, SizeY);
+			const int32 ClampedMouseX = FMath::Clamp<int32>(SoftwareCursorPosition.X / CachedGeometry.Scale, 0, SizeX);
+			const int32 ClampedMouseY = FMath::Clamp<int32>(SoftwareCursorPosition.Y / CachedGeometry.Scale, 0, SizeY);
 
 			CurrentReplyState.SetMousePos( CachedGeometry.LocalToAbsolute( FVector2D(ClampedMouseX, ClampedMouseY) ).IntPoint() );
 		}
@@ -307,6 +307,15 @@ void FSceneViewport::OnDrawViewport( const FGeometry& AllottedGeometry, const FS
 	// Switch to the viewport clients world before resizing
 	FScopedConditionalWorldSwitcher WorldSwitcher( ViewportClient );
 
+	// In order to get material parameter collections to function properly, we need the current world's Scene
+	// properly propagated through to any widgets that depend on that functionality. The SceneViewport and RetainerWidget are the 
+	// only locations where this information exists in Slate, so we push the current scene onto the current
+	// Slate application so that we can leverage it in later calls.
+	if (ViewportClient && ViewportClient->GetWorld() && ViewportClient->GetWorld()->Scene)
+	{
+		FSlateApplication::Get().GetRenderer()->RegisterCurrentScene(ViewportClient->GetWorld()->Scene);
+	}
+	
 	/** Check to see if the viewport should be resized */
 	FIntPoint DrawSize = FIntPoint( FMath::RoundToInt( AllottedGeometry.GetDrawSize().X ), FMath::RoundToInt( AllottedGeometry.GetDrawSize().Y ) );
 	if( GetSizeXY() != DrawSize )
diff --git a/Engine/Source/Runtime/Engine/Private/StaticMesh.cpp b/Engine/Source/Runtime/Engine/Private/StaticMesh.cpp
--- a/Engine/Source/Runtime/Engine/Private/StaticMesh.cpp
+++ b/Engine/Source/Runtime/Engine/Private/StaticMesh.cpp
@@ -2443,7 +2443,11 @@ void UStaticMesh::Serialize(FArchive& Ar)
 			FStaticMeshSourceModel& SrcModel = SourceModels[i];
 			SrcModel.SerializeBulkData(Ar, this);
 		}
-		SectionInfoMap.Serialize(Ar);
+
+		if (Ar.CustomVer(FEditorObjectVersion::GUID) < FEditorObjectVersion::UPropertryForMeshSection)
+		{
+			SectionInfoMap.Serialize(Ar);
+		}
 
 		// Need to set a flag rather than do conversion in place as RenderData is not
 		// created until postload and it is needed for bounding information
diff --git a/Engine/Source/Runtime/Engine/Private/Tests/AutomationCommon.cpp b/Engine/Source/Runtime/Engine/Private/Tests/AutomationCommon.cpp
--- a/Engine/Source/Runtime/Engine/Private/Tests/AutomationCommon.cpp
+++ b/Engine/Source/Runtime/Engine/Private/Tests/AutomationCommon.cpp
@@ -90,8 +90,7 @@ namespace AutomationCommon
 
 		FPaths::MakePathRelativeTo(PathName, *FPaths::RootDir());
 
-		FString AdapterName = GRHIAdapterName.IsEmpty() ? FString(TEXT("Unknown")) : GRHIAdapterName;
-		OutScreenshotName = FString::Printf(TEXT("%s/%s.png"), *PathName, *AdapterName);
+		OutScreenshotName = FString::Printf(TEXT("%s/%s.png"), *PathName, *FPlatformMisc::GetDeviceId());
 	}
 
 	FAutomationScreenshotData BuildScreenshotData(const FString& MapOrContext, const FString& TestName, int32 Width, int32 Height)
@@ -102,6 +101,7 @@ namespace AutomationCommon
 
 		Data.Name = TestName;
 		Data.Context = MapOrContext;
+		Data.Id = FGuid::NewGuid();
 
 		Data.Width = Width;
 		Data.Height = Height;
@@ -125,7 +125,7 @@ namespace AutomationCommon
 		Data.TextureQuality = QualityLevels.TextureQuality;
 		Data.EffectsQuality = QualityLevels.EffectsQuality;
 		Data.FoliageQuality = QualityLevels.FoliageQuality;
-
+		
 		//GRHIDeviceId
 
 		// TBD - 
diff --git a/Engine/Source/Runtime/Engine/Private/UserDefinedEnum.cpp b/Engine/Source/Runtime/Engine/Private/UserDefinedEnum.cpp
--- a/Engine/Source/Runtime/Engine/Private/UserDefinedEnum.cpp
+++ b/Engine/Source/Runtime/Engine/Private/UserDefinedEnum.cpp
@@ -1,6 +1,7 @@
 // Copyright 1998-2017 Epic Games, Inc. All Rights Reserved.
 
 #include "Engine/UserDefinedEnum.h"
+#include "EditorObjectVersion.h"
 #if WITH_EDITOR
 #include "Kismet2/EnumEditorUtils.h"
 #endif	// WITH_EDITOR
@@ -15,6 +16,8 @@ void UUserDefinedEnum::Serialize(FArchive& Ar)
 {
 	Super::Serialize(Ar);
 
+	Ar.UsingCustomVersion(FEditorObjectVersion::GUID);
+
 #if WITH_EDITOR
 	if (Ar.IsLoading() && Ar.IsPersistent())
 	{
@@ -22,6 +25,11 @@ void UUserDefinedEnum::Serialize(FArchive& Ar)
 		{
 			Names[i].Value = i;
 		}
+
+		if (Ar.CustomVer(FEditorObjectVersion::GUID) < FEditorObjectVersion::StableUserDefinedEnumDisplayNames)
+		{
+			FEnumEditorUtils::UpgradeDisplayNamesFromMetaData(this);
+		}
 	}
 #endif // WITH_EDITOR
 }
@@ -67,11 +75,7 @@ void UUserDefinedEnum::PostLoad()
 {
 	Super::PostLoad();
 	FEnumEditorUtils::UpdateAfterPathChanged(this);
-
-	if (GIsEditor && !IsRunningCommandlet())
-	{
-		FEnumEditorUtils::EnsureAllDisplayNamesExist(this);
-	}
+	FEnumEditorUtils::EnsureAllDisplayNamesExist(this);
 }
 
 void UUserDefinedEnum::PostEditUndo()
@@ -107,12 +111,18 @@ int64 UUserDefinedEnum::ResolveEnumerator(FArchive& Ar, int64 EnumeratorValue) c
 
 FText UUserDefinedEnum::GetEnumText(int32 InIndex) const
 {
-	if (DisplayNames.IsValidIndex(InIndex))
+	const FName EnumEntryName = *GetEnumName(InIndex);
+
+	if (!EnumEntryName.IsNone())
 	{
-		return DisplayNames[InIndex];
+		const FText* EnumEntryDisplayName = DisplayNameMap.Find(EnumEntryName);
+		if (EnumEntryDisplayName)
+		{
+			return *EnumEntryDisplayName;
+		}
 	}
 
-	return Super::GetEnumText(InIndex);
+	return FText::GetEmpty();
 }
 
 bool UUserDefinedEnum::SetEnums(TArray<TPair<FName, int64>>& InNames, ECppForm InCppForm, bool bAddMaxKeyIfMissing)
diff --git a/Engine/Source/Runtime/Engine/Private/UserDefinedStruct.cpp b/Engine/Source/Runtime/Engine/Private/UserDefinedStruct.cpp
--- a/Engine/Source/Runtime/Engine/Private/UserDefinedStruct.cpp
+++ b/Engine/Source/Runtime/Engine/Private/UserDefinedStruct.cpp
@@ -29,11 +29,7 @@ namespace
 		FStructOnScope StructData(UserDefinedStruct);
 		FStructureEditorUtils::Fill_MakeStructureDefaultValue(UserDefinedStruct, StructData.GetStructMemory());
 
-		// Iterate over all fields of the object's class.
-		for (TFieldIterator<const UProperty> PropIt(StructData.GetStruct(), EFieldIteratorFlags::IncludeSuper, EFieldIteratorFlags::ExcludeDeprecated, EFieldIteratorFlags::IncludeInterfaces); PropIt; ++PropIt)
-		{
-			PropertyLocalizationDataGatherer.GatherLocalizationDataFromChildTextProperties(PathToObject, *PropIt, PropIt->ContainerPtrToValuePtr<void>(StructData.GetStructMemory()), GatherTextFlags);
-		}
+		PropertyLocalizationDataGatherer.GatherLocalizationDataFromStructFields(PathToObject, StructData.GetStruct(), StructData.GetStructMemory(), nullptr, GatherTextFlags);
 	}
 
 	void CollectUserDefinedStructTextReferences(UObject* Object, FArchive& Ar)
diff --git a/Engine/Source/Runtime/Engine/Private/UserInterface/Canvas.cpp b/Engine/Source/Runtime/Engine/Private/UserInterface/Canvas.cpp
--- a/Engine/Source/Runtime/Engine/Private/UserInterface/Canvas.cpp
+++ b/Engine/Source/Runtime/Engine/Private/UserInterface/Canvas.cpp
@@ -830,7 +830,7 @@ void FCanvas::PushRelativeTransform(const FMatrix& Transform)
 
 void FCanvas::PushAbsoluteTransform(const FMatrix& Transform) 
 {
-	if(ensure(TransformStack.Num()>0))
+	if(TransformStack.Num()>0)
 	{
 		TransformStack.Add(FTransformEntry(Transform * TransformStack[0].GetMatrix()));
 	}
diff --git a/Engine/Source/Runtime/Engine/Public/ComponentReregisterContext.h b/Engine/Source/Runtime/Engine/Public/ComponentReregisterContext.h
--- a/Engine/Source/Runtime/Engine/Public/ComponentReregisterContext.h
+++ b/Engine/Source/Runtime/Engine/Public/ComponentReregisterContext.h
@@ -92,24 +92,24 @@ class FComponentReregisterContext : public FComponentReregisterContextBase
 {
 private:
 	/** Pointer to component we are unregistering */
-	UActorComponent* Component;
+	TWeakObjectPtr<UActorComponent> Component;
 	/** Cache pointer to world from which we were removed */
-	UWorld* World;
+	TWeakObjectPtr<UWorld> World;
 public:
 	FComponentReregisterContext(UActorComponent* InComponent)
-		: World(NULL)
+		: World(nullptr)
 	{
 		World = UnRegister(InComponent);
 		// If we didn't get a scene back NULL the component so we dont try to
 		// process it on destruction
-		Component = World ? InComponent : NULL;
+		Component = World.IsValid() ? InComponent : nullptr;
 	}
 
 	~FComponentReregisterContext()
 	{
-		if( Component != NULL )
+		if( Component.IsValid() )
 		{
-			ReRegister(Component, World);
+			ReRegister(Component.Get(), World.Get());
 		}
 	}
 };
diff --git a/Engine/Source/Runtime/Engine/Public/MaterialExpressionIO.h b/Engine/Source/Runtime/Engine/Public/MaterialExpressionIO.h
--- a/Engine/Source/Runtime/Engine/Public/MaterialExpressionIO.h
+++ b/Engine/Source/Runtime/Engine/Public/MaterialExpressionIO.h
@@ -14,11 +14,17 @@
 struct FExpressionInput
 {
 #if WITH_EDITORONLY_DATA
-	/** Material expression that this input is connected to, or NULL if not connected. */
+	/** 
+	 * Material expression that this input is connected to, or NULL if not connected. 
+	 * If you want to be safe when checking against dangling Reroute nodes, please use GetCheckedInput before accessing this property.
+	*/
 	class UMaterialExpression*	Expression;
 #endif
 
-	/** Index into Expression's outputs array that this input is connected to. */
+	/** 
+	 * Index into Expression's outputs array that this input is connected to.
+	 * If you want to be safe when checking against dangling Reroute nodes, please use GetCheckedInput before accessing this property.
+	*/
 	int32						OutputIndex;
 
 	/** 
@@ -74,6 +80,9 @@ struct FExpressionInput
 
 	/** ICPPStructOps interface */
 	ENGINE_API bool Serialize(FArchive& Ar);
+
+	/** If this input goes through reroute nodes or other paths that should not affect code, trace back on the input chain.*/
+	ENGINE_API FExpressionInput GetTracedInput() const;
 };
 
 template<>
diff --git a/Engine/Source/Runtime/Engine/Public/MaterialShared.h b/Engine/Source/Runtime/Engine/Public/MaterialShared.h
--- a/Engine/Source/Runtime/Engine/Public/MaterialShared.h
+++ b/Engine/Source/Runtime/Engine/Public/MaterialShared.h
@@ -1155,7 +1155,7 @@ public:
 		OutstandingCompileShaderMapIds.Remove( OldOutstandingCompileShaderMapId );
 	}
 
-	void AddReferencedObjects(FReferenceCollector& Collector);
+	ENGINE_API virtual void AddReferencedObjects(FReferenceCollector& Collector);
 
 	virtual const TArray<UTexture*>& GetReferencedTextures() const = 0;
 
@@ -1707,6 +1707,8 @@ public:
 
 	ENGINE_API virtual const TArray<UTexture*>& GetReferencedTextures() const override;
 
+	ENGINE_API virtual void AddReferencedObjects(FReferenceCollector& Collector) override;
+
 	ENGINE_API virtual bool GetAllowDevelopmentShaderCompile() const override;
 
 protected:
diff --git a/Engine/Source/Runtime/Engine/Public/ShowFlags.h b/Engine/Source/Runtime/Engine/Public/ShowFlags.h
--- a/Engine/Source/Runtime/Engine/Public/ShowFlags.h
+++ b/Engine/Source/Runtime/Engine/Public/ShowFlags.h
@@ -386,6 +386,7 @@ private:
 		SetMaterialTextureScaleAccuracy(false);
 		SetOutputMaterialTextureScales(false);
 		SetRequiredTextureResolution(false);
+		SetMotionBlur(InitMode != ESFIM_Editor);
 	}
 
 
diff --git a/Engine/Source/Runtime/Engine/Public/Slate/SceneViewport.h b/Engine/Source/Runtime/Engine/Public/Slate/SceneViewport.h
--- a/Engine/Source/Runtime/Engine/Public/Slate/SceneViewport.h
+++ b/Engine/Source/Runtime/Engine/Public/Slate/SceneViewport.h
@@ -253,6 +253,9 @@ public:
 
 	virtual FSlateShaderResource* GetViewportRenderTargetTexture() override;
 
+	/** Get the cached viewport geometry. */
+	const FGeometry& GetCachedGeometry() const { return CachedGeometry; }
+
 private:
 	/**
 	 * Called when this viewport is destroyed
diff --git a/Engine/Source/Runtime/Internationalization/Private/LocTextHelper.cpp b/Engine/Source/Runtime/Internationalization/Private/LocTextHelper.cpp
--- a/Engine/Source/Runtime/Internationalization/Private/LocTextHelper.cpp
+++ b/Engine/Source/Runtime/Internationalization/Private/LocTextHelper.cpp
@@ -122,6 +122,16 @@ FLocTextHelper::FLocTextHelper(FString InTargetPath, FString InManifestName, FSt
 	}
 }
 
+const FString& FLocTextHelper::GetNativeCulture() const
+{
+	return NativeCulture;
+}
+
+const TArray<FString>& FLocTextHelper::GetForeignCultures() const
+{
+	return ForeignCultures;
+}
+
 bool FLocTextHelper::HasManifest() const
 {
 	return Manifest.IsValid();
@@ -720,6 +730,11 @@ void FLocTextHelper::GetRuntimeText(const FString& InCulture, const FString& InN
 	}
 }
 
+void FLocTextHelper::AddConflict(const FString& InNamespace, const FString& InKey, const TSharedPtr<FLocMetadataObject>& InKeyMetadata, const FLocItem& InSource, const FString& InSourceLocation)
+{
+	ConflictTracker.AddConflict(InNamespace, InKey, InKeyMetadata, InSource, InSourceLocation);
+}
+
 FString FLocTextHelper::GetConflictReport() const
 {
 	return ConflictTracker.GetConflictReport();
@@ -848,7 +863,7 @@ bool FLocTextHelper::FindKeysForLegacyTranslation(const TSharedRef<const FIntern
 			}
 			else if ((!Context.KeyMetadataObj.IsValid() && !InKeyMetadataObj.IsValid()) || (*Context.KeyMetadataObj == *InKeyMetadataObj))
 			{
-				OutKeys.Add(Context.Key);
+				OutKeys.AddUnique(Context.Key);
 				bFoundKeys = true;
 			}
 		}
diff --git a/Engine/Source/Runtime/Internationalization/Public/LocTextHelper.h b/Engine/Source/Runtime/Internationalization/Public/LocTextHelper.h
--- a/Engine/Source/Runtime/Internationalization/Public/LocTextHelper.h
+++ b/Engine/Source/Runtime/Internationalization/Public/LocTextHelper.h
@@ -165,6 +165,16 @@ public:
 	 */
 	FLocTextHelper(FString InTargetPath, FString InManifestName, FString InArchiveName, FString InNativeCulture, TArray<FString> InForeignCultures, TSharedPtr<ILocFileNotifies> InLocFileNotifies);
 
+	/**
+	 * @return Get the culture code of the native culture (eg, en), or an empty string if the native culture is unknown.
+	 */
+	const FString& GetNativeCulture() const;
+
+	/**
+	 * @return Get an array of culture codes for the foreign cultures (does not include the native culture).
+	 */
+	const TArray<FString>& GetForeignCultures() const;
+
 	/**
 	 * Check to see whether we've loaded the manifest.
 	 */
@@ -561,6 +571,17 @@ public:
 	 */
 	void GetRuntimeText(const FString& InCulture, const FString& InNamespace, const FString& InKey, const TSharedPtr<FLocMetadataObject> InKeyMetadataObj, const ELocTextExportSourceMethod InSourceMethod, const FLocItem& InSource, FLocItem& OutTranslation, const bool bSkipSourceCheck) const;
 
+	/**
+	 * Add a new conflict entry.
+	 *
+	 * @param InNamespace			The namespace of the entry.
+	 * @param InKey					The key/identifier of the entry.
+	 * @param InKeyMetadata			Entry Metadata keys.
+	 * @param InSource				The source info for the conflict.
+	 * @param InSourceLocation		The source location of the conflict.
+	 */
+	void AddConflict(const FString& InNamespace, const FString& InKey, const TSharedPtr<FLocMetadataObject>& InKeyMetadata, const FLocItem& InSource, const FString& InSourceLocation);
+
 	/**
 	 * Get a conflict report that can be easily saved as a report summary.
 	 */
diff --git a/Engine/Source/Runtime/JsonUtilities/Private/JsonObjectConverter.cpp b/Engine/Source/Runtime/JsonUtilities/Private/JsonObjectConverter.cpp
--- a/Engine/Source/Runtime/JsonUtilities/Private/JsonObjectConverter.cpp
+++ b/Engine/Source/Runtime/JsonUtilities/Private/JsonObjectConverter.cpp
@@ -149,8 +149,8 @@ bool FJsonObjectConverter::UStructToJsonAttributes(const UStruct* StructDefiniti
 {
 	if (SkipFlags == 0)
 	{
-		// If we have no specified skip flags, skip deprecated by default when writing
-		SkipFlags |= CPF_Deprecated;
+		// If we have no specified skip flags, skip deprecated, transient and skip serialization by default when writing
+		SkipFlags |= CPF_Deprecated | CPF_Transient;
 	}
 
 	if (StructDefinition == FJsonObjectWrapper::StaticStruct())
