Fixed IWYU ComponentReregisterContext compile error on linux

diff --git a/Engine/Source/Runtime/Engine/Public/ComponentReregisterContext.h b/Engine/Source/Runtime/Engine/Public/ComponentReregisterContext.h
--- a/Engine/Source/Runtime/Engine/Public/ComponentReregisterContext.h
+++ b/Engine/Source/Runtime/Engine/Public/ComponentReregisterContext.h
@@ -7,6 +7,7 @@
 #include "RenderingThread.h"
 #include "UObject/UObjectIterator.h"
 #include "Components/PrimitiveComponent.h"
+#include "Engine/World.h"
 #include "AI/Navigation/NavigationSystem.h"
 #include "ContentStreaming.h"
 
