Fix non-unity/non-pch editor builds, and use that configuration for non-unity builds in future. Catches more include errors.

diff --git a/Engine/Source/Developer/AutomationDriver/Private/AutomatedApplication.cpp b/Engine/Source/Developer/AutomationDriver/Private/AutomatedApplication.cpp
--- a/Engine/Source/Developer/AutomationDriver/Private/AutomatedApplication.cpp
+++ b/Engine/Source/Developer/AutomationDriver/Private/AutomatedApplication.cpp
@@ -2,6 +2,8 @@
 
 #include "AutomatedApplication.h"
 #include "PassThroughMessageHandler.h"
+#include "Widgets/SWindow.h"
+#include "Framework/Application/SlateApplication.h"
 
 class FAutomatedCursor
 	: public ICursor
diff --git a/Engine/Source/Developer/AutomationDriver/Private/AutomatedApplication.h b/Engine/Source/Developer/AutomationDriver/Private/AutomatedApplication.h
--- a/Engine/Source/Developer/AutomationDriver/Private/AutomatedApplication.h
+++ b/Engine/Source/Developer/AutomationDriver/Private/AutomatedApplication.h
@@ -1,6 +1,10 @@
 // Copyright 1998-2017 Epic Games, Inc. All Rights Reserved.
 #pragma once
 
+#include "CoreMinimal.h"
+#include "GenericPlatform/GenericApplication.h"
+#include "GenericPlatform/ICursor.h"
+
 class GenericApplication;
 class FGenericApplicationMessageHandler;
 class IPassThroughMessageHandlerFactory;
diff --git a/Engine/Source/Developer/AutomationDriver/Private/AutomationDriver.cpp b/Engine/Source/Developer/AutomationDriver/Private/AutomationDriver.cpp
--- a/Engine/Source/Developer/AutomationDriver/Private/AutomationDriver.cpp
+++ b/Engine/Source/Developer/AutomationDriver/Private/AutomationDriver.cpp
@@ -11,6 +11,7 @@
 #include "WaitUntil.h"
 
 #include "Async.h"
+#include "Framework/Application/SlateApplication.h"
 
 FAsyncAutomationDriver::~FAsyncAutomationDriver()
 {
diff --git a/Engine/Source/Developer/AutomationDriver/Private/AutomationDriver.h b/Engine/Source/Developer/AutomationDriver/Private/AutomationDriver.h
--- a/Engine/Source/Developer/AutomationDriver/Private/AutomationDriver.h
+++ b/Engine/Source/Developer/AutomationDriver/Private/AutomationDriver.h
@@ -4,6 +4,7 @@
 #include "IAutomationDriver.h"
 #include "WaitUntil.h"
 #include "Misc/Timespan.h"
+#include "InputCoreTypes.h"
 
 class FAutomatedApplication;
 class IElementLocator;
diff --git a/Engine/Source/Developer/AutomationDriver/Private/AutomationDriverLogging.h b/Engine/Source/Developer/AutomationDriver/Private/AutomationDriverLogging.h
--- a/Engine/Source/Developer/AutomationDriver/Private/AutomationDriverLogging.h
+++ b/Engine/Source/Developer/AutomationDriver/Private/AutomationDriverLogging.h
@@ -1,6 +1,8 @@
 // Copyright 1998-2017 Epic Games, Inc. All Rights Reserved.
 #pragma once
 
+#include "CoreMinimal.h"
+
 class IApplicationElement;
 class IElementLocator;
 
diff --git a/Engine/Source/Developer/AutomationDriver/Private/AutomationDriverModule.cpp b/Engine/Source/Developer/AutomationDriver/Private/AutomationDriverModule.cpp
--- a/Engine/Source/Developer/AutomationDriver/Private/AutomationDriverModule.cpp
+++ b/Engine/Source/Developer/AutomationDriver/Private/AutomationDriverModule.cpp
@@ -4,6 +4,7 @@
 #include "PassThroughMessageHandler.h"
 #include "AutomatedApplication.h"
 #include "AutomationDriver.h"
+#include "Framework/Application/SlateApplication.h"
 
 class FAutomationDriverModule
 	: public IAutomationDriverModule
diff --git a/Engine/Source/Developer/AutomationDriver/Private/DriverElement.h b/Engine/Source/Developer/AutomationDriver/Private/DriverElement.h
--- a/Engine/Source/Developer/AutomationDriver/Private/DriverElement.h
+++ b/Engine/Source/Developer/AutomationDriver/Private/DriverElement.h
@@ -1,6 +1,8 @@
 // Copyright 1998-2017 Epic Games, Inc. All Rights Reserved.
 #pragma once
 
+#include "CoreMinimal.h"
+
 class IElementLocator;
 
 
diff --git a/Engine/Source/Developer/AutomationDriver/Private/DriverSequence.cpp b/Engine/Source/Developer/AutomationDriver/Private/DriverSequence.cpp
--- a/Engine/Source/Developer/AutomationDriver/Private/DriverSequence.cpp
+++ b/Engine/Source/Developer/AutomationDriver/Private/DriverSequence.cpp
@@ -19,6 +19,7 @@
 #include "Misc/Timespan.h"
 #include "InputCoreTypes.h"
 #include "GenericApplicationMessageHandler.h"
+#include "Framework/Application/SlateApplication.h"
 
 
 class FStep
diff --git a/Engine/Source/Developer/AutomationDriver/Private/DriverSequence.h b/Engine/Source/Developer/AutomationDriver/Private/DriverSequence.h
--- a/Engine/Source/Developer/AutomationDriver/Private/DriverSequence.h
+++ b/Engine/Source/Developer/AutomationDriver/Private/DriverSequence.h
@@ -1,6 +1,8 @@
 // Copyright 1998-2017 Epic Games, Inc. All Rights Reserved.
 #pragma once
 
+#include "CoreMinimal.h"
+
 class FAutomatedApplication;
 
 
diff --git a/Engine/Source/Developer/AutomationDriver/Private/IApplicationElement.h b/Engine/Source/Developer/AutomationDriver/Private/IApplicationElement.h
--- a/Engine/Source/Developer/AutomationDriver/Private/IApplicationElement.h
+++ b/Engine/Source/Developer/AutomationDriver/Private/IApplicationElement.h
@@ -1,6 +1,10 @@
 // Copyright 1998-2017 Epic Games, Inc. All Rights Reserved.
 #pragma once
 
+#include "CoreMinimal.h"
+#include "GenericPlatform/GenericWindow.h"
+#include "Types/SlateEnums.h"
+
 class IElementLocator;
 
 class IApplicationElement
diff --git a/Engine/Source/Developer/AutomationDriver/Private/LocateBy.cpp b/Engine/Source/Developer/AutomationDriver/Private/LocateBy.cpp
--- a/Engine/Source/Developer/AutomationDriver/Private/LocateBy.cpp
+++ b/Engine/Source/Developer/AutomationDriver/Private/LocateBy.cpp
@@ -3,6 +3,7 @@
 #include "LocateBy.h"
 #include "SlateWidgetLocatorByDelegate.h"
 #include "SlateWidgetLocatorByPath.h"
+#include "Framework/Application/SlateApplication.h"
 
 TSharedRef<IElementLocator, ESPMode::ThreadSafe> By::Delegate(const FLocateSlateWidgetElementDelegate& Value)
 {
diff --git a/Engine/Source/Developer/AutomationDriver/Private/Locators/SlateWidgetLocatorByDelegate.cpp b/Engine/Source/Developer/AutomationDriver/Private/Locators/SlateWidgetLocatorByDelegate.cpp
--- a/Engine/Source/Developer/AutomationDriver/Private/Locators/SlateWidgetLocatorByDelegate.cpp
+++ b/Engine/Source/Developer/AutomationDriver/Private/Locators/SlateWidgetLocatorByDelegate.cpp
@@ -3,6 +3,7 @@
 #include "SlateWidgetLocatorByDelegate.h"
 #include "SlateWidgetElement.h"
 #include "IElementLocator.h"
+#include "Framework/Application/SlateApplication.h"
 
 class FSlateWidgetLocatorByWidgetDelegate
 	: public IElementLocator
diff --git a/Engine/Source/Developer/AutomationDriver/Private/Locators/SlateWidgetLocatorByPath.cpp b/Engine/Source/Developer/AutomationDriver/Private/Locators/SlateWidgetLocatorByPath.cpp
--- a/Engine/Source/Developer/AutomationDriver/Private/Locators/SlateWidgetLocatorByPath.cpp
+++ b/Engine/Source/Developer/AutomationDriver/Private/Locators/SlateWidgetLocatorByPath.cpp
@@ -5,6 +5,8 @@
 #include "IElementLocator.h"
 #include "DriverIdMetaData.h"
 #include "AutomationDriverLogging.h"
+#include "Widgets/SWidget.h"
+#include "Framework/Application/SlateApplication.h"
 
 class FSlateWidgetLocatorByPath
 	: public IElementLocator
diff --git a/Engine/Source/Developer/AutomationDriver/Private/Locators/SlateWidgetLocatorByPath.h b/Engine/Source/Developer/AutomationDriver/Private/Locators/SlateWidgetLocatorByPath.h
--- a/Engine/Source/Developer/AutomationDriver/Private/Locators/SlateWidgetLocatorByPath.h
+++ b/Engine/Source/Developer/AutomationDriver/Private/Locators/SlateWidgetLocatorByPath.h
@@ -1,6 +1,8 @@
 // Copyright 1998-2017 Epic Games, Inc. All Rights Reserved.
 #pragma once
 
+#include "CoreMinimal.h"
+
 class IElementLocator;
 
 class FSlateWidgetLocatorByPathFactory
diff --git a/Engine/Source/Developer/AutomationDriver/Private/Locators/SlateWidgetLocatorByUniqueTag.cpp b/Engine/Source/Developer/AutomationDriver/Private/Locators/SlateWidgetLocatorByUniqueTag.cpp
--- a/Engine/Source/Developer/AutomationDriver/Private/Locators/SlateWidgetLocatorByUniqueTag.cpp
+++ b/Engine/Source/Developer/AutomationDriver/Private/Locators/SlateWidgetLocatorByUniqueTag.cpp
@@ -4,6 +4,7 @@
 #include "SlateWidgetElement.h"
 #include "IElementLocator.h"
 #include "DriverUniqueTagMetaData.h"
+#include "Framework/Application/SlateApplication.h"
 
 class FSlateWidgetLocatorByUniqueTag
 	: public IElementLocator
diff --git a/Engine/Source/Developer/AutomationDriver/Private/Locators/SlateWidgetLocatorByUniqueTag.h b/Engine/Source/Developer/AutomationDriver/Private/Locators/SlateWidgetLocatorByUniqueTag.h
--- a/Engine/Source/Developer/AutomationDriver/Private/Locators/SlateWidgetLocatorByUniqueTag.h
+++ b/Engine/Source/Developer/AutomationDriver/Private/Locators/SlateWidgetLocatorByUniqueTag.h
@@ -1,6 +1,8 @@
 // Copyright 1998-2017 Epic Games, Inc. All Rights Reserved.
 #pragma once
 
+#include "CoreMinimal.h"
+
 class IElementLocator;
 class FDriverUniqueTagMetaData;
 
diff --git a/Engine/Source/Developer/AutomationDriver/Private/MetaData/DriverUniqueTagMetaData.h b/Engine/Source/Developer/AutomationDriver/Private/MetaData/DriverUniqueTagMetaData.h
--- a/Engine/Source/Developer/AutomationDriver/Private/MetaData/DriverUniqueTagMetaData.h
+++ b/Engine/Source/Developer/AutomationDriver/Private/MetaData/DriverUniqueTagMetaData.h
@@ -1,6 +1,8 @@
 // Copyright 1998-2017 Epic Games, Inc. All Rights Reserved.
 #pragma once
 
+#include "CoreMinimal.h"
+
 class FDriverUniqueTagMetaData
 	: public ISlateMetaData
 {
diff --git a/Engine/Source/Developer/AutomationDriver/Private/PassThroughMessageHandler .cpp b/Engine/Source/Developer/AutomationDriver/Private/PassThroughMessageHandler .cpp
--- a/Engine/Source/Developer/AutomationDriver/Private/PassThroughMessageHandler .cpp	
+++ b/Engine/Source/Developer/AutomationDriver/Private/PassThroughMessageHandler .cpp	
@@ -1,6 +1,7 @@
 // Copyright 1998-2017 Epic Games, Inc. All Rights Reserved.
 
 #include "PassThroughMessageHandler.h"
+#include "InputCoreTypes.h"
 
 class FPassThroughMessageHandlerFactory;
 
diff --git a/Engine/Source/Developer/AutomationDriver/Private/SlateWidgetElement.cpp b/Engine/Source/Developer/AutomationDriver/Private/SlateWidgetElement.cpp
--- a/Engine/Source/Developer/AutomationDriver/Private/SlateWidgetElement.cpp
+++ b/Engine/Source/Developer/AutomationDriver/Private/SlateWidgetElement.cpp
@@ -2,6 +2,7 @@
 
 #include "SlateWidgetElement.h"
 #include "IApplicationElement.h"
+#include "Framework/Application/SlateApplication.h"
 
 #include "SlateWidgetLocatorByUniqueTag.h"
 #include "DriverUniqueTagMetaData.h"
@@ -10,6 +11,9 @@
 #include "SRichTextBlock.h"
 #include "SMultiLineEditableTextBox.h"
 #include "SMultiLineEditableText.h"
+#include "STextBlock.h"
+#include "SEditableText.h"
+#include "SEditableTextBox.h"
 
 class FSlateWidgetElement
 	: public IApplicationElement
diff --git a/Engine/Source/Developer/AutomationDriver/Private/SlateWidgetElement.h b/Engine/Source/Developer/AutomationDriver/Private/SlateWidgetElement.h
--- a/Engine/Source/Developer/AutomationDriver/Private/SlateWidgetElement.h
+++ b/Engine/Source/Developer/AutomationDriver/Private/SlateWidgetElement.h
@@ -1,6 +1,9 @@
 // Copyright 1998-2017 Epic Games, Inc. All Rights Reserved.
 #pragma once
 
+#include "CoreMinimal.h"
+#include "Layout/WidgetPath.h"
+
 class IApplicationElement;
 
 class FSlateWidgetElementFactory
diff --git a/Engine/Source/Developer/AutomationDriver/Private/Specs/AutomationDriver.spec.cpp b/Engine/Source/Developer/AutomationDriver/Private/Specs/AutomationDriver.spec.cpp
--- a/Engine/Source/Developer/AutomationDriver/Private/Specs/AutomationDriver.spec.cpp
+++ b/Engine/Source/Developer/AutomationDriver/Private/Specs/AutomationDriver.spec.cpp
@@ -7,6 +7,7 @@
 #include "AutomationDriverSpecSuiteViewModel.h"
 #include "Ticker.h"
 #include "Async.h"
+#include "Framework/Application/SlateApplication.h"
 
 #define TEST_TRUE(expression) \
 	EPIC_TEST_BOOLEAN_(TEXT(#expression), expression, true)
diff --git a/Engine/Source/Developer/AutomationDriver/Private/Specs/AutomationDriverSpecSuiteViewModel.h b/Engine/Source/Developer/AutomationDriver/Private/Specs/AutomationDriverSpecSuiteViewModel.h
--- a/Engine/Source/Developer/AutomationDriver/Private/Specs/AutomationDriverSpecSuiteViewModel.h
+++ b/Engine/Source/Developer/AutomationDriver/Private/Specs/AutomationDriverSpecSuiteViewModel.h
@@ -1,6 +1,11 @@
 // Copyright 1998-2017 Epic Games, Inc. All Rights Reserved.
 #pragma once
 
+#include "CoreMinimal.h"
+#include "Types/SlateEnums.h"
+#include "Input/Reply.h"
+#include "Layout/Visibility.h"
+
 enum class EPianoKey : uint8
 {
 	AFlat,
diff --git a/Engine/Source/Developer/AutomationDriver/Private/Specs/SAutomationDriverSpecSuite.cpp b/Engine/Source/Developer/AutomationDriver/Private/Specs/SAutomationDriverSpecSuite.cpp
--- a/Engine/Source/Developer/AutomationDriver/Private/Specs/SAutomationDriverSpecSuite.cpp
+++ b/Engine/Source/Developer/AutomationDriver/Private/Specs/SAutomationDriverSpecSuite.cpp
@@ -9,7 +9,9 @@
 #include "SScrollBox.h"
 #include "STileView.h"
 #include "SListView.h"
+#include "SOverlay.h"
 #include "SMultiLineEditableTextBox.h"
+#include "Widgets/Input/SEditableTextBox.h"
 
 class SAutomationDriverSpecSuiteImpl
 	: public SAutomationDriverSpecSuite
diff --git a/Engine/Source/Developer/AutomationDriver/Private/StepExecutor.h b/Engine/Source/Developer/AutomationDriver/Private/StepExecutor.h
--- a/Engine/Source/Developer/AutomationDriver/Private/StepExecutor.h
+++ b/Engine/Source/Developer/AutomationDriver/Private/StepExecutor.h
@@ -1,6 +1,8 @@
 // Copyright 1998-2017 Epic Games, Inc. All Rights Reserved.
 #pragma once
 
+#include "CoreMinimal.h"
+
 class IStepExecutor;
 class FDriverConfiguration;
 
diff --git a/Engine/Source/Developer/AutomationDriver/Public/IAutomationDriver.h b/Engine/Source/Developer/AutomationDriver/Public/IAutomationDriver.h
--- a/Engine/Source/Developer/AutomationDriver/Public/IAutomationDriver.h
+++ b/Engine/Source/Developer/AutomationDriver/Public/IAutomationDriver.h
@@ -4,6 +4,7 @@
 #include "AsyncResult.h"
 #include "Misc/Timespan.h"
 #include "WaitUntil.h"
+#include "GenericPlatform/GenericApplication.h"
 
 class IElementLocator;
 
diff --git a/Engine/Source/Developer/AutomationDriver/Public/IDriverElement.h b/Engine/Source/Developer/AutomationDriver/Public/IDriverElement.h
--- a/Engine/Source/Developer/AutomationDriver/Public/IDriverElement.h
+++ b/Engine/Source/Developer/AutomationDriver/Public/IDriverElement.h
@@ -3,6 +3,8 @@
 
 #include "IElementLocator.h"
 #include "AsyncResult.h"
+#include "InputCoreTypes.h"
+#include "GenericPlatform/GenericApplicationMessageHandler.h"
 
 /**
  * Represents the async API for a potential application element
diff --git a/Engine/Source/Developer/AutomationDriver/Public/IDriverSequence.h b/Engine/Source/Developer/AutomationDriver/Public/IDriverSequence.h
--- a/Engine/Source/Developer/AutomationDriver/Public/IDriverSequence.h
+++ b/Engine/Source/Developer/AutomationDriver/Public/IDriverSequence.h
@@ -3,6 +3,8 @@
 
 #include "WaitUntil.h"
 #include "Misc/Timespan.h"
+#include "InputCoreTypes.h"
+#include "GenericPlatform/GenericApplicationMessageHandler.h"
 
 class IElementLocator;
 
diff --git a/Engine/Source/Developer/AutomationDriver/Public/IElementLocator.h b/Engine/Source/Developer/AutomationDriver/Public/IElementLocator.h
--- a/Engine/Source/Developer/AutomationDriver/Public/IElementLocator.h
+++ b/Engine/Source/Developer/AutomationDriver/Public/IElementLocator.h
@@ -1,6 +1,8 @@
 // Copyright 1998-2017 Epic Games, Inc. All Rights Reserved.
 #pragma once
 
+#include "CoreMinimal.h"
+
 //@todo this type shouldn't be public at all
 class IApplicationElement; 
 
diff --git a/Engine/Source/Developer/AutomationDriver/Public/LocateBy.h b/Engine/Source/Developer/AutomationDriver/Public/LocateBy.h
--- a/Engine/Source/Developer/AutomationDriver/Public/LocateBy.h
+++ b/Engine/Source/Developer/AutomationDriver/Public/LocateBy.h
@@ -1,6 +1,9 @@
 // Copyright 1998-2017 Epic Games, Inc. All Rights Reserved.
 #pragma once
 
+#include "CoreMinimal.h"
+#include "Widgets/SWidget.h"
+
 class IElementLocator;
 
 DECLARE_DELEGATE_OneParam(FLocateSlateWidgetElementDelegate, TArray<TSharedRef<SWidget>>& /*OutWidgets*/);
diff --git a/Engine/Source/Developer/AutomationDriver/Public/WaitUntil.h b/Engine/Source/Developer/AutomationDriver/Public/WaitUntil.h
--- a/Engine/Source/Developer/AutomationDriver/Public/WaitUntil.h
+++ b/Engine/Source/Developer/AutomationDriver/Public/WaitUntil.h
@@ -1,6 +1,8 @@
 // Copyright 1998-2017 Epic Games, Inc. All Rights Reserved.
 #pragma once
 
+#include "CoreMinimal.h"
+
 class IElementLocator;
 
 /**
diff --git a/Engine/Source/Editor/UnrealEd/Private/AssetViewerSettings.cpp b/Engine/Source/Editor/UnrealEd/Private/AssetViewerSettings.cpp
--- a/Engine/Source/Editor/UnrealEd/Private/AssetViewerSettings.cpp
+++ b/Engine/Source/Editor/UnrealEd/Private/AssetViewerSettings.cpp
@@ -3,6 +3,7 @@
 #include "AssetViewerSettings.h"
 #include "UObject/UnrealType.h"
 #include "Editor/EditorPerProjectUserSettings.h"
+#include "Editor.h"
 
 
 UAssetViewerSettings::UAssetViewerSettings()
diff --git a/Engine/Source/Editor/UnrealEd/Private/Commandlets/InternationalizationExportCommandlet.cpp b/Engine/Source/Editor/UnrealEd/Private/Commandlets/InternationalizationExportCommandlet.cpp
--- a/Engine/Source/Editor/UnrealEd/Private/Commandlets/InternationalizationExportCommandlet.cpp
+++ b/Engine/Source/Editor/UnrealEd/Private/Commandlets/InternationalizationExportCommandlet.cpp
@@ -8,6 +8,7 @@
 #include "Localization/PortableObjectFormatDOM.h"
 #include "TextNamespaceUtil.h"
 #include "LocalizationTargetTypes.h"
+#include "UObject/Package.h"
 
 DEFINE_LOG_CATEGORY_STATIC(LogInternationalizationExportCommandlet, Log, All);
 
diff --git a/Engine/Source/Editor/UnrealEd/Private/EditorBuildUtils.cpp b/Engine/Source/Editor/UnrealEd/Private/EditorBuildUtils.cpp
--- a/Engine/Source/Editor/UnrealEd/Private/EditorBuildUtils.cpp
+++ b/Engine/Source/Editor/UnrealEd/Private/EditorBuildUtils.cpp
@@ -38,6 +38,7 @@
 #include "ActorEditorUtils.h"
 #include "MaterialUtilities.h"
 #include "UnrealEngine.h"
+#include "DebugViewModeHelpers.h"
 
 DEFINE_LOG_CATEGORY_STATIC(LogEditorBuildUtils, Log, All);
 
diff --git a/Engine/Source/Editor/UnrealEd/Private/EditorViewportCommands.cpp b/Engine/Source/Editor/UnrealEd/Private/EditorViewportCommands.cpp
--- a/Engine/Source/Editor/UnrealEd/Private/EditorViewportCommands.cpp
+++ b/Engine/Source/Editor/UnrealEd/Private/EditorViewportCommands.cpp
@@ -11,6 +11,8 @@
 #include "ContentBrowserModule.h"
 #include "IContentBrowserSingleton.h"
 #include "Materials/MaterialInterface.h"
+#include "MaterialShared.h"
+#include "Engine/Texture2D.h"
 
 #define LOCTEXT_NAMESPACE "EditorViewportCommands"
 
diff --git a/Engine/Source/Editor/UnrealEd/Private/Fbx/FbxMaterialImport.cpp b/Engine/Source/Editor/UnrealEd/Private/Fbx/FbxMaterialImport.cpp
--- a/Engine/Source/Editor/UnrealEd/Private/Fbx/FbxMaterialImport.cpp
+++ b/Engine/Source/Editor/UnrealEd/Private/Fbx/FbxMaterialImport.cpp
@@ -25,6 +25,7 @@
 #include "IAssetTools.h"
 #include "AssetToolsModule.h"
 #include "Misc/FbxErrors.h"
+#include "ARFilter.h"
 
 DEFINE_LOG_CATEGORY_STATIC(LogFbxMaterialImport, Log, All);
 
diff --git a/Engine/Source/Editor/UnrealEd/Private/MaterialGraphNode_Knot.cpp b/Engine/Source/Editor/UnrealEd/Private/MaterialGraphNode_Knot.cpp
--- a/Engine/Source/Editor/UnrealEd/Private/MaterialGraphNode_Knot.cpp
+++ b/Engine/Source/Editor/UnrealEd/Private/MaterialGraphNode_Knot.cpp
@@ -5,6 +5,7 @@
 #include "BlueprintNodeSpawner.h"
 #include "BlueprintActionDatabaseRegistrar.h"
 #include "MaterialEditorUtilities.h"
+#include "EdGraph/EdGraphPin.h"
 
 #define LOCTEXT_NAMESPACE "MaterialGraphNode_Knot"
 
diff --git a/Engine/Source/Editor/UnrealEd/Private/PackageTools.cpp b/Engine/Source/Editor/UnrealEd/Private/PackageTools.cpp
--- a/Engine/Source/Editor/UnrealEd/Private/PackageTools.cpp
+++ b/Engine/Source/Editor/UnrealEd/Private/PackageTools.cpp
@@ -23,6 +23,7 @@
 #include "ISourceControlModule.h"
 #include "Editor.h"
 #include "Dialogs/Dialogs.h"
+#include "Toolkits/AssetEditorManager.h"
 
 #include "ObjectTools.h"
 #include "KismetEditorUtilities.h"
diff --git a/Engine/Source/Editor/UnrealEd/Public/AutomatedAssetImportData.h b/Engine/Source/Editor/UnrealEd/Public/AutomatedAssetImportData.h
--- a/Engine/Source/Editor/UnrealEd/Public/AutomatedAssetImportData.h
+++ b/Engine/Source/Editor/UnrealEd/Public/AutomatedAssetImportData.h
@@ -8,6 +8,7 @@
 #include "AutomatedAssetImportData.generated.h"
 
 class UFactory;
+class FJsonObject;
 
 DECLARE_LOG_CATEGORY_EXTERN(LogAutomatedImport, Log, All);
 
diff --git a/Engine/Source/Runtime/CoreUObject/Private/Serialization/PropertyLocalizationDataGathering.cpp b/Engine/Source/Runtime/CoreUObject/Private/Serialization/PropertyLocalizationDataGathering.cpp
--- a/Engine/Source/Runtime/CoreUObject/Private/Serialization/PropertyLocalizationDataGathering.cpp
+++ b/Engine/Source/Runtime/CoreUObject/Private/Serialization/PropertyLocalizationDataGathering.cpp
@@ -8,6 +8,7 @@
 #include "UObject/Package.h"
 #include "UObject/UnrealType.h"
 #include "UObject/TextProperty.h"
+#include "UObject/PropertyPortFlags.h"
 #include "Internationalization/TextNamespaceUtil.h"
 #include "Internationalization/TextPackageNamespaceUtil.h"
 
diff --git a/Engine/Source/Runtime/CoreUObject/Private/UObject/PackageReload.cpp b/Engine/Source/Runtime/CoreUObject/Private/UObject/PackageReload.cpp
--- a/Engine/Source/Runtime/CoreUObject/Private/UObject/PackageReload.cpp
+++ b/Engine/Source/Runtime/CoreUObject/Private/UObject/PackageReload.cpp
@@ -6,6 +6,10 @@
 #include "UObjectHash.h"
 #include "Misc/AssetRegistryInterface.h"
 #include "Misc/ScopedSlowTask.h"
+#include "UObject/Package.h"
+#include "UObject/Linker.h"
+#include "Templates/Casts.h"
+#include "UObject/UObjectIterator.h"
 
 namespace PackageReloadInternal
 {
diff --git a/Engine/Source/Runtime/CoreUObject/Public/UObject/PackageReload.h b/Engine/Source/Runtime/CoreUObject/Public/UObject/PackageReload.h
--- a/Engine/Source/Runtime/CoreUObject/Public/UObject/PackageReload.h
+++ b/Engine/Source/Runtime/CoreUObject/Public/UObject/PackageReload.h
@@ -2,7 +2,10 @@
 
 #pragma once
 
+#include "CoreMinimal.h"
 #include "ArrayView.h"
+#include "UObject/ObjectMacros.h"
+#include "UObject/WeakObjectPtr.h"
 
 class IAssetRegistryInterface;
 
diff --git a/Engine/Source/Runtime/Engine/Private/Animation/SkeletalMeshActor.cpp b/Engine/Source/Runtime/Engine/Private/Animation/SkeletalMeshActor.cpp
--- a/Engine/Source/Runtime/Engine/Private/Animation/SkeletalMeshActor.cpp
+++ b/Engine/Source/Runtime/Engine/Private/Animation/SkeletalMeshActor.cpp
@@ -8,6 +8,7 @@
 #include "Misc/MapErrors.h"
 #include "Animation/AnimBlueprintGeneratedClass.h"
 #include "Engine/CollisionProfile.h"
+#include "Logging/MessageLog.h"
 
 #define LOCTEXT_NAMESPACE "SkeletalMeshActor"
 
diff --git a/Engine/Source/Runtime/Engine/Private/Components/ForceFeedbackComponent.cpp b/Engine/Source/Runtime/Engine/Private/Components/ForceFeedbackComponent.cpp
--- a/Engine/Source/Runtime/Engine/Private/Components/ForceFeedbackComponent.cpp
+++ b/Engine/Source/Runtime/Engine/Private/Components/ForceFeedbackComponent.cpp
@@ -5,6 +5,7 @@
 #include "GameFramework/ForceFeedbackEffect.h"
 #include "Engine/Canvas.h"
 #include "IInputInterface.h"
+#include "Engine/Texture2D.h"
 
 TArray<FForceFeedbackManager*> FForceFeedbackManager::PerWorldForceFeedbackManagers;
 FDelegateHandle FForceFeedbackManager::OnWorldCleanupHandle;
diff --git a/Engine/Source/Runtime/Engine/Private/Components/SkinnedMeshComponent.cpp b/Engine/Source/Runtime/Engine/Private/Components/SkinnedMeshComponent.cpp
--- a/Engine/Source/Runtime/Engine/Private/Components/SkinnedMeshComponent.cpp
+++ b/Engine/Source/Runtime/Engine/Private/Components/SkinnedMeshComponent.cpp
@@ -19,6 +19,7 @@
 #include "PhysicsEngine/PhysicsAsset.h"
 #include "EngineGlobals.h"
 #include "PrimitiveSceneProxy.h"
+#include "Engine/CollisionProfile.h"
 
 DEFINE_LOG_CATEGORY_STATIC(LogSkinnedMeshComp, Log, All);
 
diff --git a/Engine/Source/Runtime/Engine/Private/DebugViewModeHelpers.cpp b/Engine/Source/Runtime/Engine/Private/DebugViewModeHelpers.cpp
--- a/Engine/Source/Runtime/Engine/Private/DebugViewModeHelpers.cpp
+++ b/Engine/Source/Runtime/Engine/Private/DebugViewModeHelpers.cpp
@@ -8,6 +8,8 @@
 #include "ShaderCompiler.h"
 #include "ScopedSlowTask.h"
 #include "FeedbackContext.h"
+#include "Engine/World.h"
+#include "Components/PrimitiveComponent.h"
 
 #define LOCTEXT_NAMESPACE "LogDebugViewMode"
 
diff --git a/Engine/Source/Runtime/Engine/Private/EngineSessionManager.cpp b/Engine/Source/Runtime/Engine/Private/EngineSessionManager.cpp
--- a/Engine/Source/Runtime/Engine/Private/EngineSessionManager.cpp
+++ b/Engine/Source/Runtime/Engine/Private/EngineSessionManager.cpp
@@ -3,7 +3,11 @@
 #include "EngineSessionManager.h"
 #include "Misc/Guid.h"
 #include "Misc/CoreDelegates.h"
+#include "Misc/ConfigCacheIni.h"
 #include "Misc/EngineVersion.h"
+#include "Misc/App.h"
+#include "Misc/CommandLine.h"
+#include "HAL/PlatformProcess.h"
 #include "EngineAnalytics.h"
 #include "AnalyticsEventAttribute.h"
 #include "Interfaces/IAnalyticsProvider.h"
diff --git a/Engine/Source/Runtime/Engine/Private/PhysicsEngine/BodySetup.cpp b/Engine/Source/Runtime/Engine/Private/PhysicsEngine/BodySetup.cpp
--- a/Engine/Source/Runtime/Engine/Private/PhysicsEngine/BodySetup.cpp
+++ b/Engine/Source/Runtime/Engine/Private/PhysicsEngine/BodySetup.cpp
@@ -9,13 +9,16 @@
 #include "HAL/IConsoleManager.h"
 #include "Components/PrimitiveComponent.h"
 #include "Engine/Engine.h"
+#include "Engine/StaticMesh.h"
 #include "Components/SkinnedMeshComponent.h"
+#include "Components/StaticMeshComponent.h"
 #include "Interfaces/Interface_CollisionDataProvider.h"
 #include "PhysicsEngine/PhysicsSettings.h"
 #include "Interfaces/ITargetPlatform.h"
 #include "Interfaces/ITargetPlatformManagerModule.h"
 #include "Animation/AnimStats.h"
 #include "DerivedDataCacheInterface.h"
+#include "UObject/UObjectIterator.h"
 
 #if WITH_PHYSX
 	#include "PhysXPublic.h"
diff --git a/Engine/Source/Runtime/Engine/Private/PhysicsEngine/PhysicsAsset.cpp b/Engine/Source/Runtime/Engine/Private/PhysicsEngine/PhysicsAsset.cpp
--- a/Engine/Source/Runtime/Engine/Private/PhysicsEngine/PhysicsAsset.cpp
+++ b/Engine/Source/Runtime/Engine/Private/PhysicsEngine/PhysicsAsset.cpp
@@ -9,6 +9,7 @@
 #include "Serialization/ObjectWriter.h"
 #include "Serialization/ObjectReader.h"
 #include "Components/SkinnedMeshComponent.h"
+#include "Components/SkeletalMeshComponent.h"
 #include "PhysicsEngine/PhysicsConstraintTemplate.h"
 #include "UObject/ReleaseObjectVersion.h"
 #include "Logging/MessageLog.h"
diff --git a/Engine/Source/Runtime/Engine/Private/Rendering/SkinWeightVertexBuffer.cpp b/Engine/Source/Runtime/Engine/Private/Rendering/SkinWeightVertexBuffer.cpp
--- a/Engine/Source/Runtime/Engine/Private/Rendering/SkinWeightVertexBuffer.cpp
+++ b/Engine/Source/Runtime/Engine/Private/Rendering/SkinWeightVertexBuffer.cpp
@@ -2,6 +2,7 @@
 
 #include "Rendering/SkinWeightVertexBuffer.h"
 #include "SkeletalMeshTypes.h"
+#include "EngineUtils.h"
 
 ///
 
diff --git a/Engine/Source/Runtime/Engine/Private/SoundSubmix.cpp b/Engine/Source/Runtime/Engine/Private/SoundSubmix.cpp
--- a/Engine/Source/Runtime/Engine/Private/SoundSubmix.cpp
+++ b/Engine/Source/Runtime/Engine/Private/SoundSubmix.cpp
@@ -5,6 +5,7 @@
 #include "AudioDeviceManager.h"
 #include "EngineGlobals.h"
 #include "Engine/Engine.h"
+#include "UObject/UObjectIterator.h"
 
 #if WITH_EDITOR
 TSharedPtr<ISoundSubmixAudioEditor> USoundSubmix::SoundSubmixAudioEditor = nullptr;
diff --git a/Engine/Source/Runtime/Engine/Private/Streaming/TextureInstanceManager.h b/Engine/Source/Runtime/Engine/Private/Streaming/TextureInstanceManager.h
--- a/Engine/Source/Runtime/Engine/Private/Streaming/TextureInstanceManager.h
+++ b/Engine/Source/Runtime/Engine/Private/Streaming/TextureInstanceManager.h
@@ -11,6 +11,7 @@ TextureInstanceManager.h: Definitions of classes used for texture streaming.
 #include "ContentStreaming.h"
 #include "Streaming/TextureStreamingHelpers.h"
 #include "Components/PrimitiveComponent.h"
+#include "Async/AsyncWork.h"
 
 #define MAX_TEXTURE_SIZE (float(1 << (MAX_TEXTURE_MIP_COUNT - 1)))
 
diff --git a/Engine/Source/Runtime/Engine/Public/DebugViewModeHelpers.h b/Engine/Source/Runtime/Engine/Public/DebugViewModeHelpers.h
--- a/Engine/Source/Runtime/Engine/Public/DebugViewModeHelpers.h
+++ b/Engine/Source/Runtime/Engine/Public/DebugViewModeHelpers.h
@@ -8,6 +8,9 @@
 
 #include "CoreMinimal.h"
 #include "RHIDefinitions.h"
+#include "SceneTypes.h"
+
+class UMaterialInterface;
 
 /** 
  * Enumeration for different Quad Overdraw visualization mode.
