Copying@3231790 //UE4/Dev-Rendering to Dev-Main (//UE4/Dev-Main)

Change 3231790 on 2016/12/12 by [REDACTED_PERSON]

	Fixed particle GetUsedMaterials not containing material overrides in some situations

diff --git a/Engine/Source/Runtime/Engine/Private/Particles/ParticleComponents.cpp b/Engine/Source/Runtime/Engine/Private/Particles/ParticleComponents.cpp
--- a/Engine/Source/Runtime/Engine/Private/Particles/ParticleComponents.cpp
+++ b/Engine/Source/Runtime/Engine/Private/Particles/ParticleComponents.cpp
@@ -6834,6 +6834,8 @@ void UParticleSystemComponent::GetUsedMaterials( TArray<UMaterialInterface*>& Ou
 			}
 		}
 	}
+
+	OutMaterials.Append(EmitterMaterials);
 }
 
 FBodyInstance* UParticleSystemComponent::GetBodyInstance(FName BoneName /*= NAME_None*/, bool bGetWelded /*= true*/) const
