Fixed IWYU regression with PhysPublic.h (introduced in 3233389)

diff --git a/Engine/Source/Runtime/Engine/Public/PhysicsPublic.h b/Engine/Source/Runtime/Engine/Public/PhysicsPublic.h
--- a/Engine/Source/Runtime/Engine/Public/PhysicsPublic.h
+++ b/Engine/Source/Runtime/Engine/Public/PhysicsPublic.h
@@ -10,6 +10,7 @@
 #include "CoreMinimal.h"
 #include "Stats/Stats.h"
 #include "Engine/EngineTypes.h"
+#include "Engine/World.h"
 #include "Misc/CoreMisc.h"
 #include "EngineDefines.h"
 #include "RenderResource.h"
