Main - Copy of 3239008
DR - Fixing accidently returning a stackpointer in EnqueueRenderCommands

diff --git a/Engine/Source/Runtime/Engine/Private/GPUSkinCache.cpp b/Engine/Source/Runtime/Engine/Private/GPUSkinCache.cpp
--- a/Engine/Source/Runtime/Engine/Private/GPUSkinCache.cpp
+++ b/Engine/Source/Runtime/Engine/Private/GPUSkinCache.cpp
@@ -1094,7 +1094,7 @@ void FGPUSkinCache::CVarSinkFunction()
 
 	if (NewGPUSkinCacheValue != GEnableGPUSkinCache)
 	{
-		EnqueueUniqueRenderCommand("DoEnableSkinCaching", 
+		ENQUEUE_RENDER_COMMAND(DoEnableSkinCaching)(
 			[NewGPUSkinCacheValue](FRHICommandList& RHICmdList)
 			{
 				GEnableGPUSkinCache = NewGPUSkinCacheValue;
diff --git a/Engine/Source/Runtime/Engine/Private/KismetRenderingLibrary.cpp b/Engine/Source/Runtime/Engine/Private/KismetRenderingLibrary.cpp
--- a/Engine/Source/Runtime/Engine/Private/KismetRenderingLibrary.cpp
+++ b/Engine/Source/Runtime/Engine/Private/KismetRenderingLibrary.cpp
@@ -36,7 +36,7 @@ void UKismetRenderingLibrary::ClearRenderTarget2D(UObject* WorldContextObject, U
 		&& World)
 	{
 		FTextureRenderTargetResource* RenderTargetResource = TextureRenderTarget->GameThread_GetRenderTargetResource();
-		EnqueueUniqueRenderCommand("ClearRTCommand",
+		ENQUEUE_RENDER_COMMAND(ClearRTCommand)(
 			[RenderTargetResource, ClearColor](FRHICommandList& RHICmdList)
 			{
 				SetRenderTarget(RHICmdList, RenderTargetResource->GetRenderTargetTexture(), FTextureRHIRef(), true);
@@ -89,7 +89,7 @@ void UKismetRenderingLibrary::DrawMaterialToRenderTarget(UObject* WorldContextOb
 		TDrawEvent<FRHICommandList>* DrawMaterialToTargetEvent = new TDrawEvent<FRHICommandList>();
 
 		FName RTName = TextureRenderTarget->GetFName();
-		EnqueueUniqueRenderCommand("BeginDrawEventCommand",
+		ENQUEUE_RENDER_COMMAND(BeginDrawEventCommand)(
 			[&RTName ,DrawMaterialToTargetEvent](FRHICommandList& RHICmdList)
 			{
 				BEGIN_DRAW_EVENTF(
@@ -105,7 +105,7 @@ void UKismetRenderingLibrary::DrawMaterialToRenderTarget(UObject* WorldContextOb
 		Canvas->Canvas = NULL;
 
 		FTextureRenderTargetResource* RenderTargetResource = TextureRenderTarget->GameThread_GetRenderTargetResource();
-		EnqueueUniqueRenderCommand("CanvasRenderTargetResolveCommand",
+		ENQUEUE_RENDER_COMMAND(CanvasRenderTargetResolveCommand)(
 			[RenderTargetResource, DrawMaterialToTargetEvent](FRHICommandList& RHICmdList)
 			{
 				RHICmdList.CopyToResolveTarget(RenderTargetResource->GetRenderTargetTexture(), RenderTargetResource->TextureRHI, true, FResolveParams());
@@ -246,7 +246,7 @@ void UKismetRenderingLibrary::BeginDrawCanvasToRenderTarget(UObject* WorldContex
 
 		FName RTName = TextureRenderTarget->GetFName();
 		TDrawEvent<FRHICommandList>* DrawEvent = Context.DrawEvent;
-		EnqueueUniqueRenderCommand("BeginDrawEventCommand",
+		ENQUEUE_RENDER_COMMAND(BeginDrawEventCommand)(
 			[RTName, DrawEvent](FRHICommandList& RHICmdList)
 			{
 				BEGIN_DRAW_EVENTF(
@@ -286,7 +286,7 @@ void UKismetRenderingLibrary::EndDrawCanvasToRenderTarget(UObject* WorldContextO
 		{
 			FTextureRenderTargetResource* RenderTargetResource = Context.RenderTarget->GameThread_GetRenderTargetResource();
 			TDrawEvent<FRHICommandList>* DrawEvent = Context.DrawEvent;
-			EnqueueUniqueRenderCommand("CanvasRenderTargetResolveCommand",
+			ENQUEUE_RENDER_COMMAND(CanvasRenderTargetResolveCommand)(
 				[RenderTargetResource, DrawEvent](FRHICommandList& RHICmdList)
 				{
 					RHICmdList.CopyToResolveTarget(RenderTargetResource->GetRenderTargetTexture(), RenderTargetResource->TextureRHI, true, FResolveParams());
diff --git a/Engine/Source/Runtime/Engine/Private/LocalVertexFactory.cpp b/Engine/Source/Runtime/Engine/Private/LocalVertexFactory.cpp
--- a/Engine/Source/Runtime/Engine/Private/LocalVertexFactory.cpp
+++ b/Engine/Source/Runtime/Engine/Private/LocalVertexFactory.cpp
@@ -96,7 +96,7 @@ void FLocalVertexFactory::Copy(const FLocalVertexFactory& Other)
 {
 	FLocalVertexFactory* VertexFactory = this;
 	const FDataType* DataCopy = &Other.Data;
-	EnqueueUniqueRenderCommand("FLocalVertexFactoryCopyData",
+	ENQUEUE_RENDER_COMMAND(FLocalVertexFactoryCopyData)(
 		[VertexFactory, DataCopy](FRHICommandListImmediate& RHICmdList)
 		{
 			VertexFactory->Data = *DataCopy;
diff --git a/Engine/Source/Runtime/Engine/Private/Materials/MaterialShared.cpp b/Engine/Source/Runtime/Engine/Private/Materials/MaterialShared.cpp
--- a/Engine/Source/Runtime/Engine/Private/Materials/MaterialShared.cpp
+++ b/Engine/Source/Runtime/Engine/Private/Materials/MaterialShared.cpp
@@ -1966,7 +1966,7 @@ void FMaterialRenderProxy::CacheUniformExpressions_GameThread()
 	if (FApp::CanEverRender())
 	{
 		FMaterialRenderProxy* RenderProxy = this;
-		EnqueueUniqueRenderCommand("FCacheUniformExpressionsCommand",
+		ENQUEUE_RENDER_COMMAND(FCacheUniformExpressionsCommand)(
 			[RenderProxy](FRHICommandListImmediate& RHICmdList)
 			{
 				RenderProxy->CacheUniformExpressions();
diff --git a/Engine/Source/Runtime/Engine/Private/SkeletalRenderGPUSkin.cpp b/Engine/Source/Runtime/Engine/Private/SkeletalRenderGPUSkin.cpp
--- a/Engine/Source/Runtime/Engine/Private/SkeletalRenderGPUSkin.cpp
+++ b/Engine/Source/Runtime/Engine/Private/SkeletalRenderGPUSkin.cpp
@@ -180,7 +180,7 @@ void FSkeletalMeshObjectGPUSkin::ReleaseResources()
 	// also release morph resources
 	ReleaseMorphResources();
 	FSkeletalMeshObjectGPUSkin* MeshObject = this;
-	EnqueueUniqueRenderCommand("WaitRHIThreadFenceForDynamicData",
+	ENQUEUE_RENDER_COMMAND(WaitRHIThreadFenceForDynamicData)(
 		[MeshObject](FRHICommandList& RHICmdList)
 		{
 			FScopeCycleCounter Context(MeshObject->GetStatId());
@@ -210,7 +210,7 @@ void FSkeletalMeshObjectGPUSkin::InitMorphResources(bool bInUsePerBoneMotionBlur
 	{
 		int32 NumWeights = MorphTargetWeights.Num();
 		FVertexBufferAndSRV* MorphWeightsVertexBufferPtr = &MorphWeightsVertexBuffer;
-		EnqueueUniqueRenderCommand("InitMorphWeightsVertexBuffer",
+		ENQUEUE_RENDER_COMMAND(InitMorphWeightsVertexBuffer)(
 			[NumWeights, MorphWeightsVertexBufferPtr](FRHICommandList& RHICmdList)
 			{
 				const uint32 Size = (uint32)NumWeights * sizeof(float);
@@ -234,7 +234,7 @@ void FSkeletalMeshObjectGPUSkin::ReleaseMorphResources()
 	if (MorphWeightsVertexBuffer.VertexBufferRHI)
 	{
 		FVertexBufferAndSRV* MorphWeightsVertexBufferPtr = &MorphWeightsVertexBuffer;
-		EnqueueUniqueRenderCommand("InitMorphWeightsVertexBuffer",
+		ENQUEUE_RENDER_COMMAND(InitMorphWeightsVertexBuffer)(
 			[MorphWeightsVertexBufferPtr](FRHICommandList& RHICmdList)
 			{
 				MorphWeightsVertexBufferPtr->VertexBufferRHI.SafeRelease();
@@ -266,7 +266,7 @@ void FSkeletalMeshObjectGPUSkin::Update(int32 LODIndex,USkinnedMeshComponent* In
 
 	// queue a call to update this data
 	FSkeletalMeshObjectGPUSkin* MeshObject = this;
-	EnqueueUniqueRenderCommand("SkelMeshObjectUpdateDataCommand",
+	ENQUEUE_RENDER_COMMAND(SkelMeshObjectUpdateDataCommand)(
 		[MeshObject, FrameNumberToPrepare, NewDynamicData](FRHICommandListImmediate& RHICmdList)
 		{
 			FScopeCycleCounter Context(MeshObject->GetStatId());
@@ -288,7 +288,7 @@ void FSkeletalMeshObjectGPUSkin::UpdateRecomputeTangent(int32 MaterialIndex, int
 {
 	// queue a call to update this data
 	FSkeletalMeshObjectGPUSkin* MeshObject = this;
-	EnqueueUniqueRenderCommand("SkelMeshObjectUpdateMaterialDataCommand",
+	ENQUEUE_RENDER_COMMAND(SkelMeshObjectUpdateMaterialDataCommand)(
 		[MeshObject, MaterialIndex, LODIndex, bRecomputeTangent](FRHICommandList& RHICmdList)
 		{
 			// iterate through section and find the section that matches MaterialIndex, if so, set that flag
@@ -915,7 +915,7 @@ static VertexFactoryType* CreateVertexFactory(TArray<TUniquePtr<VertexFactoryTyp
 	TDynamicUpdateVertexFactoryData<VertexFactoryType> VertexUpdateData(VertexFactory,InVertexBuffers);
 
 	// update vertex factory components and sync it
-	EnqueueUniqueRenderCommand("InitGPUSkinVertexFactory",
+	ENQUEUE_RENDER_COMMAND(InitGPUSkinVertexFactory)(
 		[VertexUpdateData](FRHICommandList& CmdList)
 		{
 			typename VertexFactoryType::FDataType Data;
@@ -940,7 +940,7 @@ static void CreatePassthroughVertexFactory(TArray<TUniquePtr<FGPUSkinPassthrough
 	PassthroughVertexFactories.Add(TUniquePtr<FGPUSkinPassthroughVertexFactory>(NewPassthroughVertexFactory));
 
 	// update vertex factory components and sync it
-	EnqueueUniqueRenderCommand("InitPassthroughGPUSkinVertexFactory",
+	ENQUEUE_RENDER_COMMAND(InitPassthroughGPUSkinVertexFactory)(
 		[NewPassthroughVertexFactory, SourceVertexFactory](FRHICommandList& RHICmdList)
 		{
 			SourceVertexFactory->CopyDataTypeForPassthroughFactory(NewPassthroughVertexFactory);
@@ -968,7 +968,7 @@ static VertexFactoryType* CreateVertexFactoryMorph(TArray<TUniquePtr<VertexFacto
 	TDynamicUpdateVertexFactoryData<VertexFactoryType> VertexUpdateData(VertexFactory, InVertexBuffers);
 
 	// update vertex factory components and sync it
-	EnqueueUniqueRenderCommand("InitGPUSkinVertexFactoryMorph",
+	ENQUEUE_RENDER_COMMAND(InitGPUSkinVertexFactoryMorph)(
 		[VertexUpdateData](FRHICommandList& RHICmdList)
 		{
 			typename VertexFactoryType::FDataType Data;
@@ -1005,7 +1005,7 @@ static void CreateVertexFactoryCloth(TArray<TUniquePtr<VertexFactoryTypeBase>>&
 	TDynamicUpdateVertexFactoryData<VertexFactoryType> VertexUpdateData(VertexFactory, InVertexBuffers);
 
 	// update vertex factory components and sync it
-	EnqueueUniqueRenderCommand("InitGPUSkinAPEXClothVertexFactory",
+	ENQUEUE_RENDER_COMMAND(InitGPUSkinAPEXClothVertexFactory)(
 		[VertexUpdateData](FRHICommandList& RHICmdList)
 		{
 			typename VertexFactoryType::FDataType Data;
diff --git a/Engine/Source/Runtime/Engine/Private/TileRendering.cpp b/Engine/Source/Runtime/Engine/Private/TileRendering.cpp
--- a/Engine/Source/Runtime/Engine/Private/TileRendering.cpp
+++ b/Engine/Source/Runtime/Engine/Private/TileRendering.cpp
@@ -342,8 +342,7 @@ bool FCanvasTileRendererItem::Render_GameThread(const FCanvas* Canvas)
 		bNeedsToSwitchVerticalAxis ? uint32(1) : uint32(0),
 		Canvas->GetAllowedModes()
 	};
-	EnqueueUniqueRenderCommand(
-		"DrawTileCommand",
+	ENQUEUE_RENDER_COMMAND(DrawTileCommand)(
 		[DrawTileParameters](FRHICommandListImmediate& RHICmdList)
 		{
 			SCOPED_DRAW_EVENTF(RHICmdList, CanvasDrawTile, *DrawTileParameters.RenderData->MaterialRenderProxy->GetMaterial(GMaxRHIFeatureLevel)->GetFriendlyName());
diff --git a/Engine/Source/Runtime/Engine/Private/UserInterface/Canvas.cpp b/Engine/Source/Runtime/Engine/Private/UserInterface/Canvas.cpp
--- a/Engine/Source/Runtime/Engine/Private/UserInterface/Canvas.cpp
+++ b/Engine/Source/Runtime/Engine/Private/UserInterface/Canvas.cpp
@@ -455,7 +455,7 @@ bool FCanvasBatchedElementRenderItem::Render_GameThread(const FCanvas* Canvas)
 			Canvas->GetFeatureLevel(),
 			Canvas->GetShaderPlatform()
 		};
-		EnqueueUniqueRenderCommand("BatchedDrawCommand",
+		ENQUEUE_RENDER_COMMAND(BatchedDrawCommand)(
 			[DrawParameters](FRHICommandList& RHICmdList)
 			{
 				// draw batched items
@@ -758,7 +758,7 @@ void FCanvas::Flush_GameThread(bool bForce)
 	};
 	bool bEmitCanvasDrawEvents = GEmitDrawEvents;
 
-	EnqueueUniqueRenderCommand("CanvasFlushSetupCommand",
+	ENQUEUE_RENDER_COMMAND(CanvasFlushSetupCommand)(
 		[FlushParameters](FRHICommandList& RHICmdList)
 		{
 			// Set the RHI render target.
@@ -900,7 +900,7 @@ void FCanvas::SetRenderTargetRect( const FIntRect& InViewRect )
 void FCanvas::Clear(const FLinearColor& ClearColor)
 {
 	FRenderTarget* CanvasRenderTarget = GetRenderTarget();
-	EnqueueUniqueRenderCommand("ClearCommand",
+	ENQUEUE_RENDER_COMMAND(ClearCommand)(
 		[ClearColor, CanvasRenderTarget](FRHICommandList& RHICmdList)
 		{
 			SCOPED_DRAW_EVENT(RHICmdList, CanvasClear);
diff --git a/Engine/Source/Runtime/RenderCore/Public/RenderingThread.h b/Engine/Source/Runtime/RenderCore/Public/RenderingThread.h
--- a/Engine/Source/Runtime/RenderCore/Public/RenderingThread.h
+++ b/Engine/Source/Runtime/RenderCore/Public/RenderingThread.h
@@ -178,7 +178,7 @@ DECLARE_STATS_GROUP(TEXT("Render Thread Commands"), STATGROUP_RenderThreadComman
 			RETURN_QUICK_DECLARE_CYCLE_STAT(TypeName, STATGROUP_RenderThreadCommands); \
 		}
 
-template<typename LAMBDA>
+template<typename TSTR, typename LAMBDA>
 class TEnqueueUniqueRenderCommandType : public FRenderCommand
 {
 public:
@@ -193,7 +193,6 @@ public:
 	FORCEINLINE_DEBUGGABLE TStatId GetStatId() const
 	{
 #if STATS
-		FStat_EnqueueUniqueRenderCommandType::TypeName = TypeName;
 		static struct FThreadSafeStaticStat<FStat_EnqueueUniqueRenderCommandType> StatPtr_EnqueueUniqueRenderCommandType;
 		return StatPtr_EnqueueUniqueRenderCommandType.GetStatId();
 #else
@@ -208,11 +207,11 @@ private:
 		typedef FStatGroup_STATGROUP_RenderThreadCommands TGroup;
 		static FORCEINLINE const char* GetStatName()
 		{
-			return TypeName;
+			return TSTR::CStr();
 		}
 		static FORCEINLINE const TCHAR* GetDescription()
 		{
-			return ANSI_TO_TCHAR(TypeName);
+			return TSTR::TStr();
 		}
 		static FORCEINLINE EStatDataType::Type GetStatType()
 		{
@@ -230,34 +229,20 @@ private:
 		{
 			return FPlatformMemory::MCR_Invalid;
 		}
-
-		static const char* TypeName;
 	};
 #endif
 
-public:
-	static const char* TypeName;
-
 private:
 	LAMBDA Lambda;
 };
 
-template<typename LAMBDA> 
-const char* TEnqueueUniqueRenderCommandType<LAMBDA>::TypeName = "Undefined_TEnqueueUniqueRenderCommandType";
-
-#if STATS
-template<typename LAMBDA>
-const char* TEnqueueUniqueRenderCommandType<LAMBDA>::FStat_EnqueueUniqueRenderCommandType::TypeName = "Undefined_TEnqueueUniqueRenderCommandType";
-#endif
-
-template<typename LAMBDA>
-FORCEINLINE_DEBUGGABLE void EnqueueUniqueRenderCommand(const char* TypeName, LAMBDA&& Lambda)
+template<typename TSTR, typename LAMBDA>
+FORCEINLINE_DEBUGGABLE void EnqueueUniqueRenderCommand(LAMBDA&& Lambda)
 {
-	typedef TEnqueueUniqueRenderCommandType<LAMBDA> EURCType;
-	EURCType::TypeName = TypeName;
+	typedef TEnqueueUniqueRenderCommandType<TSTR, LAMBDA> EURCType;
 
 #if 0 // UE_SERVER && UE_BUILD_DEBUG
-	UE_LOG(LogRHI, Warning, TEXT("Render command '%s' is being executed on a dedicated server."), ANSI_TO_TCHAR(TypeName))
+	UE_LOG(LogRHI, Warning, TEXT("Render command '%s' is being executed on a dedicated server."), TSTR::TStr())
 #endif
 	if (ShouldExecuteOnRenderThread())
 	{
@@ -272,19 +257,18 @@ FORCEINLINE_DEBUGGABLE void EnqueueUniqueRenderCommand(const char* TypeName, LAM
 	}
 }
 
-namespace Detail
-{
-	template<typename T>
-	struct FalseType
-	{
-		static CONSTEXPR const bool False = false;
-	};
-}
+#define ENQUEUE_RENDER_COMMAND(Type) \
+	struct Type##Name \
+	{  \
+		static const char* CStr() { return #Type; } \
+		static const TCHAR* TStr() { return TEXT(#Type); } \
+	}; \
+	EnqueueUniqueRenderCommand<Type##Name>
 
 template<typename LAMBDA>
-FORCEINLINE_DEBUGGABLE void EnqueueUniqueRenderCommand(const char* TypeName, LAMBDA& Lambda)
+FORCEINLINE_DEBUGGABLE void EnqueueUniqueRenderCommand(LAMBDA& Lambda)
 {
-	static_assert(Detail::FalseType<LAMBDA>::False, "EnqueueUniqueRenderCommand enforces use of rvalue and therefore move to avoid an extra copy of the Lambda");
+	static_assert(sizeof(LAMBDA) == 0, "EnqueueUniqueRenderCommand enforces use of rvalue and therefore move to avoid an extra copy of the Lambda");
 }
 
 /**
diff --git a/Engine/Source/Runtime/Renderer/Private/RendererScene.cpp b/Engine/Source/Runtime/Renderer/Private/RendererScene.cpp
--- a/Engine/Source/Runtime/Renderer/Private/RendererScene.cpp
+++ b/Engine/Source/Runtime/Renderer/Private/RendererScene.cpp
@@ -2070,7 +2070,7 @@ void FScene::UpdateSpeedTreeWind(double CurrentTime)
 	UniformParameters.Prev##name = *(FVector4*)(WindShaderValues + FSpeedTreeWind::offset + FSpeedTreeWind::NUM_SHADER_VALUES);
 
 	FScene* Scene = this;
-	EnqueueUniqueRenderCommand("FUpdateSpeedTreeWindCommand",
+	ENQUEUE_RENDER_COMMAND(FUpdateSpeedTreeWindCommand)(
 		[Scene, CurrentTime](FRHICommandListImmediate& RHICmdList)
 		{
 			FVector WindDirection;
