Merging //UE4/Private-Geometry to Dev-VREditor (//UE4/Dev-VREditor) 3339870

diff --git a/Engine/Source/Runtime/Core/Private/Misc/FbxErrors.cpp b/Engine/Source/Runtime/Core/Private/Misc/FbxErrors.cpp
--- a/Engine/Source/Runtime/Core/Private/Misc/FbxErrors.cpp
+++ b/Engine/Source/Runtime/Core/Private/Misc/FbxErrors.cpp
@@ -77,8 +77,6 @@ FName FFbxErrors::SkeletalMesh_RestoreSortingForSectionNumber(TEXT("SkeletalMesh
 
 FName FFbxErrors::SkeletalMesh_NoMeshFoundOnRoot(TEXT("SkeletalMesh_NoMeshFoundOnRoot"));
 
-FName FFbxErrors::SkeletalMesh_ImportSubDSurface(TEXT("SkeletalMesh_ImportSubDSurface"));
-
 FName FFbxErrors::SkeletalMesh_InvalidRoot(TEXT("SkeletalMesh_InvalidRoot"));
 
 FName FFbxErrors::SkeletalMesh_InvalidBone(TEXT("SkeletalMesh_InvalidBone"));
diff --git a/Engine/Source/Runtime/Core/Private/UObject/DevObjectVersion.cpp b/Engine/Source/Runtime/Core/Private/UObject/DevObjectVersion.cpp
--- a/Engine/Source/Runtime/Core/Private/UObject/DevObjectVersion.cpp
+++ b/Engine/Source/Runtime/Core/Private/UObject/DevObjectVersion.cpp
@@ -14,6 +14,7 @@
 #include "UObject/RenderingObjectVersion.h"
 #include "UObject/SequencerObjectVersion.h"
 #include "UObject/VRObjectVersion.h"
+#include "UObject/GeometryObjectVersion.h"
 
 DEFINE_LOG_CATEGORY_STATIC(LogDevObjectVersion, Log, All);
 
@@ -108,3 +109,10 @@ FDevVersionRegistration GRegisterVRObjectVersion(FVRObjectVersion::GUID, FVRObje
 const FGuid FLoadTimesObjectVersion::GUID(0xC2A15278, 0xBFE74AFE, 0x6C1790FF, 0x531DF755);
 // Register LoadTimes custom version with Core
 FDevVersionRegistration GRegisterLoadTimesObjectVersion(FLoadTimesObjectVersion::GUID, FLoadTimesObjectVersion::LatestVersion, TEXT("Dev-LoadTimes"));
+
+
+// Unique Geometry Object version id
+const FGuid FGeometryObjectVersion::GUID(0x6EACA3D4, 0x40EC4CC1, 0xb7868BED, 0x9428FC5);
+// Register Geometry custom version with Core
+FDevVersionRegistration GRegisterGeometryObjectVersion(FGeometryObjectVersion::GUID, FGeometryObjectVersion::LatestVersion, TEXT("Private-Geometry"));
+
diff --git a/Engine/Source/Runtime/Core/Public/Misc/FbxErrors.h b/Engine/Source/Runtime/Core/Public/Misc/FbxErrors.h
--- a/Engine/Source/Runtime/Core/Public/Misc/FbxErrors.h
+++ b/Engine/Source/Runtime/Core/Public/Misc/FbxErrors.h
@@ -89,8 +89,6 @@ struct CORE_API FFbxErrors
 
 	static FName SkeletalMesh_NoMeshFoundOnRoot;
 	
-	static FName SkeletalMesh_ImportSubDSurface;
-
 	static FName SkeletalMesh_InvalidRoot;
 
 	static FName SkeletalMesh_InvalidBone;
diff --git a/Engine/Source/Runtime/Engine/Classes/Components/StaticMeshComponent.h b/Engine/Source/Runtime/Engine/Classes/Components/StaticMeshComponent.h
--- a/Engine/Source/Runtime/Engine/Classes/Components/StaticMeshComponent.h
+++ b/Engine/Source/Runtime/Engine/Classes/Components/StaticMeshComponent.h
@@ -537,6 +537,9 @@ public:
 	/** Whether or not the component supports default collision from its static mesh asset */
 	virtual bool SupportsDefaultCollision();
 
+	/** Whether we can support dithered LOD transitions (default behavior checks all materials). Used for HISMC LOD. */
+	virtual bool SupportsDitheredLODTransitions();
+
 private:
 	/** Initializes the resources used by the static mesh component. */
 	void InitResources();
diff --git a/Engine/Source/Runtime/Engine/Classes/Engine/StaticMesh.h b/Engine/Source/Runtime/Engine/Classes/Engine/StaticMesh.h
--- a/Engine/Source/Runtime/Engine/Classes/Engine/StaticMesh.h
+++ b/Engine/Source/Runtime/Engine/Classes/Engine/StaticMesh.h
@@ -238,6 +238,7 @@ struct FMeshSectionInfoMap
 	GENERATED_USTRUCT_BODY()
 
 	/** Maps an LOD+Section to the material it should render with. */
+	UPROPERTY()
 	TMap<uint32,FMeshSectionInfo> Map;
 
 	/** Serialize. */
@@ -426,7 +427,7 @@ class UStaticMesh : public UObject, public IInterface_CollisionDataProvider, pub
 	FMeshSectionInfoMap SectionInfoMap;
 
 	/** The LOD group to which this mesh belongs. */
-	UPROPERTY(AssetRegistrySearchable)
+	UPROPERTY(EditAnywhere, AssetRegistrySearchable, Category=StaticMesh)
 	FName LODGroup;
 
 	/** If true, the screen sizees at which LODs swap are computed automatically. */
@@ -604,7 +605,7 @@ public:
 	ENGINE_API virtual void PreEditChange(UProperty* PropertyAboutToChange) override;
 	ENGINE_API virtual void PostEditChangeProperty(FPropertyChangedEvent& PropertyChangedEvent) override;
 	ENGINE_API virtual void GetAssetRegistryTagMetadata(TMap<FName, FAssetRegistryTagMetadata>& OutMetadata) const override;
-	ENGINE_API void SetLODGroup(FName NewGroup);
+	ENGINE_API void SetLODGroup(FName NewGroup, bool bRebuildImmediately = true);
 	ENGINE_API void BroadcastNavCollisionChange();
 #endif // WITH_EDITOR
 	ENGINE_API virtual void Serialize(FArchive& Ar) override;
diff --git a/Engine/Source/Runtime/Engine/Classes/Materials/Material.h b/Engine/Source/Runtime/Engine/Classes/Materials/Material.h
--- a/Engine/Source/Runtime/Engine/Classes/Materials/Material.h
+++ b/Engine/Source/Runtime/Engine/Classes/Materials/Material.h
@@ -447,9 +447,9 @@ public:
 	UPROPERTY(EditAnywhere, Category=Material)
 	uint32 TwoSided:1;
 
-	/** Whether the material should support a dithered LOD transition. */
-	UPROPERTY(EditAnywhere, Category=Material, AdvancedDisplay)
-	uint32 DitheredLODTransition:1;
+	/** Whether meshes rendered with the material should support dithered LOD transitions. */
+	UPROPERTY(EditAnywhere, Category = Material, AdvancedDisplay, meta = (DisplayName = "Dithered LOD Transition"))
+	uint32 DitheredLODTransition : 1;
 
 	/** Dither opacity mask. When combined with Temporal AA this can be used as a form of limited translucency which supports all lighting features. */
 	UPROPERTY(EditAnywhere, Category=Material, AdvancedDisplay)
diff --git a/Engine/Source/Runtime/Engine/Engine.Build.cs b/Engine/Source/Runtime/Engine/Engine.Build.cs
--- a/Engine/Source/Runtime/Engine/Engine.Build.cs
+++ b/Engine/Source/Runtime/Engine/Engine.Build.cs
@@ -102,7 +102,7 @@ public class Engine : ModuleRules
                 "MaterialShaderQualitySettings",
                 "CinematicCamera",
 				"Analytics",
-				"AnalyticsET",
+				"AnalyticsET"
             }
 		);
 
@@ -336,8 +336,7 @@ public class Engine : ModuleRules
 				"UEOgg",
 				"Vorbis",
 				"VorbisFile",
-				"libOpus",
-			    "OpenSubdiv"
+				"libOpus"
 				);
 
 			if (UEBuildConfiguration.bCompileLeanAndMeanUE == false)
diff --git a/Engine/Source/Runtime/Engine/Private/Components/StaticMeshComponent.cpp b/Engine/Source/Runtime/Engine/Private/Components/StaticMeshComponent.cpp
--- a/Engine/Source/Runtime/Engine/Private/Components/StaticMeshComponent.cpp
+++ b/Engine/Source/Runtime/Engine/Private/Components/StaticMeshComponent.cpp
@@ -1580,6 +1580,20 @@ bool UStaticMeshComponent::SupportsDefaultCollision()
 	return GetStaticMesh() && GetBodySetup() == GetStaticMesh()->BodySetup;
 }
 
+bool UStaticMeshComponent::SupportsDitheredLODTransitions()
+{
+	// Only support dithered transitions if all materials do.
+	TArray<class UMaterialInterface*> Materials = GetMaterials();
+	for (UMaterialInterface* Material : Materials)
+	{
+		if (Material && !Material->IsDitheredLODTransition())
+		{
+			return false;
+		}
+	}
+	return true;
+}
+
 void UStaticMeshComponent::UpdateCollisionFromStaticMesh()
 {
 	if(bUseDefaultCollision && SupportsDefaultCollision())
diff --git a/Engine/Source/Runtime/Engine/Private/DrawDebugHelpers.cpp b/Engine/Source/Runtime/Engine/Private/DrawDebugHelpers.cpp
--- a/Engine/Source/Runtime/Engine/Private/DrawDebugHelpers.cpp
+++ b/Engine/Source/Runtime/Engine/Private/DrawDebugHelpers.cpp
@@ -436,8 +436,8 @@ void DrawDebug2DDonut(const UWorld* InWorld, const FMatrix& TransformMatrix, flo
 
 			// Need at least 4 segments
 			Segments = FMath::Max((Segments - 4) / 2, 4);
-			InternalDrawDebugCircle(InWorld, TransformMatrix, InnerRadius, Segments, Color, bPersistentLines, LifeTime, DepthPriority);
-			InternalDrawDebugCircle(InWorld, TransformMatrix, OuterRadius, Segments, Color, bPersistentLines, LifeTime, DepthPriority);
+			InternalDrawDebugCircle(InWorld, TransformMatrix, InnerRadius, Segments, Color, bPersistentLines, LifeTime, DepthPriority, Thickness);
+			InternalDrawDebugCircle(InWorld, TransformMatrix, OuterRadius, Segments, Color, bPersistentLines, LifeTime, DepthPriority, Thickness );
 		
 			const FVector Center = TransformMatrix.GetOrigin();
 			const FVector AxisY = TransformMatrix.GetScaledAxis( EAxis::Y );
diff --git a/Engine/Source/Runtime/Engine/Private/HierarchicalInstancedStaticMesh.cpp b/Engine/Source/Runtime/Engine/Private/HierarchicalInstancedStaticMesh.cpp
--- a/Engine/Source/Runtime/Engine/Private/HierarchicalInstancedStaticMesh.cpp
+++ b/Engine/Source/Runtime/Engine/Private/HierarchicalInstancedStaticMesh.cpp
@@ -27,6 +27,7 @@
 #include "StaticMeshResources.h"
 #include "Components/HierarchicalInstancedStaticMeshComponent.h"
 #include "InstancedStaticMesh.h"
+#include "SceneManagement.h"
 
 static TAutoConsoleVariable<int32> CVarFoliageSplitFactor(
 	TEXT("foliage.SplitFactor"),
@@ -747,6 +748,7 @@ class FHierarchicalStaticMeshSceneProxy : public FInstancedStaticMeshSceneProxy
 	TMap<uint32, FFoliageOcclusionResults> OcclusionResults;
 	bool bIsGrass;
 	uint32 SceneProxyCreatedFrameNumberRenderThread;
+	bool bDitheredLODTransitions;
 
 #if !(UE_BUILD_SHIPPING || UE_BUILD_TEST)
 	mutable TArray<uint32> SingleDebugRuns[MAX_STATIC_MESH_LODS];
@@ -767,6 +769,7 @@ public:
 		, LastUnbuiltIndex(InComponent->GetNumRenderInstances()-1)
 		, bIsGrass(bInIsGrass)
 		, SceneProxyCreatedFrameNumberRenderThread(UINT32_MAX)
+		, bDitheredLODTransitions(InComponent->SupportsDitheredLODTransitions())
 #if !(UE_BUILD_SHIPPING || UE_BUILD_TEST)
 		, CaptureTag(0)
 #endif
@@ -783,6 +786,8 @@ public:
 		, FirstUnbuiltIndex(InComponent->NumBuiltRenderInstances)
 		, LastUnbuiltIndex(InComponent->GetNumRenderInstances()-1)
 		, bIsGrass(bInIsGrass)
+		, SceneProxyCreatedFrameNumberRenderThread(UINT32_MAX)
+		, bDitheredLODTransitions(InComponent->SupportsDitheredLODTransitions())
 #if !(UE_BUILD_SHIPPING || UE_BUILD_TEST)
 		, CaptureTag(0)
 #endif
@@ -1315,7 +1320,7 @@ void FHierarchicalStaticMeshSceneProxy::GetDynamicMeshElements(const TArray<cons
 
 	QUICK_SCOPE_CYCLE_COUNTER(STAT_HierarchicalInstancedStaticMeshSceneProxy_GetMeshElements);
 
-	bool bMultipleSections = ALLOW_DITHERED_LOD_FOR_INSTANCED_STATIC_MESHES && CVarDitheredLOD.GetValueOnRenderThread() > 0;
+	bool bMultipleSections = ALLOW_DITHERED_LOD_FOR_INSTANCED_STATIC_MESHES && bDitheredLODTransitions && CVarDitheredLOD.GetValueOnRenderThread() > 0;
 	bool bSingleSections = !bMultipleSections;
 	bool bOverestimate = CVarOverestimateLOD.GetValueOnRenderThread() > 0;
 
@@ -1327,6 +1332,9 @@ void FHierarchicalStaticMeshSceneProxy::GetDynamicMeshElements(const TArray<cons
 		{
 			const FSceneView* View = Views[ViewIndex];
 
+			// With stereo rendering views may have asymmetric FOVs, so always use the left eye's view for LOD calculation
+			const FSceneView* ViewForLOD = ( View->Family != nullptr && View->StereoPass == eSSP_RIGHT_EYE ) ? View->Family->Views[ 0 ] : View;
+
 			FFoliageElementParams ElementParams;
 			ElementParams.bSelectionRenderEnabled = GIsEditor && ViewFamily.EngineShowFlags.Selection;
 			ElementParams.NumSelectionGroups = (ElementParams.bSelectionRenderEnabled && bHasSelectedInstances) ? 2 : 1;
@@ -1395,7 +1403,7 @@ void FHierarchicalStaticMeshSceneProxy::GetDynamicMeshElements(const TArray<cons
 					}
 					else
 					{
-						const FMatrix LocalViewProjForCulling = GetLocalToWorld() * View->ViewMatrices.GetViewProjectionMatrix();
+						const FMatrix LocalViewProjForCulling = GetLocalToWorld() * ViewForLOD->ViewMatrices.GetViewProjectionMatrix();
 						GetViewFrustumBounds(InstanceParams.ViewFrustumLocal, LocalViewProjForCulling, false);
 					}
 
@@ -1448,39 +1456,25 @@ void FHierarchicalStaticMeshSceneProxy::GetDynamicMeshElements(const TArray<cons
 				float MaxDrawDistanceScale = GetCachedScalabilityCVars().ViewDistanceScale;
 				float SphereRadius = RenderData->Bounds.SphereRadius;
 
-				for (int32 LODIndex = 0; LODIndex < InstanceParams.LODs; LODIndex++)
+				float FinalCull = MAX_flt;
+				if (MinSize > 0.0)
 				{
-					InstanceParams.LODPlanesMax[LODIndex] = MIN_flt;
-					InstanceParams.LODPlanesMin[LODIndex] = MAX_flt;
+					FinalCull = ComputeBoundsDrawDistance(MinSize, SphereRadius, ViewForLOD->ViewMatrices.GetProjectionMatrix()) * LODScale;
 				}
-				ElementParams.FinalCullDistance = MIN_flt;
-				for (int32 SampleIndex = 0; SampleIndex < 2; SampleIndex++)
+				if (UserData_AllInstances.EndCullDistance > 0.0f)
 				{
-					// we aren't going to accurately figure out LOD planes for both samples and try all 4 possibilities, rather we just consider the worst case LOD range; the shaders will render correctly
-					float Fac = View->GetTemporalLODDistanceFactor(SampleIndex, bMultipleSections) * SphereRadius * LODScale;
-
-					float FinalCull = MAX_flt;
-
-					if (MinSize > 0.0)
-					{
-						FinalCull = FMath::Sqrt(Fac / MinSize);
-					}
-					if (UserData_AllInstances.EndCullDistance > 0.0f && UserData_AllInstances.EndCullDistance < FinalCull)
-					{
-						FinalCull = UserData_AllInstances.EndCullDistance;
-					}
-					FinalCull *= MaxDrawDistanceScale;
-					ElementParams.FinalCullDistance = FMath::Max(ElementParams.FinalCullDistance, FinalCull);
+					FinalCull = FMath::Min(FinalCull, UserData_AllInstances.EndCullDistance * MaxDrawDistanceScale);
+				}
+				ElementParams.FinalCullDistance = FinalCull;
 
-					for (int32 LODIndex = 1; LODIndex < InstanceParams.LODs; LODIndex++)
-					{
-						float Val = FMath::Min(FMath::Sqrt(Fac / RenderData->ScreenSize[LODIndex]), FinalCull);
-						InstanceParams.LODPlanesMin[LODIndex - 1] = FMath::Min(InstanceParams.LODPlanesMin[LODIndex - 1], Val - LODRandom);
-						InstanceParams.LODPlanesMax[LODIndex - 1] = FMath::Max(InstanceParams.LODPlanesMax[LODIndex - 1], Val);
-					}
-					InstanceParams.LODPlanesMin[InstanceParams.LODs - 1] = FMath::Min(InstanceParams.LODPlanesMin[InstanceParams.LODs - 1], FinalCull - LODRandom);
-					InstanceParams.LODPlanesMax[InstanceParams.LODs - 1] = FMath::Max(InstanceParams.LODPlanesMax[InstanceParams.LODs - 1], FinalCull);
+				for (int32 LODIndex = 1; LODIndex < InstanceParams.LODs; LODIndex++)
+				{
+					float Distance = ComputeBoundsDrawDistance(RenderData->ScreenSize[LODIndex], SphereRadius, ViewForLOD->ViewMatrices.GetProjectionMatrix()) * LODScale;
+					InstanceParams.LODPlanesMin[LODIndex - 1] = Distance - LODRandom;
+					InstanceParams.LODPlanesMax[LODIndex - 1] = Distance;
 				}
+				InstanceParams.LODPlanesMin[InstanceParams.LODs - 1] = FinalCull - LODRandom;
+				InstanceParams.LODPlanesMax[InstanceParams.LODs - 1] = FinalCull;
 
 				for (int32 LODIndex = 0; LODIndex < InstanceParams.LODs; LODIndex++)
 				{
@@ -1628,38 +1622,26 @@ void FHierarchicalStaticMeshSceneProxy::GetDynamicMeshElements(const TArray<cons
 						const float SphereRadius = RenderData->Bounds.SphereRadius;
 
 						checkSlow(NumLODs > 0);
-						for (int32 LODIndex = 0; LODIndex < NumLODs; LODIndex++)
+						
+						float FinalCull = MAX_flt;
+						if (MinSize > 0.0)
 						{
-							LODPlanesMax[LODIndex] = MIN_flt;
-							LODPlanesMin[LODIndex] = MAX_flt;
+							FinalCull = ComputeBoundsDrawDistance(MinSize, SphereRadius, ViewForLOD->ViewMatrices.GetProjectionMatrix()) * LODScale;
 						}
-						ElementParams.FinalCullDistance = MIN_flt;
-						for (int32 SampleIndex = 0; SampleIndex < 2; SampleIndex++)
+						if (UserData_AllInstances.EndCullDistance > 0.0f)
 						{
-							// we aren't going to accurately figure out LOD planes for both samples and try all 4 possibilities, rather we just consider the worst case LOD range; the shaders will render correctly
-							float Fac = View->GetTemporalLODDistanceFactor(SampleIndex, bMultipleSections) * SphereRadius * LODScale;
-							float FinalCull = MAX_flt;
-
-							if (MinSize > 0.0)
-							{
-								FinalCull = FMath::Sqrt(Fac / MinSize);
-							}
-							if (UserData_AllInstances.EndCullDistance > 0.0f && UserData_AllInstances.EndCullDistance < FinalCull)
-							{
-								FinalCull = UserData_AllInstances.EndCullDistance;
-							}
-							FinalCull *= MaxDrawDistanceScale;
-							ElementParams.FinalCullDistance = FMath::Max(ElementParams.FinalCullDistance, FinalCull);
+							FinalCull = FMath::Min(FinalCull, UserData_AllInstances.EndCullDistance * MaxDrawDistanceScale);
+						}
+						ElementParams.FinalCullDistance = FinalCull;
 
-							for (int32 LODIndex = 1; LODIndex < NumLODs; LODIndex++)
-							{
-								float Val = FMath::Min(FMath::Sqrt(Fac / RenderData->ScreenSize[LODIndex]), FinalCull);
-								LODPlanesMin[LODIndex - 1] = FMath::Min(LODPlanesMin[LODIndex - 1], Val - LODRandom);
-								LODPlanesMax[LODIndex - 1] = FMath::Max(LODPlanesMax[LODIndex - 1], Val);
-							}
-							LODPlanesMin[NumLODs - 1] = FMath::Min(LODPlanesMin[NumLODs - 1], FinalCull - LODRandom);
-							LODPlanesMax[NumLODs - 1] = FMath::Max(LODPlanesMax[NumLODs - 1], FinalCull);
+						for (int32 LODIndex = 1; LODIndex < NumLODs; LODIndex++)
+						{
+							float Distance = ComputeBoundsDrawDistance(RenderData->ScreenSize[LODIndex], SphereRadius, ViewForLOD->ViewMatrices.GetProjectionMatrix()) * LODScale;
+							LODPlanesMin[LODIndex - 1] = Distance - LODRandom;
+							LODPlanesMax[LODIndex - 1] = Distance;
 						}
+						LODPlanesMin[NumLODs - 1] = FinalCull - LODRandom;
+						LODPlanesMax[NumLODs - 1] = FinalCull;
 
 						// calculate runs
 						int32 MinLOD = 0;
diff --git a/Engine/Source/Runtime/Engine/Private/InstancedStaticMesh.cpp b/Engine/Source/Runtime/Engine/Private/InstancedStaticMesh.cpp
--- a/Engine/Source/Runtime/Engine/Private/InstancedStaticMesh.cpp
+++ b/Engine/Source/Runtime/Engine/Private/InstancedStaticMesh.cpp
@@ -22,6 +22,7 @@
 #include "PhysicsEngine/BodySetup.h"
 #include "GameFramework/WorldSettings.h"
 #include "ComponentRecreateRenderStateContext.h"
+#include "SceneManagement.h"
 
 const int32 InstancedStaticMeshMaxTexCoord = 8;
 
@@ -2013,6 +2014,9 @@ void FInstancedStaticMeshVertexFactoryShaderParameters::SetMesh( FRHICommandList
 		InstancingWorldViewOriginOne.W = 1.0f;
 		if (InstancingUserData && BatchElement.InstancedLODRange)
 		{
+			// With stereo rendering views may have asymmetric FOVs, so always use the left eye's view for LOD calculation
+			const FSceneView& ViewForLOD = ( View.Family != nullptr && View.StereoPass == eSSP_RIGHT_EYE ) ? *View.Family->Views[ 0 ] : View;
+
 			int32 FirstLOD = InstancingUserData->MinLOD;
 
 			int32 DebugMin = FMath::Min(CVarMinLOD.GetValueOnRenderThread(), InstancingUserData->MeshRenderData->LODResources.Num() - 1);
@@ -2048,34 +2052,31 @@ void FInstancedStaticMeshVertexFactoryShaderParameters::SetMesh( FRHICommandList
 			for (int32 SampleIndex = 0; SampleIndex < 2; SampleIndex++)
 			{
 				FVector4& InstancingViewZCompare(SampleIndex ? InstancingViewZCompareOne : InstancingViewZCompareZero);
-				float Fac = View.GetTemporalLODDistanceFactor(SampleIndex) * SphereRadius * LODScale;
 
 				float FinalCull = MAX_flt;
 				if (MinSize > 0.0)
 				{
-					FinalCull = FMath::Sqrt(Fac / MinSize);
+					FinalCull = ComputeBoundsDrawDistance(MinSize, SphereRadius, ViewForLOD.ViewMatrices.GetProjectionMatrix()) * LODScale;
 				}
-				if (InstancingUserData->EndCullDistance > 0.0f && InstancingUserData->EndCullDistance < FinalCull)
+				if (InstancingUserData->EndCullDistance > 0.0f)
 				{
-					FinalCull = InstancingUserData->EndCullDistance;
+					FinalCull = FMath::Min(FinalCull, InstancingUserData->EndCullDistance * MaxDrawDistanceScale);
 				}
-				FinalCull *= MaxDrawDistanceScale;
 
 				InstancingViewZCompare.Z = FinalCull;
 				if (BatchElement.InstancedLODIndex < InstancingUserData->MeshRenderData->LODResources.Num() - 1)
 				{
-					float NextCut = InstancingUserData->MeshRenderData->ScreenSize[BatchElement.InstancedLODIndex + 1];
-					InstancingViewZCompare.Z = FMath::Min(FMath::Sqrt(Fac / NextCut), FinalCull);
+					float NextCut = ComputeBoundsDrawDistance(InstancingUserData->MeshRenderData->ScreenSize[BatchElement.InstancedLODIndex + 1], SphereRadius, ViewForLOD.ViewMatrices.GetProjectionMatrix()) * LODScale;
+					InstancingViewZCompare.Z = FMath::Min(NextCut, FinalCull);
 				}
 
 				InstancingViewZCompare.X = MIN_flt;
 				if (BatchElement.InstancedLODIndex > FirstLOD)
 				{
-					float CurCut = FMath::Sqrt(Fac / InstancingUserData->MeshRenderData->ScreenSize[BatchElement.InstancedLODIndex]);
+					float CurCut = ComputeBoundsDrawDistance(InstancingUserData->MeshRenderData->ScreenSize[BatchElement.InstancedLODIndex], SphereRadius, ViewForLOD.ViewMatrices.GetProjectionMatrix()) * LODScale;
 					if (CurCut < FinalCull)
 					{
 						InstancingViewZCompare.Y = CurCut;
-
 					}
 					else
 					{
@@ -2095,7 +2096,7 @@ void FInstancedStaticMeshVertexFactoryShaderParameters::SetMesh( FRHICommandList
 			InstancingWorldViewOriginZero = View.GetTemporalLODOrigin(0);
 			InstancingWorldViewOriginOne = View.GetTemporalLODOrigin(1);
 
-			float Alpha = View.GetTemporalLODTransition();
+			float Alpha = ViewForLOD.GetTemporalLODTransition();
 			InstancingWorldViewOriginZero.W = 1.0f - Alpha;
 			InstancingWorldViewOriginOne.W = Alpha;
 
diff --git a/Engine/Source/Runtime/Engine/Private/SceneManagement.cpp b/Engine/Source/Runtime/Engine/Private/SceneManagement.cpp
--- a/Engine/Source/Runtime/Engine/Private/SceneManagement.cpp
+++ b/Engine/Source/Runtime/Engine/Private/SceneManagement.cpp
@@ -375,7 +375,9 @@ int8 ComputeTemporalStaticMeshLOD( const FStaticMeshRenderData* RenderData, cons
 {
 	const int32 NumLODs = MAX_STATIC_MESH_LODS;
 
-	const float ScreenRadiusSquared = ComputeTemporalLODBoundsScreenRadiusSquared(Origin, SphereRadius, View, SampleIndex) * FactorScale * FactorScale;
+	// With stereo rendering views may have asymmetric FOVs, so always use the left eye's view for LOD calculation
+	const FSceneView& ViewForLOD = ( View.Family != nullptr && View.StereoPass == eSSP_RIGHT_EYE ) ? *View.Family->Views[ 0 ] : View;
+	const float ScreenRadiusSquared = ComputeTemporalLODBoundsScreenRadiusSquared(Origin, SphereRadius, ViewForLOD, SampleIndex) * FactorScale * FactorScale;
 
 	// Walk backwards and return the first matching LOD
 	for(int32 LODIndex = NumLODs - 1 ; LODIndex >= 0 ; --LODIndex)
@@ -395,7 +397,9 @@ int8 ComputeStaticMeshLOD( const FStaticMeshRenderData* RenderData, const FVecto
 	{
 		const int32 NumLODs = MAX_STATIC_MESH_LODS;
 
-		const float ScreenRadiusSquared = ComputeBoundsScreenRadiusSquared(Origin, SphereRadius, View) * FactorScale * FactorScale;
+		// With stereo rendering views may have asymmetric FOVs, so always use the left eye's view for LOD calculation
+		const FSceneView& ViewForLOD = ( View.Family != nullptr && View.StereoPass == eSSP_RIGHT_EYE ) ? *View.Family->Views[ 0 ] : View;
+		const float ScreenRadiusSquared = ComputeBoundsScreenRadiusSquared(Origin, SphereRadius, ViewForLOD) * FactorScale * FactorScale;
 
 		// Walk backwards and return the first matching LOD
 		for (int32 LODIndex = NumLODs - 1; LODIndex >= 0; --LODIndex)
@@ -430,13 +434,15 @@ FLODMask ComputeLODForMeshes( const TIndirectArray<class FStaticMesh>& StaticMes
 	{
 		int32 NumMeshes = StaticMeshes.Num();
 
+		// With stereo rendering views may have asymmetric FOVs, so always use the left eye's view for LOD calculation
+		const FSceneView& ViewForLOD = ( View.Family != nullptr && View.StereoPass == eSSP_RIGHT_EYE ) ? *View.Family->Views[ 0 ] : View;
 		if (NumMeshes && StaticMeshes[0].bDitheredLODTransition)
 		{
 			for (int32 SampleIndex = 0; SampleIndex < 2; SampleIndex++)
 			{
 				int32 MinLODFound = INT_MAX;
 				bool bFoundLOD = false;
-				const float ScreenRadiusSquared = ComputeTemporalLODBoundsScreenRadiusSquared(Origin, SphereRadius, View, SampleIndex);
+				const float ScreenRadiusSquared = ComputeTemporalLODBoundsScreenRadiusSquared(Origin, SphereRadius, ViewForLOD, SampleIndex);
 
 				for(int32 MeshIndex = NumMeshes-1 ; MeshIndex >= 0 ; --MeshIndex)
 				{
@@ -464,7 +470,7 @@ FLODMask ComputeLODForMeshes( const TIndirectArray<class FStaticMesh>& StaticMes
 		{
 			int32 MinLODFound = INT_MAX;
 			bool bFoundLOD = false;
-			const float ScreenRadiusSquared = ComputeBoundsScreenRadiusSquared(Origin, SphereRadius, View);
+			const float ScreenRadiusSquared = ComputeBoundsScreenRadiusSquared(Origin, SphereRadius, ViewForLOD);
 
 			for(int32 MeshIndex = NumMeshes-1 ; MeshIndex >= 0 ; --MeshIndex)
 			{
diff --git a/Engine/Source/Runtime/Engine/Private/SceneView.cpp b/Engine/Source/Runtime/Engine/Private/SceneView.cpp
--- a/Engine/Source/Runtime/Engine/Private/SceneView.cpp
+++ b/Engine/Source/Runtime/Engine/Private/SceneView.cpp
@@ -797,13 +797,16 @@ static TAutoConsoleVariable<int32> CVarCompensateForFOV(
 
 float FSceneView::GetLODDistanceFactor() const
 {
+	// With stereo rendering views may have asymmetric FOVs, so always use the left eye's view for LOD calculation
+	const FSceneView* ViewForLOD = ( this->Family != nullptr && this->StereoPass == eSSP_RIGHT_EYE ) ? Family->Views[0] : this;
+
 	bool bCompensateForFOV = bUseFieldOfViewForLOD && CVarCompensateForFOV.GetValueOnAnyThread() != 0;
-	float ScreenScaleX = bCompensateForFOV ? ViewMatrices.GetProjectionMatrix().M[0][0] : 1.0f;
-	float ScreenScaleY = bCompensateForFOV ? ViewMatrices.GetProjectionMatrix().M[1][1] : (static_cast<float>(ViewRect.Width()) / ViewRect.Height());
+	float ScreenScaleX = bCompensateForFOV ? ViewForLOD->ViewMatrices.GetProjectionMatrix().M[0][0] : 1.0f;
+	float ScreenScaleY = bCompensateForFOV ? ViewForLOD->ViewMatrices.GetProjectionMatrix().M[1][1] : (static_cast<float>(ViewForLOD->ViewRect.Width()) / ViewForLOD->ViewRect.Height());
 
-	const float ScreenMultiple = FMath::Max(ViewRect.Width() / 2.0f * ScreenScaleX,
-		ViewRect.Height() / 2.0f * ScreenScaleY);
-	float Fac = PI * ScreenMultiple * ScreenMultiple / ViewRect.Area();
+	const float ScreenMultiple = FMath::Max(ViewForLOD->ViewRect.Width() / 2.0f * ScreenScaleX,
+		ViewForLOD->ViewRect.Height() / 2.0f * ScreenScaleY);
+	float Fac = PI * ScreenMultiple * ScreenMultiple / ViewForLOD->ViewRect.Area();
 	return Fac;
 }
 
diff --git a/Engine/Source/Runtime/Engine/Private/SkeletalRender.cpp b/Engine/Source/Runtime/Engine/Private/SkeletalRender.cpp
--- a/Engine/Source/Runtime/Engine/Private/SkeletalRender.cpp
+++ b/Engine/Source/Runtime/Engine/Private/SkeletalRender.cpp
@@ -68,7 +68,9 @@ void FSkeletalMeshObject::UpdateMinDesiredLODLevel(const FSceneView* View, const
 	static const auto* SkeletalMeshLODRadiusScale = IConsoleManager::Get().FindTConsoleVariableDataFloat(TEXT("r.SkeletalMeshLODRadiusScale"));
 	float LODScale = FMath::Clamp(SkeletalMeshLODRadiusScale->GetValueOnRenderThread(), 0.25f, 1.0f);
 
-	const float ScreenRadiusSquared = ComputeBoundsScreenRadiusSquared(Bounds.Origin, Bounds.SphereRadius, *View) * LODScale * LODScale;
+	// With stereo rendering views may have asymmetric FOVs, so always use the left eye's view for LOD calculation
+	const FSceneView* ViewForLOD = ( View->Family != nullptr && View->StereoPass == eSSP_RIGHT_EYE ) ? View->Family->Views[ 0 ] : View;
+	const float ScreenRadiusSquared = ComputeBoundsScreenRadiusSquared(Bounds.Origin, Bounds.SphereRadius, *ViewForLOD) * LODScale * LODScale;
 
 	check( SkeletalMeshLODInfo.Num() == SkeletalMeshResource->LODModels.Num() );
 
diff --git a/Engine/Source/Runtime/Engine/Private/Slate/SceneViewport.cpp b/Engine/Source/Runtime/Engine/Private/Slate/SceneViewport.cpp
--- a/Engine/Source/Runtime/Engine/Private/Slate/SceneViewport.cpp
+++ b/Engine/Source/Runtime/Engine/Private/Slate/SceneViewport.cpp
@@ -1538,6 +1538,20 @@ void FSceneViewport::Tick( const FGeometry& AllottedGeometry, double InCurrentTi
 
 void FSceneViewport::OnPlayWorldViewportSwapped( const FSceneViewport& OtherViewport )
 {
+	// We need to call WindowRenderTargetUpdate() to make sure the Slate renderer is updated to render
+	// to the viewport client we'll be using for PIE/SIE.  Otherwise if stereo rendering is enabled, Slate
+	// could render the HMD mirror to a game viewport client which is not visible on screen!
+	TSharedPtr<SWidget> PinnedViewport = ViewportWidget.Pin();
+	if( PinnedViewport.IsValid() )
+	{
+		TSharedPtr<FSlateRenderer> Renderer = FSlateApplication::Get().GetRenderer();
+
+		FWidgetPath WidgetPath;
+		TSharedPtr<SWindow> Window = FSlateApplication::Get().FindWidgetWindow( PinnedViewport.ToSharedRef(), WidgetPath );
+
+		WindowRenderTargetUpdate( Renderer.Get(), Window.Get() );
+	}
+
 	// Play world viewports should always be the same size.  Resize to other viewports size
 	if( GetSizeXY() != OtherViewport.GetSizeXY() )
 	{
diff --git a/Engine/Source/Runtime/Engine/Private/StaticMesh.cpp b/Engine/Source/Runtime/Engine/Private/StaticMesh.cpp
--- a/Engine/Source/Runtime/Engine/Private/StaticMesh.cpp
+++ b/Engine/Source/Runtime/Engine/Private/StaticMesh.cpp
@@ -1876,6 +1876,14 @@ void UStaticMesh::PostEditChangeProperty(FPropertyChangedEvent& PropertyChangedE
 	UProperty* PropertyThatChanged = PropertyChangedEvent.Property;
 	const FName PropertyName = PropertyThatChanged ? PropertyThatChanged->GetFName() : NAME_None;
 	
+	if (PropertyName == GET_MEMBER_NAME_CHECKED(UStaticMesh, LODGroup))
+	{
+		// Force an update of LOD group settings
+
+		// Dont rebuild inside here.  We're doing that below.
+		bool bRebuild = false;
+		SetLODGroup(LODGroup, bRebuild);
+	}
 	LightMapResolution = FMath::Max(LightMapResolution, 0);
 
 	if (PropertyChangedEvent.MemberProperty 
@@ -1919,7 +1927,7 @@ void UStaticMesh::PostEditChangeProperty(FPropertyChangedEvent& PropertyChangedE
 	Super::PostEditChangeProperty(PropertyChangedEvent);
 }
 
-void UStaticMesh::SetLODGroup(FName NewGroup)
+void UStaticMesh::SetLODGroup(FName NewGroup, bool bRebuildImmediately)
 {
 #if WITH_EDITORONLY_DATA
 	const bool bBeforeDerivedDataCached = (RenderData == nullptr);
@@ -1959,6 +1967,11 @@ void UStaticMesh::SetLODGroup(FName NewGroup)
 	if (!bBeforeDerivedDataCached)
 	{
 		bAutoComputeLODScreenSize = true;
+		
+	}
+
+	if (bRebuildImmediately && !bBeforeDerivedDataCached)
+	{
 		PostEditChange();
 	}
 #endif
@@ -2226,7 +2239,12 @@ FArchive& operator<<(FArchive& Ar, FMeshSectionInfo& Info)
 
 void FMeshSectionInfoMap::Serialize(FArchive& Ar)
 {
-	Ar << Map;
+	Ar.UsingCustomVersion(FReleaseObjectVersion::GUID);
+
+	if (Ar.CustomVer(FReleaseObjectVersion::GUID) < FReleaseObjectVersion::UPropertryForMeshSectionSerialize)
+	{
+		Ar << Map;
+	}
 }
 
 #endif // #if WITH_EDITORONLY_DATA
@@ -2656,6 +2674,28 @@ void UStaticMesh::PostLoad()
 	{
 		CalculateExtendedBounds();
 	}
+
+	if (SectionInfoMap.Map.Num() == 0)
+	{
+		UE_LOG(LogStaticMesh, Warning, TEXT("StaticMesh '%s' has no SectionInfoMap. Please re-export from the original source data."), *GetName());
+
+		// Before this serialization issue was fixed, some assets were resaved and permanently lost their section info map.
+		// This attempts to recreate it based on the render data.
+		SectionInfoMap.Clear();
+		for (int32 LODResourceIndex = 0; LODResourceIndex < RenderData->LODResources.Num(); ++LODResourceIndex)
+		{
+			FStaticMeshLODResources& LOD = RenderData->LODResources[LODResourceIndex];
+			const int32 NumSections = LOD.Sections.Num();
+			for (int32 SectionIndex = 0; SectionIndex < NumSections; ++SectionIndex)
+			{
+				const int32 MaterialIndex = LOD.Sections[SectionIndex].MaterialIndex;
+				if (StaticMaterials.IsValidIndex(MaterialIndex))
+				{
+					SectionInfoMap.Set(LODResourceIndex, SectionIndex, FMeshSectionInfo(MaterialIndex));
+				}
+			}
+		}
+	}
 #endif // #if WITH_EDITOR
 
 	// We want to always have a BodySetup, its used for per-poly collision as well
@@ -2664,6 +2704,7 @@ void UStaticMesh::PostLoad()
 		CreateBodySetup();
 	}
 
+
 	CreateNavCollision();
 }
 
diff --git a/Engine/Source/Runtime/Engine/Public/Engine.h b/Engine/Source/Runtime/Engine/Public/Engine.h
--- a/Engine/Source/Runtime/Engine/Public/Engine.h
+++ b/Engine/Source/Runtime/Engine/Public/Engine.h
@@ -110,7 +110,6 @@ MONOLITHIC_HEADER_BOILERPLATE()
 #include "Curves/CurveFloat.h"
 #include "GameFramework/ForceFeedbackEffect.h"
 #include "Engine/SubsurfaceProfile.h"
-#include "Engine/SubDSurface.h"
 #include "Engine/DebugCameraController.h"
 #include "Engine/DecalActor.h"
 #include "CanvasItem.h"
diff --git a/Engine/Source/Runtime/GameLiveStreaming/GameLiveStreaming.cpp b/Engine/Source/Runtime/GameLiveStreaming/GameLiveStreaming.cpp
--- a/Engine/Source/Runtime/GameLiveStreaming/GameLiveStreaming.cpp
+++ b/Engine/Source/Runtime/GameLiveStreaming/GameLiveStreaming.cpp
@@ -448,7 +448,7 @@ void FGameLiveStreaming::StartCopyingNextGameVideoFrame( const FViewportRHIRef&
 		SetRenderTarget(RHICmdList, DestRenderTarget.TargetableTexture, FTextureRHIRef());
 		RHICmdList.SetViewport(0, 0, 0.0f, Context.ResizeTo.X, Context.ResizeTo.Y, 1.0f);
 
-		RHICmdList.SetBlendState(TStaticBlendState<>::GetRHI());
+		RHICmdList.SetBlendState(TStaticBlendState<CW_RGB>::GetRHI());
 		RHICmdList.SetRasterizerState(TStaticRasterizerState<>::GetRHI());
 		RHICmdList.SetDepthStencilState(TStaticDepthStencilState<false,CF_Always>::GetRHI());
 
diff --git a/Engine/Source/Runtime/GeometryCache/Classes/GeometryCacheMeshData.h b/Engine/Source/Runtime/GeometryCache/Classes/GeometryCacheMeshData.h
--- a/Engine/Source/Runtime/GeometryCache/Classes/GeometryCacheMeshData.h
+++ b/Engine/Source/Runtime/GeometryCache/Classes/GeometryCacheMeshData.h
@@ -5,7 +5,7 @@
 #include "UObject/ObjectMacros.h"
 #include "ProfilingDebugging/ResourceSize.h"
 #include "DynamicMeshBuilder.h"
-
+#include "UObject/GeometryObjectVersion.h"
 #include "GeometryCacheMeshData.generated.h"
 
 /** Stores per-batch data used for rendering */
@@ -56,6 +56,8 @@ struct FGeometryCacheMeshData
 	/** Serialization for FVertexAnimationSample. */
 	friend FArchive& operator<<(FArchive& Ar, FGeometryCacheMeshData& Mesh)
 	{
+		Ar.UsingCustomVersion(FGeometryObjectVersion::GUID);
+
 		int32 NumVertices = 0;
 		
 		if (Ar.IsSaving())
@@ -69,19 +71,26 @@ struct FGeometryCacheMeshData
 			Mesh.Vertices.AddUninitialized(NumVertices);
 		}
 
-		for (int32 VertexIndex = 0; VertexIndex < NumVertices; ++VertexIndex)
+		if (Ar.CustomVer(FGeometryObjectVersion::GUID) < FGeometryObjectVersion::CompressGeometryCache)
+		{
+			for (int32 VertexIndex = 0; VertexIndex < NumVertices; ++VertexIndex)
+			{
+				FDynamicMeshVertex& Vertex = Mesh.Vertices[VertexIndex];
+				Ar << Vertex.Position;
+				Ar << Vertex.TextureCoordinate;
+				Ar << Vertex.TangentX;
+				Ar << Vertex.TangentZ;
+				Ar << Vertex.Color;
+			}
+		}
+		else if(NumVertices)
 		{
-			FDynamicMeshVertex& Vertex = Mesh.Vertices[VertexIndex];
-			Ar << Vertex.Position;
-			Ar << Vertex.TextureCoordinate;
-			Ar << Vertex.TangentX;
-			Ar << Vertex.TangentZ;
-			Ar << Vertex.Color;
+			Ar.SerializeCompressed(&Mesh.Vertices[0], Mesh.Vertices.Num()*Mesh.Vertices.GetTypeSize(), COMPRESS_ZLIB);
 		}
 
+	
 		Ar << Mesh.BoundingBox;
 		Ar << Mesh.BatchesInfo;
-
 		Ar << Mesh.Indices;	
 
 		return Ar;
diff --git a/Engine/Source/Runtime/HeadMountedDisplay/Private/HeadMountedDisplayTypes.cpp b/Engine/Source/Runtime/HeadMountedDisplay/Private/HeadMountedDisplayTypes.cpp
--- a/Engine/Source/Runtime/HeadMountedDisplay/Private/HeadMountedDisplayTypes.cpp
+++ b/Engine/Source/Runtime/HeadMountedDisplay/Private/HeadMountedDisplayTypes.cpp
@@ -26,7 +26,17 @@ FHMDViewMesh::~FHMDViewMesh()
 
 void FHMDViewMesh::BuildMesh(const FVector2D Positions[], uint32 VertexCount, EHMDMeshType MeshType)
 {
-	check(pVertices == nullptr);
+	if( pVertices != nullptr )
+	{
+		delete[] pVertices;
+		pVertices = nullptr;
+	}
+	if( pIndices != nullptr )
+	{
+		delete[] pIndices;
+		pIndices = nullptr;
+	}
+
 	check(VertexCount > 2 && VertexCount % 3 == 0);
 
 	NumVertices = VertexCount;
diff --git a/Engine/Source/Runtime/Launch/Private/LaunchEngineLoop.cpp b/Engine/Source/Runtime/Launch/Private/LaunchEngineLoop.cpp
--- a/Engine/Source/Runtime/Launch/Private/LaunchEngineLoop.cpp
+++ b/Engine/Source/Runtime/Launch/Private/LaunchEngineLoop.cpp
@@ -2987,15 +2987,28 @@ void FEngineLoop::Tick()
 
 		// @todo vreditor urgent: Temporary hack to allow world-to-meters to be set before
 		// input is polled for motion controller devices each frame.
-#if WITH_ENGINE
 		extern ENGINE_API float GNewWorldToMetersScale;
-		UWorld* HackGWorld = GWorld;
-		if( GNewWorldToMetersScale != 0.0f && HackGWorld != nullptr )
+		if( GNewWorldToMetersScale != 0.0f  )
 		{
-			if( GNewWorldToMetersScale != HackGWorld->GetWorldSettings()->WorldToMeters )
+#if WITH_ENGINE
+			UWorld* WorldToScale = GWorld;
+
+#if WITH_EDITOR
+			if( GIsEditor && GEditor->PlayWorld != nullptr && GEditor->bIsSimulatingInEditor )
 			{
-				HackGWorld->GetWorldSettings()->WorldToMeters = GNewWorldToMetersScale;
+				WorldToScale = GEditor->PlayWorld;
 			}
+#endif	// WITH_EDITOR
+
+			if( WorldToScale != nullptr )
+			{
+				if( GNewWorldToMetersScale != WorldToScale->GetWorldSettings()->WorldToMeters )
+				{
+					WorldToScale->GetWorldSettings()->WorldToMeters = GNewWorldToMetersScale;
+				}
+			}
+
+			GNewWorldToMetersScale = 0.0f;
 		}
 #endif	// WITH_ENGINE
 
diff --git a/Engine/Source/Runtime/MovieScene/Private/Evaluation/MovieSceneEvaluationTemplate.cpp b/Engine/Source/Runtime/MovieScene/Private/Evaluation/MovieSceneEvaluationTemplate.cpp
--- a/Engine/Source/Runtime/MovieScene/Private/Evaluation/MovieSceneEvaluationTemplate.cpp
+++ b/Engine/Source/Runtime/MovieScene/Private/Evaluation/MovieSceneEvaluationTemplate.cpp
@@ -26,15 +26,18 @@ void FCachedMovieSceneEvaluationTemplate::Regenerate(const FMovieSceneTrackCompi
 {
 	if (IsOutOfDate(Params))
 	{
-		ForceRegenerate(Params);
+		RegenerateImpl(Params);
 	}
 }
 
 void FCachedMovieSceneEvaluationTemplate::ForceRegenerate(const FMovieSceneTrackCompilationParams& Params)
 {
-	// Reassign the template back to default
-	static_cast<FMovieSceneEvaluationTemplate&>(*this) = FMovieSceneEvaluationTemplate();
+	ResetGeneratedData();
+	RegenerateImpl(Params);
+}
 
+void FCachedMovieSceneEvaluationTemplate::RegenerateImpl(const FMovieSceneTrackCompilationParams& Params)
+{
 	CachedSignatures.Reset();
 	CachedCompilationParams = Params;
 
@@ -165,6 +168,18 @@ void FMovieSceneGenerationLedger::AddTrack(const FGuid& InSignature, FMovieScene
 	++TrackReferenceCounts.FindOrAdd(Identifier);
 }
 
+void FMovieSceneEvaluationTemplate::ResetGeneratedData()
+{
+	Ledger.TrackSignatureToTrackIdentifier.Reset();
+	Ledger.TrackReferenceCounts.Reset();
+
+	Tracks.Reset();
+	StaleTracks.Reset();
+	EvaluationField = FMovieSceneEvaluationField();
+	Hierarchy = FMovieSceneSequenceHierarchy();
+	bHasLegacyTrackInstances = false;
+}
+
 FMovieSceneTrackIdentifier FMovieSceneEvaluationTemplate::AddTrack(const FGuid& InSignature, FMovieSceneEvaluationTrack&& InTrack)
 {
 	FMovieSceneTrackIdentifier NewIdentifier = ++Ledger.LastTrackIdentifier;
diff --git a/Engine/Source/Runtime/MovieScene/Private/Evaluation/MovieSceneEvaluationTemplateInstance.cpp b/Engine/Source/Runtime/MovieScene/Private/Evaluation/MovieSceneEvaluationTemplateInstance.cpp
--- a/Engine/Source/Runtime/MovieScene/Private/Evaluation/MovieSceneEvaluationTemplateInstance.cpp
+++ b/Engine/Source/Runtime/MovieScene/Private/Evaluation/MovieSceneEvaluationTemplateInstance.cpp
@@ -89,6 +89,14 @@ void FMovieSceneRootEvaluationTemplateInstance::Initialize(UMovieSceneSequence&
 {
 	Reset();
 
+	if (RootSequence.Get() != &InRootSequence)
+	{
+		// Always ensure that there is no persistent data when initializing a new sequence
+		// to ensure we don't collide with the previous sequence's entity keys
+		Player.State.PersistentEntityData.Reset();
+		Player.State.PersistentSharedData.Reset();
+	}
+
 	const bool bAddEvents = TemplateStore->AreTemplatesVolatile();
 	
 	RootSequence = &InRootSequence;
@@ -365,7 +373,7 @@ void FMovieSceneRootEvaluationTemplateInstance::CallSetupTearDown(IMovieScenePla
 		else
 		{
 			PersistentDataProxy.SetSectionKey(Key);
-			if (Track)
+			if (Track && Track->HasChildTemplate(Key.SectionIdentifier))
 			{
 				Track->GetChildTemplate(Key.SectionIdentifier).OnEndEvaluation(PersistentDataProxy, Player);
 			}
diff --git a/Engine/Source/Runtime/MovieScene/Private/Evaluation/MovieSceneSequenceTemplateStore.cpp b/Engine/Source/Runtime/MovieScene/Private/Evaluation/MovieSceneSequenceTemplateStore.cpp
--- a/Engine/Source/Runtime/MovieScene/Private/Evaluation/MovieSceneSequenceTemplateStore.cpp
+++ b/Engine/Source/Runtime/MovieScene/Private/Evaluation/MovieSceneSequenceTemplateStore.cpp
@@ -6,9 +6,9 @@
 FMovieSceneEvaluationTemplate& FMovieSceneSequenceTemplateStore::GetCompiledTemplate(UMovieSceneSequence& Sequence)
 {
 #if WITH_EDITORONLY_DATA
-	if (AreTemplatesVolatile())
+	if (AreTemplatesVolatile() && Sequence.EvaluationTemplate.IsOutOfDate(Sequence.TemplateParameters))
 	{
-		Sequence.EvaluationTemplate.Regenerate(Sequence.TemplateParameters);
+		Sequence.EvaluationTemplate.ForceRegenerate(Sequence.TemplateParameters);
 	}
 #endif
 	return Sequence.EvaluationTemplate;
diff --git a/Engine/Source/Runtime/MovieScene/Public/Evaluation/MovieSceneEvaluationTemplate.h b/Engine/Source/Runtime/MovieScene/Public/Evaluation/MovieSceneEvaluationTemplate.h
--- a/Engine/Source/Runtime/MovieScene/Public/Evaluation/MovieSceneEvaluationTemplate.h
+++ b/Engine/Source/Runtime/MovieScene/Public/Evaluation/MovieSceneEvaluationTemplate.h
@@ -95,6 +95,11 @@ struct FMovieSceneEvaluationTemplate
 
 public:
 
+	/**
+	 * Reset this evaluation template, whilst ensuring that track identifiers do not get reused
+	 */
+	MOVIESCENE_API void ResetGeneratedData();
+
 	/**
 	 * Attempt to locate a track with the specified identifier
 	 */
@@ -236,6 +241,8 @@ struct FCachedMovieSceneEvaluationTemplate : public FMovieSceneEvaluationTemplat
 
 private:
 
+	void RegenerateImpl(const FMovieSceneTrackCompilationParams& NewParams);
+
 	/** Transient data */
 	TWeakObjectPtr<UMovieSceneSequence> SourceSequence;
 	FMovieSceneSequenceTemplateStore* Origin;
diff --git a/Engine/Source/Runtime/MovieSceneTracks/Private/Evaluation/MovieSceneAudioTemplate.cpp b/Engine/Source/Runtime/MovieSceneTracks/Private/Evaluation/MovieSceneAudioTemplate.cpp
--- a/Engine/Source/Runtime/MovieSceneTracks/Private/Evaluation/MovieSceneAudioTemplate.cpp
+++ b/Engine/Source/Runtime/MovieSceneTracks/Private/Evaluation/MovieSceneAudioTemplate.cpp
@@ -268,7 +268,7 @@ void FMovieSceneAudioSectionTemplateData::EnsureAudioIsPlaying(UAudioComponent&
 		AudioComponent.Stop();
 		AudioComponent.SetSound(Sound);
 		AudioComponent.bIsUISound = true;
-		AudioComponent.Play(AudioStartOffset < 0 ? 0 : AudioStartOffset);
+		AudioComponent.Play((Context.GetTime() - AudioRange.GetLowerBoundValue()) + FMath::Max(AudioStartOffset, 0.f));
 
 		if (Context.GetStatus() == EMovieScenePlayerStatus::Scrubbing)
 		{
diff --git a/Engine/Source/Runtime/Renderer/Private/IndirectLightingCache.cpp b/Engine/Source/Runtime/Renderer/Private/IndirectLightingCache.cpp
--- a/Engine/Source/Runtime/Renderer/Private/IndirectLightingCache.cpp
+++ b/Engine/Source/Runtime/Renderer/Private/IndirectLightingCache.cpp
@@ -326,7 +326,7 @@ void FIndirectLightingCache::CalculateBlockScaleAndAdd(FIntVector InTexelMin, in
 
 FIndirectLightingCacheAllocation* FIndirectLightingCache::AllocatePrimitive(const FPrimitiveSceneInfo* PrimitiveSceneInfo, bool bUnbuiltPreview)
 {
-	const bool bPointSample = PrimitiveSceneInfo->Proxy->GetIndirectLightingCacheQuality() == ILCQ_Point || bUnbuiltPreview;
+	const bool bPointSample = !bUnbuiltPreview && PrimitiveSceneInfo->Proxy->GetIndirectLightingCacheQuality() == ILCQ_Point;
 	const int32 BlockSize = bPointSample ? 1 : GLightingCacheMovableObjectAllocationSize;
 	return PrimitiveAllocations.Add(PrimitiveSceneInfo->PrimitiveComponentId, CreateAllocation(BlockSize, PrimitiveSceneInfo->Proxy->GetBounds(), bPointSample, bUnbuiltPreview));
 }
@@ -686,7 +686,7 @@ void FIndirectLightingCache::UpdateCachePrimitive(
 		{
 			FIndirectLightingCacheAllocation* OriginalAllocation = PrimitiveAllocation;
 			const bool bUnbuiltPreview = bAllowUnbuiltPreview && !bIsMovable;
-			const bool bPointSample = PrimitiveSceneProxy->GetIndirectLightingCacheQuality() == ILCQ_Point || bUnbuiltPreview || !bOpaqueRelevance;
+			const bool bPointSample = !bUnbuiltPreview && ( PrimitiveSceneProxy->GetIndirectLightingCacheQuality() == ILCQ_Point || !bOpaqueRelevance );
 			const int32 BlockSize = bPointSample ? 1 : GLightingCacheMovableObjectAllocationSize;
 
 			// Light with the cumulative bounds of the entire attachment group
diff --git a/Engine/Source/Runtime/Renderer/Private/PostProcess/PostProcessing.cpp b/Engine/Source/Runtime/Renderer/Private/PostProcess/PostProcessing.cpp
--- a/Engine/Source/Runtime/Renderer/Private/PostProcess/PostProcessing.cpp
+++ b/Engine/Source/Runtime/Renderer/Private/PostProcess/PostProcessing.cpp
@@ -1768,7 +1768,6 @@ void FPostProcessing::Process(FRHICommandListImmediate& RHICmdList, const FViewI
 		if ( GIsEditor
 			&& View.Family->EngineShowFlags.SelectionOutline
 			&& !(View.Family->EngineShowFlags.Wireframe)
-			&& ( !GIsDemoMode || ( GIsDemoMode && !View.Family->EngineShowFlags.Game ) ) 
 			&& !bVisualizeBloom
 			&& !View.Family->EngineShowFlags.VisualizeHDR)
 		{
@@ -2358,7 +2357,6 @@ void FPostProcessing::ProcessES2(FRHICommandListImmediate& RHICmdList, const FVi
 		if ( GIsEditor
 			&& View.Family->EngineShowFlags.SelectionOutline
 			&& !(View.Family->EngineShowFlags.Wireframe)
-			&& ( !GIsDemoMode || ( GIsDemoMode && !View.Family->EngineShowFlags.Game ) ) 
 			)
 		{
 			// Editor selection outline
diff --git a/Engine/Source/Runtime/Renderer/Private/VelocityRendering.cpp b/Engine/Source/Runtime/Renderer/Private/VelocityRendering.cpp
--- a/Engine/Source/Runtime/Renderer/Private/VelocityRendering.cpp
+++ b/Engine/Source/Runtime/Renderer/Private/VelocityRendering.cpp
@@ -39,6 +39,11 @@ static TAutoConsoleVariable<int32> CVarRHICmdVelocityPassDeferredContexts(
 	1,
 	TEXT("True to use deferred contexts to parallelize velocity pass command list execution."));
 
+RENDERER_API TAutoConsoleVariable<int32> CVarAllowMotionBlurInVR(
+	TEXT("r.AllowMotionBlurInVR"),
+	0,
+	TEXT("For projects with motion blur enabled, this allows motion blur to be enabled even while in VR."));
+
 DECLARE_FLOAT_COUNTER_STAT(TEXT("Render Velocities"), Stat_GPU_RenderVelocities, STATGROUP_GPU);
 
 bool IsParallelVelocity()
@@ -607,7 +612,7 @@ bool IsMotionBlurEnabled(const FViewInfo& View)
 		&& View.Family->bRealtimeUpdate
 		&& MotionBlurQuality > 0
 		&& !View.bIsSceneCapture
-		&& !(View.Family->Views.Num() > 1);
+		&& (CVarAllowMotionBlurInVR->GetInt() != 0 || !(View.Family->Views.Num() > 1));
 }
 
 void FDeferredShadingSceneRenderer::RenderDynamicVelocitiesMeshElementsInner(FRHICommandList& RHICmdList, const FViewInfo& View, const FDrawingPolicyRenderState& DrawRenderState, int32 FirstIndex, int32 LastIndex)
diff --git a/Engine/Source/Runtime/Slate/Private/Framework/Application/SlateApplication.cpp b/Engine/Source/Runtime/Slate/Private/Framework/Application/SlateApplication.cpp
--- a/Engine/Source/Runtime/Slate/Private/Framework/Application/SlateApplication.cpp
+++ b/Engine/Source/Runtime/Slate/Private/Framework/Application/SlateApplication.cpp
@@ -3042,7 +3042,7 @@ void FSlateApplication::ProcessReply( const FWidgetPath& CurrentEventPath, const
 	// are all operations that we shouldn't do if our application isn't Active (The OS ignores half of it, and we'd be in a half state)
 	// We do allow the release of capture and lock when deactivated, this is innocuous of some platforms but required on others when 
 	// the Application deactivated before the window. (Mac is an example of this)
-	if (bAppIsActive)
+	if (bAppIsActive || bIsVirtualInteraction)
 	{
 		TSharedPtr<SWidget> RequestedMouseCaptor = TheReply.GetMouseCaptor();
 		// Do not capture the mouse if we are also starting a drag and drop.
@@ -5035,6 +5035,15 @@ FReply FSlateApplication::RoutePointerUpEvent(FWidgetPath& WidgetsUnderPointer,
 	TSharedPtr<SWindow> TopLevelWindow;
 #endif
 
+	if ( DragDetector.DetectDragForWidget.IsValid() )
+	{
+		if ( PointerEvent.GetEffectingButton() == DragDetector.DetectDragButton && PointerEvent.GetPointerIndex() == DragDetector.DetectDragPointerIndex )
+		{
+			// The user has released the button (or finger) that was supposed to start the drag; stop detecting it.
+			DragDetector = FDragDetector();
+		}
+	}
+
 	if (MouseCaptor.HasCaptureForPointerIndex(PointerEvent.GetUserIndex(), PointerEvent.GetPointerIndex()))
 	{
 		//FWidgetPath MouseCaptorPath = MouseCaptor.ToWidgetPath(PointerEvent.GetPointerIndex());
@@ -5141,7 +5150,14 @@ bool FSlateApplication::RoutePointerMoveEvent(const FWidgetPath& WidgetsUnderPoi
 
 	// User asked us to detect a drag.
 	bool bDragDetected = false;
-	if( DragDetector.DetectDragForWidget.IsValid() && !bIsSynthetic )
+	bool bShouldStartDetectingDrag = true;
+
+	//@TODO VREDITOR - Remove and move to interaction component
+	if (OnDragDropCheckOverride.IsBound())
+	{
+		bShouldStartDetectingDrag = OnDragDropCheckOverride.Execute();
+	}
+	if( DragDetector.DetectDragForWidget.IsValid() && !bIsSynthetic && bShouldStartDetectingDrag)
 	{	
 		const FVector2D DragDelta = (DragDetector.DetectDragStartLocation - PointerEvent.GetScreenSpacePosition());
 		bDragDetected = ( DragDelta.SizeSquared() > FMath::Square(FSlateApplication::Get().GetDragTriggerDistance()) );
@@ -5347,6 +5363,13 @@ bool FSlateApplication::RoutePointerMoveEvent(const FWidgetPath& WidgetsUnderPoi
 	if (IsDragDropping())
 	{
 		FDragDropEvent DragDropEvent( PointerEvent, DragDropContent );
+		//@TODO VREDITOR - Remove and move to interaction component
+		if (OnDragDropCheckOverride.IsBound() && DragDropEvent.GetOperation().IsValid())
+		{
+			DragDropEvent.GetOperation()->SetDecoratorVisibility(false);
+			DragDropEvent.GetOperation()->SetCursorOverride(EMouseCursor::None);
+			DragDropContent->SetCursorOverride(EMouseCursor::None);
+		}
 		FScopedSwitchWorldHack SwitchWorld( WidgetsUnderPointer );
 		DragDropContent->OnDragged( DragDropEvent );
 
@@ -5502,15 +5525,6 @@ bool FSlateApplication::ProcessMouseButtonUpEvent( FPointerEvent& MouseEvent )
 		return true;
 	}
 
-	if ( DragDetector.DetectDragForWidget.IsValid() )
-	{
-		if ( MouseEvent.GetEffectingButton() == DragDetector.DetectDragButton && MouseEvent.GetPointerIndex() == DragDetector.DetectDragPointerIndex )
-		{
-			// The user has released the button (or finger) that was supposed to start the drag; stop detecting it.
-			DragDetector = FDragDetector();
-		}
-	}
-
 	// An empty widget path is passed in.  As an optimization, one will be generated only if a captured mouse event isn't routed
 	FWidgetPath EmptyPath;
 	const bool bHandled = RoutePointerUpEvent( EmptyPath, MouseEvent ).IsEventHandled();
diff --git a/Engine/Source/Runtime/Slate/Private/Framework/Docking/TabManager.cpp b/Engine/Source/Runtime/Slate/Private/Framework/Docking/TabManager.cpp
--- a/Engine/Source/Runtime/Slate/Private/Framework/Docking/TabManager.cpp
+++ b/Engine/Source/Runtime/Slate/Private/Framework/Docking/TabManager.cpp
@@ -1906,7 +1906,10 @@ void FGlobalTabmanager::DrawAttentionToTabManager( const TSharedRef<FTabManager>
 		// #HACK VREDITOR
 		if ( ProxyTabManager.IsValid() )
 		{
-			ProxyTabManager->DrawAttention(Tab.ToSharedRef());
+			if( ProxyTabManager->IsTabSupported( Tab->GetLayoutIdentifier() ) )
+			{
+				ProxyTabManager->DrawAttention(Tab.ToSharedRef());
+			}
 		}
 	}
 }
@@ -2114,7 +2117,7 @@ void FGlobalTabmanager::UpdateStats()
 
 void FGlobalTabmanager::OpenUnmanagedTab(FName PlaceholderId, const FSearchPreference& SearchPreference, const TSharedRef<SDockTab>& UnmanagedTab)
 {
-	if ( ProxyTabManager.IsValid() )
+	if ( ProxyTabManager.IsValid() && ProxyTabManager->IsTabSupported( UnmanagedTab->GetLayoutIdentifier() ) )
 	{
 		ProxyTabManager->OpenUnmanagedTab(PlaceholderId, SearchPreference, UnmanagedTab);
 	}
@@ -2129,29 +2132,43 @@ void FGlobalTabmanager::SetProxyTabManager(TSharedPtr<FProxyTabmanager> InProxyT
 	ProxyTabManager = InProxyTabManager;
 }
 
-void FProxyTabmanager::OpenUnmanagedTab(FName PlaceholderId, const FSearchPreference& SearchPreference, const TSharedRef<SDockTab>& UnmanagedTab)
+bool FProxyTabmanager::IsTabSupported( const FTabId TabId ) const
 {
-	// No layout info about this tab found; start 
-	TSharedRef<FArea> NewAreaForTab = FTabManager::NewPrimaryArea()
-		->Split
-		(
-			FTabManager::NewStack()
-			->AddTab(UnmanagedTab->GetLayoutIdentifier(), ETabState::OpenedTab)
-		);
+	bool bIsTabSupported = true;
+	if( OnIsTabSupported.IsBound() )
+	{
+		OnIsTabSupported.Broadcast( TabId, /* In/Out */ bIsTabSupported );
+	}
 
+	return bIsTabSupported;
+}
+
+void FProxyTabmanager::OpenUnmanagedTab(FName PlaceholderId, const FSearchPreference& SearchPreference, const TSharedRef<SDockTab>& UnmanagedTab)
+{
 	TSharedPtr<SWindow> ParentWindowPtr = ParentWindow.Pin();
-	TSharedRef<SDockingArea> DockingArea = RestoreArea(NewAreaForTab, ParentWindowPtr, false);
-	ParentWindowPtr->SetContent(DockingArea);
+	if( ensure( ParentWindowPtr.IsValid() ) )
+	{
+		// No layout info about this tab found; start 
+		TSharedRef<FArea> NewAreaForTab = FTabManager::NewPrimaryArea()
+			->Split
+			(
+				FTabManager::NewStack()
+				->AddTab( UnmanagedTab->GetLayoutIdentifier(), ETabState::OpenedTab )
+			);
+
+		TSharedRef<SDockingArea> DockingArea = RestoreArea(NewAreaForTab, ParentWindowPtr, false);
+		ParentWindowPtr->SetContent(DockingArea);
 
-	const TSharedPtr<SDockTab> NewlyOpenedTab = DockingArea->GetAllChildTabs()[0];
-	check(NewlyOpenedTab.IsValid());
+		const TSharedPtr<SDockTab> NewlyOpenedTab = DockingArea->GetAllChildTabs()[0];
+		check(NewlyOpenedTab.IsValid());
 		
-	NewlyOpenedTab->GetParent()->GetParentDockTabStack()->OpenTab(UnmanagedTab);
-	NewlyOpenedTab->RequestCloseTab();
+		NewlyOpenedTab->GetParent()->GetParentDockTabStack()->OpenTab(UnmanagedTab);
+		NewlyOpenedTab->RequestCloseTab();
 
-	MainNonCloseableTab = UnmanagedTab;
+		MainNonCloseableTab = UnmanagedTab;
 
-	OnTabOpened.Broadcast(UnmanagedTab);
+		OnTabOpened.Broadcast(UnmanagedTab);
+	}
 }
 
 void FProxyTabmanager::DrawAttention(const TSharedRef<SDockTab>& TabToHighlight)
@@ -2162,4 +2179,9 @@ void FProxyTabmanager::DrawAttention(const TSharedRef<SDockTab>& TabToHighlight)
 }
 
 
+void FProxyTabmanager::SetParentWindow(TSharedRef<SWindow> InParentWindow)
+{
+	ParentWindow = InParentWindow;
+}
+
 #undef LOCTEXT_NAMESPACE
diff --git a/Engine/Source/Runtime/Slate/Private/Framework/MultiBox/SMenuEntryBlock.cpp b/Engine/Source/Runtime/Slate/Private/Framework/MultiBox/SMenuEntryBlock.cpp
--- a/Engine/Source/Runtime/Slate/Private/Framework/MultiBox/SMenuEntryBlock.cpp
+++ b/Engine/Source/Runtime/Slate/Private/Framework/MultiBox/SMenuEntryBlock.cpp
@@ -493,7 +493,8 @@ TSharedRef< SWidget > SMenuEntryBlock::BuildMenuEntryWidget( const FMenuEntryBui
 				EVisibility::Hidden;
 
 	// Collapse (rather than Hide) the checkbox when using a custom menu widget with a button action, otherwise we add additional padding around the user defined widget
-	if (MenuEntryBlock->EntryWidget.IsValid() && UserInterfaceType == EUserInterfaceActionType::Button)
+	if ((MenuEntryBlock->EntryWidget.IsValid() && UserInterfaceType == EUserInterfaceActionType::Button) ||
+		UserInterfaceType == EUserInterfaceActionType::CollapsedButton)
 	{
 		CheckBoxVisibility = EVisibility::Collapsed;
 	}
diff --git a/Engine/Source/Runtime/Slate/Private/Framework/Notifications/NotificationManager.cpp b/Engine/Source/Runtime/Slate/Private/Framework/Notifications/NotificationManager.cpp
--- a/Engine/Source/Runtime/Slate/Private/Framework/Notifications/NotificationManager.cpp
+++ b/Engine/Source/Runtime/Slate/Private/Framework/Notifications/NotificationManager.cpp
@@ -63,6 +63,7 @@ void FSlateNotificationManager::FRegionalNotificationList::Arrange()
 }
 
 FSlateNotificationManager::FSlateNotificationManager()
+	: bAllowNotifications( true )
 {
 }
 
diff --git a/Engine/Source/Runtime/Slate/Private/Widgets/Layout/SScissorRectBox.cpp b/Engine/Source/Runtime/Slate/Private/Widgets/Layout/SScissorRectBox.cpp
--- a/Engine/Source/Runtime/Slate/Private/Widgets/Layout/SScissorRectBox.cpp
+++ b/Engine/Source/Runtime/Slate/Private/Widgets/Layout/SScissorRectBox.cpp
@@ -4,6 +4,12 @@
 #include "Types/PaintArgs.h"
 #include "Layout/ArrangedChildren.h"
 #include "Rendering/RenderingCommon.h"
+#include "IConsoleManager.h"
+
+namespace Slate
+{
+	static FAutoConsoleVariable OverrideScissorRect( TEXT( "Slate.OverrideScissorRect" ), 1, TEXT( "Whether to allow Slate to apply a scissor clip to UI elements to prevent certain artifacts.  You might need to disable this if you are drawing UI over a mirror backbuffer while in VR." ) );
+}
 
 SScissorRectBox::SScissorRectBox()
 {
@@ -56,19 +62,26 @@ int32 SScissorRectBox::OnPaint(const FPaintArgs& Args, const FGeometry& Allotted
 		check( ArrangedChildren.Num() == 1 );
 		FArrangedWidget& TheChild = ArrangedChildren[0];
 
-		// Child clipping rect is made effectively infinite to ensure that any rotated children do not get prematurely cut off. We let the scissor rect do all the work.
-		const FSlateRect ChildClippingRect = AllottedGeometry.GetClippingRect().IntersectionWith(MyClippingRect);
-
 		// set the hacky, back door global scissor rect settings.
 		// This will get picked up by the element batcher and added to the draw element.
 		// We do it this way because this scissor rect is a real hack until better solution is made available.
 		// In the meantime, we don't want to modify the public interface of all element drawing to ensure all widgets pick up this scissor rect.
 		extern SLATECORE_API TOptional<FShortRect> GSlateScissorRect;
 		TOptional<FShortRect> OldRect = GSlateScissorRect;
-		GSlateScissorRect = FShortRect(ChildClippingRect);
+		const bool bOverrideScissorRect = Slate::OverrideScissorRect->GetInt() != 0;
+		if( bOverrideScissorRect )
+		{
+			// Child clipping rect is made effectively infinite to ensure that any rotated children do not get prematurely cut off. We let the scissor rect do all the work.
+			const FSlateRect ChildClippingRect = AllottedGeometry.GetClippingRect().IntersectionWith( MyClippingRect );
+			GSlateScissorRect = FShortRect(ChildClippingRect);
+		}
 		int32 Result = TheChild.Widget->Paint( Args.WithNewParent(this), TheChild.Geometry, FSlateRect(FVector2D(-FLT_MAX, -FLT_MAX), FVector2D(FLT_MAX, FLT_MAX)), OutDrawElements, LayerId, InWidgetStyle, ShouldBeEnabled( bParentEnabled ) );
-		// We've added all the draw elements, so put the scissor rect back
-		GSlateScissorRect = OldRect;
+
+		if( bOverrideScissorRect )
+		{
+			// We've added all the draw elements, so put the scissor rect back
+			GSlateScissorRect = OldRect;
+		}
 		
 		return Result;
 	}
diff --git a/Engine/Source/Runtime/Slate/Private/Widgets/Notifications/SNotificationList.cpp b/Engine/Source/Runtime/Slate/Private/Widgets/Notifications/SNotificationList.cpp
--- a/Engine/Source/Runtime/Slate/Private/Widgets/Notifications/SNotificationList.cpp
+++ b/Engine/Source/Runtime/Slate/Private/Widgets/Notifications/SNotificationList.cpp
@@ -15,6 +15,7 @@
 #include "Widgets/Input/SHyperlink.h"
 #include "Widgets/Images/SThrobber.h"
 #include "Widgets/Notifications/INotificationWidget.h"
+#include "Framework/Notifications/NotificationManager.h"
 
 /////////////////////////////////////////////////
 // SNotificationExtendable
@@ -682,54 +683,64 @@ TSharedRef<SNotificationItem> SNotificationList::AddNotification(const FNotifica
 {
 	TSharedPtr<SNotificationExtendable> NewItem;
 
-	if (Info.ContentWidget.IsValid())
+	if (FSlateNotificationManager::Get().AreNotificationsAllowed())
 	{
-		NewItem = SNew(SNotificationItemExternalImpl)
-			.ContentWidget(Info.ContentWidget)
-			.FadeInDuration(Info.FadeInDuration)
-			.ExpireDuration(Info.ExpireDuration)
-			.FadeOutDuration(Info.FadeOutDuration);
-	}
-	else
-	{
-		static const FSlateBrush* CachedImage = FCoreStyle::Get().GetBrush("NotificationList.DefaultMessage");
-
-		// Create notification.
-		NewItem = SNew(SNotificationItemImpl)
-			.Text(Info.Text)
-			.Font(Font)
-			.ButtonDetails(Info.ButtonDetails)
-			.Image((Info.Image != nullptr) ? Info.Image : CachedImage)
-			.FadeInDuration(Info.FadeInDuration)
-			.ExpireDuration(Info.ExpireDuration)
-			.FadeOutDuration(Info.FadeOutDuration)
-			.bUseThrobber(Info.bUseThrobber)
-			.bUseSuccessFailIcons(Info.bUseSuccessFailIcons)
-			.bUseLargeFont(Info.bUseLargeFont)
-			.WidthOverride(Info.WidthOverride)
-			.CheckBoxState(Info.CheckBoxState)
-			.CheckBoxStateChanged(Info.CheckBoxStateChanged)
-			.CheckBoxText(Info.CheckBoxText)
-			.Hyperlink(Info.Hyperlink)
-			.HyperlinkText(Info.HyperlinkText);
-	}
+		if (Info.ContentWidget.IsValid())
+		{
+			NewItem = SNew(SNotificationItemExternalImpl)
+				.ContentWidget(Info.ContentWidget)
+				.FadeInDuration(Info.FadeInDuration)
+				.ExpireDuration(Info.ExpireDuration)
+				.FadeOutDuration(Info.FadeOutDuration);
+		}
+		else
+		{
+			static const FSlateBrush* CachedImage = FCoreStyle::Get().GetBrush("NotificationList.DefaultMessage");
+
+			// Create notification.
+			NewItem = SNew(SNotificationItemImpl)
+				.Text(Info.Text)
+				.Font(Font)
+				.ButtonDetails(Info.ButtonDetails)
+				.Image((Info.Image != nullptr) ? Info.Image : CachedImage)
+				.FadeInDuration(Info.FadeInDuration)
+				.ExpireDuration(Info.ExpireDuration)
+				.FadeOutDuration(Info.FadeOutDuration)
+				.bUseThrobber(Info.bUseThrobber)
+				.bUseSuccessFailIcons(Info.bUseSuccessFailIcons)
+				.bUseLargeFont(Info.bUseLargeFont)
+				.WidthOverride(Info.WidthOverride)
+				.CheckBoxState(Info.CheckBoxState)
+				.CheckBoxStateChanged(Info.CheckBoxStateChanged)
+				.CheckBoxText(Info.CheckBoxText)
+				.Hyperlink(Info.Hyperlink)
+				.HyperlinkText(Info.HyperlinkText);
+		}
 		
-	NewItem->MyList = SharedThis(this);
+		NewItem->MyList = SharedThis(this);
 
-	MessageItemBoxPtr->AddSlot()
-		.AutoHeight()
-		.HAlign(HAlign_Right)
-		[
-			NewItem.ToSharedRef()
-		];
+		MessageItemBoxPtr->AddSlot()
+			.AutoHeight()
+			.HAlign(HAlign_Right)
+			[
+				NewItem.ToSharedRef()
+			];
 
-	NewItem->Fadein( Info.bAllowThrottleWhenFrameRateIsLow );
+		NewItem->Fadein( Info.bAllowThrottleWhenFrameRateIsLow );
 
-	if (Info.bFireAndForget)
+		if (Info.bFireAndForget)
+		{
+			NewItem->ExpireAndFadeout();
+		}
+	}
+	else
 	{
-		NewItem->ExpireAndFadeout();
+		// When notifications are not allowed we want to return an empty notification.
+		NewItem = SNew(SNotificationItemImpl);
 	}
 
+	LastNotification = NewItem;
+
 	return NewItem.ToSharedRef();
 }
 
diff --git a/Engine/Source/Runtime/Slate/Private/Widgets/SViewport.cpp b/Engine/Source/Runtime/Slate/Private/Widgets/SViewport.cpp
--- a/Engine/Source/Runtime/Slate/Private/Widgets/SViewport.cpp
+++ b/Engine/Source/Runtime/Slate/Private/Widgets/SViewport.cpp
@@ -14,6 +14,7 @@ DECLARE_CYCLE_STAT(TEXT("Game UI Tick/Paint"), STAT_ViewportUpdateTime, STATGROU
 SViewport::SViewport()
 	: bRenderDirectlyToWindow(false)
 	, bEnableGammaCorrection(true)
+	, bReverseGammaCorrection(false)
 	, bEnableStereoRendering(false)
 { }
 
@@ -26,6 +27,7 @@ void SViewport::Construct( const FArguments& InArgs )
 	ShowDisabledEffect = InArgs._ShowEffectWhenDisabled;
 	bRenderDirectlyToWindow = InArgs._RenderDirectlyToWindow;
 	bEnableGammaCorrection = InArgs._EnableGammaCorrection;
+	bReverseGammaCorrection = InArgs._ReverseGammaCorrection;
 	bEnableBlending = InArgs._EnableBlending;
 	bEnableStereoRendering = InArgs._EnableStereoRendering;
 	bIgnoreTextureAlpha = InArgs._IgnoreTextureAlpha;
@@ -80,6 +82,12 @@ int32 SViewport::OnPaint( const FPaintArgs& Args, const FGeometry& AllottedGeome
 		DrawEffects |= ESlateDrawEffect::NoGamma;
 	}
 
+	// Should we reverse gamma correction
+	if (bReverseGammaCorrection)
+	{
+		DrawEffects |= ESlateDrawEffect::ReverseGamma;
+	}
+
 	// Show we enable blending?
 	if( !bEnableBlending )
 	{
diff --git a/Engine/Source/Runtime/Slate/Public/Framework/Application/SlateApplication.h b/Engine/Source/Runtime/Slate/Public/Framework/Application/SlateApplication.h
--- a/Engine/Source/Runtime/Slate/Public/Framework/Application/SlateApplication.h
+++ b/Engine/Source/Runtime/Slate/Public/Framework/Application/SlateApplication.h
@@ -51,6 +51,8 @@ DECLARE_DELEGATE_RetVal_TwoParams(bool, FOnWindowAction, const TSharedRef<FGener
 /** Delegate for overriding the behavior when a navigation action is taken, Not to be confused with FNavigationDelegate which allows a specific widget to override behavior for itself */
 DECLARE_DELEGATE_RetVal_OneParam(bool, FCustomNavigationHandler, TSharedPtr<SWidget>);
 
+DECLARE_DELEGATE_RetVal(bool, FDragDropCheckingOverride);
+
 extern SLATE_API const FName NAME_UnrealOS;
 
 
@@ -1490,6 +1492,9 @@ public:
 	*/
 	const void* FindBestParentWindowHandleForDialogs(const TSharedPtr<SWidget>& InWidget);
 
+public:
+	FDragDropCheckingOverride OnDragDropCheckOverride;
+
 private:
 
 	TSharedRef< FGenericWindow > MakeWindow( TSharedRef<SWindow> InSlateWindow, const bool bShowImmediately );
diff --git a/Engine/Source/Runtime/Slate/Public/Framework/Docking/TabManager.h b/Engine/Source/Runtime/Slate/Public/Framework/Docking/TabManager.h
--- a/Engine/Source/Runtime/Slate/Public/Framework/Docking/TabManager.h
+++ b/Engine/Source/Runtime/Slate/Public/Framework/Docking/TabManager.h
@@ -1009,6 +1009,7 @@ private:
 
 //#HACK VREDITOR - Had to introduce the proxy tab manager to steal asset tabs.
 
+DECLARE_MULTICAST_DELEGATE_TwoParams(FIsTabSupportedEvent, FTabId /* TabId */, bool& /* bOutIsSupported */ );
 DECLARE_MULTICAST_DELEGATE_OneParam(FOnTabEvent, TSharedPtr<SDockTab>);
 
 
@@ -1025,8 +1026,12 @@ public:
 
 	virtual void OpenUnmanagedTab(FName PlaceholderId, const FSearchPreference& SearchPreference, const TSharedRef<SDockTab>& UnmanagedTab) override;
 	virtual void DrawAttention(const TSharedRef<SDockTab>& TabToHighlight) override;
+	
+	bool IsTabSupported( const FTabId TabId ) const;
+	void SetParentWindow(TSharedRef<SWindow> InParentWindow);
 
 public:
+	FIsTabSupportedEvent OnIsTabSupported;
 	FOnTabEvent OnTabOpened;
 	FOnTabEvent OnAttentionDrawnToTab;
 
diff --git a/Engine/Source/Runtime/Slate/Public/Framework/Notifications/NotificationManager.h b/Engine/Source/Runtime/Slate/Public/Framework/Notifications/NotificationManager.h
--- a/Engine/Source/Runtime/Slate/Public/Framework/Notifications/NotificationManager.h
+++ b/Engine/Source/Runtime/Slate/Public/Framework/Notifications/NotificationManager.h
@@ -57,6 +57,21 @@ public:
 	 */
 	void GetWindows(TArray< TSharedRef<SWindow> >& OutWindows) const;
 
+	/**
+	 * Sets whether notifications should be displayed at all 
+	 * @param	bShouldAllow	Whether notifications should be enabled.  It defaults to on.
+	 */
+	void SetAllowNotifications( const bool bShouldAllow )
+	{
+		this->bAllowNotifications = bShouldAllow;
+	}
+
+	/** @return	Checks whether notifications are currently enabled */
+	bool AreNotificationsAllowed() const
+	{
+		return this->bAllowNotifications;
+	}
+
 protected:
 	/** Protect constructor as this is a singleton */
 	FSlateNotificationManager();
@@ -96,4 +111,7 @@ private:
 
 	/** Thread safe queue of notifications to display */
 	TLockFreePointerListLIFO<FNotificationInfo> PendingNotifications;
+
+	/** Whether notifications should be displayed or not.  This can be used to globally suppress notification pop-ups */
+	bool bAllowNotifications;
 };
diff --git a/Engine/Source/Runtime/Slate/Public/Widgets/Notifications/SNotificationList.h b/Engine/Source/Runtime/Slate/Public/Widgets/Notifications/SNotificationList.h
--- a/Engine/Source/Runtime/Slate/Public/Widgets/Notifications/SNotificationList.h
+++ b/Engine/Source/Runtime/Slate/Public/Widgets/Notifications/SNotificationList.h
@@ -259,4 +259,10 @@ protected:
 
 	/** Flag to auto-destroy this list. */
 	bool bDone;
+
+private:
+
+	/** The last notification added to the list. To avoid weak pointer pointing to a destroyed notification when notification are not allowed. */
+	TSharedPtr<SNotificationItem> LastNotification;
+
 };
diff --git a/Engine/Source/Runtime/Slate/Public/Widgets/SViewport.h b/Engine/Source/Runtime/Slate/Public/Widgets/SViewport.h
--- a/Engine/Source/Runtime/Slate/Public/Widgets/SViewport.h
+++ b/Engine/Source/Runtime/Slate/Public/Widgets/SViewport.h
@@ -29,6 +29,7 @@ public:
 		, _ShowEffectWhenDisabled(true)
 		, _RenderDirectlyToWindow(false)
 		, _EnableGammaCorrection(true)
+		, _ReverseGammaCorrection(false)
 		, _EnableBlending(false)
 		, _EnableStereoRendering(false)
 		, _PreMultipliedAlpha(true)
@@ -51,6 +52,9 @@ public:
 		/** Whether or not to enable gamma correction. Doesn't apply when rendering directly to a backbuffer. */
 		SLATE_ARGUMENT( bool, EnableGammaCorrection )
 
+		/** Whether or not to reverse the gamma correction done to the texture in this viewport.  Ignores the bEnableGammaCorrection setting */
+		SLATE_ARGUMENT( bool, ReverseGammaCorrection )
+
 		/** Allow this viewport to blend with its background. */
 		SLATE_ARGUMENT( bool, EnableBlending )
 
@@ -257,6 +261,9 @@ private:
 	/** Whether or not to apply gamma correction on the render target supplied by the ISlateViewport. */
 	bool bEnableGammaCorrection;
 
+	/** Whether or not to reverse the gamma correction done by the texture in this viewport.  Ignores the bEnableGammaCorrection setting */
+	bool bReverseGammaCorrection;
+
 	/** Whether or not to blend this viewport with the background. */
 	bool bEnableBlending;
 
diff --git a/Engine/Source/Runtime/SlateCore/Private/Rendering/DrawElements.cpp b/Engine/Source/Runtime/SlateCore/Private/Rendering/DrawElements.cpp
--- a/Engine/Source/Runtime/SlateCore/Private/Rendering/DrawElements.cpp
+++ b/Engine/Source/Runtime/SlateCore/Private/Rendering/DrawElements.cpp
@@ -65,6 +65,10 @@ void FSlateDrawElement::Init(uint32 InLayer, const FPaintGeometry& PaintGeometry
 	{
 		DataPayload.BatchFlags |= ESlateBatchDrawFlag::NoGamma;
 	}
+	if (InDrawEffects & ESlateDrawEffect::ReverseGamma)
+	{
+		DataPayload.BatchFlags |= ESlateBatchDrawFlag::ReverseGamma;
+	}
 }
 
 void FSlateDrawElement::MakeDebugQuad( FSlateWindowElementList& ElementList, uint32 InLayer, const FPaintGeometry& PaintGeometry, const FSlateRect& InClippingRect)
diff --git a/Engine/Source/Runtime/SlateCore/Public/Rendering/RenderingCommon.h b/Engine/Source/Runtime/SlateCore/Public/Rendering/RenderingCommon.h
--- a/Engine/Source/Runtime/SlateCore/Public/Rendering/RenderingCommon.h
+++ b/Engine/Source/Runtime/SlateCore/Public/Rendering/RenderingCommon.h
@@ -74,6 +74,8 @@ namespace ESlateDrawEffect
 	const Type PreMultipliedAlpha	= 1 << 3;
 	/** Advanced: No gamma correction should be done */
 	const Type NoGamma				= 1 << 4;
+	/** Advanced: Existing Gamma correction should be reversed */
+	const Type ReverseGamma			= 1 << 5;
 };
 
 
@@ -95,6 +97,8 @@ namespace ESlateBatchDrawFlag
 	const Type TileU				= 1 << 4;
 	/** The element should be tiled vertically */
 	const Type TileV				= 1 << 5;
+	/** Reverse gamma correction */
+	const Type ReverseGamma			 = 1 << 6;
 };
 
 namespace ESlateLineJoinType
diff --git a/Engine/Source/Runtime/SlateRHIRenderer/Private/SlateRHIRenderer.cpp b/Engine/Source/Runtime/SlateRHIRenderer/Private/SlateRHIRenderer.cpp
--- a/Engine/Source/Runtime/SlateRHIRenderer/Private/SlateRHIRenderer.cpp
+++ b/Engine/Source/Runtime/SlateRHIRenderer/Private/SlateRHIRenderer.cpp
@@ -55,6 +55,13 @@ static TAutoConsoleVariable<int32> CVarUICompositeMode(
 	TEXT("1: Shader pass to improve HDR blending\n"),
 	ECVF_RenderThreadSafe);
 
+static TAutoConsoleVariable<float> CVarDrawToVRRenderTarget(
+	TEXT("Slate.DrawToVRRenderTarget"),
+	1,
+	TEXT("If enabled while in VR. Slate UI will be drawn into the render target texture where the VR imagery for either eye was rendered, allow the viewer of the HMD to see the UI (for better or worse.)  This render target will then be cropped/scaled into the back buffer, if mirroring is enabled.  When disabled, Slate UI will be drawn on top of the backbuffer (not to the HMD) after the mirror texture has been cropped/scaled into the backbuffer."),
+	ECVF_RenderThreadSafe);
+
+
 void FSlateCrashReportResource::InitDynamicRHI()
 {
 	int32 Width = VirtualScreen.Width();
@@ -650,6 +657,13 @@ void FSlateRHIRenderer::DrawWindow_RenderThread(FRHICommandListImmediate& RHICmd
 	ViewportInfo.ColorSpaceLUTOutputDevice = HDROutputDevice;
 	ViewportInfo.ColorSpaceLUTOutputGamut = HDROutputGamut;
 	
+	bool bRenderedStereo = false;
+	if( CVarDrawToVRRenderTarget->GetInt() == 0 && GEngine && IsValidRef( ViewportInfo.GetRenderTargetTexture() ) && GEngine->StereoRenderingDevice.IsValid() )
+	{
+		GEngine->StereoRenderingDevice->RenderTexture_RenderThread( RHICmdList, RHICmdList.GetViewportBackBuffer( ViewportInfo.ViewportRHI ), ViewportInfo.GetRenderTargetTexture() );
+		bRenderedStereo = true;
+	}
+
 	{
 		SCOPED_GPU_STAT(RHICmdList, Stat_GPU_SlateUI);
 		SCOPE_CYCLE_COUNTER( STAT_SlateRenderingRTTime );
@@ -673,9 +687,8 @@ void FSlateRHIRenderer::DrawWindow_RenderThread(FRHICommandListImmediate& RHICmd
 		// should have been created by the game thread
 		check( IsValidRef(ViewportInfo.ViewportRHI) );
 
-		FTexture2DRHIRef ViewportRT = ViewportInfo.GetRenderTargetTexture();
-		FTexture2DRHIRef BackBuffer = (ViewportRT) ?
-		ViewportRT : RHICmdList.GetViewportBackBuffer(ViewportInfo.ViewportRHI);
+		FTexture2DRHIRef ViewportRT = bRenderedStereo ? nullptr : ViewportInfo.GetRenderTargetTexture();
+		FTexture2DRHIRef BackBuffer = (ViewportRT) ? ViewportRT : RHICmdList.GetViewportBackBuffer(ViewportInfo.ViewportRHI);
 		
 		const uint32 ViewportWidth = (ViewportRT) ? ViewportRT->GetSizeX() : ViewportInfo.Width;
 		const uint32 ViewportHeight = (ViewportRT) ? ViewportRT->GetSizeY() : ViewportInfo.Height;
@@ -832,7 +845,7 @@ void FSlateRHIRenderer::DrawWindow_RenderThread(FRHICommandListImmediate& RHICmd
 		}
 	}
 
-	if (GEngine && IsValidRef(ViewportInfo.GetRenderTargetTexture()) && GEngine->StereoRenderingDevice.IsValid())
+	if (!bRenderedStereo && GEngine && IsValidRef(ViewportInfo.GetRenderTargetTexture()) && GEngine->StereoRenderingDevice.IsValid())
 	{
 		GEngine->StereoRenderingDevice->RenderTexture_RenderThread(RHICmdList, RHICmdList.GetViewportBackBuffer(ViewportInfo.ViewportRHI), ViewportInfo.GetRenderTargetTexture());
 	}
diff --git a/Engine/Source/Runtime/SlateRHIRenderer/Private/SlateRHIRenderingPolicy.cpp b/Engine/Source/Runtime/SlateRHIRenderer/Private/SlateRHIRenderingPolicy.cpp
--- a/Engine/Source/Runtime/SlateRHIRenderer/Private/SlateRHIRenderingPolicy.cpp
+++ b/Engine/Source/Runtime/SlateRHIRenderer/Private/SlateRHIRenderingPolicy.cpp
@@ -239,7 +239,8 @@ void FSlateRHIRenderingPolicy::DrawElements(FRHICommandListImmediate& RHICmdList
 
 	static const FEngineShowFlags DefaultShowFlags(ESFIM_Game);
 
-	const float DisplayGamma = bGammaCorrect ? GEngine ? GEngine->GetDisplayGamma() : 2.2f : 1.0f;
+	const float EngineGamma = GEngine ? GEngine->GetDisplayGamma() : 2.2f;
+	const float DisplayGamma = bGammaCorrect ? EngineGamma : 1.0f;
 
 	// In order to support MaterialParameterCollections, we need to create multiple FSceneViews for 
 	// each possible Scene that we encounter. The following code creates these as separate arrays, where the first 
@@ -326,6 +327,8 @@ void FSlateRHIRenderingPolicy::DrawElements(FRHICommandListImmediate& RHICmdList
 		const ESlateShader::Type ShaderType = RenderBatch.ShaderType;
 		const FShaderParams& ShaderParams = RenderBatch.ShaderParams;
 
+		const float FinalGamma = DrawFlags & ESlateBatchDrawFlag::ReverseGamma ? 1.0f / EngineGamma : (DrawFlags & ESlateBatchDrawFlag::NoGamma) ? 1.0f : DisplayGamma;
+
 		auto UpdateScissorRect = [&RHICmdList, &RenderBatch]()
 		{
 			if (RenderBatch.ScissorRect.IsSet())
@@ -474,7 +477,9 @@ void FSlateRHIRenderingPolicy::DrawElements(FRHICommandListImmediate& RHICmdList
                 
 				PixelShader->SetTexture(RHICmdList, TextureRHI, SamplerState);
 				PixelShader->SetShaderParams(RHICmdList, ShaderParams.PixelParams);
-				PixelShader->SetDisplayGamma(RHICmdList, (DrawFlags & ESlateBatchDrawFlag::NoGamma) ? 1.0f : DisplayGamma);
+
+				
+				PixelShader->SetDisplayGamma(RHICmdList, FinalGamma);
 
 	
 				// for RHIs that can't handle VertexOffset, we need to offset the stream source each time
@@ -547,8 +552,8 @@ void FSlateRHIRenderingPolicy::DrawElements(FRHICommandListImmediate& RHICmdList
 						VertexShader->SetVerticalAxisMultiplier(RHICmdList, bAllowSwitchVerticalAxis && RHINeedsToSwitchVerticalAxis(GShaderPlatformForFeatureLevel[GMaxRHIFeatureLevel]) ? -1.0f : 1.0f);
 						VertexShader->SetMaterialShaderParameters(RHICmdList, ActiveSceneView, MaterialRenderProxy, Material);
 
-						PixelShader->SetParameters(RHICmdList, ActiveSceneView, MaterialRenderProxy, Material, DisplayGamma, ShaderParams.PixelParams);
-						PixelShader->SetDisplayGamma(RHICmdList, (DrawFlags & ESlateBatchDrawFlag::NoGamma) ? 1.0f : DisplayGamma);
+						PixelShader->SetParameters(RHICmdList, ActiveSceneView, MaterialRenderProxy, Material, FinalGamma, ShaderParams.PixelParams);
+						PixelShader->SetDisplayGamma(RHICmdList, FinalGamma);
 
 						FSlateShaderResource* MaskResource = MaterialShaderResource->GetTextureMaskResource();
 						if (MaskResource)
diff --git a/Engine/Source/Runtime/UMG/Private/Components/WidgetComponent.cpp b/Engine/Source/Runtime/UMG/Private/Components/WidgetComponent.cpp
--- a/Engine/Source/Runtime/UMG/Private/Components/WidgetComponent.cpp
+++ b/Engine/Source/Runtime/UMG/Private/Components/WidgetComponent.cpp
@@ -1223,11 +1223,14 @@ void UWidgetComponent::UpdateWidget()
 				NewSlateWidget = Widget->TakeWidget();
 			}
 
+			bool bNeededNewWindow = false;
 			if ( !SlateWindow.IsValid() )
 			{
 				SlateWindow = SNew(SVirtualWindow).Size(DrawSize);
 				SlateWindow->SetIsFocusable(bWindowFocusable);
 				RegisterWindow();
+
+				bNeededNewWindow = true;
 			}
 
 			if ( !HitTestGrid.IsValid() )
@@ -1239,7 +1242,7 @@ void UWidgetComponent::UpdateWidget()
 
 			if ( NewSlateWidget.IsValid() )
 			{
-				if ( NewSlateWidget != CurrentSlateWidget )
+				if ( NewSlateWidget != CurrentSlateWidget || bNeededNewWindow )
 				{
 					CurrentSlateWidget = NewSlateWidget;
 					SlateWindow->SetContent(NewSlateWidget.ToSharedRef());
@@ -1247,7 +1250,7 @@ void UWidgetComponent::UpdateWidget()
 			}
 			else if( SlateWidget.IsValid() )
 			{
-				if ( SlateWidget != CurrentSlateWidget )
+				if ( SlateWidget != CurrentSlateWidget || bNeededNewWindow )
 				{
 					CurrentSlateWidget = SlateWidget;
 					SlateWindow->SetContent(SlateWidget.ToSharedRef());
