Automation - Fixing the stereo flag when running tests, it's no longer part of the hardware string.  Automation map loading now waits for the map to finish initializing actors.

diff --git a/Engine/Source/Editor/UnrealEd/Private/EditorEngine.cpp b/Engine/Source/Editor/UnrealEd/Private/EditorEngine.cpp
--- a/Engine/Source/Editor/UnrealEd/Private/EditorEngine.cpp
+++ b/Engine/Source/Editor/UnrealEd/Private/EditorEngine.cpp
@@ -6998,11 +6998,9 @@ void UEditorEngine::AutomationLoadMap(const FString& MapName, FString* OutError)
 		}
 	}
 
-	//should really be wait until the map is properly loaded....in PIE or gameplay....
 	if (bNeedPieStart)
 	{
-		//TODO NICKD We need a better way to determine when to start the map.
-		ADD_LATENT_AUTOMATION_COMMAND(FWaitLatentCommand(10.f));
+		ADD_LATENT_AUTOMATION_COMMAND(FWaitForMapToLoadCommand);
 	}
 #endif
 	return;
diff --git a/Engine/Source/Runtime/Engine/Private/Tests/AutomationCommon.cpp b/Engine/Source/Runtime/Engine/Private/Tests/AutomationCommon.cpp
--- a/Engine/Source/Runtime/Engine/Private/Tests/AutomationCommon.cpp
+++ b/Engine/Source/Runtime/Engine/Private/Tests/AutomationCommon.cpp
@@ -17,6 +17,7 @@
 #include "GameFramework/GameStateBase.h"
 #include "Scalability.h"
 #include "Matinee/MatineeActor.h"
+#include "IHeadMountedDisplay.h"
 
 #if (WITH_DEV_AUTOMATION_TESTS || WITH_PERF_AUTOMATION_TESTS)
 
@@ -66,11 +67,6 @@ namespace AutomationCommon
 			HardwareDetailsString = ( HardwareDetailsString + TEXT("_") ) + FeatureLevelString;
 		}
 
-		if ( GEngine->StereoRenderingDevice.IsValid() )
-		{
-			HardwareDetailsString = ( HardwareDetailsString + TEXT("_") ) + TEXT("STEREO");
-		}
-
 		if ( HardwareDetailsString.Len() > 0 )
 		{
 			//Get rid of the leading "_"
@@ -108,7 +104,7 @@ namespace AutomationCommon
 		Data.Platform = FPlatformProperties::PlatformName();
 		Data.Rhi = FHardwareInfo::GetHardwareInfo(NAME_RHI);
 		GetFeatureLevelName(GMaxRHIFeatureLevel, Data.FeatureLevel);
-		Data.bIsStereo = GEngine->StereoRenderingDevice.IsValid();
+		Data.bIsStereo = GEngine->HMDDevice.IsValid() ? GEngine->HMDDevice->IsStereoEnabled() : false;
 		Data.Vendor = RHIVendorIdToString();
 		Data.AdapterName = GRHIAdapterName;
 		Data.AdapterInternalDriverVersion = GRHIAdapterInternalDriverVersion;
