Fixed serialization crash for simplified skeletal meshes leading to corrupted assets that crash on load after skin weight buffer changes.

diff --git a/Engine/Source/Runtime/Engine/Private/Rendering/SkinWeightVertexBuffer.cpp b/Engine/Source/Runtime/Engine/Private/Rendering/SkinWeightVertexBuffer.cpp
--- a/Engine/Source/Runtime/Engine/Private/Rendering/SkinWeightVertexBuffer.cpp
+++ b/Engine/Source/Runtime/Engine/Private/Rendering/SkinWeightVertexBuffer.cpp
@@ -93,7 +93,7 @@ FArchive& operator<<(FArchive& Ar, FSkinWeightVertexBuffer& VertexBuffer)
 	Ar << VertexBuffer.bExtraBoneInfluences;
 	Ar << VertexBuffer.NumVertices;
 
-	if (Ar.IsLoading())
+	if (Ar.IsLoading() || VertexBuffer.WeightData == NULL)
 	{
 		// allocate vertex data on load
 		VertexBuffer.AllocateData();
