Patch for possible crash when updating the dynamic component bounds for texture streamer because of dangling pointer.
The crash is only possible following CL 3231696 (last CL affecting streaming)

diff --git a/Engine/Source/Runtime/Engine/Private/Streaming/TextureStreamingHelpers.cpp b/Engine/Source/Runtime/Engine/Private/Streaming/TextureStreamingHelpers.cpp
--- a/Engine/Source/Runtime/Engine/Private/Streaming/TextureStreamingHelpers.cpp
+++ b/Engine/Source/Runtime/Engine/Private/Streaming/TextureStreamingHelpers.cpp
@@ -185,7 +185,7 @@ TAutoConsoleVariable<int32> CVarStreamingNumStaticComponentsProcessedPerFrame(
 
 TAutoConsoleVariable<int32> CVarStreamingDefragDynamicBounds(
 	TEXT("r.Streaming.DefragDynamicBounds"),
-	1,
+	0,
 	TEXT("If non-zero, unused dynamic bounds will be removed from the update loop"),
 	ECVF_Default);
 
