Changed ProcessQueuedBunch warning to clog for cooked platforms

diff --git a/Engine/Source/Runtime/Engine/Private/DataChannel.cpp b/Engine/Source/Runtime/Engine/Private/DataChannel.cpp
--- a/Engine/Source/Runtime/Engine/Private/DataChannel.cpp
+++ b/Engine/Source/Runtime/Engine/Private/DataChannel.cpp
@@ -1951,7 +1951,7 @@ bool UActorChannel::ProcessQueuedBunches()
 
 		if ( FPlatformTime::Seconds() - QueuedBunchStartTime > QUEUED_BUNCH_TIMEOUT_IN_SECONDS )
 		{
-			UE_LOG( LogNet, Warning, TEXT( "UActorChannel::ProcessQueuedBunches: Queued bunches for longer than normal. ChIndex: %i, Actor: %s, Queued: %i" ), ChIndex, Actor != NULL ? *Actor->GetPathName() : TEXT( "NULL" ), QueuedBunches.Num() );
+			UE_CLOG(FPlatformProperties::RequiresCookedData(), LogNet, Warning, TEXT( "UActorChannel::ProcessQueuedBunches: Queued bunches for longer than normal. ChIndex: %i, Actor: %s, Queued: %i" ), ChIndex, Actor != NULL ? *Actor->GetPathName() : TEXT( "NULL" ), QueuedBunches.Num() );
 			QueuedBunchStartTime = FPlatformTime::Seconds();
 		}
 	}
