Copying //UE4/Dev-Editor to //UE4/Dev-Main (Source: //UE4/Dev-Editor @ 3279756)

diff --git a/Engine/Source/Runtime/CoreUObject/Private/Serialization/PropertyLocalizationDataGathering.cpp b/Engine/Source/Runtime/CoreUObject/Private/Serialization/PropertyLocalizationDataGathering.cpp
--- a/Engine/Source/Runtime/CoreUObject/Private/Serialization/PropertyLocalizationDataGathering.cpp
+++ b/Engine/Source/Runtime/CoreUObject/Private/Serialization/PropertyLocalizationDataGathering.cpp
@@ -9,6 +9,7 @@
 #include "UObject/UnrealType.h"
 #include "UObject/TextProperty.h"
 #include "UObject/PropertyPortFlags.h"
+#include "HAL/UnrealMemory.h"
 #include "Internationalization/TextNamespaceUtil.h"
 #include "Internationalization/TextPackageNamespaceUtil.h"
 
@@ -379,159 +380,166 @@ void FPropertyLocalizationDataGatherer::GatherTextInstance(const FText& Text, co
 	}
 }
 
-void FPropertyLocalizationDataGatherer::GatherScriptBytecode(const FString& PathToScript, const TArray<uint8>& ScriptData, const bool bIsEditorOnly)
+struct FGatherTextFromScriptBytecode
 {
-	struct FGatherTextFromScriptBytecode
+public:
+	FGatherTextFromScriptBytecode(const TCHAR* InSourceDescription, const TArray<uint8>& InScript, FPropertyLocalizationDataGatherer& InPropertyLocalizationDataGatherer, const bool InTreatAsEditorOnlyData)
+		: SourceDescription(InSourceDescription)
+		, Script(const_cast<TArray<uint8>&>(InScript)) // We won't change the script, but we have to lie so that the code in ScriptSerialization.h will compile :(
+		, PropertyLocalizationDataGatherer(InPropertyLocalizationDataGatherer)
+		, bTreatAsEditorOnlyData(InTreatAsEditorOnlyData)
+		, bIsParsingText(false)
 	{
-	public:
-		FGatherTextFromScriptBytecode(const TCHAR* InSourceDescription, const TArray<uint8>& InScript, FPropertyLocalizationDataGatherer& InPropertyLocalizationDataGatherer, const bool InTreatAsEditorOnlyData)
-			: SourceDescription(InSourceDescription)
-			, Script(const_cast<TArray<uint8>&>(InScript)) // We won't change the script, but we have to lie so that the code in ScriptSerialization.h will compile :(
-			, PropertyLocalizationDataGatherer(InPropertyLocalizationDataGatherer)
-			, bTreatAsEditorOnlyData(InTreatAsEditorOnlyData)
-			, bIsParsingText(false)
-		{
-			const int32 ScriptSizeBytes = Script.Num();
+		const int32 ScriptSizeBytes = Script.Num();
 
-			int32 iCode = 0;
-			while (iCode < ScriptSizeBytes)
-			{
-				SerializeExpr(iCode, DummyArchive);
-			}
-		}
-
-	private:
-		FLinker* GetLinker()
+		int32 iCode = 0;
+		while (iCode < ScriptSizeBytes)
 		{
-			return nullptr;
+			SerializeExpr(iCode, DummyArchive);
 		}
+	}
 
-		EExprToken SerializeExpr(int32& iCode, FArchive& Ar)
-		{
-		#define XFERSTRING()		SerializeString(iCode, Ar)
-		#define XFERUNICODESTRING() SerializeUnicodeString(iCode, Ar)
-		#define XFERTEXT()			SerializeText(iCode, Ar)
-
-		#define SERIALIZEEXPR_INC
-		#define SERIALIZEEXPR_AUTO_UNDEF_XFER_MACROS
-		#include "ScriptSerialization.h"
-			return Expr;
-		#undef SERIALIZEEXPR_INC
-		#undef SERIALIZEEXPR_AUTO_UNDEF_XFER_MACROS
-		}
+private:
+	FLinker* GetLinker()
+	{
+		return nullptr;
+	}
+
+	EExprToken SerializeExpr(int32& iCode, FArchive& Ar)
+	{
+	#define XFERSTRING()		SerializeString(iCode, Ar)
+	#define XFERUNICODESTRING() SerializeUnicodeString(iCode, Ar)
+	#define XFERTEXT()			SerializeText(iCode, Ar)
+
+	#define SERIALIZEEXPR_INC
+	#define SERIALIZEEXPR_AUTO_UNDEF_XFER_MACROS
+	#include "ScriptSerialization.h"
+		return Expr;
+	#undef SERIALIZEEXPR_INC
+	#undef SERIALIZEEXPR_AUTO_UNDEF_XFER_MACROS
+	}
 
-		void SerializeString(int32& iCode, FArchive& Ar)
+	void SerializeString(int32& iCode, FArchive& Ar)
+	{
+		if (bIsParsingText)
 		{
-			if (bIsParsingText)
-			{
-				LastParsedString.Reset();
+			LastParsedString.Reset();
 
-				do
-				{
-					LastParsedString += (char)(Script[iCode]);
+			do
+			{
+				LastParsedString += (char)(Script[iCode]);
 
-					iCode += sizeof(uint8);
-				}
-				while (Script[iCode-1]);
+				iCode += sizeof(uint8);
 			}
-			else
+			while (Script[iCode-1]);
+		}
+		else
+		{
+			do
 			{
-				do
-				{
-					iCode += sizeof(uint8);
-				}
-				while (Script[iCode-1]);
+				iCode += sizeof(uint8);
 			}
+			while (Script[iCode-1]);
 		}
+	}
 
-		void SerializeUnicodeString(int32& iCode, FArchive& Ar)
+	void SerializeUnicodeString(int32& iCode, FArchive& Ar)
+	{
+		if (bIsParsingText)
 		{
-			if (bIsParsingText)
-			{
-				LastParsedString.Reset();
+			LastParsedString.Reset();
 
-				do
-				{
-					uint16 UnicodeChar = 0;
+			do
+			{
+				uint16 UnicodeChar = 0;
 #ifdef REQUIRES_ALIGNED_INT_ACCESS
-					FMemory::Memcpy(&UnicodeChar, &Script[iCode], sizeof(uint16));
+				FMemory::Memcpy(&UnicodeChar, &Script[iCode], sizeof(uint16));
 #else
-					UnicodeChar = *((uint16*)(&Script[iCode]));
+				UnicodeChar = *((uint16*)(&Script[iCode]));
 #endif
-					LastParsedString += (TCHAR)UnicodeChar;
+				LastParsedString += (TCHAR)UnicodeChar;
 
-					iCode += sizeof(uint16);
-				}
-				while (Script[iCode-1] || Script[iCode-2]);
+				iCode += sizeof(uint16);
 			}
-			else
+			while (Script[iCode-1] || Script[iCode-2]);
+		}
+		else
+		{
+			do
 			{
-				do
-				{
-					iCode += sizeof(uint16);
-				}
-				while (Script[iCode-1] || Script[iCode-2]);
+				iCode += sizeof(uint16);
 			}
+			while (Script[iCode-1] || Script[iCode-2]);
 		}
+	}
 
-		void SerializeText(int32& iCode, FArchive& Ar)
+	void SerializeText(int32& iCode, FArchive& Ar)
+	{
+		// What kind of text are we dealing with?
+		const EBlueprintTextLiteralType TextLiteralType = (EBlueprintTextLiteralType)Script[iCode++];
+
+		switch (TextLiteralType)
 		{
-			// What kind of text are we dealing with?
-			const EBlueprintTextLiteralType TextLiteralType = (EBlueprintTextLiteralType)Script[iCode++];
+		case EBlueprintTextLiteralType::Empty:
+			// Don't need to gather empty text
+			break;
 
-			switch (TextLiteralType)
+		case EBlueprintTextLiteralType::LocalizedText:
 			{
-			case EBlueprintTextLiteralType::Empty:
-				// Don't need to gather empty text
-				break;
+				bIsParsingText = true;
 
-			case EBlueprintTextLiteralType::LocalizedText:
-				{
-					bIsParsingText = true;
+				SerializeExpr(iCode, Ar);
+				const FString SourceString = MoveTemp(LastParsedString);
 
-					SerializeExpr(iCode, Ar);
-					const FString SourceString = MoveTemp(LastParsedString);
+				SerializeExpr(iCode, Ar);
+				const FString TextKey = MoveTemp(LastParsedString);
 
-					SerializeExpr(iCode, Ar);
-					const FString TextKey = MoveTemp(LastParsedString);
+				SerializeExpr(iCode, Ar);
+				const FString TextNamespace = MoveTemp(LastParsedString);
 
-					SerializeExpr(iCode, Ar);
-					const FString TextNamespace = MoveTemp(LastParsedString);
+				bIsParsingText = false;
 
-					bIsParsingText = false;
+				const FText TextInstance = FInternationalization::ForUseOnlyByLocMacroAndGraphNodeTextLiterals_CreateText(*SourceString, *TextNamespace, *TextKey);
 
-					const FText TextInstance = FInternationalization::ForUseOnlyByLocMacroAndGraphNodeTextLiterals_CreateText(*SourceString, *TextNamespace, *TextKey);
+				PropertyLocalizationDataGatherer.GatherTextInstance(TextInstance, FString::Printf(TEXT("%s [Script Bytecode]"), SourceDescription), bTreatAsEditorOnlyData);
+			}
+			break;
 
-					PropertyLocalizationDataGatherer.GatherTextInstance(TextInstance, FString::Printf(TEXT("%s [Script Bytecode]"), SourceDescription), bTreatAsEditorOnlyData);
-				}
-				break;
+		case EBlueprintTextLiteralType::InvariantText:
+			// Don't need to gather invariant text, but we do need to walk over the string in the buffer
+			SerializeExpr(iCode, Ar);
+			break;
 
-			case EBlueprintTextLiteralType::InvariantText:
-				// Don't need to gather invariant text, but we do need to walk over the string in the buffer
-				SerializeExpr(iCode, Ar);
-				break;
+		case EBlueprintTextLiteralType::LiteralString:
+			// Don't need to gather literal strings, but we do need to walk over the string in the buffer
+			SerializeExpr(iCode, Ar);
+			break;
 
-			case EBlueprintTextLiteralType::LiteralString:
-				// Don't need to gather literal strings, but we do need to walk over the string in the buffer
-				SerializeExpr(iCode, Ar);
-				break;
+		case EBlueprintTextLiteralType::StringTableEntry:
+			// Don't need to gather string table entries, but we do need to walk over the strings in the buffer
+			iCode += sizeof(ScriptPointerType); // String Table asset (if any)
+			SerializeExpr(iCode, Ar);
+			SerializeExpr(iCode, Ar);
+			break;
 
-			default:
-				checkf(false, TEXT("Unknown EBlueprintTextLiteralType! Please update FGatherTextFromScriptBytecode::SerializeText to handle this type of text."));
-				break;
-			}
+		default:
+			checkf(false, TEXT("Unknown EBlueprintTextLiteralType! Please update FGatherTextFromScriptBytecode::SerializeText to handle this type of text."));
+			break;
 		}
+	}
 
-		const TCHAR* SourceDescription;
-		TArray<uint8>& Script;
-		FPropertyLocalizationDataGatherer& PropertyLocalizationDataGatherer;
-		bool bTreatAsEditorOnlyData;
+	const TCHAR* SourceDescription;
+	TArray<uint8>& Script;
+	FPropertyLocalizationDataGatherer& PropertyLocalizationDataGatherer;
+	bool bTreatAsEditorOnlyData;
 
-		FArchive DummyArchive;
-		bool bIsParsingText;
-		FString LastParsedString;
-	};
+	FArchive DummyArchive;
+	bool bIsParsingText;
+	FString LastParsedString;
+};
 
+void FPropertyLocalizationDataGatherer::GatherScriptBytecode(const FString& PathToScript, const TArray<uint8>& ScriptData, const bool bIsEditorOnly)
+{
 	if (ScriptData.Num() > 0)
 	{
 		ResultFlags |= EPropertyLocalizationGathererResultFlags::HasScript;
diff --git a/Engine/Source/Runtime/CoreUObject/Private/UObject/ScriptCore.cpp b/Engine/Source/Runtime/CoreUObject/Private/UObject/ScriptCore.cpp
--- a/Engine/Source/Runtime/CoreUObject/Private/UObject/ScriptCore.cpp
+++ b/Engine/Source/Runtime/CoreUObject/Private/UObject/ScriptCore.cpp
@@ -2464,6 +2464,20 @@ void UObject::execTextConst( FFrame& Stack, RESULT_DECL )
 		}
 		break;
 
+	case EBlueprintTextLiteralType::StringTableEntry:
+		{
+			Stack.ReadObject(); // String Table asset (if any)
+
+			FString TableIdString;
+			Stack.Step(Stack.Object, &TableIdString);
+
+			FString KeyString;
+			Stack.Step(Stack.Object, &KeyString);
+
+			*(FText*)RESULT_PARAM = FText::FromStringTable(FName(*TableIdString), KeyString);
+		}
+		break;
+
 	default:
 		checkf(false, TEXT("Unknown EBlueprintTextLiteralType! Please update UObject::execTextConst to handle this type of text."));
 		break;
diff --git a/Engine/Source/Runtime/CoreUObject/Private/UObject/TextProperty.cpp b/Engine/Source/Runtime/CoreUObject/Private/UObject/TextProperty.cpp
--- a/Engine/Source/Runtime/CoreUObject/Private/UObject/TextProperty.cpp
+++ b/Engine/Source/Runtime/CoreUObject/Private/UObject/TextProperty.cpp
@@ -4,6 +4,7 @@
 #include "Internationalization/ITextData.h"
 #include "UObject/PropertyPortFlags.h"
 #include "Internationalization/TextPackageNamespaceUtil.h"
+#include "Internationalization/StringTableRegistry.h"
 
 bool UTextProperty::ConvertFromType(const FPropertyTag& Tag, FArchive& Ar, uint8* Data, UStruct* DefaultsStruct, bool& bOutAdvanceProperty)
 {
@@ -147,6 +148,28 @@ FString UTextProperty::GenerateCppCodeForTextValue(const FText& InValue, const F
 	{
 		CppCode += TEXT("FText::GetEmpty()");
 	}
+	else if (InValue.IsFromStringTable())
+	{
+		FName TableId;
+		FString Key;
+		FStringTableRegistry::Get().FindTableIdAndKey(InValue, TableId, Key);
+
+		// Produces FText::FromStringTable(TEXT("..."), TEXT("..."))
+		CppCode += TEXT("FText::FromStringTable(\n");
+
+		CppCode += Indent;
+		CppCode += TEXT("\t");
+		CppCode += UStrProperty::ExportCppHardcodedText(TableId.ToString(), Indent + TEXT("\t\t"));
+		CppCode += TEXT(", /* String Table ID */\n");
+
+		CppCode += Indent;
+		CppCode += TEXT("\t");
+		CppCode += UStrProperty::ExportCppHardcodedText(Key, Indent + TEXT("\t\t"));
+		CppCode += TEXT(" /* Key */\n");
+
+		CppCode += Indent;
+		CppCode += TEXT("\t)");
+	}
 	else if (InValue.IsCultureInvariant())
 	{
 		// Produces FText::AsCultureInvariant(TEXT("..."))
diff --git a/Engine/Source/Runtime/CoreUObject/Public/UObject/Script.h b/Engine/Source/Runtime/CoreUObject/Public/UObject/Script.h
--- a/Engine/Source/Runtime/CoreUObject/Public/UObject/Script.h
+++ b/Engine/Source/Runtime/CoreUObject/Public/UObject/Script.h
@@ -282,6 +282,8 @@ enum class EBlueprintTextLiteralType : uint8
 	InvariantText,
 	/** Text is a literal FString. The bytecode will contain one string, and you should use FText::FromString to initialize the FText instance. */
 	LiteralString,
+	/** Text is from a string table. The bytecode will contain an object pointer (not used) and two strings - the table ID, and key - and should be found via FText::FromStringTable */
+	StringTableEntry,
 };
 
 // Kinds of Blueprint exceptions
diff --git a/Engine/Source/Runtime/CoreUObject/Public/UObject/ScriptSerialization.h b/Engine/Source/Runtime/CoreUObject/Public/UObject/ScriptSerialization.h
--- a/Engine/Source/Runtime/CoreUObject/Public/UObject/ScriptSerialization.h
+++ b/Engine/Source/Runtime/CoreUObject/Public/UObject/ScriptSerialization.h
@@ -91,6 +91,12 @@
 		case EBlueprintTextLiteralType::LiteralString: \
 			SerializeExpr( iCode, Ar );	\
 			break; \
+		case EBlueprintTextLiteralType::StringTableEntry: \
+			XFER_OBJECT_POINTER( UObject* ); \
+			FIXUP_EXPR_OBJECT_POINTER( UObject* ); \
+			SerializeExpr( iCode, Ar );	\
+			SerializeExpr( iCode, Ar ); \
+			break; \
 		default: \
 			checkf(false, TEXT("Unknown EBlueprintTextLiteralType! Please update XFERTEXT to handle this type of text.")); \
 			break; \
diff --git a/Engine/Source/Runtime/Engine/Classes/Animation/SkeletalMeshActor.h b/Engine/Source/Runtime/Engine/Classes/Animation/SkeletalMeshActor.h
--- a/Engine/Source/Runtime/Engine/Classes/Animation/SkeletalMeshActor.h
+++ b/Engine/Source/Runtime/Engine/Classes/Animation/SkeletalMeshActor.h
@@ -69,9 +69,6 @@ public:
 	//~ Begin UObject Interface
 protected:
 	virtual FString GetDetailedInfoInternal() const override;
-#if WITH_EDITOR
-	virtual void PostEditChangeChainProperty(FPropertyChangedChainEvent& PropertyChangedEvent) override;
-#endif // WITH_EDITOR
 public:
 	//~ End UObject Interface
 
diff --git a/Engine/Source/Runtime/Engine/Classes/Components/MeshComponent.h b/Engine/Source/Runtime/Engine/Classes/Components/MeshComponent.h
--- a/Engine/Source/Runtime/Engine/Classes/Components/MeshComponent.h
+++ b/Engine/Source/Runtime/Engine/Classes/Components/MeshComponent.h
@@ -50,7 +50,9 @@ class ENGINE_API UMeshComponent : public UPrimitiveComponent
 	/** 
 	 * This empties all override materials and used by editor when replacing preview mesh 
 	 */
-	void EmptyOverrideMaterials(); 
+	void EmptyOverrideMaterials();
+
+	virtual void PostEditChangeChainProperty(FPropertyChangedChainEvent& PropertyChangedEvent) override;
 #endif
 
 	//~ Begin UObject Interface
diff --git a/Engine/Source/Runtime/Engine/Classes/Engine/GameViewportDelegates.h b/Engine/Source/Runtime/Engine/Classes/Engine/GameViewportDelegates.h
--- a/Engine/Source/Runtime/Engine/Classes/Engine/GameViewportDelegates.h
+++ b/Engine/Source/Runtime/Engine/Classes/Engine/GameViewportDelegates.h
@@ -39,7 +39,7 @@ DECLARE_DELEGATE_RetVal(bool, FOnWindowCloseRequested);
  * The first parameter is the viewport being closed.
  * @see UGameViewportClient
  */
-DECLARE_DELEGATE_OneParam(FOnCloseRequested, FViewport*);
+DECLARE_MULTICAST_DELEGATE_OneParam(FOnCloseRequested, FViewport*);
 
 /**
  * Delegate type used by UGameViewportClient for when a player is added or removed
diff --git a/Engine/Source/Runtime/Engine/Classes/Engine/Scene.h b/Engine/Source/Runtime/Engine/Classes/Engine/Scene.h
--- a/Engine/Source/Runtime/Engine/Classes/Engine/Scene.h
+++ b/Engine/Source/Runtime/Engine/Classes/Engine/Scene.h
@@ -504,50 +504,50 @@ struct FPostProcessSettings
 	float WhiteTint;
 
 	// Color Correction controls
-	UPROPERTY(interp, BlueprintReadWrite, Category = "Color Grading|Global", meta = (UIMin = "0.0", UIMax = "2.0", Delta = "0.01", ColorGradingMode = "saturation", ShiftMouseMovePixelPerDelta = "10", editcondition = "bOverride_ColorSaturation", DisplayName = "Saturation"))
+	UPROPERTY(interp, BlueprintReadWrite, Category = "Color Grading|Global", meta = (UIMin = "0.0", UIMax = "2.0", Delta = "0.01", DisplayColorChannel = "true", ColorGradingMode = "saturation", ShiftMouseMovePixelPerDelta = "10", editcondition = "bOverride_ColorSaturation", DisplayName = "Saturation"))
 	FVector4 ColorSaturation;
-	UPROPERTY(interp, BlueprintReadWrite, Category = "Color Grading|Global", meta = (UIMin = "0.0", UIMax = "2.0", Delta = "0.01", ColorGradingMode = "contrast", ShiftMouseMovePixelPerDelta = "10", editcondition = "bOverride_ColorContrast", DisplayName = "Contrast"))
+	UPROPERTY(interp, BlueprintReadWrite, Category = "Color Grading|Global", meta = (UIMin = "0.0", UIMax = "2.0", Delta = "0.01", DisplayColorChannel = "true", ColorGradingMode = "contrast", ShiftMouseMovePixelPerDelta = "10", editcondition = "bOverride_ColorContrast", DisplayName = "Contrast"))
 	FVector4 ColorContrast;
-	UPROPERTY(interp, BlueprintReadWrite, Category = "Color Grading|Global", meta = (UIMin = "0.0", UIMax = "2.0", Delta = "0.01", ColorGradingMode = "gamma", ShiftMouseMovePixelPerDelta = "10", editcondition = "bOverride_ColorGamma", DisplayName = "Gamma"))
+	UPROPERTY(interp, BlueprintReadWrite, Category = "Color Grading|Global", meta = (UIMin = "0.0", UIMax = "2.0", Delta = "0.01", DisplayColorChannel = "true", ColorGradingMode = "gamma", ShiftMouseMovePixelPerDelta = "10", editcondition = "bOverride_ColorGamma", DisplayName = "Gamma"))
 	FVector4 ColorGamma;
-	UPROPERTY(interp, BlueprintReadWrite, Category = "Color Grading|Global", meta = (UIMin = "0.0", UIMax = "2.0", Delta = "0.01", ColorGradingMode = "gain", ShiftMouseMovePixelPerDelta = "10", editcondition = "bOverride_ColorGain", DisplayName = "Gain"))
+	UPROPERTY(interp, BlueprintReadWrite, Category = "Color Grading|Global", meta = (UIMin = "0.0", UIMax = "2.0", Delta = "0.01", DisplayColorChannel = "true", ColorGradingMode = "gain", ShiftMouseMovePixelPerDelta = "10", editcondition = "bOverride_ColorGain", DisplayName = "Gain"))
 	FVector4 ColorGain;
-	UPROPERTY(interp, BlueprintReadWrite, Category = "Color Grading|Global", meta = (UIMin = "-1.0", UIMax = "1.0", Delta = "0.001", ColorGradingMode = "offset", ShiftMouseMovePixelPerDelta = "20", editcondition = "bOverride_ColorOffset", DisplayName = "Offset"))
+	UPROPERTY(interp, BlueprintReadWrite, Category = "Color Grading|Global", meta = (UIMin = "-1.0", UIMax = "1.0", Delta = "0.001", DisplayColorChannel = "true", ColorGradingMode = "offset", ShiftMouseMovePixelPerDelta = "20", editcondition = "bOverride_ColorOffset", DisplayName = "Offset"))
 	FVector4 ColorOffset;
 
-	UPROPERTY(interp, BlueprintReadWrite, Category = "Color Grading|Shadows", meta = (UIMin = "0.0", UIMax = "2.0", Delta = "0.01", ColorGradingMode = "saturation", ShiftMouseMovePixelPerDelta = "10", editcondition = "bOverride_ColorSaturationShadows", DisplayName = "Saturation"))
+	UPROPERTY(interp, BlueprintReadWrite, Category = "Color Grading|Shadows", meta = (UIMin = "0.0", UIMax = "2.0", Delta = "0.01", DisplayColorChannel = "true", ColorGradingMode = "saturation", ShiftMouseMovePixelPerDelta = "10", editcondition = "bOverride_ColorSaturationShadows", DisplayName = "Saturation"))
 	FVector4 ColorSaturationShadows;
-	UPROPERTY(interp, BlueprintReadWrite, Category = "Color Grading|Shadows", meta = (UIMin = "0.0", UIMax = "2.0", Delta = "0.01", ColorGradingMode = "contrast", ShiftMouseMovePixelPerDelta = "10", editcondition = "bOverride_ColorContrastShadows", DisplayName = "Contrast"))
+	UPROPERTY(interp, BlueprintReadWrite, Category = "Color Grading|Shadows", meta = (UIMin = "0.0", UIMax = "2.0", Delta = "0.01", DisplayColorChannel = "true", ColorGradingMode = "contrast", ShiftMouseMovePixelPerDelta = "10", editcondition = "bOverride_ColorContrastShadows", DisplayName = "Contrast"))
 	FVector4 ColorContrastShadows;
-	UPROPERTY(interp, BlueprintReadWrite, Category = "Color Grading|Shadows", meta = (UIMin = "0.0", UIMax = "2.0", Delta = "0.01", ColorGradingMode = "gamma", ShiftMouseMovePixelPerDelta = "10", editcondition = "bOverride_ColorGammaShadows", DisplayName = "Gamma"))
+	UPROPERTY(interp, BlueprintReadWrite, Category = "Color Grading|Shadows", meta = (UIMin = "0.0", UIMax = "2.0", Delta = "0.01", DisplayColorChannel = "true", ColorGradingMode = "gamma", ShiftMouseMovePixelPerDelta = "10", editcondition = "bOverride_ColorGammaShadows", DisplayName = "Gamma"))
 	FVector4 ColorGammaShadows;
-	UPROPERTY(interp, BlueprintReadWrite, Category = "Color Grading|Shadows", meta = (UIMin = "0.0", UIMax = "2.0", Delta = "0.01", ColorGradingMode = "gain", ShiftMouseMovePixelPerDelta = "10", editcondition = "bOverride_ColorGainShadows", DisplayName = "Gain"))
+	UPROPERTY(interp, BlueprintReadWrite, Category = "Color Grading|Shadows", meta = (UIMin = "0.0", UIMax = "2.0", Delta = "0.01", DisplayColorChannel = "true", ColorGradingMode = "gain", ShiftMouseMovePixelPerDelta = "10", editcondition = "bOverride_ColorGainShadows", DisplayName = "Gain"))
 	FVector4 ColorGainShadows;
-	UPROPERTY(interp, BlueprintReadWrite, Category = "Color Grading|Shadows", meta = (UIMin = "-1.0", UIMax = "1.0", Delta = "0.001", ColorGradingMode = "offset", ShiftMouseMovePixelPerDelta = "20", editcondition = "bOverride_ColorOffsetShadows", DisplayName = "Offset"))
+	UPROPERTY(interp, BlueprintReadWrite, Category = "Color Grading|Shadows", meta = (UIMin = "-1.0", UIMax = "1.0", Delta = "0.001", DisplayColorChannel = "true", ColorGradingMode = "offset", ShiftMouseMovePixelPerDelta = "20", editcondition = "bOverride_ColorOffsetShadows", DisplayName = "Offset"))
 	FVector4 ColorOffsetShadows;
 	UPROPERTY(interp, BlueprintReadWrite, Category = "Color Grading|Shadows", meta = (UIMin = "-1.0", UIMax = "1.0", editcondition = "bOverride_ColorCorrectionShadowsMax", DisplayName = "ShadowsMax"))
 	float ColorCorrectionShadowsMax;
 
-	UPROPERTY(interp, BlueprintReadWrite, Category = "Color Grading|Midtones", meta = (UIMin = "0.0", UIMax = "2.0", Delta = "0.01", ColorGradingMode = "saturation", ShiftMouseMovePixelPerDelta = "10", editcondition = "bOverride_ColorSaturationMidtones", DisplayName = "Saturation"))
+	UPROPERTY(interp, BlueprintReadWrite, Category = "Color Grading|Midtones", meta = (UIMin = "0.0", UIMax = "2.0", Delta = "0.01", DisplayColorChannel = "true", ColorGradingMode = "saturation", ShiftMouseMovePixelPerDelta = "10", editcondition = "bOverride_ColorSaturationMidtones", DisplayName = "Saturation"))
 	FVector4 ColorSaturationMidtones;
-	UPROPERTY(interp, BlueprintReadWrite, Category = "Color Grading|Midtones", meta = (UIMin = "0.0", UIMax = "2.0", Delta = "0.01", ColorGradingMode = "contrast", ShiftMouseMovePixelPerDelta = "10", editcondition = "bOverride_ColorContrastMidtones", DisplayName = "Contrast"))
+	UPROPERTY(interp, BlueprintReadWrite, Category = "Color Grading|Midtones", meta = (UIMin = "0.0", UIMax = "2.0", Delta = "0.01", DisplayColorChannel = "true", ColorGradingMode = "contrast", ShiftMouseMovePixelPerDelta = "10", editcondition = "bOverride_ColorContrastMidtones", DisplayName = "Contrast"))
 	FVector4 ColorContrastMidtones;
-	UPROPERTY(interp, BlueprintReadWrite, Category = "Color Grading|Midtones", meta = (UIMin = "0.0", UIMax = "2.0", Delta = "0.01", ColorGradingMode = "gamma", ShiftMouseMovePixelPerDelta = "10", editcondition = "bOverride_ColorGammaMidtones", DisplayName = "Gamma"))
+	UPROPERTY(interp, BlueprintReadWrite, Category = "Color Grading|Midtones", meta = (UIMin = "0.0", UIMax = "2.0", Delta = "0.01", DisplayColorChannel = "true", ColorGradingMode = "gamma", ShiftMouseMovePixelPerDelta = "10", editcondition = "bOverride_ColorGammaMidtones", DisplayName = "Gamma"))
 	FVector4 ColorGammaMidtones;
-	UPROPERTY(interp, BlueprintReadWrite, Category = "Color Grading|Midtones", meta = (UIMin = "0.0", UIMax = "2.0", Delta = "0.01", ColorGradingMode = "gain", ShiftMouseMovePixelPerDelta = "10", editcondition = "bOverride_ColorGainMidtones", DisplayName = "Gain"))
+	UPROPERTY(interp, BlueprintReadWrite, Category = "Color Grading|Midtones", meta = (UIMin = "0.0", UIMax = "2.0", Delta = "0.01", DisplayColorChannel = "true", ColorGradingMode = "gain", ShiftMouseMovePixelPerDelta = "10", editcondition = "bOverride_ColorGainMidtones", DisplayName = "Gain"))
 	FVector4 ColorGainMidtones;
-	UPROPERTY(interp, BlueprintReadWrite, Category = "Color Grading|Midtones", meta = (UIMin = "-1.0", UIMax = "1.0", Delta = "0.001", ColorGradingMode = "offset", ShiftMouseMovePixelPerDelta = "20", editcondition = "bOverride_ColorOffsetMidtones", DisplayName = "Offset"))
+	UPROPERTY(interp, BlueprintReadWrite, Category = "Color Grading|Midtones", meta = (UIMin = "-1.0", UIMax = "1.0", Delta = "0.001", DisplayColorChannel = "true", ColorGradingMode = "offset", ShiftMouseMovePixelPerDelta = "20", editcondition = "bOverride_ColorOffsetMidtones", DisplayName = "Offset"))
 	FVector4 ColorOffsetMidtones;
 
-	UPROPERTY(interp, BlueprintReadWrite, Category = "Color Grading|Highlights", meta = (UIMin = "0.0", UIMax = "2.0", Delta = "0.01", ColorGradingMode = "saturation", ShiftMouseMovePixelPerDelta = "10", editcondition = "bOverride_ColorSaturationHighlights", DisplayName = "Saturation"))
+	UPROPERTY(interp, BlueprintReadWrite, Category = "Color Grading|Highlights", meta = (UIMin = "0.0", UIMax = "2.0", Delta = "0.01", DisplayColorChannel = "true", ColorGradingMode = "saturation", ShiftMouseMovePixelPerDelta = "10", editcondition = "bOverride_ColorSaturationHighlights", DisplayName = "Saturation"))
 	FVector4 ColorSaturationHighlights;
-	UPROPERTY(interp, BlueprintReadWrite, Category = "Color Grading|Highlights", meta = (UIMin = "0.0", UIMax = "2.0", Delta = "0.01", ColorGradingMode = "contrast", ShiftMouseMovePixelPerDelta = "10", editcondition = "bOverride_ColorContrastHighlights", DisplayName = "Contrast"))
+	UPROPERTY(interp, BlueprintReadWrite, Category = "Color Grading|Highlights", meta = (UIMin = "0.0", UIMax = "2.0", Delta = "0.01", DisplayColorChannel = "true", ColorGradingMode = "contrast", ShiftMouseMovePixelPerDelta = "10", editcondition = "bOverride_ColorContrastHighlights", DisplayName = "Contrast"))
 	FVector4 ColorContrastHighlights;
-	UPROPERTY(interp, BlueprintReadWrite, Category = "Color Grading|Highlights", meta = (UIMin = "0.0", UIMax = "2.0", Delta = "0.01", ColorGradingMode = "gamma", ShiftMouseMovePixelPerDelta = "10", editcondition = "bOverride_ColorGammaHighlights", DisplayName = "Gamma"))
+	UPROPERTY(interp, BlueprintReadWrite, Category = "Color Grading|Highlights", meta = (UIMin = "0.0", UIMax = "2.0", Delta = "0.01", DisplayColorChannel = "true", ColorGradingMode = "gamma", ShiftMouseMovePixelPerDelta = "10", editcondition = "bOverride_ColorGammaHighlights", DisplayName = "Gamma"))
 	FVector4 ColorGammaHighlights;
-	UPROPERTY(interp, BlueprintReadWrite, Category = "Color Grading|Highlights", meta = (UIMin = "0.0", UIMax = "2.0", Delta = "0.01", ColorGradingMode = "gain", ShiftMouseMovePixelPerDelta = "10", editcondition = "bOverride_ColorGainHighlights", DisplayName = "Gain"))
+	UPROPERTY(interp, BlueprintReadWrite, Category = "Color Grading|Highlights", meta = (UIMin = "0.0", UIMax = "2.0", Delta = "0.01", DisplayColorChannel = "true", ColorGradingMode = "gain", ShiftMouseMovePixelPerDelta = "10", editcondition = "bOverride_ColorGainHighlights", DisplayName = "Gain"))
 	FVector4 ColorGainHighlights;
-	UPROPERTY(interp, BlueprintReadWrite, Category = "Color Grading|Highlights", meta = (UIMin = "-1.0", UIMax = "1.0", Delta = "0.001", ColorGradingMode = "offset", ShiftMouseMovePixelPerDelta = "20", editcondition = "bOverride_ColorOffsetHighlights", DisplayName = "Offset"))
+	UPROPERTY(interp, BlueprintReadWrite, Category = "Color Grading|Highlights", meta = (UIMin = "-1.0", UIMax = "1.0", Delta = "0.001", DisplayColorChannel = "true", ColorGradingMode = "offset", ShiftMouseMovePixelPerDelta = "20", editcondition = "bOverride_ColorOffsetHighlights", DisplayName = "Offset"))
 	FVector4 ColorOffsetHighlights;
 	UPROPERTY(interp, BlueprintReadWrite, Category = "Color Grading|Highlights", meta = (UIMin = "-1.0", UIMax = "1.0", editcondition = "bOverride_ColorCorrectionHighlightsMin", DisplayName = "HighlightsMin"))
 	float ColorCorrectionHighlightsMin;
diff --git a/Engine/Source/Runtime/Engine/Classes/Engine/UserDefinedStruct.h b/Engine/Source/Runtime/Engine/Classes/Engine/UserDefinedStruct.h
--- a/Engine/Source/Runtime/Engine/Classes/Engine/UserDefinedStruct.h
+++ b/Engine/Source/Runtime/Engine/Classes/Engine/UserDefinedStruct.h
@@ -60,6 +60,8 @@ public:
 	virtual void InitializeDefaultValue(uint8* StructData) const override;
 	// End of UScriptStruct interface.
 
+	bool DiffersFromDefaultValue(uint8* StructData) const;
+
 	void ValidateGuid();
 #endif	// WITH_EDITOR
 
diff --git a/Engine/Source/Runtime/Engine/Classes/Kismet/KismetStringTableLibrary.h b/Engine/Source/Runtime/Engine/Classes/Kismet/KismetStringTableLibrary.h
new file mode 100644
--- /dev/null
+++ b/Engine/Source/Runtime/Engine/Classes/Kismet/KismetStringTableLibrary.h
@@ -0,0 +1,47 @@
+﻿// Copyright 1998-2017 Epic Games, Inc. All Rights Reserved.
+
+#pragma once
+
+#include "CoreMinimal.h"
+#include "UObject/ObjectMacros.h"
+#include "UObject/TextProperty.h"
+#include "Kismet/BlueprintFunctionLibrary.h"
+#include "KismetStringTableLibrary.generated.h"
+
+UCLASS(meta=(BlueprintThreadSafe))
+class ENGINE_API UKismetStringTableLibrary : public UBlueprintFunctionLibrary
+{
+	GENERATED_UCLASS_BODY()
+
+	/** Returns true if the given table ID corresponds to a registered string table. */
+	UFUNCTION(BlueprintPure, Category="Utilities|String Table", meta=(DisplayName="Is String Table Registered"))
+	static bool IsRegisteredTableId(const FName TableId);
+
+	/** Returns true if the given table ID corresponds to a registered string table, and that table has. */
+	UFUNCTION(BlueprintPure, Category="Utilities|String Table", meta=(DisplayName="Is String Table Entry Registered"))
+	static bool IsRegisteredTableEntry(const FName TableId, const FString& Key);
+
+	/** Returns the namespace of the given string table. */
+	UFUNCTION(BlueprintPure, Category="Utilities|String Table", meta=(DisplayName="Get String Table Namespace"))
+	static FString GetTableNamespace(const FName TableId);
+
+	/** Returns the source string of the given string table entry (or an empty string). */
+	UFUNCTION(BlueprintPure, Category="Utilities|String Table", meta=(DisplayName="Get String Table Entry Source String"))
+	static FString GetTableEntrySourceString(const FName TableId, const FString& Key);
+
+	/** Returns the specified meta-data of the given string table entry (or an empty string). */
+	UFUNCTION(BlueprintPure, Category="Utilities|String Table", meta=(DisplayName="Get String Table Entry Meta-Data"))
+	static FString GetTableEntryMetaData(const FName TableId, const FString& Key, const FName MetaDataId);
+
+	/** Returns an array of all registered string table IDs */
+	UFUNCTION(BlueprintPure, Category="Utilities|String Table", meta=(DisplayName="Get Registered String Tables"))
+	static TArray<FName> GetRegisteredStringTables();
+
+	/** Returns an array of all keys within the given string table */
+	UFUNCTION(BlueprintPure, Category="Utilities|String Table", meta=(DisplayName="Get Keys from String Table"))
+	static TArray<FString> GetKeysFromStringTable(const FName TableId);
+
+	/** Returns an array of all meta-data IDs within the given string table entry */
+	UFUNCTION(BlueprintPure, Category="Utilities|String Table", meta=(DisplayName="Get Meta-Data IDs from String Table Entry"))
+	static TArray<FName> GetMetaDataIdsFromStringTableEntry(const FName TableId, const FString& Key);
+};
diff --git a/Engine/Source/Runtime/Engine/Classes/Kismet/KismetTextLibrary.h b/Engine/Source/Runtime/Engine/Classes/Kismet/KismetTextLibrary.h
--- a/Engine/Source/Runtime/Engine/Classes/Kismet/KismetTextLibrary.h
+++ b/Engine/Source/Runtime/Engine/Classes/Kismet/KismetTextLibrary.h
@@ -1,4 +1,4 @@
-// Copyright 1998-2017 Epic Games, Inc. All Rights Reserved.
+﻿// Copyright 1998-2017 Epic Games, Inc. All Rights Reserved.
 
 #pragma once
 
@@ -256,4 +256,23 @@ class ENGINE_API UKismetTextLibrary : public UBlueprintFunctionLibrary
 	/* Used for formatting text using the FText::Format function and utilized by the UK2Node_FormatText */
 	UFUNCTION(BlueprintPure, meta=(BlueprintInternalUseOnly = "true"))
 	static FText Format(FText InPattern, TArray<FFormatArgumentData> InArgs);
+
+	/** Returns true if the given text is referencing a string table. */
+	UFUNCTION(BlueprintPure, Category="Utilities|Text", meta=(DisplayName="Is Text from String Table"))
+	static bool TextIsFromStringTable(const FText& Text);
+
+	/**
+	 * Attempts to create a text instance from a string table ID and key.
+	 * @note This exists to allow programmatic ‎look-up of a string table entry from dynamic content - you should favor setting your string table reference on a text property or pin wherever possible as it is significantly more robust (see "Make Literal Text").
+	 * @return The found text, or a dummy text if the entry could not be found.
+	 */
+	UFUNCTION(BlueprintPure, Category="Utilities|Text", meta=(DisplayName="Make Text from String Table (Advanced)"))
+	static FText TextFromStringTable(const FName TableId, const FString& Key);
+
+	/**
+	 * Attempts to find the String Table ID and key used by the given text.
+	 * @return True if the String Table ID and key were found, false otherwise.
+	 */
+	UFUNCTION(BlueprintPure, Category="Utilities|Text", meta=(DisplayName="Find String Table ID and Key from Text"))
+	static bool StringTableIdAndKeyFromText(FText Text, FName& OutTableId, FString& OutKey);
 };
diff --git a/Engine/Source/Runtime/Engine/Private/AI/Navigation/NavigationSystem.cpp b/Engine/Source/Runtime/Engine/Private/AI/Navigation/NavigationSystem.cpp
--- a/Engine/Source/Runtime/Engine/Private/AI/Navigation/NavigationSystem.cpp
+++ b/Engine/Source/Runtime/Engine/Private/AI/Navigation/NavigationSystem.cpp
@@ -3585,8 +3585,10 @@ void UNavigationSystem::CleanUp(ECleanupMode Mode)
 	FWorldDelegates::LevelRemovedFromWorld.RemoveAll(this);
 
 #if WITH_HOT_RELOAD
-	IHotReloadInterface& HotReloadSupport = FModuleManager::LoadModuleChecked<IHotReloadInterface>("HotReload");
-	HotReloadSupport.OnHotReload().Remove(HotReloadDelegateHandle);
+	if (IHotReloadInterface* HotReloadSupport = FModuleManager::GetModulePtr<IHotReloadInterface>("HotReload"))
+	{
+		HotReloadSupport->OnHotReload().Remove(HotReloadDelegateHandle);
+	}
 #endif
 
 	DestroyNavOctree();
diff --git a/Engine/Source/Runtime/Engine/Private/Animation/SkeletalMeshActor.cpp b/Engine/Source/Runtime/Engine/Private/Animation/SkeletalMeshActor.cpp
--- a/Engine/Source/Runtime/Engine/Private/Animation/SkeletalMeshActor.cpp
+++ b/Engine/Source/Runtime/Engine/Private/Animation/SkeletalMeshActor.cpp
@@ -164,21 +164,6 @@ FString ASkeletalMeshActor::GetDetailedInfoInternal() const
 	return Result;
 }
 
-#if WITH_EDITOR
-void ASkeletalMeshActor::PostEditChangeChainProperty(FPropertyChangedChainEvent& PropertyChangedEvent)
-{
-	Super::PostEditChangeChainProperty(PropertyChangedEvent);
-
-	if (PropertyChangedEvent.Property != nullptr)
-	{
-		if (PropertyChangedEvent.Property->GetFName() == GET_MEMBER_NAME_STRING_CHECKED(ASkeletalMeshActor, SkeletalMeshComponent) && SkeletalMeshComponent->SkeletalMesh != nullptr)
-		{
-			SkeletalMeshComponent->CleanUpOverrideMaterials();
-		}
-	}
-}
-#endif
-
 void ASkeletalMeshActor::PostInitializeComponents()
 {
 	Super::PostInitializeComponents();
diff --git a/Engine/Source/Runtime/Engine/Private/BlueprintGeneratedClass.cpp b/Engine/Source/Runtime/Engine/Private/BlueprintGeneratedClass.cpp
--- a/Engine/Source/Runtime/Engine/Private/BlueprintGeneratedClass.cpp
+++ b/Engine/Source/Runtime/Engine/Private/BlueprintGeneratedClass.cpp
@@ -229,7 +229,7 @@ void UBlueprintGeneratedClass::ConditionalRecompileClass(TArray<UObject*>* ObjLo
 			GeneratingBP->bIsRegeneratingOnLoad = true;
 
 			{
-				UPackage* const Package = Cast<UPackage>(GeneratingBP->GetOutermost());
+				UPackage* const Package = GeneratingBP->GetOutermost();
 				const bool bStartedWithUnsavedChanges = Package != nullptr ? Package->IsDirty() : true;
 
 				// Make sure that nodes are up to date, so that we get any updated blueprint signatures
diff --git a/Engine/Source/Runtime/Engine/Private/Components/MeshComponent.cpp b/Engine/Source/Runtime/Engine/Private/Components/MeshComponent.cpp
--- a/Engine/Source/Runtime/Engine/Private/Components/MeshComponent.cpp
+++ b/Engine/Source/Runtime/Engine/Private/Components/MeshComponent.cpp
@@ -109,6 +109,19 @@ int32 UMeshComponent::GetNumOverrideMaterials() const
 }
 
 #if WITH_EDITOR
+void UMeshComponent::PostEditChangeChainProperty(FPropertyChangedChainEvent& PropertyChangedEvent)
+{
+	Super::PostEditChangeChainProperty(PropertyChangedEvent);
+
+	if (PropertyChangedEvent.Property != nullptr)
+	{
+		if (PropertyChangedEvent.Property->GetFName() == GET_MEMBER_NAME_STRING_CHECKED(UMeshComponent, OverrideMaterials))
+		{
+			CleanUpOverrideMaterials();
+		}
+	}
+}
+
 void UMeshComponent::CleanUpOverrideMaterials()
 {
 	//We have to remove material override Ids that are bigger then the material list
diff --git a/Engine/Source/Runtime/Engine/Private/EngineSessionManager.cpp b/Engine/Source/Runtime/Engine/Private/EngineSessionManager.cpp
--- a/Engine/Source/Runtime/Engine/Private/EngineSessionManager.cpp
+++ b/Engine/Source/Runtime/Engine/Private/EngineSessionManager.cpp
@@ -681,7 +681,20 @@ void FEngineSessionManager::StartWatchdog(const FString& RunType, const FString&
 	const int SuccessfulRtnCode = 0;	// hardcode this for now, zero might not always be correct
 
 	FString WatchdogClientArguments =
-		FString::Printf(TEXT("-PID=%u -RunType=%s -ProjectName=\"%s\" -Platform=%s -SessionId=%s -EngineVersion=%s -SuccessfulRtnCode=%d -HeartbeatSeconds=%d"), ProcessId, *RunType, *ProjectName, *PlatformName, *SessionId, *EngineVersion, SuccessfulRtnCode, SessionManagerDefs::HeartbeatPeriodSeconds);
+		FString::Printf(TEXT(
+			"-PID=%u -RunType=%s -ProjectName=\"%s\" -Platform=%s -SessionId=%s -EngineVersion=%s -SuccessfulRtnCode=%d"),
+			ProcessId, *RunType, *ProjectName, *PlatformName, *SessionId, *EngineVersion, SuccessfulRtnCode);
+
+	bool bAllowWatchdogDetectHangs = false;
+	GConfig->GetBool(TEXT("EngineSessionManager"), TEXT("AllowWatchdogDetectHangs"), bAllowWatchdogDetectHangs, GEngineIni);
+
+	if (bAllowWatchdogDetectHangs)
+	{
+		int HangSeconds = 120;
+		GConfig->GetInt(TEXT("EngineSessionManager"), TEXT("WatchdogMinimumHangSeconds"), HangSeconds, GEngineIni);
+
+		WatchdogClientArguments.Append(FString::Printf(TEXT(" -DetectHangs -HangSeconds=%d"), HangSeconds));
+	}
 
 	if (FEngineBuildSettings::IsInternalBuild())
 	{
diff --git a/Engine/Source/Runtime/Engine/Private/GameEngine.cpp b/Engine/Source/Runtime/Engine/Private/GameEngine.cpp
--- a/Engine/Source/Runtime/Engine/Private/GameEngine.cpp
+++ b/Engine/Source/Runtime/Engine/Private/GameEngine.cpp
@@ -354,7 +354,7 @@ TSharedRef<SWindow> UGameEngine::CreateGameWindow()
 	const bool bAllowMinimize = GetDefault<UGeneralProjectSettings>()->bAllowMinimize;
 
 	// Allow optional winX/winY parameters to set initial window position
-	EAutoCenter::Type AutoCenterType = EAutoCenter::PrimaryWorkArea;
+	EAutoCenter AutoCenterType = EAutoCenter::PrimaryWorkArea;
 	int32 WinX=0;
 	int32 WinY=0;
 	if (FParse::Value(FCommandLine::Get(), TEXT("WinX="), WinX) && FParse::Value(FCommandLine::Get(), TEXT("WinY="), WinY))
diff --git a/Engine/Source/Runtime/Engine/Private/GameViewportClient.cpp b/Engine/Source/Runtime/Engine/Private/GameViewportClient.cpp
--- a/Engine/Source/Runtime/Engine/Private/GameViewportClient.cpp
+++ b/Engine/Source/Runtime/Engine/Private/GameViewportClient.cpp
@@ -1529,7 +1529,7 @@ void UGameViewportClient::CloseRequested(FViewport* InViewport)
 	FSlateApplication::Get().SetGameIsFakingTouchEvents(false);
 	
 	// broadcast close request to anyone that registered an interest
-	CloseRequestedDelegate.ExecuteIfBound(InViewport);
+	CloseRequestedDelegate.Broadcast(InViewport);
 
 	SetViewportFrame(NULL);
 
diff --git a/Engine/Source/Runtime/Engine/Private/Internationalization/StringTable.cpp b/Engine/Source/Runtime/Engine/Private/Internationalization/StringTable.cpp
new file mode 100644
--- /dev/null
+++ b/Engine/Source/Runtime/Engine/Private/Internationalization/StringTable.cpp
@@ -0,0 +1,231 @@
+// Copyright 1998-2017 Epic Games, Inc. All Rights Reserved.
+
+#include "StringTable.h"
+#include "StringTableCore.h"
+#include "StringTableRegistry.h"
+#include "StringAssetReference.h"
+#include "PackageName.h"
+#include "GCObject.h"
+#include "Misc/ScopeLock.h"
+#include "Templates/Casts.h"
+#include "SlateApplicationBase.h"
+#include "Serialization/PropertyLocalizationDataGathering.h"
+
+#if WITH_EDITORONLY_DATA
+namespace
+{
+	static void GatherStringTableForLocalization(const UObject* const Object, FPropertyLocalizationDataGatherer& PropertyLocalizationDataGatherer, const EPropertyLocalizationGathererTextFlags GatherTextFlags)
+	{
+		FStringTableConstRef StringTable = CastChecked<UStringTable>(Object)->GetStringTable();
+
+		auto FindOrAddTextData = [&](const FString& InText) -> FGatherableTextData&
+		{
+			check(!InText.IsEmpty());
+
+			FTextSourceData SourceData;
+			SourceData.SourceString = InText;
+
+			auto& GatherableTextDataArray = PropertyLocalizationDataGatherer.GetGatherableTextDataArray();
+			FGatherableTextData* GatherableTextData = GatherableTextDataArray.FindByPredicate([&](const FGatherableTextData& Candidate)
+			{
+				return Candidate.NamespaceName.Equals(StringTable->GetNamespace(), ESearchCase::CaseSensitive)
+					&& Candidate.SourceData.SourceString.Equals(SourceData.SourceString, ESearchCase::CaseSensitive)
+					&& Candidate.SourceData.SourceStringMetaData == SourceData.SourceStringMetaData;
+			});
+			if (!GatherableTextData)
+			{
+				GatherableTextData = &GatherableTextDataArray[GatherableTextDataArray.AddDefaulted()];
+				GatherableTextData->NamespaceName = StringTable->GetNamespace();
+				GatherableTextData->SourceData = SourceData;
+			}
+
+			return *GatherableTextData;
+		};
+
+		const FString SourceLocation = Object->GetPathName();
+
+		StringTable->EnumerateSourceStrings([&](const FString& InKey, const FString& InSourceString) -> bool
+		{
+			if (!InSourceString.IsEmpty())
+			{
+				FGatherableTextData& GatherableTextData = FindOrAddTextData(InSourceString);
+
+				FTextSourceSiteContext& SourceSiteContext = GatherableTextData.SourceSiteContexts[GatherableTextData.SourceSiteContexts.AddDefaulted()];
+				SourceSiteContext.KeyName = InKey;
+				SourceSiteContext.SiteDescription = SourceLocation;
+				SourceSiteContext.IsEditorOnly = false;
+				SourceSiteContext.IsOptional = false;
+
+				StringTable->EnumerateMetaData(InKey, [&](const FName InMetaDataId, const FString& InMetaData)
+				{
+					SourceSiteContext.InfoMetaData.SetStringField(InMetaDataId.ToString(), InMetaData);
+					return true; // continue enumeration
+				});
+			}
+
+			return true; // continue enumeration
+		});
+	}
+}
+#endif
+
+class FStringTableEngineBridge : public IStringTableEngineBridge, public FGCObject
+{
+public:
+	static void Initialize()
+	{
+		static FStringTableEngineBridge Instance;
+		InstancePtr = &Instance;
+	}
+
+private:
+	//~ IStringTableEngineInterop interface
+	virtual void RedirectAndLoadStringTableAssetImpl(FName& InOutTableId, const EStringTableLoadingPolicy InLoadingPolicy) override
+	{
+		FStringAssetReference StringTableAssetReference = GetAssetReference(InOutTableId);
+		if (StringTableAssetReference.IsValid())
+		{
+			UStringTable* StringTableAsset = Cast<UStringTable>(StringTableAssetReference.ResolveObject());
+			if ((!StringTableAsset && InLoadingPolicy != EStringTableLoadingPolicy::Find) || (StringTableAsset && StringTableAsset->HasAnyFlags(RF_NeedLoad | RF_NeedPostLoad) && InLoadingPolicy == EStringTableLoadingPolicy::FindOrFullyLoad))
+			{
+				StringTableAsset = Cast<UStringTable>(StringTableAssetReference.TryLoad());
+			}
+			if (StringTableAsset)
+			{
+				InOutTableId = StringTableAsset->GetStringTableId();
+
+				// Prevent the string table from being GC'd
+				{
+					FScopeLock KeepAliveStringTablesLock(&KeepAliveStringTablesCS);
+					KeepAliveStringTables.AddUnique(StringTableAsset);
+				}
+			}
+		}
+	}
+
+	virtual void CollectStringTableAssetReferencesImpl(const FName InTableId, FArchive& InAr) override
+	{
+		check(InAr.IsObjectReferenceCollector());
+
+		UStringTable* StringTableAsset = FStringTableRegistry::Get().FindStringTableAsset(InTableId);
+		InAr << StringTableAsset;
+	}
+
+	static FStringAssetReference GetAssetReference(const FName InTableId)
+	{
+		const FString StringTableAssetName = InTableId.ToString();
+
+		FString StringTablePackageName = StringTableAssetName;
+		{
+			int32 DotIndex = INDEX_NONE;
+			if (StringTablePackageName.FindChar(TEXT('.'), DotIndex))
+			{
+				StringTablePackageName = StringTablePackageName.Left(DotIndex);
+			}
+		}
+
+		FStringAssetReference StringTableAssetReference;
+		if (FPackageName::IsValidLongPackageName(StringTablePackageName, /*bIncludeReadOnlyRoots*/true) && FPackageName::DoesPackageExist(StringTablePackageName))
+		{
+			StringTableAssetReference.SetPath(StringTableAssetName);
+		}
+		
+		return StringTableAssetReference;
+	}
+
+	//~ FGCObject interface
+	virtual void AddReferencedObjects(FReferenceCollector& Collector) override
+	{
+		FScopeLock KeepAliveStringTablesLock(&KeepAliveStringTablesCS);
+		Collector.AddReferencedObjects(KeepAliveStringTables);
+	}
+
+private:
+	/** Array of string table assets that have been loaded and should be kept alive */
+	TArray<UStringTable*> KeepAliveStringTables;
+
+	/** Critical section preventing concurrent modification of KeepAliveStringTables */
+	mutable FCriticalSection KeepAliveStringTablesCS;
+};
+
+UStringTable::UStringTable()
+	: StringTable(FStringTable::NewStringTable())
+	, StringTableId(*GetPathName())
+{
+	StringTable->SetOwnerAsset(this);
+	StringTable->IsLoaded(!HasAnyFlags(RF_NeedLoad | RF_NeedPostLoad));
+
+	StringTable->SetNamespace(GetName());
+
+	if (!HasAnyFlags(RF_ClassDefaultObject))
+	{
+		FStringTableRegistry::Get().RegisterStringTable(GetStringTableId(), StringTable.ToSharedRef());
+	}
+
+#if WITH_EDITORONLY_DATA
+	{ static const FAutoRegisterLocalizationDataGatheringCallback AutomaticRegistrationOfLocalizationGatherer(UStringTable::StaticClass(), &GatherStringTableForLocalization); }
+#endif
+}
+
+void UStringTable::InitializeEngineBridge()
+{
+	FStringTableEngineBridge::Initialize();
+}
+
+void UStringTable::FinishDestroy()
+{
+	if (!HasAnyFlags(RF_ClassDefaultObject))
+	{
+		FStringTableRegistry::Get().UnregisterStringTable(GetStringTableId());
+	}
+	StringTable.Reset();
+
+	Super::FinishDestroy();
+}
+
+void UStringTable::Serialize(FArchive& Ar)
+{
+	Super::Serialize(Ar);
+
+	StringTable->Serialize(Ar);
+}
+
+void UStringTable::PostLoad()
+{
+	Super::PostLoad();
+
+	StringTable->IsLoaded(true);
+
+	if (FSlateApplicationBase::IsInitialized())
+	{
+		// Ensure all invalidation panels are updated now that the string data is loaded
+		FSlateApplicationBase::Get().InvalidateAllWidgets();
+	}
+}
+
+bool UStringTable::Rename(const TCHAR* NewName, UObject* NewOuter, ERenameFlags Flags)
+{
+	const bool bRenamed = Super::Rename(NewName, NewOuter, Flags);
+	if (bRenamed && !HasAnyFlags(RF_ClassDefaultObject))
+	{
+		FStringTableRegistry::Get().UnregisterStringTable(GetStringTableId());
+		StringTableId = *GetPathName();
+		FStringTableRegistry::Get().RegisterStringTable(GetStringTableId(), StringTable.ToSharedRef());
+	}
+	return bRenamed;
+}
+
+FName UStringTable::GetStringTableId() const
+{
+	return StringTableId;
+}
+
+FStringTableConstRef UStringTable::GetStringTable() const
+{
+	return StringTable.ToSharedRef();
+}
+
+FStringTableRef UStringTable::GetMutableStringTable() const
+{
+	return StringTable.ToSharedRef();
+}
diff --git a/Engine/Source/Runtime/Engine/Private/KismetTextLibrary.cpp b/Engine/Source/Runtime/Engine/Private/KismetTextLibrary.cpp
--- a/Engine/Source/Runtime/Engine/Private/KismetTextLibrary.cpp
+++ b/Engine/Source/Runtime/Engine/Private/KismetTextLibrary.cpp
@@ -2,6 +2,7 @@
 
 #include "Kismet/KismetTextLibrary.h"
 #include "Internationalization/TextFormatter.h"
+#include "Internationalization/StringTableRegistry.h"
 
 #define LOCTEXT_NAMESPACE "Kismet"
 
@@ -306,4 +307,21 @@ FText UKismetTextLibrary::Format(FText InPattern, TArray<FFormatArgumentData> In
 }
 
 
+bool UKismetTextLibrary::TextIsFromStringTable(const FText& Text)
+{
+	return Text.IsFromStringTable();
+}
+
+
+FText UKismetTextLibrary::TextFromStringTable(const FName TableId, const FString& Key)
+{
+	return FText::FromStringTable(TableId, Key);
+}
+
+
+bool UKismetTextLibrary::StringTableIdAndKeyFromText(FText Text, FName& OutTableId, FString& OutKey)
+{
+	return FStringTableRegistry::Get().FindTableIdAndKey(Text, OutTableId, OutKey);
+}
+
 #undef LOCTEXT_NAMESPACE
diff --git a/Engine/Source/Runtime/Engine/Private/LevelUtils.cpp b/Engine/Source/Runtime/Engine/Private/LevelUtils.cpp
--- a/Engine/Source/Runtime/Engine/Private/LevelUtils.cpp
+++ b/Engine/Source/Runtime/Engine/Private/LevelUtils.cpp
@@ -133,7 +133,7 @@ bool FLevelUtils::IsLevelLocked(ULevel* Level)
 				LevelData.LastUpdateTime = Level->OwningWorld->GetRealTimeSeconds();
 				//If we dont find package we dont consider it as readonly
 				LevelData.IsReadOnly = false;
-				const UPackage* pPackage = Cast<UPackage>(Level->GetOutermost());
+				const UPackage* pPackage = Level->GetOutermost();
 				if (pPackage)
 				{
 					FString PackageFileName;
diff --git a/Engine/Source/Runtime/Engine/Private/ShaderCompiler/ShaderCompiler.cpp b/Engine/Source/Runtime/Engine/Private/ShaderCompiler/ShaderCompiler.cpp
--- a/Engine/Source/Runtime/Engine/Private/ShaderCompiler/ShaderCompiler.cpp
+++ b/Engine/Source/Runtime/Engine/Private/ShaderCompiler/ShaderCompiler.cpp
@@ -303,7 +303,8 @@ static const TArray<const IShaderFormat*>& GetShaderFormats()
 
 		for (int32 Index = 0; Index < Modules.Num(); Index++)
 		{
-			IShaderFormat* Format = FModuleManager::LoadModuleChecked<IShaderFormatModule>(Modules[Index]).GetShaderFormat();
+			IShaderFormat* Format = FModuleManager::GetModuleChecked<IShaderFormatModule>(Modules[Index]).GetShaderFormat();
+
 			if (Format != nullptr)
 			{
 				Results.Add(Format);
@@ -385,11 +386,11 @@ static bool DoWriteTasks(const TArray<FShaderCommonCompileJob*>& QueuedJobs, FAr
 		TransferFile << SingleJobHeader;
 
 		int32 NumBatches = QueuedSingleJobs.Num();
-	TransferFile << NumBatches;
+		TransferFile << NumBatches;
 
-	// Serialize all the batched jobs
+		// Serialize all the batched jobs
 		for (int32 JobIndex = 0; JobIndex < QueuedSingleJobs.Num(); JobIndex++)
-	{
+		{
 			TransferFile << QueuedSingleJobs[JobIndex]->Input;
 		}
 	}
diff --git a/Engine/Source/Runtime/Engine/Private/Slate/SceneViewport.cpp b/Engine/Source/Runtime/Engine/Private/Slate/SceneViewport.cpp
--- a/Engine/Source/Runtime/Engine/Private/Slate/SceneViewport.cpp
+++ b/Engine/Source/Runtime/Engine/Private/Slate/SceneViewport.cpp
@@ -45,6 +45,7 @@ FSceneViewport::FSceneViewport( FViewportClient* InViewportClient, TSharedPtr<SV
 	, NumBufferedFrames(1)
 	, CurrentBufferedTargetIndex(0)
 	, NextBufferedTargetIndex(0)
+	, NumTouches(0)
 {
 	bIsSlateViewport = true;
 	RenderThreadSlateTexture = new FSlateRenderTargetRHI(nullptr, 0, 0);
@@ -717,10 +718,11 @@ FReply FSceneViewport::OnTouchStarted( const FGeometry& MyGeometry, const FPoint
 {
 	// Start a new reply state
 	CurrentReplyState = FReply::Handled().PreventThrottling(); 
+	++NumTouches;
 
 	UpdateCachedMousePos(MyGeometry, TouchEvent);
 	UpdateCachedGeometry(MyGeometry);
-
+	
 	if( ViewportClient )
 	{
 		// Switch to the viewport clients world before processing input
@@ -766,7 +768,14 @@ FReply FSceneViewport::OnTouchEnded( const FGeometry& MyGeometry, const FPointer
 	// Start a new reply state
 	CurrentReplyState = FReply::Handled(); 
 
-	UpdateCachedMousePos(MyGeometry, TouchEvent);
+	if (--NumTouches > 0)
+	{
+		UpdateCachedMousePos(MyGeometry, TouchEvent);
+	}
+	else
+	{
+		CachedMousePos = FIntPoint(-1, -1);
+	}
 	UpdateCachedGeometry(MyGeometry);
 
 	if( ViewportClient )
@@ -1626,7 +1635,7 @@ void FSceneViewport::OnPostResizeWindowBackbuffer(void* Backbuffer)
 {
 	check(IsInGameThread());
 
-	if(!UseSeparateRenderTarget() && !IsValidRef(ViewportRHI))
+	if(!UseSeparateRenderTarget() && !IsValidRef(ViewportRHI) && ViewportWidget.IsValid())
 	{
 		TSharedPtr<FSlateRenderer> Renderer = FSlateApplication::Get().GetRenderer();
 		FWidgetPath WidgetPath;
diff --git a/Engine/Source/Runtime/Engine/Private/StaticMesh.cpp b/Engine/Source/Runtime/Engine/Private/StaticMesh.cpp
--- a/Engine/Source/Runtime/Engine/Private/StaticMesh.cpp
+++ b/Engine/Source/Runtime/Engine/Private/StaticMesh.cpp
@@ -2018,8 +2018,10 @@ FArchive& operator<<(FArchive& Ar, FMeshSectionInfo& Info)
 void FMeshSectionInfoMap::Serialize(FArchive& Ar)
 {
 	Ar.UsingCustomVersion(FReleaseObjectVersion::GUID);
+	Ar.UsingCustomVersion(FEditorObjectVersion::GUID);
 
-	if (Ar.CustomVer(FReleaseObjectVersion::GUID) < FReleaseObjectVersion::UPropertryForMeshSectionSerialize)
+	if ( Ar.CustomVer(FReleaseObjectVersion::GUID) < FReleaseObjectVersion::UPropertryForMeshSectionSerialize // Release-4.15 change
+		&& Ar.CustomVer(FEditorObjectVersion::GUID) < FEditorObjectVersion::UPropertryForMeshSectionSerialize) // Dev-Editor change
 	{
 		Ar << Map;
 	}
diff --git a/Engine/Source/Runtime/Engine/Private/StaticMeshActor.cpp b/Engine/Source/Runtime/Engine/Private/StaticMeshActor.cpp
--- a/Engine/Source/Runtime/Engine/Private/StaticMeshActor.cpp
+++ b/Engine/Source/Runtime/Engine/Private/StaticMeshActor.cpp
@@ -108,10 +108,6 @@ void AStaticMeshActor::PostEditChangeChainProperty(FPropertyChangedChainEvent& P
 			bReplicateMovement = bStaticMeshReplicateMovement;
 			SetReplicates(bReplicateMovement);
 		}
-		else if (PropertyChangedEvent.Property->GetFName() == GET_MEMBER_NAME_STRING_CHECKED(AStaticMeshActor, StaticMeshComponent) && StaticMeshComponent->GetStaticMesh() != nullptr)
-		{
-			StaticMeshComponent->CleanUpOverrideMaterials();
-		}
 	}
 
 	// Get 'deepest' property name we changed.
diff --git a/Engine/Source/Runtime/Engine/Private/Tests/AutomationCommon.cpp b/Engine/Source/Runtime/Engine/Private/Tests/AutomationCommon.cpp
--- a/Engine/Source/Runtime/Engine/Private/Tests/AutomationCommon.cpp
+++ b/Engine/Source/Runtime/Engine/Private/Tests/AutomationCommon.cpp
@@ -2,6 +2,7 @@
 
 #include "Tests/AutomationCommon.h"
 #include "Misc/Paths.h"
+#include "Misc/EngineVersion.h"
 #include "EngineGlobals.h"
 #include "Widgets/SWidget.h"
 #include "Engine/GameViewportClient.h"
@@ -98,6 +99,7 @@ namespace AutomationCommon
 		Data.Name = TestName;
 		Data.Context = MapOrContext;
 		Data.Id = FGuid::NewGuid();
+		Data.Commit = FEngineVersion::Current().HasChangelist() ? FString::FromInt(FEngineVersion::Current().GetChangelist()) : FString(TEXT(""));
 
 		Data.Width = Width;
 		Data.Height = Height;
diff --git a/Engine/Source/Runtime/Engine/Private/UnrealEngine.cpp b/Engine/Source/Runtime/Engine/Private/UnrealEngine.cpp
--- a/Engine/Source/Runtime/Engine/Private/UnrealEngine.cpp
+++ b/Engine/Source/Runtime/Engine/Private/UnrealEngine.cpp
@@ -11567,7 +11567,7 @@ void UEngine::CopyPropertiesForUnrelatedObjects(UObject* OldObject, UObject* New
 	});
 
 	// Replace references to old classes and instances on this object with the corresponding new ones
-	UPackage* NewPackage = Cast<UPackage>(NewObject->GetOutermost());
+	UPackage* NewPackage = NewObject->GetOutermost();
 	FArchiveReplaceOrClearExternalReferences<UObject> ReplaceInCDOAr(NewObject, ReferenceReplacementMap, NewPackage);
 
 	// Replace references inside each individual component. This is always required because if something is in ReferenceReplacementMap, the above replace code will skip fixing child properties
diff --git a/Engine/Source/Runtime/Engine/Private/UserDefinedStruct.cpp b/Engine/Source/Runtime/Engine/Private/UserDefinedStruct.cpp
--- a/Engine/Source/Runtime/Engine/Private/UserDefinedStruct.cpp
+++ b/Engine/Source/Runtime/Engine/Private/UserDefinedStruct.cpp
@@ -138,6 +138,11 @@ void UUserDefinedStruct::InitializeDefaultValue(uint8* StructData) const
 	FStructureEditorUtils::Fill_MakeStructureDefaultValue(this, StructData);
 }
 
+bool UUserDefinedStruct::DiffersFromDefaultValue(uint8* StructData) const
+{
+	return FStructureEditorUtils::DiffersFromDefaultValue(this, StructData);
+}
+
 void UUserDefinedStruct::ValidateGuid()
 {
 	// Backward compatibility:
diff --git a/Engine/Source/Runtime/Engine/Private/World.cpp b/Engine/Source/Runtime/Engine/Private/World.cpp
--- a/Engine/Source/Runtime/Engine/Private/World.cpp
+++ b/Engine/Source/Runtime/Engine/Private/World.cpp
@@ -2444,7 +2444,7 @@ void FLevelStreamingGCHelper::PrepareStreamedOutLevelsForGC()
 
 		if( Level && (!GIsEditor || Level->GetOutermost()->HasAnyPackageFlags(PKG_PlayInEditor) ))
 		{
-			UPackage* LevelPackage = Cast<UPackage>(Level->GetOutermost());
+			UPackage* LevelPackage = Level->GetOutermost();
 			UE_LOG(LogStreaming, Log, TEXT("PrepareStreamedOutLevelsForGC called on '%s'"), *LevelPackage->GetName() );
 			
 			for (const FWorldContext& Context : GEngine->GetWorldContexts())
diff --git a/Engine/Source/Runtime/Engine/Private/WorldComposition.cpp b/Engine/Source/Runtime/Engine/Private/WorldComposition.cpp
--- a/Engine/Source/Runtime/Engine/Private/WorldComposition.cpp
+++ b/Engine/Source/Runtime/Engine/Private/WorldComposition.cpp
@@ -872,7 +872,7 @@ void UWorldComposition::OnLevelRemovedFromWorld(ULevel* InLevel)
 
 void UWorldComposition::OnLevelPostLoad(ULevel* InLevel)
 {
-	UPackage* LevelPackage = Cast<UPackage>(InLevel->GetOutermost());	
+	UPackage* LevelPackage = InLevel->GetOutermost();	
 	if (LevelPackage && InLevel->OwningWorld)
 	{
 		FWorldTileInfo Info;
@@ -923,7 +923,7 @@ void UWorldComposition::OnLevelPostSave(ULevel* InLevel)
 FIntVector UWorldComposition::GetLevelOffset(ULevel* InLevel) const
 {
 	UWorld* OwningWorld = GetWorld();
-	UPackage* LevelPackage = Cast<UPackage>(InLevel->GetOutermost());
+	UPackage* LevelPackage = InLevel->GetOutermost();
 	
 	FIntVector LevelPosition = FIntVector::ZeroValue;
 	if (LevelPackage->WorldTileInfo)
@@ -938,7 +938,7 @@ FIntVector UWorldComposition::GetLevelOffset(ULevel* InLevel) const
 FBox UWorldComposition::GetLevelBounds(ULevel* InLevel) const
 {
 	UWorld* OwningWorld = GetWorld();
-	UPackage* LevelPackage = Cast<UPackage>(InLevel->GetOutermost());
+	UPackage* LevelPackage = InLevel->GetOutermost();
 	
 	FBox LevelBBox(0);
 	if (LevelPackage->WorldTileInfo)
diff --git a/Engine/Source/Runtime/Engine/Public/Internationalization/StringTable.h b/Engine/Source/Runtime/Engine/Public/Internationalization/StringTable.h
new file mode 100644
--- /dev/null
+++ b/Engine/Source/Runtime/Engine/Public/Internationalization/StringTable.h
@@ -0,0 +1,43 @@
+// Copyright 1998-2017 Epic Games, Inc. All Rights Reserved.
+#pragma once
+
+#include "CoreTypes.h"
+#include "StringTableCoreFwd.h"
+#include "UObject/ObjectMacros.h"
+#include "UObject/Object.h"
+#include "StringTable.generated.h"
+
+/** String table wrapper asset */
+UCLASS()
+class ENGINE_API UStringTable : public UObject
+{
+	GENERATED_BODY()
+
+public:
+	UStringTable();
+
+	/** Called during Engine init to initialize the engine bridge instance */
+	static void InitializeEngineBridge();
+
+	//~ UObject interface
+	virtual void FinishDestroy() override;
+	virtual void Serialize(FArchive& Ar) override;
+	virtual void PostLoad() override;
+	virtual bool Rename(const TCHAR* NewName = nullptr, UObject* NewOuter = nullptr, ERenameFlags Flags = REN_None) override;
+
+	/** Get the string table ID that should be used by this asset */
+	FName GetStringTableId() const;
+
+	/** Get the underlying string table owned by this asset */
+	FStringTableConstRef GetStringTable() const;
+
+	/** Get the underlying string table owned by this asset */
+	FStringTableRef GetMutableStringTable() const;
+
+private:
+	/** Internal string table. Will never be null, but is a TSharedPtr rather than a TSharedRef so we can use a forward declaration */
+	FStringTablePtr StringTable;
+
+	/** Cached ID of this string table asset in the string table registry */
+	FName StringTableId;
+};
diff --git a/Engine/Source/Runtime/Engine/Public/Slate/SceneViewport.h b/Engine/Source/Runtime/Engine/Public/Slate/SceneViewport.h
--- a/Engine/Source/Runtime/Engine/Public/Slate/SceneViewport.h
+++ b/Engine/Source/Runtime/Engine/Public/Slate/SceneViewport.h
@@ -416,4 +416,7 @@ private:
 	int32 NumBufferedFrames;
 	int32 CurrentBufferedTargetIndex;
 	int32 NextBufferedTargetIndex;
+
+	/** Tracks the number of touches currently active on the viewport */
+	int32 NumTouches;
 };
diff --git a/Engine/Source/Runtime/Internationalization/Private/TextLocalizationResourceGenerator.cpp b/Engine/Source/Runtime/Internationalization/Private/TextLocalizationResourceGenerator.cpp
--- a/Engine/Source/Runtime/Internationalization/Private/TextLocalizationResourceGenerator.cpp
+++ b/Engine/Source/Runtime/Internationalization/Private/TextLocalizationResourceGenerator.cpp
@@ -219,9 +219,9 @@ bool FTextLocalizationResourceGenerator::GenerateAndUpdateLiveEntriesFromConfig(
 		return false;
 	}
 
-	// Get manifest name.
+	// Get archive name.
 	FString ArchiveName;
-	if (!GConfig->GetString(*SectionName, TEXT("ArchiveName"), ManifestName, InConfigFilePath))
+	if (!GConfig->GetString(*SectionName, TEXT("ArchiveName"), ArchiveName, InConfigFilePath))
 	{
 		UE_LOG(LogTextLocalizationResourceGenerator, Error, TEXT("No archive name specified."));
 		return false;
diff --git a/Engine/Source/Runtime/JsonUtilities/JsonUtilities.Build.cs b/Engine/Source/Runtime/JsonUtilities/JsonUtilities.Build.cs
--- a/Engine/Source/Runtime/JsonUtilities/JsonUtilities.Build.cs
+++ b/Engine/Source/Runtime/JsonUtilities/JsonUtilities.Build.cs
@@ -10,13 +10,8 @@ public class JsonUtilities : ModuleRules
 			new string[]
 			{
 				"Core",
-				"Json",
-			}
-		);
-		PrivateDependencyModuleNames.AddRange(
-			new string[] 
-			{ 			
 				"CoreUObject",
+				"Json",
 			}
 		);
 	}
diff --git a/Engine/Source/Runtime/JsonUtilities/Public/JsonObjectConverter.h b/Engine/Source/Runtime/JsonUtilities/Public/JsonObjectConverter.h
--- a/Engine/Source/Runtime/JsonUtilities/Public/JsonObjectConverter.h
+++ b/Engine/Source/Runtime/JsonUtilities/Public/JsonObjectConverter.h
@@ -3,6 +3,7 @@
 #pragma once
 
 #include "CoreMinimal.h"
+#include "UObject/Class.h"
 #include "Serialization/JsonTypes.h"
 #include "Dom/JsonObject.h"
 #include "Serialization/JsonReader.h"
diff --git a/Engine/Source/Runtime/Landscape/Classes/LandscapeLayerInfoObject.h b/Engine/Source/Runtime/Landscape/Classes/LandscapeLayerInfoObject.h
--- a/Engine/Source/Runtime/Landscape/Classes/LandscapeLayerInfoObject.h
+++ b/Engine/Source/Runtime/Landscape/Classes/LandscapeLayerInfoObject.h
@@ -27,6 +27,9 @@ class ULandscapeLayerInfoObject : public UObject
 #if WITH_EDITORONLY_DATA
 	UPROPERTY(VisibleAnywhere, Category=LandscapeLayerInfoObject)
 	uint32 bNoWeightBlend:1;
+
+	UPROPERTY(NonTransactional, Transient)
+	bool IsReferencedFromLoadedData;
 #endif // WITH_EDITORONLY_DATA
 
 	/* The color to use for layer usage debug */
diff --git a/Engine/Source/Runtime/Landscape/Classes/LandscapeProxy.h b/Engine/Source/Runtime/Landscape/Classes/LandscapeProxy.h
--- a/Engine/Source/Runtime/Landscape/Classes/LandscapeProxy.h
+++ b/Engine/Source/Runtime/Landscape/Classes/LandscapeProxy.h
@@ -194,6 +194,19 @@ enum class ELandscapeLayerPaintingRestriction : uint8
 	UseComponentWhitelist UMETA(DisplayName="Component Whitelist"),
 };
 
+UENUM()
+enum class ELandscapeLayerDisplayMode : uint8
+{
+	/** Material sorting display mode */
+	Default,
+
+	/** Alphabetical sorting display mode */
+	Alphabetical,
+
+	/** User specific sorting display mode */
+	UserSpecific,
+};
+
 UENUM()
 namespace ELandscapeLODFalloff
 {
@@ -350,6 +363,14 @@ public:
 	/** LOD level to use when exporting the landscape to obj or FBX */
 	UPROPERTY(EditAnywhere, Category=LOD, AdvancedDisplay)
 	int32 ExportLOD;
+
+	/** Display Order of the targets */
+	UPROPERTY(NonTransactional)
+	TArray<FName> TargetDisplayOrderList;
+
+	/** Display Order mode for the targets */
+	UPROPERTY(NonTransactional)
+	ELandscapeLayerDisplayMode TargetDisplayOrder;
 #endif
 
 	/** LOD level to use when running lightmass (increase to 1 or 2 for large landscapes to stop lightmass crashing) */
diff --git a/Engine/Source/Runtime/Landscape/Private/Landscape.cpp b/Engine/Source/Runtime/Landscape/Private/Landscape.cpp
--- a/Engine/Source/Runtime/Landscape/Private/Landscape.cpp
+++ b/Engine/Source/Runtime/Landscape/Private/Landscape.cpp
@@ -647,6 +647,9 @@ void ULandscapeComponent::PostLoad()
 
 ALandscapeProxy::ALandscapeProxy(const FObjectInitializer& ObjectInitializer)
 	: Super(ObjectInitializer)
+#if WITH_EDITORONLY_DATA
+	, TargetDisplayOrder(ELandscapeLayerDisplayMode::Default)
+#endif // WITH_EDITORONLY_DATA
 	, bHasLandscapeGrass(true)
 {
 	PrimaryActorTick.bCanEverTick = true;
@@ -1382,6 +1385,13 @@ bool ULandscapeInfo::UpdateLayerInfoMap(ALandscapeProxy* Proxy /*= nullptr*/, bo
 							int32 LayerInfoIndex = GetLayerInfoIndex(LayerInfo);
 							bool bValid = LayerNames.Contains(LayerInfo->LayerName);
 
+							#if WITH_EDITORONLY_DATA
+							if (bValid)
+							{
+								//LayerInfo->IsReferencedFromLoadedData = true;
+							}
+							#endif
+
 							if (LayerInfoIndex != INDEX_NONE)
 							{
 								FLandscapeInfoLayerSettings& LayerSettings = Layers[LayerInfoIndex];
@@ -1633,6 +1643,18 @@ void ALandscapeProxy::ConditionalAssignCommonProperties(ALandscape* Landscape)
 		bUpdated = true;
 	}
 
+	if (TargetDisplayOrder != Landscape->TargetDisplayOrder)
+	{
+		TargetDisplayOrder = Landscape->TargetDisplayOrder;
+		bUpdated = true;
+	}
+
+	if (TargetDisplayOrderList != Landscape->TargetDisplayOrderList)
+	{
+		TargetDisplayOrderList = Landscape->TargetDisplayOrderList;
+		bUpdated = true;
+	}
+
 	if (bUpdated)
 	{
 		MarkPackageDirty();
diff --git a/Engine/Source/Runtime/Landscape/Private/LandscapeEdit.cpp b/Engine/Source/Runtime/Landscape/Private/LandscapeEdit.cpp
--- a/Engine/Source/Runtime/Landscape/Private/LandscapeEdit.cpp
+++ b/Engine/Source/Runtime/Landscape/Private/LandscapeEdit.cpp
@@ -3337,6 +3337,9 @@ void ALandscape::PostDuplicate(bool bDuplicateForPIE)
 
 ULandscapeLayerInfoObject::ULandscapeLayerInfoObject(const FObjectInitializer& ObjectInitializer)
 	: Super(ObjectInitializer)
+#if WITH_EDITORONLY_DATA
+	, IsReferencedFromLoadedData(false)
+#endif // WITH_EDITORONLY_DATA
 {
 	Hardness = 0.5f;
 #if WITH_EDITORONLY_DATA
diff --git a/Engine/Source/Runtime/Landscape/Private/LandscapeEditInterface.cpp b/Engine/Source/Runtime/Landscape/Private/LandscapeEditInterface.cpp
--- a/Engine/Source/Runtime/Landscape/Private/LandscapeEditInterface.cpp
+++ b/Engine/Source/Runtime/Landscape/Private/LandscapeEditInterface.cpp
@@ -1872,10 +1872,40 @@ void FLandscapeEditDataInterface::FillLayer(ULandscapeLayerInfoObject* LayerInfo
 		return;
 	}
 
+	LayerInfo->IsReferencedFromLoadedData = true;
+
 	for (auto& XYComponentPair : LandscapeInfo->XYtoComponentMap)
 	{
 		ULandscapeComponent* Component = XYComponentPair.Value;
 		Component->FillLayer(LayerInfo, *this);
+	}	
+
+	// Flush dynamic data (e.g. grass)
+	TSet<ULandscapeComponent*> Components;
+	for (auto& XYComponentPair : LandscapeInfo->XYtoComponentMap)
+	{
+		Components.Add(XYComponentPair.Value);
+	}
+	ALandscapeProxy::InvalidateGeneratedComponentData(Components);
+}
+
+void FLandscapeEditDataInterface::FillEmptyLayers(ULandscapeLayerInfoObject* LayerInfo)
+{
+	if (!LandscapeInfo)
+	{
+		return;
+	}
+
+	LayerInfo->IsReferencedFromLoadedData = true;
+
+	for (auto& XYComponentPair : LandscapeInfo->XYtoComponentMap)
+	{
+		ULandscapeComponent* Component = XYComponentPair.Value;
+
+		if (Component->WeightmapLayerAllocations.Num() == 0)
+		{
+			Component->FillLayer(LayerInfo, *this);
+		}
 	}
 
 	// Flush dynamic data (e.g. grass)
diff --git a/Engine/Source/Runtime/Landscape/Private/LandscapeSplineRaster.cpp b/Engine/Source/Runtime/Landscape/Private/LandscapeSplineRaster.cpp
--- a/Engine/Source/Runtime/Landscape/Private/LandscapeSplineRaster.cpp
+++ b/Engine/Source/Runtime/Landscape/Private/LandscapeSplineRaster.cpp
@@ -62,7 +62,7 @@ protected:
 		const float Alpha = CosInterpX * CosInterpY;
 		uint16& Dest = Data[(Y - MinY)*(1 + MaxX - MinX) + X - MinX];
 		float Value = FMath::Lerp((float)Dest, Interpolant.Z, Alpha);
-		uint16 DValue = (uint32)FMath::Clamp<float>(Value, 0, LandscapeDataAccess::MaxValue);
+		uint16 DValue = (uint16)FMath::Clamp<float>(Value, 0, (float)LandscapeDataAccess::MaxValue);
 		if ((bRaiseTerrain && DValue > Dest) ||
 			(bLowerTerrain && DValue < Dest))
 		{
@@ -257,6 +257,7 @@ void RasterizeControlPointAlpha(int32& MinX, int32& MinY, int32& MaxX, int32& Ma
 	}
 
 	LandscapeEdit.SetAlphaData(LayerInfo, MinX, MinY, MaxX, MaxY, Data.GetData(), 0, ELandscapeLayerPaintingRestriction::None, !LayerInfo->bNoWeightBlend, false);
+
 	LandscapeEdit.GetComponentsInRegion(MinX, MinY, MaxX, MaxY, &ModifiedComponents);
 }
 
@@ -414,6 +415,7 @@ void RasterizeSegmentAlpha(int32& MinX, int32& MinY, int32& MaxX, int32& MaxY, F
 	}
 
 	LandscapeEdit.SetAlphaData(LayerInfo, MinX, MinY, MaxX, MaxY, Data.GetData(), 0, ELandscapeLayerPaintingRestriction::None, !LayerInfo->bNoWeightBlend, false);
+
 	LandscapeEdit.GetComponentsInRegion(MinX, MinY, MaxX, MaxY, &ModifiedComponents);
 }
 
diff --git a/Engine/Source/Runtime/Landscape/Public/LandscapeDataAccess.h b/Engine/Source/Runtime/Landscape/Public/LandscapeDataAccess.h
--- a/Engine/Source/Runtime/Landscape/Public/LandscapeDataAccess.h
+++ b/Engine/Source/Runtime/Landscape/Public/LandscapeDataAccess.h
@@ -33,7 +33,7 @@ namespace LandscapeDataAccess
 
 	FORCEINLINE uint16 GetTexHeight(float Height)
 	{
-		return FMath::Clamp<float>(Height * LANDSCAPE_INV_ZSCALE + MidValue, 0.f, MaxValue);
+		return FMath::RoundToInt(FMath::Clamp<float>(Height * LANDSCAPE_INV_ZSCALE + MidValue, 0.f, MaxValue));		
 	}
 };
 
diff --git a/Engine/Source/Runtime/Landscape/Public/LandscapeEdit.h b/Engine/Source/Runtime/Landscape/Public/LandscapeEdit.h
--- a/Engine/Source/Runtime/Landscape/Public/LandscapeEdit.h
+++ b/Engine/Source/Runtime/Landscape/Public/LandscapeEdit.h
@@ -160,6 +160,8 @@ struct LANDSCAPE_API FLandscapeEditDataInterface : public FLandscapeTextureDataI
 	void DeleteLayer(ULandscapeLayerInfoObject* LayerInfo);
 	// Fill a layer and re-normalize other layers
 	void FillLayer(ULandscapeLayerInfoObject* LayerInfo);
+	// Fill all empty layers and re-normalize layers
+	void FillEmptyLayers(ULandscapeLayerInfoObject* LayerInfo);
 	// Replace/merge a layer
 	void ReplaceLayer(ULandscapeLayerInfoObject* FromLayerInfo, ULandscapeLayerInfoObject* ToLayerInfo);
 
diff --git a/Engine/Source/Runtime/Launch/Private/Android/AndroidEventManager.h b/Engine/Source/Runtime/Launch/Private/Android/AndroidEventManager.h
--- a/Engine/Source/Runtime/Launch/Private/Android/AndroidEventManager.h
+++ b/Engine/Source/Runtime/Launch/Private/Android/AndroidEventManager.h
@@ -1,7 +1,6 @@
 // Copyright 1998-2017 Epic Games, Inc. All Rights Reserved.
 
-#ifndef ANDROID_EVENT_MANAGER_H
-#define ANDROID_EVENT_MANAGER_H
+#pragma once
 
 #include "CoreMinimal.h"
 #include "Containers/Queue.h"
@@ -93,6 +92,3 @@ private:
 	bool bHaveGame;
 	bool bRunning;
 };
-
-
-#endif
diff --git a/Engine/Source/Runtime/Launch/Private/LaunchEngineLoop.cpp b/Engine/Source/Runtime/Launch/Private/LaunchEngineLoop.cpp
--- a/Engine/Source/Runtime/Launch/Private/LaunchEngineLoop.cpp
+++ b/Engine/Source/Runtime/Launch/Private/LaunchEngineLoop.cpp
@@ -99,6 +99,7 @@
 	#include "Materials/MaterialInterface.h"
 	#include "TextureResource.h"
 	#include "Engine/Texture2D.h"
+	#include "StringTable.h"
 	#include "SceneUtils.h"
 	#include "ParticleHelper.h"
 	#include "PhysicsPublic.h"
@@ -1580,6 +1581,8 @@ int32 FEngineLoop::PreInit( const TCHAR* CmdLine )
 
 	EndInitTextLocalization();
 
+	UStringTable::InitializeEngineBridge();
+
 	if (FApp::ShouldUseThreadingForPerformance() && FPlatformMisc::AllowAudioThread())
 	{
 		bool bUseThreadedAudio = false;
@@ -2335,6 +2338,9 @@ bool FEngineLoop::LoadStartupCoreModules()
 	IAudioEditorModule* AudioEditorModule = &FModuleManager::LoadModuleChecked<IAudioEditorModule>("AudioEditor");
 	AudioEditorModule->RegisterAssetActions();
 
+	// Load the StringTableEditor module to register its asset actions
+	FModuleManager::Get().LoadModule("StringTableEditor");
+
 	if( !IsRunningDedicatedServer() )
 	{
 		// VREditor needs to be loaded in non-server editor builds early, so engine content Blueprints can be loaded during DDC generation
diff --git a/Engine/Source/Runtime/Mac/CoreAudio/Public/CoreAudioDevice.h b/Engine/Source/Runtime/Mac/CoreAudio/Public/CoreAudioDevice.h
--- a/Engine/Source/Runtime/Mac/CoreAudio/Public/CoreAudioDevice.h
+++ b/Engine/Source/Runtime/Mac/CoreAudio/Public/CoreAudioDevice.h
@@ -4,8 +4,7 @@
  	CoreAudioDevice.h: Unreal CoreAudio audio interface object.
  =============================================================================*/
 
-#ifndef _INC_COREAUDIODEVICE
-#define _INC_COREAUDIODEVICE
+#pragma once
 
 #include "CoreMinimal.h"
 #include "AudioDecompress.h"
@@ -444,5 +443,3 @@ private:
 	friend class FCoreAudioSoundSource;
 	friend class FCoreAudioEffectsManager;
 };
-
-#endif
diff --git a/Engine/Source/Runtime/Mac/CoreAudio/Public/CoreAudioEffects.h b/Engine/Source/Runtime/Mac/CoreAudio/Public/CoreAudioEffects.h
--- a/Engine/Source/Runtime/Mac/CoreAudio/Public/CoreAudioEffects.h
+++ b/Engine/Source/Runtime/Mac/CoreAudio/Public/CoreAudioEffects.h
@@ -4,8 +4,7 @@
 	CodeAudioEffects.h: Unreal CoreAudio audio effects interface object.
 =============================================================================*/
 
-#ifndef _INC_COREAUDIOEFFECTS
-#define _INC_COREAUDIOEFFECTS
+#pragma once
 
 #include "CoreAudioDevice.h"
 #include "AudioEffect.h"
@@ -55,5 +54,3 @@ private:
 	friend class FCoreAudioDevice;
 	friend class FCoreAudioSoundSource;
 };
-
-#endif
diff --git a/Engine/Source/Runtime/Niagara/Private/NiagaraMeshRendererProperties.cpp b/Engine/Source/Runtime/Niagara/Private/NiagaraMeshRendererProperties.cpp
--- a/Engine/Source/Runtime/Niagara/Private/NiagaraMeshRendererProperties.cpp
+++ b/Engine/Source/Runtime/Niagara/Private/NiagaraMeshRendererProperties.cpp
@@ -36,7 +36,7 @@ void UNiagaraMeshRendererProperties::FixMaterial(UMaterial* Material)
 
 void UNiagaraMeshRendererProperties::PostEditChangeProperty(FPropertyChangedEvent& PropertyChangedEvent)
 {
-	if (ParticleMesh && PropertyChangedEvent.Property->GetName() == "ParticleMesh")
+	if (ParticleMesh && PropertyChangedEvent.Property && PropertyChangedEvent.Property->GetName() == "ParticleMesh")
 	{
 		const FStaticMeshLODResources& LODModel = ParticleMesh->RenderData->LODResources[0];
 		for (int32 SectionIndex = 0; SectionIndex < LODModel.Sections.Num(); SectionIndex++)
diff --git a/Engine/Source/Runtime/Online/BuildPatchServices/Private/BuildPatchError.h b/Engine/Source/Runtime/Online/BuildPatchServices/Private/BuildPatchError.h
--- a/Engine/Source/Runtime/Online/BuildPatchServices/Private/BuildPatchError.h
+++ b/Engine/Source/Runtime/Online/BuildPatchServices/Private/BuildPatchError.h
@@ -4,9 +4,6 @@
 	BuildPatchError.h: Declares classes involved setting and getting error information.
 =============================================================================*/
 
-#ifndef __BuildPatchError_h__
-#define __BuildPatchError_h__
-
 #pragma once
 
 #include "CoreMinimal.h"
@@ -183,5 +180,3 @@ public:
 	 */
 	static const FString& EnumToString(const EBuildPatchInstallError& ErrorType);
 };
-
-#endif // __BuildPatchChunk_h__
diff --git a/Engine/Source/Runtime/Online/BuildPatchServices/Private/BuildPatchHTTP.h b/Engine/Source/Runtime/Online/BuildPatchServices/Private/BuildPatchHTTP.h
--- a/Engine/Source/Runtime/Online/BuildPatchServices/Private/BuildPatchHTTP.h
+++ b/Engine/Source/Runtime/Online/BuildPatchServices/Private/BuildPatchHTTP.h
@@ -5,9 +5,6 @@
 	in order to insure the HTTP module is only used from the main thread
 =============================================================================*/
 
-#ifndef __BuildPatchHTTP_h__
-#define __BuildPatchHTTP_h__
-
 #pragma once
 
 #include "CoreMinimal.h"
@@ -141,5 +138,3 @@ private:
 	static TSharedPtr<FBuildPatchHTTP> SingletonInstance;
 	static FDelegateHandle TickDelegateHandle;
 };
-
-#endif // __BuildPatchHTTP_h__
diff --git a/Engine/Source/Runtime/Online/HTTP/Private/GenericPlatform/GenericPlatformHttp.cpp b/Engine/Source/Runtime/Online/HTTP/Private/GenericPlatform/GenericPlatformHttp.cpp
--- a/Engine/Source/Runtime/Online/HTTP/Private/GenericPlatform/GenericPlatformHttp.cpp
+++ b/Engine/Source/Runtime/Online/HTTP/Private/GenericPlatform/GenericPlatformHttp.cpp
@@ -2,6 +2,7 @@
 
 #include "GenericPlatform/GenericPlatformHttp.h"
 #include "Interfaces/IHttpRequest.h"
+#include "Paths.h"
 
 
 /**
@@ -167,3 +168,74 @@ FString FGenericPlatformHttp::UrlDecode(const FString &EncodedString)
 	Data.Add('\0');
 	return FString(UTF8_TO_TCHAR(Data.GetData()));
 }
+
+FString FGenericPlatformHttp::HtmlEncode(const FString& UnencodedString)
+{
+	FString EncodedString;
+	EncodedString.Reserve(UnencodedString.Len());
+
+	for ( int32 Index = 0; Index != UnencodedString.Len(); ++Index )
+	{
+		switch ( UnencodedString[Index] )
+		{
+			case '&':  EncodedString.Append(TEXT("&amp;"));					break;
+			case '\"': EncodedString.Append(TEXT("&quot;"));				break;
+			case '\'': EncodedString.Append(TEXT("&apos;"));				break;
+			case '<':  EncodedString.Append(TEXT("&lt;"));					break;
+			case '>':  EncodedString.Append(TEXT("&gt;"));					break;
+			default:   EncodedString.AppendChar(UnencodedString[Index]);	break;
+		}
+	}
+
+	return EncodedString;
+}
+
+FString FGenericPlatformHttp::GetMimeType(const FString& FilePath)
+{
+	const FString FileExtension = FPaths::GetExtension(FilePath, true);
+
+	static TMap<FString, FString> ExtensionMimeTypeMap;
+	if ( ExtensionMimeTypeMap.Num() == 0 )
+	{
+		// Web
+		ExtensionMimeTypeMap.Add(TEXT(".html"), TEXT("text/html"));
+		ExtensionMimeTypeMap.Add(TEXT(".css"), TEXT("text/css"));
+		ExtensionMimeTypeMap.Add(TEXT(".js"), TEXT("application/x-javascript"));
+
+		// Video
+		ExtensionMimeTypeMap.Add(TEXT(".avi"), TEXT("video/msvideo, video/avi, video/x-msvideo"));
+		ExtensionMimeTypeMap.Add(TEXT(".mpeg"), TEXT("video/mpeg"));
+
+		// Image
+		ExtensionMimeTypeMap.Add(TEXT(".bmp"), TEXT("image/bmp"));
+		ExtensionMimeTypeMap.Add(TEXT(".gif"), TEXT("image/gif"));
+		ExtensionMimeTypeMap.Add(TEXT(".jpg"), TEXT("image/jpeg"));
+		ExtensionMimeTypeMap.Add(TEXT(".jpeg"), TEXT("image/jpeg"));
+		ExtensionMimeTypeMap.Add(TEXT(".png"), TEXT("image/png"));
+		ExtensionMimeTypeMap.Add(TEXT(".svg"), TEXT("image/svg+xml"));
+		ExtensionMimeTypeMap.Add(TEXT(".tiff"), TEXT("image/tiff"));
+
+		// Audio
+		ExtensionMimeTypeMap.Add(TEXT(".midi"), TEXT("audio/x-midi"));
+		ExtensionMimeTypeMap.Add(TEXT(".mp3"), TEXT("audio/mpeg"));
+		ExtensionMimeTypeMap.Add(TEXT(".ogg"), TEXT("audio/vorbis, application/ogg"));
+		ExtensionMimeTypeMap.Add(TEXT(".wav"), TEXT("audio/wav, audio/x-wav"));
+
+		// Documents
+		ExtensionMimeTypeMap.Add(TEXT(".xml"), TEXT("application/xml"));
+		ExtensionMimeTypeMap.Add(TEXT(".txt"), TEXT("text/plain"));
+		ExtensionMimeTypeMap.Add(TEXT(".tsv"), TEXT("text/tab-separated-values"));
+		ExtensionMimeTypeMap.Add(TEXT(".csv"), TEXT("text/csv"));
+		ExtensionMimeTypeMap.Add(TEXT(".json"), TEXT("application/json"));
+		
+		// Compressed
+		ExtensionMimeTypeMap.Add(TEXT(".zip"), TEXT("application/zip, application/x-compressed-zip"));
+	}
+
+	if ( FString* FoundMimeType = ExtensionMimeTypeMap.Find(FileExtension) )
+	{
+		return *FoundMimeType;
+	}
+
+	return TEXT("application/unknown");
+}
\ No newline at end of file
diff --git a/Engine/Source/Runtime/Online/HTTP/Private/Windows/WindowsPlatformHttp.cpp b/Engine/Source/Runtime/Online/HTTP/Private/Windows/WindowsPlatformHttp.cpp
--- a/Engine/Source/Runtime/Online/HTTP/Private/Windows/WindowsPlatformHttp.cpp
+++ b/Engine/Source/Runtime/Online/HTTP/Private/Windows/WindowsPlatformHttp.cpp
@@ -75,4 +75,27 @@ IHttpRequest* FWindowsPlatformHttp::ConstructRequest()
 	{
 		return new FHttpRequestWinInet();
 	}
+}
+
+FString FWindowsPlatformHttp::GetMimeType(const FString& FilePath)
+{
+	FString MimeType = TEXT("application/unknown");
+	const FString FileExtension = FPaths::GetExtension(FilePath, true);
+
+	HKEY hKey;
+	if ( ::RegOpenKeyEx(HKEY_CLASSES_ROOT, *FileExtension, 0, KEY_READ, &hKey) == ERROR_SUCCESS )
+	{
+		TCHAR MimeTypeBuffer[128];
+		DWORD MimeTypeBufferSize = sizeof(MimeTypeBuffer);
+		DWORD KeyType = 0;
+
+		if ( ::RegQueryValueEx(hKey, TEXT("Content Type"), NULL, &KeyType, (BYTE*)MimeTypeBuffer, &MimeTypeBufferSize) == ERROR_SUCCESS && KeyType == REG_SZ )
+		{
+			MimeType = MimeTypeBuffer;
+		}
+
+		::RegCloseKey(hKey);
+	}
+
+	return MimeType;
 }
\ No newline at end of file
diff --git a/Engine/Source/Runtime/Online/HTTP/Public/GenericPlatform/GenericPlatformHttp.h b/Engine/Source/Runtime/Online/HTTP/Public/GenericPlatform/GenericPlatformHttp.h
--- a/Engine/Source/Runtime/Online/HTTP/Public/GenericPlatform/GenericPlatformHttp.h
+++ b/Engine/Source/Runtime/Online/HTTP/Public/GenericPlatform/GenericPlatformHttp.h
@@ -57,4 +57,17 @@ public:
 	 * @return The decoded string
 	 */
 	static FString UrlDecode(const FString& EncodedString);
+
+	/**
+	 * Returns the &lt; &gt...etc encoding for strings between HTML elements.
+	 *
+	 * @param UnencodedString The unencoded string to convert to html encoding.
+	 * @return The html encoded string
+	 */
+	static FString HtmlEncode(const FString& UnencodedString);
+
+	/**
+	 * Returns the mime type for the file.
+	 */
+	static FString GetMimeType(const FString& FilePath);
 };
diff --git a/Engine/Source/Runtime/Online/HTTP/Public/Windows/WindowsPlatformHttp.h b/Engine/Source/Runtime/Online/HTTP/Public/Windows/WindowsPlatformHttp.h
--- a/Engine/Source/Runtime/Online/HTTP/Public/Windows/WindowsPlatformHttp.h
+++ b/Engine/Source/Runtime/Online/HTTP/Public/Windows/WindowsPlatformHttp.h
@@ -8,7 +8,7 @@
 /**
  * Platform specific HTTP implementations.
  */
-class FWindowsPlatformHttp : public FGenericPlatformHttp
+class HTTP_API FWindowsPlatformHttp : public FGenericPlatformHttp
 {
 public:
 
@@ -31,6 +31,11 @@ public:
 	 * @return The request object.
 	 */
 	static IHttpRequest* ConstructRequest();
+
+	/**
+	 * Returns the mime type for the file.
+	 */
+	static FString GetMimeType(const FString& FilePath);
 };
 
 
diff --git a/Engine/Source/Runtime/Slate/Private/Framework/Application/HardwareCursor.cpp b/Engine/Source/Runtime/Slate/Private/Framework/Application/HardwareCursor.cpp
new file mode 100644
--- /dev/null
+++ b/Engine/Source/Runtime/Slate/Private/Framework/Application/HardwareCursor.cpp
@@ -0,0 +1,372 @@
+// Copyright 1998-2017 Epic Games, Inc. All Rights Reserved.
+
+#include "Framework/Application/HardwareCursor.h"
+#include "Logging/LogMacros.h"
+#include "Misc/Paths.h"
+#include "HAL/PlatformProcess.h"
+#include "Misc/FileHelper.h"
+#include "ModuleManager.h"
+#include "Interfaces/IImageWrapperModule.h"
+#include "SlateApplication.h"
+
+#if PLATFORM_LINUX
+#include "SDL.h"
+#include "SDL_surface.h"
+#endif
+
+DEFINE_LOG_CATEGORY_STATIC(LogHardwareCursor, Log, All);
+
+FHardwareCursor::FHardwareCursor(const FString& InPathToCursorWithoutExtension, FVector2D InHotSpot)
+#if PLATFORM_WINDOWS
+	: CursorHandle(NULL)
+#elif PLATFORM_MAC
+	: CursorHandle(nullptr)
+#endif
+{
+	// The hotspot should be in normalized coordinates, because we multiply it by the width & height of the
+	// actual image we load into memory, which may vary per platform, to get the actual Pixel Hotspot.
+	check(InHotSpot.X >= 0.0f && InHotSpot.X <= 1.0f);
+	check(InHotSpot.Y >= 0.0f && InHotSpot.Y <= 1.0f);
+
+	// NOTE:
+	// The reason we don't include the file extension is so we can support prioritization
+	// of the different formats supported on each platform.  So on Windows we can support .ani -> .cur -> .png.
+	//
+	// Additionally - we load all the files into memory and then find ways to then convert them into
+	// OS cursors.  Several of the OSes have functions for loading a cursor from a file,
+	// which works great, unless your game uses Pak files.
+
+#if PLATFORM_WINDOWS
+
+	const FString AniCursor = InPathToCursorWithoutExtension + TEXT(".ani");
+	const FString CurCursor = InPathToCursorWithoutExtension + TEXT(".cur");
+
+	TArray<uint8> CursorFileData;
+
+	if (FFileHelper::LoadFileToArray(CursorFileData, *AniCursor, FILEREAD_Silent) || FFileHelper::LoadFileToArray(CursorFileData, *CurCursor, FILEREAD_Silent))
+	{
+		//TODO Would be nice to find a way to do this that doesn't involve the temp file copy.
+
+		// The cursors may be in a pak file, if that's the case we need to write it to a temporary file
+		// and then load that file as the cursor.  It's a workaround because there doesn't appear to be
+		// a good way to load a cursor from anything other than a loose file or a resource.
+		FString TempCursorFile = FPaths::CreateTempFilename(FPlatformProcess::UserTempDir(), TEXT("Cursor-"), TEXT(".temp"));
+		if (FFileHelper::SaveArrayToFile(CursorFileData, *TempCursorFile))
+		{
+			CursorHandle = (HCURSOR)LoadImage(NULL,
+				*TempCursorFile,
+				IMAGE_CURSOR,
+				0,
+				0,
+				LR_LOADFROMFILE);
+
+			IFileManager::Get().Delete(*TempCursorFile);
+		}
+	}
+	else if (FindAndLoadCorrectResolutionPng(CursorFileData, InPathToCursorWithoutExtension))
+	{
+		IImageWrapperModule& ImageWrapperModule = FModuleManager::LoadModuleChecked<IImageWrapperModule>(FName("ImageWrapper"));
+
+		IImageWrapperPtr PngImageWrapper = ImageWrapperModule.CreateImageWrapper(EImageFormat::PNG);
+		if (PngImageWrapper.IsValid() && PngImageWrapper->SetCompressed(CursorFileData.GetData(), CursorFileData.Num()))
+		{
+			const TArray<uint8>* RawImageData = nullptr;
+			if (PngImageWrapper->GetRaw(ERGBFormat::RGBA, 8, RawImageData))
+			{
+				const int32 Width = PngImageWrapper->GetWidth();
+				const int32 Height = PngImageWrapper->GetHeight();
+
+				CreateCursorFromRGBABuffer((FColor*)RawImageData->GetData(), Width, Height, InHotSpot);
+			}
+		}
+	}
+
+#elif PLATFORM_MAC
+
+	const FString TiffCursor = InPathToCursorWithoutExtension + TEXT(".tiff");
+
+	TArray<uint8> CursorFileData;
+	if (FFileHelper::LoadFileToArray(CursorFileData, *TiffCursor, FILEREAD_Silent) || FindAndLoadCorrectResolutionPng(CursorFileData, InPathToCursorWithoutExtension))
+	{
+		NSData* CursorData = [[NSData alloc] initWithBytes:CursorFileData.GetData() length : CursorFileData.Num()];
+		NSImage* CursorImage = [[NSImage alloc] initWithData:CursorData];
+		NSImageRep* CursorImageRep = [[CursorImage representations] objectAtIndex:0];
+
+		const int32 PixelHotspotX = FMath::RoundToInt(InHotSpot.X * CursorImageRep.pixelsWide);
+		const int32 PixelHotspotY = FMath::RoundToInt(InHotSpot.Y * CursorImageRep.pixelsHigh);
+
+		CursorHandle = [[NSCursor alloc] initWithImage:CursorImage hotSpot : NSMakePoint(PixelHotspotX, PixelHotspotY)];
+		[CursorImage release];
+		[CursorData release];
+	}
+
+#elif PLATFORM_LINUX
+
+	TArray<uint8> CursorFileData;
+	if (FindAndLoadCorrectResolutionPng(CursorFileData, InPathToCursorWithoutExtension))
+	{
+		IImageWrapperModule& ImageWrapperModule = FModuleManager::LoadModuleChecked<IImageWrapperModule>(FName("ImageWrapper"));
+
+		IImageWrapperPtr PngImageWrapper = ImageWrapperModule.CreateImageWrapper(EImageFormat::PNG);
+		if (PngImageWrapper.IsValid() && PngImageWrapper->SetCompressed(CursorFileData.GetData(), CursorFileData.Num()))
+		{
+			const TArray<uint8>* RawImageData = nullptr;
+			if (PngImageWrapper->GetRaw(ERGBFormat::RGBA, 8, RawImageData))
+			{
+				const int32 Width = PngImageWrapper->GetWidth();
+				const int32 Height = PngImageWrapper->GetHeight();
+
+				CreateCursorFromRGBABuffer((FColor*)RawImageData->GetData(), Width, Height, InHotSpot);
+			}
+		}
+	}
+
+#endif
+
+#if PLATFORM_WINDOWS || PLATFORM_MAC || PLATFORM_LINUX
+
+	if (!GetHandle())
+	{
+		UE_LOG(LogHardwareCursor, Error, TEXT("Failed to load cursor '%s'.  If you included a file extension or the size information '_WidthxHeight' remove them."), *InPathToCursorWithoutExtension);
+	}
+
+#endif
+}
+
+FHardwareCursor::FHardwareCursor(const TArray<FColor>& Pixels, FIntPoint InSize, FVector2D InHotSpot)
+#if PLATFORM_WINDOWS
+	: CursorHandle(NULL)
+#elif PLATFORM_MAC
+	: CursorHandle(nullptr)
+#endif
+{
+	// The hotspot should be in normalized coordinates, because we multiply it by the width & height of the
+	// actual image we load into memory, which may vary per platform, to get the actual Pixel Hotspot.
+	check(InHotSpot.X >= 0.0f && InHotSpot.X <= 1.0f);
+	check(InHotSpot.Y >= 0.0f && InHotSpot.Y <= 1.0f);
+
+	const int32 Width = InSize.X;
+	const int32 Height = InSize.Y;
+
+	CreateCursorFromRGBABuffer((FColor*)Pixels.GetData(), Width, Height, InHotSpot);
+}
+
+void FHardwareCursor::CreateCursorFromRGBABuffer(const FColor* Pixels, int32 Width, int32 Height, FVector2D InHotSpot)
+{
+#if PLATFORM_WINDOWS
+
+	TArray<FColor> BGRAPixels;
+	BGRAPixels.AddUninitialized(Width * Height);
+
+	for (int32 Index = 0; Index < BGRAPixels.Num(); Index++)
+	{
+		const FColor& SrcPixel = Pixels[Index];
+		BGRAPixels[Index] = FColor(SrcPixel.B, SrcPixel.G, SrcPixel.R, SrcPixel.A);
+	}
+
+	// The bitmap created is already in BGRA format, so we can just hand over the buffer.
+	HBITMAP CursorColor = ::CreateBitmap(Width, Height, 1, 32, BGRAPixels.GetData());
+	// Just create a dummy mask, we're making a full 32bit bitmap with support for alpha, so no need to worry about the mask.
+	HBITMAP CursorMask = ::CreateCompatibleBitmap(::GetDC(NULL), Width, Height);
+
+	ICONINFO IconInfo = { 0 };
+	IconInfo.fIcon = FALSE;
+	IconInfo.xHotspot = FMath::RoundToInt(InHotSpot.X * Width);
+	IconInfo.yHotspot = FMath::RoundToInt(InHotSpot.Y * Height);
+	IconInfo.hbmColor = CursorColor;
+	IconInfo.hbmMask = CursorMask;
+
+	CursorHandle = ::CreateIconIndirect(&IconInfo);
+
+	::DeleteObject(CursorColor);
+	::DeleteObject(CursorMask);
+
+#elif PLATFORM_MAC
+
+	NSBitmapImageRep* CursorImageRep = [
+		[NSBitmapImageRep alloc]
+			initWithBitmapDataPlanes:NULL
+			pixelsWide : Width
+			pixelsHigh : Height
+			bitsPerSample : 8
+			samplesPerPixel : 4
+			hasAlpha : YES
+			isPlanar : NO
+			colorSpaceName : NSCalibratedRGBColorSpace
+			bitmapFormat : NSAlphaFirstBitmapFormat
+			bytesPerRow : Width * 4
+			bitsPerPixel : 32];
+
+	uint8* CursorPixels = [CursorImageRep bitmapData];
+	for (int32 X = 0; X < Width; X++)
+	{
+		for (int32 Y = 0; Y < Height; Y++)
+		{
+			uint8* TargetPixel = (uint8*)CursorPixels + Y * Width + X * 4;
+			*(uint32*)TargetPixel = *(((uint32*)Pixels) + (Y * Width) + X);
+		}
+	}
+
+	NSImage* CursorImage = [[NSImage alloc] initWithSize:NSMakeSize(Width, Height)];
+	[CursorImage addRepresentation : CursorImageRep];
+
+	const int32 PixelHotspotX = FMath::RoundToInt(InHotSpot.X * Width);
+	const int32 PixelHotspotY = FMath::RoundToInt(InHotSpot.Y * Height);
+
+	CursorHandle = [[NSCursor alloc] initWithImage:CursorImage hotSpot : NSMakePoint(PixelHotspotX, PixelHotspotY)];
+	[CursorImage release];
+	[CursorImageRep release];
+
+#elif PLATFORM_LINUX
+
+	uint32 Rmask = 0x000000ff;
+	uint32 Gmask = 0x0000ff00;
+	uint32 Bmask = 0x00ff0000;
+	uint32 Amask = 0xff000000;
+
+	SDL_Surface* Surface = SDL_CreateRGBSurfaceFrom((void*)Pixels, Width, Height, 32, Width * 4, Rmask, Gmask, Bmask, Amask);
+	if (Surface)
+	{
+		CursorHandle = SDL_CreateColorCursor(Surface, FMath::RoundToInt(InHotSpot.X * Width), FMath::RoundToInt(InHotSpot.Y * Height));
+
+		SDL_FreeSurface(Surface);
+	}
+
+#endif
+}
+
+bool FHardwareCursor::FindAndLoadCorrectResolutionPng(TArray<uint8>& Result, const FString& InPathToCursorWithoutExtension)
+{
+#if PLATFORM_WINDOWS
+	const int32 SystemCursorWidth = ::GetSystemMetrics(SM_CXCURSOR);
+	const int32 SystemCursorHeight = ::GetSystemMetrics(SM_CYCURSOR);
+#elif PLATFORM_MAC
+	// TODO: Is there a way to tell the current desired cursor size on OSX?
+	const int32 SystemCursorWidth = 32;
+	const int32 SystemCursorHeight = 32;
+#elif PLATFORM_LINUX
+	// TODO: Is there a way to tell the current desired cursor size on Linux ?
+	const int32 SystemCursorWidth = 32;
+	const int32 SystemCursorHeight = 32;
+#endif
+
+#if PLATFORM_DESKTOP
+
+	FString CursorsWithSizeSearch = FPaths::GetCleanFilename(InPathToCursorWithoutExtension) + TEXT("_*x*.png");
+
+	TArray<FString> PngCursorFiles;
+	IFileManager::Get().FindFilesRecursive(PngCursorFiles, *FPaths::GetPath(InPathToCursorWithoutExtension), *CursorsWithSizeSearch, true, false, false);
+
+	FString PngCursorPath = InPathToCursorWithoutExtension + TEXT(".png");
+	int32 SizeDistance = INT_MAX;
+	int32 CursorWidth = 0;
+	int32 CursorHeight = 0;
+
+	for (const FString& CursorFile : PngCursorFiles)
+	{
+		FString Dummy;
+		FString SizeAndExtension;
+		FString Size;
+		FString Width, Height;
+
+		if (!CursorFile.Split(TEXT("_"), &Dummy, &SizeAndExtension, ESearchCase::IgnoreCase, ESearchDir::FromEnd))
+		{
+			continue;
+		}
+
+		if (!SizeAndExtension.Split(TEXT("."), &Size, &Dummy))
+		{
+			continue;
+		}
+
+		if (!Size.Split(TEXT("x"), &Width, &Height))
+		{
+			continue;
+		}
+
+		if (FCString::IsPureAnsi(*Width) && FCString::IsPureAnsi(*Height))
+		{
+			int32 TempCursorWidth = FCString::Atoi(*Width);
+			int32 TempCursorHeight = FCString::Atoi(*Height);
+
+			int32 Distance = FMath::Abs(SystemCursorWidth - CursorWidth);
+			if (Distance < SizeDistance)
+			{
+				PngCursorPath = CursorFile;
+				CursorWidth = TempCursorWidth;
+				CursorHeight = TempCursorHeight;
+				SizeDistance = Distance;
+
+				if (SizeDistance == 0)
+				{
+					break;
+				}
+			}
+		}
+	}
+
+	if (FFileHelper::LoadFileToArray(Result, *PngCursorPath, FILEREAD_Silent))
+	{
+		UE_LOG(LogHardwareCursor, Log, TEXT("Loading Cursor '%s'."), *PngCursorPath);
+
+		if ( SizeDistance != 0 )
+		{
+			FString PngCursorFileName = FPaths::GetCleanFilename(PngCursorPath);
+			UE_LOG(LogHardwareCursor, Warning, TEXT("The cursor '%s' is %dx%d, but the OS requested a cursor of size %dx%d.  You should add another variation to get a pixel perfect cursor."), 
+				*PngCursorFileName, CursorWidth, CursorHeight, SystemCursorWidth, SystemCursorHeight);
+		}
+
+		return true;
+	}
+	else
+	{
+		UE_LOG(LogHardwareCursor, Error, TEXT("Failed to load cursor '%s', couldn't find any matching cursors."), *PngCursorPath);
+	}
+
+#endif
+
+	return false;
+}
+
+FHardwareCursor::~FHardwareCursor()
+{
+#if PLATFORM_WINDOWS
+	if (CursorHandle)
+	{
+		DestroyCursor(CursorHandle);
+	}
+#elif PLATFORM_MAC
+	if (CursorHandle)
+	{
+		[CursorHandle release];
+	}
+#elif PLATFORM_LINUX
+	if (CursorHandle)
+	{
+		SDL_FreeCursor(CursorHandle);
+	}
+#endif
+}
+
+void* FHardwareCursor::GetHandle()
+{
+#if PLATFORM_WINDOWS
+	return CursorHandle;
+#elif PLATFORM_MAC
+	return CursorHandle;
+#elif PLATFORM_LINUX
+	return CursorHandle;
+#else
+	return nullptr;
+#endif
+}
+
+void FHardwareCursor::ReplaceTheShapeFor(EMouseCursor::Type InCursorType)
+{
+	TSharedPtr<ICursor> PlatformCursor = FSlateApplication::Get().GetPlatformCursor();
+
+	if (ICursor* Cursor = PlatformCursor.Get())
+	{
+		Cursor->SetTypeShape(InCursorType, GetHandle());
+	}
+}
\ No newline at end of file
