Automation - Disabling the OSS tests.

diff --git a/Engine/Source/Runtime/Engine/Private/Tests/OSSTests.cpp b/Engine/Source/Runtime/Engine/Private/Tests/OSSTests.cpp
--- a/Engine/Source/Runtime/Engine/Private/Tests/OSSTests.cpp
+++ b/Engine/Source/Runtime/Engine/Private/Tests/OSSTests.cpp
@@ -20,7 +20,7 @@ DEFINE_LOG_CATEGORY_STATIC(LogHackAutomationTests, Log, All);
 //////////////////////////////////////////////////////////////////////////
 
 
-IMPLEMENT_SIMPLE_AUTOMATION_TEST(FHttpTest, "System.OSS.Test Http", EAutomationTestFlags::EditorContext | EAutomationTestFlags::EngineFilter)
+IMPLEMENT_SIMPLE_AUTOMATION_TEST(FHttpTest, "System.OSS.Test Http", EAutomationTestFlags::Disabled | EAutomationTestFlags::EditorContext | EAutomationTestFlags::EngineFilter)
 
 bool FHttpTest::RunTest(const FString& Parameters)
 {
@@ -30,7 +30,7 @@ bool FHttpTest::RunTest(const FString& Parameters)
 	return true;
 }
 
-IMPLEMENT_SIMPLE_AUTOMATION_TEST(FMCPFriendsTest, "System.OSS.MCP.Test Friends", EAutomationTestFlags::EditorContext | EAutomationTestFlags::EngineFilter)
+IMPLEMENT_SIMPLE_AUTOMATION_TEST(FMCPFriendsTest, "System.OSS.MCP.Test Friends", EAutomationTestFlags::Disabled | EAutomationTestFlags::EditorContext | EAutomationTestFlags::EngineFilter)
 
 bool FMCPFriendsTest::RunTest(const FString& Parameters)
 {
@@ -41,7 +41,7 @@ bool FMCPFriendsTest::RunTest(const FString& Parameters)
 	return true;
 }
 
-IMPLEMENT_SIMPLE_AUTOMATION_TEST(FAmazonFriendsTest, "System.OSS.Amazon.Test Friends", EAutomationTestFlags::EditorContext | EAutomationTestFlags::EngineFilter)
+IMPLEMENT_SIMPLE_AUTOMATION_TEST(FAmazonFriendsTest, "System.OSS.Amazon.Test Friends", EAutomationTestFlags::Disabled | EAutomationTestFlags::EditorContext | EAutomationTestFlags::EngineFilter)
 
 bool FAmazonFriendsTest::RunTest(const FString& Parameters)
 {
@@ -51,7 +51,7 @@ bool FAmazonFriendsTest::RunTest(const FString& Parameters)
 	return true;
 }
 
-IMPLEMENT_SIMPLE_AUTOMATION_TEST(FFacebookFriendsTest, "System.OSS.Facebook.Test Friends", EAutomationTestFlags::EditorContext | EAutomationTestFlags::EngineFilter)
+IMPLEMENT_SIMPLE_AUTOMATION_TEST(FFacebookFriendsTest, "System.OSS.Facebook.Test Friends", EAutomationTestFlags::Disabled | EAutomationTestFlags::EditorContext | EAutomationTestFlags::EngineFilter)
 
 bool FFacebookFriendsTest::RunTest(const FString& Parameters)
 {
@@ -61,7 +61,7 @@ bool FFacebookFriendsTest::RunTest(const FString& Parameters)
 	return true;
 }
 
-IMPLEMENT_SIMPLE_AUTOMATION_TEST(FSteamFriendsTest, "System.OSS.Steam.Test Friends", EAutomationTestFlags::EditorContext | EAutomationTestFlags::EngineFilter)
+IMPLEMENT_SIMPLE_AUTOMATION_TEST(FSteamFriendsTest, "System.OSS.Steam.Test Friends", EAutomationTestFlags::Disabled | EAutomationTestFlags::EditorContext | EAutomationTestFlags::EngineFilter)
 
 bool FSteamFriendsTest::RunTest(const FString& Parameters)
 {
@@ -71,7 +71,9 @@ bool FSteamFriendsTest::RunTest(const FString& Parameters)
 	return true;
 }
 
-IMPLEMENT_SIMPLE_AUTOMATION_TEST(FIOSFriendsTest, "System.OSS.IOS.Test Friends", EAutomationTestFlags::EditorContext | EAutomationTestFlags::EngineFilter)
+#if PLATFORM_IOS
+
+IMPLEMENT_SIMPLE_AUTOMATION_TEST(FIOSFriendsTest, "System.OSS.IOS.Test Friends", EAutomationTestFlags::Disabled | EAutomationTestFlags::EditorContext | EAutomationTestFlags::EngineFilter)
 
 bool FIOSFriendsTest::RunTest(const FString& Parameters)
 {
@@ -81,7 +83,9 @@ bool FIOSFriendsTest::RunTest(const FString& Parameters)
 	return true;
 }
 
-IMPLEMENT_SIMPLE_AUTOMATION_TEST(FOculusFriendsTest, "System.OSS.Oculus.Test Friends", EAutomationTestFlags::EditorContext | EAutomationTestFlags::EngineFilter)
+#endif
+
+IMPLEMENT_SIMPLE_AUTOMATION_TEST(FOculusFriendsTest, "System.OSS.Oculus.Test Friends", EAutomationTestFlags::Disabled | EAutomationTestFlags::EditorContext | EAutomationTestFlags::EngineFilter)
 
 bool FOculusFriendsTest::RunTest(const FString& Parameters)
 {
@@ -107,7 +111,7 @@ bool FNullFriendsTest::RunTest(const FString& Parameters)
 }
 */
 
-IMPLEMENT_SIMPLE_AUTOMATION_TEST(FMCPCloudTest, "System.OSS.MCP.Test Cloud", EAutomationTestFlags::EditorContext | EAutomationTestFlags::EngineFilter)
+IMPLEMENT_SIMPLE_AUTOMATION_TEST(FMCPCloudTest, "System.OSS.MCP.Test Cloud", EAutomationTestFlags::Disabled | EAutomationTestFlags::EditorContext | EAutomationTestFlags::EngineFilter)
 
 bool FMCPCloudTest::RunTest(const FString& Parameters)
 {
@@ -117,7 +121,7 @@ bool FMCPCloudTest::RunTest(const FString& Parameters)
 	return true;
 }
 
-IMPLEMENT_SIMPLE_AUTOMATION_TEST(FSteamCloudTest, "System.OSS.Steam.Test Cloud", EAutomationTestFlags::EditorContext | EAutomationTestFlags::EngineFilter)
+IMPLEMENT_SIMPLE_AUTOMATION_TEST(FSteamCloudTest, "System.OSS.Steam.Test Cloud", EAutomationTestFlags::Disabled | EAutomationTestFlags::EditorContext | EAutomationTestFlags::EngineFilter)
 
 bool FSteamCloudTest::RunTest(const FString& Parameters)
 {
@@ -157,6 +161,8 @@ bool FNullLeaderBoardTest::RunTest(const FString& Parameters)
 }
 */
 
+#if PLATFORM_IOS
+
 IMPLEMENT_SIMPLE_AUTOMATION_TEST(FIosLeaderBoardTest, "System.OSS.IOS.Test Leaderboards", EAutomationTestFlags::EditorContext | EAutomationTestFlags::EngineFilter)
 
 bool FIosLeaderBoardTest::RunTest(const FString& Parameters)
@@ -167,8 +173,9 @@ bool FIosLeaderBoardTest::RunTest(const FString& Parameters)
 	return true;
 }
 
+#endif
 
-IMPLEMENT_SIMPLE_AUTOMATION_TEST(FSteamLeaderBoardTest, "System.OSS.Steam.Test Leaderboards", EAutomationTestFlags::EditorContext | EAutomationTestFlags::EngineFilter)
+IMPLEMENT_SIMPLE_AUTOMATION_TEST(FSteamLeaderBoardTest, "System.OSS.Steam.Test Leaderboards", EAutomationTestFlags::Disabled | EAutomationTestFlags::EditorContext | EAutomationTestFlags::EngineFilter)
 
 bool FSteamLeaderBoardTest::RunTest(const FString& Parameters)
 {
@@ -178,7 +185,7 @@ bool FSteamLeaderBoardTest::RunTest(const FString& Parameters)
 	return true;
 }
 
-IMPLEMENT_SIMPLE_AUTOMATION_TEST(FMCPTimeTest, "System.OSS.MCP.Test Time", EAutomationTestFlags::EditorContext | EAutomationTestFlags::EngineFilter)
+IMPLEMENT_SIMPLE_AUTOMATION_TEST(FMCPTimeTest, "System.OSS.MCP.Test Time", EAutomationTestFlags::Disabled | EAutomationTestFlags::EditorContext | EAutomationTestFlags::EngineFilter)
 
 bool FMCPTimeTest::RunTest(const FString& Parameters)
 {
@@ -188,7 +195,7 @@ bool FMCPTimeTest::RunTest(const FString& Parameters)
 	return true;
 }
 
-IMPLEMENT_SIMPLE_AUTOMATION_TEST(FMCPIdentityTest, "System.OSS.MCP.Test Identity", EAutomationTestFlags::EditorContext | EAutomationTestFlags::EngineFilter)
+IMPLEMENT_SIMPLE_AUTOMATION_TEST(FMCPIdentityTest, "System.OSS.MCP.Test Identity", EAutomationTestFlags::Disabled | EAutomationTestFlags::EditorContext | EAutomationTestFlags::EngineFilter)
 
 bool FMCPIdentityTest::RunTest(const FString& Parameters)
 {
@@ -198,7 +205,7 @@ bool FMCPIdentityTest::RunTest(const FString& Parameters)
 	return true;
 }
 
-IMPLEMENT_SIMPLE_AUTOMATION_TEST(FSteamIdentityTest, "System.OSS.Steam.Test Identity", EAutomationTestFlags::EditorContext | EAutomationTestFlags::EngineFilter)
+IMPLEMENT_SIMPLE_AUTOMATION_TEST(FSteamIdentityTest, "System.OSS.Steam.Test Identity", EAutomationTestFlags::Disabled | EAutomationTestFlags::EditorContext | EAutomationTestFlags::EngineFilter)
 
 bool FSteamIdentityTest::RunTest(const FString& Parameters)
 {
@@ -208,7 +215,7 @@ bool FSteamIdentityTest::RunTest(const FString& Parameters)
 	return true;
 }
 
-IMPLEMENT_SIMPLE_AUTOMATION_TEST(FAmazonIdentityTest, "System.OSS.Amazon.Test Identity", EAutomationTestFlags::EditorContext | EAutomationTestFlags::EngineFilter)
+IMPLEMENT_SIMPLE_AUTOMATION_TEST(FAmazonIdentityTest, "System.OSS.Amazon.Test Identity", EAutomationTestFlags::Disabled | EAutomationTestFlags::EditorContext | EAutomationTestFlags::EngineFilter)
 
 bool FAmazonIdentityTest::RunTest(const FString& Parameters)
 {
@@ -218,7 +225,7 @@ bool FAmazonIdentityTest::RunTest(const FString& Parameters)
 	return true;
 }
 
-IMPLEMENT_SIMPLE_AUTOMATION_TEST(FFacebookIdentityTest, "System.OSS.Facebook.Test Identity", EAutomationTestFlags::EditorContext | EAutomationTestFlags::EngineFilter)
+IMPLEMENT_SIMPLE_AUTOMATION_TEST(FFacebookIdentityTest, "System.OSS.Facebook.Test Identity", EAutomationTestFlags::Disabled | EAutomationTestFlags::EditorContext | EAutomationTestFlags::EngineFilter)
 
 bool FFacebookIdentityTest::RunTest(const FString& Parameters)
 {
@@ -228,7 +235,9 @@ bool FFacebookIdentityTest::RunTest(const FString& Parameters)
 	return true;
 }
 
-IMPLEMENT_SIMPLE_AUTOMATION_TEST(FIOSIdentityTest, "System.OSS.IOS.Test Identity", EAutomationTestFlags::EditorContext | EAutomationTestFlags::EngineFilter)
+#if PLATFORM_IOS
+
+IMPLEMENT_SIMPLE_AUTOMATION_TEST(FIOSIdentityTest, "System.OSS.IOS.Test Identity", EAutomationTestFlags::Disabled | EAutomationTestFlags::EditorContext | EAutomationTestFlags::EngineFilter)
 
 bool FIOSIdentityTest::RunTest(const FString& Parameters)
 {
@@ -238,7 +247,9 @@ bool FIOSIdentityTest::RunTest(const FString& Parameters)
 	return true;
 }
 
-IMPLEMENT_SIMPLE_AUTOMATION_TEST(FOculusIdentityTest, "System.OSS.Oculus.Test Identity", EAutomationTestFlags::EditorContext | EAutomationTestFlags::EngineFilter)
+#endif
+
+IMPLEMENT_SIMPLE_AUTOMATION_TEST(FOculusIdentityTest, "System.OSS.Oculus.Test Identity", EAutomationTestFlags::Disabled | EAutomationTestFlags::EditorContext | EAutomationTestFlags::EngineFilter)
 
 bool FOculusIdentityTest::RunTest(const FString& Parameters)
 {
@@ -263,7 +274,7 @@ bool FNullIdentityTest::RunTest(const FString& Parameters)
 }
 */
 
-IMPLEMENT_SIMPLE_AUTOMATION_TEST(FMCPEntitlementsTest, "System.OSS.MCP.Test Entitlements", EAutomationTestFlags::EditorContext | EAutomationTestFlags::EngineFilter)
+IMPLEMENT_SIMPLE_AUTOMATION_TEST(FMCPEntitlementsTest, "System.OSS.MCP.Test Entitlements", EAutomationTestFlags::Disabled | EAutomationTestFlags::EditorContext | EAutomationTestFlags::EngineFilter)
 
 bool FMCPEntitlementsTest::RunTest(const FString& Parameters)
 {
@@ -273,7 +284,7 @@ bool FMCPEntitlementsTest::RunTest(const FString& Parameters)
 	return true;
 }
 
-IMPLEMENT_SIMPLE_AUTOMATION_TEST(FMCPTitleFileTest, "System.OSS.MCP.Test Title File", EAutomationTestFlags::EditorContext | EAutomationTestFlags::EngineFilter)
+IMPLEMENT_SIMPLE_AUTOMATION_TEST(FMCPTitleFileTest, "System.OSS.MCP.Test Title File", EAutomationTestFlags::Disabled | EAutomationTestFlags::EditorContext | EAutomationTestFlags::EngineFilter)
 
 bool FMCPTitleFileTest::RunTest(const FString& Parameters)
 {
@@ -283,7 +294,7 @@ bool FMCPTitleFileTest::RunTest(const FString& Parameters)
 	return true;
 }
 
-IMPLEMENT_SIMPLE_AUTOMATION_TEST(FSteamAchievementsTest, "System.OSS.Steam.Test Achievements", EAutomationTestFlags::EditorContext | EAutomationTestFlags::EngineFilter)
+IMPLEMENT_SIMPLE_AUTOMATION_TEST(FSteamAchievementsTest, "System.OSS.Steam.Test Achievements", EAutomationTestFlags::Disabled | EAutomationTestFlags::EditorContext | EAutomationTestFlags::EngineFilter)
 
 bool FSteamAchievementsTest::RunTest(const FString& Parameters)
 {
