Updated a setting's tooltip to indicate dependency on another setting.

diff --git a/Engine/Source/Runtime/Engine/Classes/Engine/RendererSettings.h b/Engine/Source/Runtime/Engine/Classes/Engine/RendererSettings.h
--- a/Engine/Source/Runtime/Engine/Classes/Engine/RendererSettings.h
+++ b/Engine/Source/Runtime/Engine/Classes/Engine/RendererSettings.h
@@ -239,7 +239,7 @@ class ENGINE_API URendererSettings : public UDeveloperSettings
 
 	UPROPERTY(config, EditAnywhere, Category=Translucency, meta=(
 		ConsoleVariable="r.TranslucentSortPolicy",DisplayName="Translucent Sort Policy",
-		ToolTip="The sort mode for translucent primitives, affecting how they are ordered and how they change order as the camera moves."))
+		ToolTip="The sort mode for translucent primitives, affecting how they are ordered and how they change order as the camera moves. Requires that Separate Translucency (under Postprocessing) is true."))
 	TEnumAsByte<ETranslucentSortPolicy::Type> TranslucentSortPolicy;
 
 	UPROPERTY(config, EditAnywhere, Category=Translucency, meta=(
