Copying //UE4/Dev-Core to //UE4/Dev-Main (Source: //UE4/Dev-Core @ 3283640)

diff --git a/Engine/Source/Runtime/CoreUObject/Private/UObject/ScriptCore.cpp b/Engine/Source/Runtime/CoreUObject/Private/UObject/ScriptCore.cpp
--- a/Engine/Source/Runtime/CoreUObject/Private/UObject/ScriptCore.cpp
+++ b/Engine/Source/Runtime/CoreUObject/Private/UObject/ScriptCore.cpp
@@ -1096,7 +1096,7 @@ bool UObject::CallFunctionByNameWithArguments(const TCHAR* Str, FOutputDevice& A
 			FFormatNamedArguments Arguments;
 			Arguments.Add(TEXT("Message"), FText::FromName( Message ));
 			Arguments.Add(TEXT("PropertyName"), FText::FromString( It->GetName() ));
-			Ar.Logf( *FText::Format( NSLOCTEXT( "Core", "BadProperty", "'{Message}': Bad or missing property '{PropertyName}'" ), Arguments ).ToString() );
+			Ar.Logf( TEXT("%s"), *FText::Format( NSLOCTEXT( "Core", "BadProperty", "'{Message}': Bad or missing property '{PropertyName}'" ), Arguments ).ToString() );
 			Failed = true;
 
 			break;
diff --git a/Engine/Source/Runtime/CoreUObject/Private/UObject/UObjectBase.cpp b/Engine/Source/Runtime/CoreUObject/Private/UObject/UObjectBase.cpp
--- a/Engine/Source/Runtime/CoreUObject/Private/UObject/UObjectBase.cpp
+++ b/Engine/Source/Runtime/CoreUObject/Private/UObject/UObjectBase.cpp
@@ -17,6 +17,7 @@
 #include "UObject/Package.h"
 #include "Templates/Casts.h"
 #include "UObject/GCObject.h"
+#include "LinkerLoad.h"
 
 DEFINE_LOG_CATEGORY_STATIC(LogUObjectBase, Log, All);
 DEFINE_STAT(STAT_UObjectsStatGroupTester);
@@ -168,8 +169,10 @@ void UObjectBase::DeferredRegister(UClass *UClassStaticClass,const TCHAR* Packag
 {
 	check(Internal::GObjInitialized);
 	// Set object properties.
-	OuterPrivate        = CreatePackage(nullptr,PackageName);
-	check(OuterPrivate);
+	UPackage* Package = CreatePackage(nullptr, PackageName);
+	check(Package);
+	Package->SetPackageFlags(PKG_CompiledIn);
+	OuterPrivate = Package;
 
 	check(UClassStaticClass);
 	check(!ClassPrivate);
@@ -438,7 +441,7 @@ TMap<FName, UScriptStruct *(*)()>& GetDynamicStructMap()
 	return DynamicStructMap;
 }
 
-void UObjectCompiledInDeferStruct(class UScriptStruct *(*InRegister)(), const TCHAR* PackageName, const FName PathName, bool bDynamic)
+void UObjectCompiledInDeferStruct(class UScriptStruct *(*InRegister)(), const TCHAR* PackageName, const FName ObjectName, bool bDynamic, const TCHAR* DynamicPathName)
 {
 	if (!bDynamic)
 	{
@@ -449,13 +452,18 @@ void UObjectCompiledInDeferStruct(class UScriptStruct *(*InRegister)(), const TC
 	}
 	else
 	{
-		GetDynamicStructMap().Add(PathName, InRegister);
+		GetDynamicStructMap().Add(DynamicPathName, InRegister);
 	}
+	NotifyRegistrationEvent(PackageName, *ObjectName.ToString(), ENotifyRegistrationType::NRT_Struct, ENotifyRegistrationPhase::NRP_Added, (UObject *(*)())(InRegister), bDynamic);
+
 }
 
 class UScriptStruct *GetStaticStruct(class UScriptStruct *(*InRegister)(), UObject* StructOuter, const TCHAR* StructName, SIZE_T Size, uint32 Crc)
 {
-	return (*InRegister)();
+	NotifyRegistrationEvent(*StructOuter->GetOutermost()->GetName(), StructName, ENotifyRegistrationType::NRT_Struct, ENotifyRegistrationPhase::NRP_Started);
+	UScriptStruct *Result = (*InRegister)();
+	NotifyRegistrationEvent(*StructOuter->GetOutermost()->GetName(), StructName, ENotifyRegistrationType::NRT_Struct, ENotifyRegistrationPhase::NRP_Finished);
+	return Result;
 }
 
 /**
@@ -491,7 +499,7 @@ TMap<FName, UEnum *(*)()>& GetDynamicEnumMap()
 	return DynamicEnumMap;
 }
 
-void UObjectCompiledInDeferEnum(class UEnum *(*InRegister)(), const TCHAR* PackageName, const FName PathName, bool bDynamic)
+void UObjectCompiledInDeferEnum(class UEnum *(*InRegister)(), const TCHAR* PackageName, const FName ObjectName, bool bDynamic, const TCHAR* DynamicPathName)
 {
 	if (!bDynamic)
 	{
@@ -502,13 +510,17 @@ void UObjectCompiledInDeferEnum(class UEnum *(*InRegister)(), const TCHAR* Packa
 	}
 	else
 	{
-		GetDynamicEnumMap().Add(PathName, InRegister);
+		GetDynamicEnumMap().Add(DynamicPathName, InRegister);
 	}
+	NotifyRegistrationEvent(PackageName, *ObjectName.ToString(), ENotifyRegistrationType::NRT_Enum, ENotifyRegistrationPhase::NRP_Added, (UObject *(*)())(InRegister), bDynamic);
 }
 
 class UEnum *GetStaticEnum(class UEnum *(*InRegister)(), UObject* EnumOuter, const TCHAR* EnumName)
 {
-	return (*InRegister)();
+	NotifyRegistrationEvent(*EnumOuter->GetOutermost()->GetName(), EnumName, ENotifyRegistrationType::NRT_Enum, ENotifyRegistrationPhase::NRP_Started);
+	UEnum *Result = (*InRegister)();
+	NotifyRegistrationEvent(*EnumOuter->GetOutermost()->GetName(), EnumName, ENotifyRegistrationType::NRT_Enum, ENotifyRegistrationPhase::NRP_Finished);
+	return Result;
 }
 
 static TArray<class UClass *(*)()>& GetDeferredCompiledInRegistration()
@@ -640,7 +652,7 @@ TMap<FName, FDynamicClassStaticData>& GetDynamicClassMap()
 	return DynamicClassMap;
 }
 
-void UObjectCompiledInDefer(UClass *(*InRegister)(), UClass *(*InStaticClass)(), const TCHAR* Name, bool bDynamic, const TCHAR* DynamicPathName, void (*InInitSearchableValues)(TMap<FName, FName>&))
+void UObjectCompiledInDefer(UClass *(*InRegister)(), UClass *(*InStaticClass)(), const TCHAR* Name, const TCHAR* PackageName, bool bDynamic, const TCHAR* DynamicPathName, void (*InInitSearchableValues)(TMap<FName, FName>&))
 {
 	if (!bDynamic)
 	{
@@ -649,6 +661,9 @@ void UObjectCompiledInDefer(UClass *(*InRegister)(), UClass *(*InStaticClass)(),
 		if (!GIsHotReload || GetDeferRegisterClassMap().FindChecked(Name)->bHasChanged)
 #endif
 		{
+			FString NoPrefix(RemoveClassPrefix(Name));
+			NotifyRegistrationEvent(PackageName, *NoPrefix, ENotifyRegistrationType::NRT_Class, ENotifyRegistrationPhase::NRP_Added, (UObject *(*)())(InRegister), false);
+			NotifyRegistrationEvent(PackageName, *(FString(DEFAULT_OBJECT_PREFIX) + NoPrefix), ENotifyRegistrationType::NRT_ClassCDO, ENotifyRegistrationPhase::NRP_Added, (UObject *(*)())(InRegister), false);
 			checkSlow(!GetDeferredCompiledInRegistration().Contains(InRegister));
 			GetDeferredCompiledInRegistration().Add(InRegister);
 		}
@@ -663,6 +678,15 @@ void UObjectCompiledInDefer(UClass *(*InRegister)(), UClass *(*InStaticClass)(),
 			InInitSearchableValues(ClassFunctions.SelectedSearchableValues);
 		}
 		GetDynamicClassMap().Add(FName(DynamicPathName), ClassFunctions);
+
+		FString OriginalPackageName = DynamicPathName;
+		check(OriginalPackageName.EndsWith(Name));
+		OriginalPackageName.RemoveFromEnd(FString(Name));
+		check(OriginalPackageName.EndsWith(TEXT(".")));
+		OriginalPackageName.RemoveFromEnd(FString(TEXT(".")));
+
+		NotifyRegistrationEvent(*OriginalPackageName, Name, ENotifyRegistrationType::NRT_Class, ENotifyRegistrationPhase::NRP_Added, (UObject *(*)())(InRegister), true);
+		NotifyRegistrationEvent(*OriginalPackageName, *(FString(DEFAULT_OBJECT_PREFIX) + Name), ENotifyRegistrationType::NRT_ClassCDO, ENotifyRegistrationPhase::NRP_Added, (UObject *(*)())(InRegister), true);
 	}
 }
 
@@ -768,10 +792,23 @@ static void UObjectLoadAllCompiledInDefaultProperties()
 	if( bHaveRegistrants )
 	{
 		TArray<UClass*> NewClasses;
+		TArray<UClass*> NewClassesInCoreUObject;
 		TArray<UClass* (*)()> PendingRegistrants = MoveTemp(DeferredCompiledInRegistration);
 		for (UClass* (*Registrant)() : PendingRegistrants)
 		{
-			NewClasses.Add(Registrant());
+			UClass* Class = Registrant();
+			if (Class->GetOutermost()->GetFName() == GLongCoreUObjectPackageName)
+			{
+				NewClassesInCoreUObject.Add(Class);
+			}
+			else
+			{
+				NewClasses.Add(Class);
+			}
+		}
+		for (UClass* Class : NewClassesInCoreUObject) // we do these first because we assume these never trigger loads
+		{
+			Class->GetDefaultObject();
 		}
 		for (UClass* Class : NewClasses)
 		{
@@ -834,6 +871,8 @@ void ProcessNewlyLoadedUObjects()
 {
 	DECLARE_SCOPE_CYCLE_COUNTER(TEXT("ProcessNewlyLoadedUObjects"), STAT_ProcessNewlyLoadedUObjects, STATGROUP_ObjectVerbose);
 
+
+
 #if WITH_HOT_RELOAD
 	UClassGenerateCDODuplicatesForHotReload();
 #endif
@@ -1089,6 +1128,11 @@ namespace
 		{
 			return nullptr;
 		}
+		virtual const TCHAR* ClassPackage() const override
+		{
+			check(0);
+			return nullptr;
+		}
 	};
 
 
@@ -1220,6 +1264,27 @@ FName GetDynamicTypeClassName(FName TypePathName)
 	{
 		Result = UEnum::StaticClass()->GetFName();
 	}
+	if (false && Result == NAME_None)
+	{
+		UE_LOG(LogUObjectBase, Warning, TEXT("GetDynamicTypeClassName %s not found."), *TypePathName.ToString());
+		UE_LOG(LogUObjectBase, Warning, TEXT("---- classes"));
+		for (auto& Pair : GetDynamicClassMap())
+		{
+			UE_LOG(LogUObjectBase, Warning, TEXT("    %s"), *Pair.Key.ToString());
+		}
+		UE_LOG(LogUObjectBase, Warning, TEXT("---- structs"));
+		for (auto& Pair : GetDynamicStructMap())
+		{
+			UE_LOG(LogUObjectBase, Warning, TEXT("    %s"), *Pair.Key.ToString());
+		}
+		UE_LOG(LogUObjectBase, Warning, TEXT("---- enums"));
+		for (auto& Pair : GetDynamicEnumMap())
+		{
+			UE_LOG(LogUObjectBase, Warning, TEXT("    %s"), *Pair.Key.ToString());
+		}
+		UE_LOG(LogUObjectBase, Fatal, TEXT("GetDynamicTypeClassName %s not found."), *TypePathName.ToString());
+	}
+	UE_CLOG(Result == NAME_None, LogUObjectBase, Warning, TEXT("GetDynamicTypeClassName %s not found."), *TypePathName.ToString());
 	return Result;
 }
 
diff --git a/Engine/Source/Runtime/CoreUObject/Private/UObject/UObjectGlobals.cpp b/Engine/Source/Runtime/CoreUObject/Private/UObject/UObjectGlobals.cpp
--- a/Engine/Source/Runtime/CoreUObject/Private/UObject/UObjectGlobals.cpp
+++ b/Engine/Source/Runtime/CoreUObject/Private/UObject/UObjectGlobals.cpp
@@ -1138,26 +1138,15 @@ static UPackage* LoadPackageInternalInner(UPackage* InOuter, const TCHAR* InLong
 		Result = FindObjectFast<UPackage>(nullptr, PackageFName);
 		if (!Result || !Result->IsFullyLoaded())
 		{
-			if (FLinkerLoad::IsKnownMissingPackage(PackageFName) || (Result && Result->IsPendingKill()))
+			int32 RequestID = LoadPackageAsync(InName, nullptr, *InPackageName);
+			FlushAsyncLoading(RequestID);
+		}
+
+		if (InOuter)
 			{
-				// Don't retry if known missing or pending kill
-				UE_LOG(LogUObjectGlobals, Verbose, TEXT("Failing load of package %s because it's already failed once %s."), *InPackageName);
+			return InOuter;
 			}
-			else
-			{
-				FlushAsyncLoading();
 				Result = FindObjectFast<UPackage>(nullptr, PackageFName);
-			
-				check(!Result || Result->IsFullyLoaded());
-				if (!Result)
-				{
-					LoadPackageAsync(InName, nullptr, *InPackageName);
-					FlushAsyncLoading();
-					Result = FindObjectFast<UPackage>(nullptr, PackageFName);
-				}
-			}
-		}
-
 		return Result;
 }
 
diff --git a/Engine/Source/Runtime/CoreUObject/Public/Serialization/AsyncLoading.h b/Engine/Source/Runtime/CoreUObject/Public/Serialization/AsyncLoading.h
--- a/Engine/Source/Runtime/CoreUObject/Public/Serialization/AsyncLoading.h
+++ b/Engine/Source/Runtime/CoreUObject/Public/Serialization/AsyncLoading.h
@@ -643,6 +643,7 @@ public:
 	void Event_ProcessPostloadWait();
 	void Event_StartPostload();
 
+	void MarkNewObjectForLoadIfItIsAnExport(UObject *Object);
 	bool AnyImportsAndExportWorkOutstanding();
 	void ConditionalQueueProcessImportsAndExports(bool bRequeueForTimeout = false);
 	void ConditionalQueueProcessPostloadWait();
diff --git a/Engine/Source/Runtime/CoreUObject/Public/UObject/Class.h b/Engine/Source/Runtime/CoreUObject/Public/UObject/Class.h
--- a/Engine/Source/Runtime/CoreUObject/Public/UObject/Class.h
+++ b/Engine/Source/Runtime/CoreUObject/Public/UObject/Class.h
@@ -281,10 +281,6 @@ public:
 	/** Array of object references embedded in script code. Mirrored for easy access by realtime garbage collection code */
 	TArray<UObject*> ScriptObjectReferences;
 
-	/** Map of Class Name to Map of Old Property Name to New Property Name */
-	static TMap<FName,TMap<FName,FName> > TaggedPropertyRedirects;
-	static void InitTaggedPropertyRedirectsMap();
-
 public:
 	// Constructors.
 	UStruct( EStaticConstructor, int32 InSize, EObjectFlags InFlags );
@@ -1796,12 +1792,6 @@ public:
 	}
 
 private:
-	/** Map of Enum Name to Map of Old Enum entry to New Enum entry */
-	static TMap<FName,TMap<FName,FName> > EnumRedirects;
-	/** Map of Enum Name to Map of Old Enum substring to New Enum substring, to handle many renames at once */
-	static TMap<FName,TMap<FString,FString> > EnumSubstringRedirects;
-	static void InitEnumRedirectsMap();
-
 	FORCEINLINE static FString GetIndexAsString_Internal( const TCHAR* EnumPath, const int32 Index )
 	{
 		UEnum* EnumClass = FindObject<UEnum>( nullptr, EnumPath );
diff --git a/Engine/Source/Runtime/CoreUObject/Public/UObject/CoreRedirects.h b/Engine/Source/Runtime/CoreUObject/Public/UObject/CoreRedirects.h
new file mode 100644
--- /dev/null
+++ b/Engine/Source/Runtime/CoreUObject/Public/UObject/CoreRedirects.h
@@ -0,0 +1,274 @@
+// Copyright 1998-2017 Epic Games, Inc. All Rights Reserved.
+
+/*=============================================================================
+	CoreRedirects.h: Object/Class/Field redirects read from ini files or registered at startup
+=============================================================================*/
+
+#pragma once
+
+#include "CoreMinimal.h"
+#include "Misc/EnumClassFlags.h"
+
+/** 
+ * Flags describing the type and properties of this redirect
+ */
+enum class ECoreRedirectFlags : int32
+{
+	None = 0,
+
+	// Core type of the thing being redirected, multiple can be set
+	Type_Object =			0x00000001, // UObject
+	Type_Class =			0x00000002, // UClass
+	Type_Struct =			0x00000004, // UStruct
+	Type_Enum =				0x00000008, // UEnum
+	Type_Function =			0x00000010, // UFunction
+	Type_Property =			0x00000020, // UProperty
+	Type_Package =			0x00000040, // UPackage
+	Type_Delegate =			0x00000080, // Dynamic delegate, special type of UFunction
+
+	// Option flags, specify rules for this redirect
+	Option_InstanceOnly =	0x00010000, // Only redirect instances of this type, not the type itself
+	Option_Removed =		0x00020000, // This type was explicitly removed, new name isn't valid
+	Option_MatchSubstring = 0x00040000, // Does a slow substring match
+};
+ENUM_CLASS_FLAGS(ECoreRedirectFlags);
+
+
+/**
+ * An object path extracted into component names for matching. TODO merge with FStringAssetReference?
+ */
+struct COREUOBJECT_API FCoreRedirectObjectName
+{
+	/** Raw name of object */
+	FName ObjectName;
+
+	/** String of outer chain, may be empty */
+	FName OuterName;
+
+	/** Package this was in before, may be extracted out of OldName */
+	FName PackageName;
+
+	/** Default to invalid names */
+	FCoreRedirectObjectName() {}
+
+	/** Copy constructor */
+	FCoreRedirectObjectName(const FCoreRedirectObjectName& Other)
+		: ObjectName(Other.ObjectName), OuterName(Other.OuterName), PackageName(Other.PackageName)
+	{
+	}
+
+	/** Construct from FNames that are already expanded */
+	FCoreRedirectObjectName(FName InObjectName, FName InOuterName, FName InPackageName)
+		: ObjectName(InObjectName), OuterName(InOuterName), PackageName(InPackageName)
+	{
+
+	}
+
+	/** Construct from a path string, this handles full paths with packages, or partial paths without */
+	FCoreRedirectObjectName(const FString& InString);
+
+	/** Construct from object in memory */
+	FCoreRedirectObjectName(const class UObject* Object);
+
+	/** Creates FString version */
+	FString ToString() const;
+
+	/** Sets back to invalid state */
+	void Reset();
+
+	/** Checks for exact equality */
+	bool operator==(const FCoreRedirectObjectName& Other) const
+	{
+		return ObjectName == Other.ObjectName && OuterName == Other.OuterName && PackageName == Other.PackageName;
+	}
+
+	bool operator!=(const FCoreRedirectObjectName& Other) const
+	{
+		return !(*this == Other);
+	}
+
+	/** Returns true if the passed in name matches requirements, will ignore names that are none */
+	bool Matches(const FCoreRedirectObjectName& Other, bool bCheckSubstring = false) const;
+
+	/** Returns integer of degree of match. 0 if doesn't match at all, higher integer for better matches */
+	int32 MatchScore(const FCoreRedirectObjectName& Other) const;
+
+	/** Returns the name used as the key into the acceleration map */
+	FName GetSearchKey(ECoreRedirectFlags Type) const
+	{
+		if ((Type & ECoreRedirectFlags::Option_MatchSubstring) == ECoreRedirectFlags::Option_MatchSubstring)
+		{
+			static FName SubstringName = FName(TEXT("*SUBSTRING*"));
+
+			// All substring matches pass initial test as they need to be manually checked
+			return SubstringName;
+		}
+
+		if ((Type & ECoreRedirectFlags::Type_Package) == ECoreRedirectFlags::Type_Package)
+		{
+			return PackageName;
+		}
+
+		return ObjectName;
+	}
+
+	/** Returns true if this refers to an actual object */
+	bool IsValid() const
+	{
+		return ObjectName != NAME_None || PackageName != NAME_None;
+	}
+
+	/** Returns true if all names have valid characters */
+	bool HasValidCharacters() const;
+
+	/** Expand OldName/NewName as needed */
+	static bool ExpandNames(const FString& FullString, FName& OutName, FName& OutOuter, FName &OutPackage);
+
+	/** Turn it back into an FString */
+	static FString CombineNames(FName NewName, FName NewOuter, FName NewPackage);
+};
+
+/** 
+ * A single redirection from an old name to a new name, parsed out of an ini file
+ */
+struct COREUOBJECT_API FCoreRedirect
+{
+	/** Flags of this redirect */
+	ECoreRedirectFlags RedirectFlags;
+
+	/** Name of object to look for */
+	FCoreRedirectObjectName OldName;
+
+	/** Name to replace with */
+	FCoreRedirectObjectName NewName;
+
+	/** Change the class of this object when doing a redirect */
+	FCoreRedirectObjectName OverrideClassName;
+
+	/** Map of value changes, from old value to new value */
+	TMap<FString, FString> ValueChanges;
+
+	/** Construct from name strings, which may get parsed out */
+	FCoreRedirect(const FCoreRedirect& Other)
+		: RedirectFlags(Other.RedirectFlags), OldName(Other.OldName), NewName(Other.NewName), OverrideClassName(Other.OverrideClassName), ValueChanges(Other.ValueChanges)
+	{
+	}
+
+	/** Construct from name strings, which may get parsed out */
+	FCoreRedirect(ECoreRedirectFlags InRedirectFlags, FString InOldName, FString InNewName)
+		: RedirectFlags(InRedirectFlags), OldName(InOldName), NewName(InNewName)
+	{
+		NormalizeNewName();
+	}
+	
+	/** Construct parsed out object names */
+	FCoreRedirect(ECoreRedirectFlags InRedirectFlags, const FCoreRedirectObjectName& InOldName, const FCoreRedirectObjectName& InNewName)
+		: RedirectFlags(InRedirectFlags), OldName(InOldName), NewName(InNewName)
+	{
+		NormalizeNewName();
+	}
+
+	/** Normalizes NewName with data from OldName */
+	void NormalizeNewName();
+
+	/** Parses a char buffer into the ValueChanges map */
+	const TCHAR* ParseValueChanges(const TCHAR* Buffer);
+
+	/** Returns true if the passed in name matches requirements */
+	bool Matches(ECoreRedirectFlags InFlags, const FCoreRedirectObjectName& InName) const;
+
+	/** Returns true if this has value redirects */
+	bool HasValueChanges() const;
+
+	/** Returns true if this is a substring match */
+	bool IsSubstringMatch() const;
+
+	/** Convert to new names based on mapping */
+	FCoreRedirectObjectName RedirectName(const FCoreRedirectObjectName& OldObjectName) const;
+
+	/** See if search criteria is identical */
+	bool IdenticalMatchRules(const FCoreRedirect& Other) const
+	{
+		return RedirectFlags == Other.RedirectFlags && OldName == Other.OldName;
+	}
+
+	/** Returns the name used as the key into the acceleration map */
+	FName GetSearchKey() const
+	{
+		return OldName.GetSearchKey(RedirectFlags);
+	}
+};
+
+/**
+ * A container for all of the registered core-level redirects 
+ */
+struct COREUOBJECT_API FCoreRedirects
+{
+	/** Returns a redirected version of the object name. If there are no valid redirects, it will return the original name */
+	static FCoreRedirectObjectName GetRedirectedName(ECoreRedirectFlags Type, const FCoreRedirectObjectName& OldObjectName);
+
+	/** Returns map of String->String value redirects for the object name, or nullptr if none found */
+	static const TMap<FString, FString>* GetValueRedirects(ECoreRedirectFlags Type, const FCoreRedirectObjectName& OldObjectName);
+
+	/** Performs both a name redirect and gets a value redirect struct if it exists. Returns true if either redirect found */
+	static bool RedirectNameAndValues(ECoreRedirectFlags Type, const FCoreRedirectObjectName& OldObjectName, FCoreRedirectObjectName& NewObjectName, const FCoreRedirect** FoundValueRedirect);
+
+	/** Returns true if this name has been registered as explicitly missing */
+	static bool IsKnownMissing(ECoreRedirectFlags Type, const FCoreRedirectObjectName& ObjectName);
+
+	/** Adds this as a missing name */
+	static bool AddKnownMissing(ECoreRedirectFlags Type, const FCoreRedirectObjectName& ObjectName);
+
+	/** Returns list of names it may have been before */
+	static bool FindPreviousNames(ECoreRedirectFlags Type, const FCoreRedirectObjectName& NewObjectName, TArray<FCoreRedirectObjectName>& PreviousNames);
+
+	/** Returns list of all core redirects that match requirements */
+	static bool GetMatchingRedirects(ECoreRedirectFlags Type, const FCoreRedirectObjectName& OldObjectName, TArray<const FCoreRedirect*>& FoundRedirects);
+
+	/** Parse all redirects out of a given ini file */
+	static bool ReadRedirectsFromIni(const FString& IniName);
+
+	/** Adds an array of redirects to global list */
+	static bool AddRedirectList(const TArray<FCoreRedirect>& Redirects, const FString& SourceString);
+
+	/** Returns true if this has ever been initialized from ini */
+	static bool IsInitialized() { return bInitialized; }
+
+	/** Gets map from config key -> Flags */
+	static const TMap<FName, ECoreRedirectFlags>& GetConfigKeyMap() { return ConfigKeyMap; }
+
+	/** Goes from the containing package and name of the type to the type flag */
+	static ECoreRedirectFlags GetFlagsForTypeName(FName PackageName, FName TypeName);
+
+	/** Goes from UClass Type to the type flag */
+	static ECoreRedirectFlags GetFlagsForTypeClass(UClass *TypeClass);
+
+	/** Runs set of redirector tests */
+	static void RunTests();
+
+private:
+	/** Static only class, never constructed */
+	FCoreRedirects();
+
+	/** Add a single redirect to a type map */
+	static bool AddSingleRedirect(const FCoreRedirect& NewRedirect, const FString& SourceString);
+
+	/** Add native redirects, called before ini is parsed for the first time */
+	static void RegisterNativeRedirects();
+
+	/** There is one of these for each registered set of redirect flags */
+	struct FRedirectNameMap
+	{
+		/** Map from name of thing being mapped to full list. List must be filtered further */
+		TMap<FName, TArray<FCoreRedirect> > RedirectMap;
+	};
+
+	/** Rather this has been initialized at least once */
+	static bool bInitialized;
+
+	/** Map from config name to flag */
+	static TMap<FName, ECoreRedirectFlags> ConfigKeyMap;
+
+	/** Map from name of thing being mapped to full list. List must be filtered further */
+	static TMap<ECoreRedirectFlags, FRedirectNameMap> RedirectTypeMap;
+};
diff --git a/Engine/Source/Runtime/CoreUObject/Public/UObject/LinkerLoad.h b/Engine/Source/Runtime/CoreUObject/Public/UObject/LinkerLoad.h
--- a/Engine/Source/Runtime/CoreUObject/Public/UObject/LinkerLoad.h
+++ b/Engine/Source/Runtime/CoreUObject/Public/UObject/LinkerLoad.h
@@ -149,39 +149,6 @@ public:
 	*/
 	TArray<FPackageIndex> PreloadDependencies;
 
-	/** OldClassName to NewClassName for ImportMap */
-	static TMap<FName, FName> ObjectNameRedirects;
-	/** Additional info for some ObjectName redirects to also redirect the class and class package  */
-	static TMap<FName, TPair<FName, FName>> ObjectNameClassRedirects;
-	/** OldClassName to NewClassName for ExportMap */
-	static TMap<FName, FName> ObjectNameRedirectsInstanceOnly;
-	/** Object name to NewClassName for export map */
-	static TMap<FName, FName> ObjectNameRedirectsObjectOnly;	
-	/** Old game name to new game name for ImportMap */
-	static TMap<FName, FName> GameNameRedirects;
-	/** Add a new redirect from old game name to new game name for ImportMap */
-	COREUOBJECT_API static void AddGameNameRedirect(const FName OldName, const FName NewName);
-	/** Old struct name to new struct name mapping */
-	static TMap<FName, FName> StructNameRedirects;
-	/** Old plugin name to new plugin name mapping */
-	static TMap<FString, FString> PluginNameRedirects;
-
-	/** Object name to required class and new name for load-time remapping */
-	struct FSubobjectRedirect
-	{
-		FName MatchClass;
-		FName NewName;
-		FSubobjectRedirect(FName InMatchClass, FName InNewName)
-			: MatchClass(InMatchClass)
-			, NewName(InNewName)
-		{
-		}
-	};
-	static TMap<FName, FSubobjectRedirect> SubobjectNameRedirects;	
-
-	/* Makes sure active redirects map has been initialized */
-	static bool bActiveRedirectsMapInitialized;
-
 	/** 
 	 * Utility functions to query the object name redirects list for previous names for a class
 	 * @param CurrentClassPath The current name of the class, with a full path
@@ -228,10 +195,12 @@ public:
 	 */
 	COREUOBJECT_API bool HasAnyObjectsPendingLoad() const;
 
-private:
+	/** 
+	 * Add a new redirect from old game name to new game name for ImportMap 
+	 */
+	COREUOBJECT_API static void AddGameNameRedirect(const FName OldName, const FName NewName);
 
-	/** Packages that are known to be missing when verifying imports that we don't want a message about */
-	COREUOBJECT_API static TSet<FName> KnownMissingPackages;
+private:
 
 	// Variables used during async linker creation.
 
@@ -331,6 +300,9 @@ private:
 
 	static FName NAME_LoadErrors;
 
+	/** Makes sure the deprecated active redirects inis have been read */
+	static bool bActiveRedirectsMapInitialized;
+
 #if WITH_EDITOR
 	/** Feedback scope that is created to house the slow task of an asynchronous linker load. Raw ptr so we don't pull in TUniquePtr for everything. */
 	FScopedSlowTask* LoadProgressScope;
@@ -375,6 +347,11 @@ public:
 	 */
 	bool IsValidPackageIndex(FPackageIndex InIndex);
 
+	/**
+	 * Locates package index for a UPackage import
+	 */
+	COREUOBJECT_API bool FindImportPackage(FName PackageName, FPackageIndex& PackageIdx);
+
 	/**
 	 * Locates the class adjusted index and its package adjusted index for a given class name in the import map
 	 */
@@ -596,7 +573,7 @@ public:
 	COREUOBJECT_API static void InvalidateExport(UObject* OldObject);
 
 	/** Used by Matinee to fixup component renaming */
-	COREUOBJECT_API static FName FindSubobjectRedirectName(const FName& Name);
+	COREUOBJECT_API static FName FindSubobjectRedirectName(const FName& Name, UClass* Class);
 
 private:
 #if WITH_EDITOR
@@ -886,11 +863,6 @@ private:
 	 */
 	ELinkerStatus FixupImportMap();
 
-	/**
-	 * Changes imports before the objects are created.
-	 */
-	ELinkerStatus RemapImports();
-
 	/**
 	 * Serializes the export map.
 	 */
@@ -1142,3 +1114,23 @@ private:
 	// FLinkerLoad creation helpers END
 	//
 };
+
+// used by the EDL at boot time to coordinate loading with what is going on with the deferred registration stuff
+enum class ENotifyRegistrationType
+{
+	NRT_Class,
+	NRT_ClassCDO,
+	NRT_Struct,
+	NRT_Enum,
+	NRT_Package,
+};
+
+enum class ENotifyRegistrationPhase
+{
+	NRP_Added,
+	NRP_Started,
+	NRP_Finished,
+};
+
+COREUOBJECT_API void NotifyRegistrationEvent(const TCHAR* PackageName, const TCHAR* Name, ENotifyRegistrationType NotifyRegistrationType, ENotifyRegistrationPhase NotifyRegistrationPhase, UObject *(*InRegister)() = nullptr, bool InbDynamic = false);
+COREUOBJECT_API void NotifyRegistrationComplete();
diff --git a/Engine/Source/Runtime/CoreUObject/Public/UObject/Object.h b/Engine/Source/Runtime/CoreUObject/Public/UObject/Object.h
--- a/Engine/Source/Runtime/CoreUObject/Public/UObject/Object.h
+++ b/Engine/Source/Runtime/CoreUObject/Public/UObject/Object.h
@@ -156,6 +156,14 @@ public:
 	 */
 	virtual void PostInitProperties();
 
+	/**
+	* Called after the C++ constructor has run on the CDO for a class. This is an obscure routine used to deal with the recursion 
+	* in the construction of the default materials
+	*/
+	virtual void PostCDOContruct()
+	{
+	}
+
 	/**
 	 * Called from within SavePackage on the passed in base/ root. The return value of this function will be passed to
 	 * PostSaveRoot. This is used to allow objects used as base to perform required actions before saving and cleanup
diff --git a/Engine/Source/Runtime/CoreUObject/Public/UObject/ObjectMacros.h b/Engine/Source/Runtime/CoreUObject/Public/UObject/ObjectMacros.h
--- a/Engine/Source/Runtime/CoreUObject/Public/UObject/ObjectMacros.h
+++ b/Engine/Source/Runtime/CoreUObject/Public/UObject/ObjectMacros.h
@@ -492,10 +492,11 @@ enum EObjectFlags
 	RF_StrongRefOnFrame			= 0x01000000,	///< References to this object from persistent function frame are handled as strong ones.
 	RF_NonPIEDuplicateTransient		= 0x02000000,  ///< Object should not be included for duplication unless it's being duplicated for a PIE session
 	RF_Dynamic = 0x04000000, // Field Only. Dynamic field - doesn't get constructed during static initialization, can be constructed multiple times
+	RF_WillBeLoaded = 0x08000000, // This object was constructed during load and will be loaded shortly
 };
 
 	// Special all and none masks
-#define RF_AllFlags				(EObjectFlags)0x07ffffff	///< All flags, used mainly for error checking
+#define RF_AllFlags				(EObjectFlags)0x0fffffff	///< All flags, used mainly for error checking
 
 	// Predefined groups of the above
 #define RF_Load						((EObjectFlags)(RF_Public | RF_Standalone | RF_Transactional | RF_ClassDefaultObject | RF_ArchetypeObject | RF_DefaultSubObject | RF_TextExportTransient | RF_InheritableComponentTemplate | RF_DuplicateTransient | RF_NonPIEDuplicateTransient)) // Flags to load from Unrealfiles.
@@ -638,9 +639,6 @@ namespace UC
 		/// This keyword is used to set the actor group that the class is show in, in the editor.
 		classGroup,
 
-		/// This keyword specifies that this class cannot have its code generated until the given class has been processed.
-		DependsOn,
-
 		/// Declares that instances of this class should always have an outer of the specified class.  This is inherited by subclasses unless overridden.
 		Within, /* =OuterClassName */
 
@@ -754,9 +752,6 @@ namespace UI
 	// valid keywords for the UINTERFACE macro, see the UCLASS versions, above
 	enum 
 	{
-		/// This keyword specifies that this class cannot have its code generated until the given class has been processed.
-		DependsOn,
-
 		/// This keyword indicates that the interface should be accessible outside of it's module, but does not need all methods exported.
 		/// It exports only the autogenerated methods required for dynamic_cast<>, etc... to work.
 		MinimalAPI,
@@ -1319,7 +1314,7 @@ Class declaration macros.
 private: \
     TClass& operator=(TClass&&);   \
     TClass& operator=(const TClass&);   \
-	TRequiredAPI static UClass* GetPrivateStaticClass(const TCHAR* Package); \
+	TRequiredAPI static UClass* GetPrivateStaticClass(); \
 public: \
 	/** Bitwise union of #EClassFlags pertaining to this class.*/ \
 	enum {StaticClassFlags=TStaticFlags}; \
@@ -1330,7 +1325,12 @@ public: \
 	/** Returns a UClass object representing this class at runtime */ \
 	inline static UClass* StaticClass() \
 	{ \
-		return GetPrivateStaticClass(TPackage); \
+		return GetPrivateStaticClass(); \
+	} \
+	/** Returns the package this class belongs in */ \
+	inline static const TCHAR* StaticPackage() \
+	{ \
+		return TPackage; \
 	} \
 	/** Returns the StaticClassFlags for this class */ \
 	inline static EClassCastFlags StaticClassCastFlags() \
@@ -1511,14 +1511,14 @@ public: \
 // Register a class at startup time.
 #define IMPLEMENT_CLASS(TClass, TClassCrc) \
 	static TClassCompiledInDefer<TClass> AutoInitialize##TClass(TEXT(#TClass), sizeof(TClass), TClassCrc); \
-	UClass* TClass::GetPrivateStaticClass(const TCHAR* Package) \
+	UClass* TClass::GetPrivateStaticClass() \
 	{ \
 		static UClass* PrivateStaticClass = NULL; \
 		if (!PrivateStaticClass) \
 		{ \
 			/* this could be handled with templates, but we want it external to avoid code bloat */ \
 			GetPrivateStaticClassBody( \
-				Package, \
+				StaticPackage(), \
 				(TCHAR*)TEXT(#TClass) + 1 + ((StaticClassFlags & CLASS_Deprecated) ? 11 : 0), \
 				PrivateStaticClass, \
 				StaticRegisterNatives##TClass, \
@@ -1537,7 +1537,7 @@ public: \
 	}
 
 // Used for intrinsics, this sets up the boiler plate, plus an initialization singleton, which can create properties and GC tokens
-#define IMPLEMENT_INTRINSIC_CLASS(TClass, TRequiredAPI, TSuperClass, TSuperRequiredAPI, InitCode) \
+#define IMPLEMENT_INTRINSIC_CLASS(TClass, TRequiredAPI, TSuperClass, TSuperRequiredAPI, TPackage, InitCode) \
 	IMPLEMENT_CLASS(TClass, 0) \
 	TRequiredAPI UClass* Z_Construct_UClass_##TClass(); \
 	UClass* Z_Construct_UClass_##TClass() \
@@ -1556,16 +1556,16 @@ public: \
 		check(Class->GetClass()); \
 		return Class; \
 	} \
-	static FCompiledInDefer Z_CompiledInDefer_UClass_##TClass(Z_Construct_UClass_##TClass, &TClass::StaticClass, TEXT(#TClass), false);
+	static FCompiledInDefer Z_CompiledInDefer_UClass_##TClass(Z_Construct_UClass_##TClass, &TClass::StaticClass, TEXT(TPackage), TEXT(#TClass), false);
 
 #define IMPLEMENT_CORE_INTRINSIC_CLASS(TClass, TSuperClass, InitCode) \
-	IMPLEMENT_INTRINSIC_CLASS(TClass, COREUOBJECT_API, TSuperClass, COREUOBJECT_API, InitCode)
+	IMPLEMENT_INTRINSIC_CLASS(TClass, COREUOBJECT_API, TSuperClass, COREUOBJECT_API, "/Script/CoreUObject" ,InitCode)
 
 // Register a dynamic class (created at runtime, not startup). Explicit ClassName parameter because Blueprint types can have names that can't be used natively:
 #define IMPLEMENT_DYNAMIC_CLASS(TClass, ClassName, TClassCrc) \
-	UClass* TClass::GetPrivateStaticClass(const TCHAR* Package) \
+	UClass* TClass::GetPrivateStaticClass() \
 	{ \
-		UPackage* PrivateStaticClassOuter = FindOrConstructDynamicTypePackage(Package); \
+		UPackage* PrivateStaticClassOuter = FindOrConstructDynamicTypePackage(StaticPackage()); \
 		UClass* PrivateStaticClass = Cast<UClass>(StaticFindObjectFast(UClass::StaticClass(), PrivateStaticClassOuter, (TCHAR*)ClassName)); \
 		if (!PrivateStaticClass) \
 		{ \
@@ -1578,7 +1578,7 @@ public: \
 		{ \
 			/* this could be handled with templates, but we want it external to avoid code bloat */ \
 			GetPrivateStaticClassBody( \
-			Package, \
+			StaticPackage(), \
 			(TCHAR*)ClassName, \
 			PrivateStaticClass, \
 			StaticRegisterNatives##TClass, \
@@ -1672,5 +1672,3 @@ public:
 	**/
 	virtual void Restore() const=0;
 };
-
-
diff --git a/Engine/Source/Runtime/CoreUObject/Public/UObject/ObjectResource.h b/Engine/Source/Runtime/CoreUObject/Public/UObject/ObjectResource.h
--- a/Engine/Source/Runtime/CoreUObject/Public/UObject/ObjectResource.h
+++ b/Engine/Source/Runtime/CoreUObject/Public/UObject/ObjectResource.h
@@ -154,6 +154,7 @@ struct FObjectResource
 #if WITH_EDITOR
 	/**
 	 * Name of the class this object was serialized with (in case active class redirects have changed it)
+	 * If this is a class and was directly redirected, this is what it was redirected from
 	 */
 	FName			OldClassName;
 #endif
@@ -214,14 +215,14 @@ struct FObjectExport : public FObjectResource
 	 * The number of bytes to serialize when saving/loading this export's UObject.
 	 * Serialized
 	 */
-	int32         	SerialSize;
+	int64         	SerialSize;
 
 	/**
 	 * The location (into the FLinker's underlying file reader archive) of the beginning of the
 	 * data for this export's UObject.  Used for verification only.
 	 * Serialized
 	 */
-	int32         	SerialOffset;
+	int64         	SerialOffset;
 
 	/**
 	 * The location (into the FLinker's underlying file reader archive) of the beginning of the
diff --git a/Engine/Source/Runtime/CoreUObject/Public/UObject/ReferenceChainSearch.h b/Engine/Source/Runtime/CoreUObject/Public/UObject/ReferenceChainSearch.h
--- a/Engine/Source/Runtime/CoreUObject/Public/UObject/ReferenceChainSearch.h
+++ b/Engine/Source/Runtime/CoreUObject/Public/UObject/ReferenceChainSearch.h
@@ -162,7 +162,7 @@ private:
 	/** Performs the referencer search */
 	void PerformSearch();
 	/** Processes an object and returns whether it is part of the reference chain or not */
-	bool ProcessObject(UObject* CurrentObject);
+	void ProcessObject(UObject* CurrentObject);
 	/** Inserts the reference chain into the result set and tries to merge it into an already
 		present one, if possible */
 	void InsertReferenceChain(FReferenceChain& Referencer);
diff --git a/Engine/Source/Runtime/CoreUObject/Public/UObject/UObjectBase.h b/Engine/Source/Runtime/CoreUObject/Public/UObject/UObjectBase.h
--- a/Engine/Source/Runtime/CoreUObject/Public/UObject/UObjectBase.h
+++ b/Engine/Source/Runtime/CoreUObject/Public/UObject/UObjectBase.h
@@ -278,6 +278,9 @@ struct FFieldCompiledInInfo
 	/** Registers the native class (constructs a UClass object) */
 	virtual UClass* Register() const = 0;
 
+	/** Return the package the class belongs in */
+	virtual const TCHAR* ClassPackage() const = 0;
+
 	/** Size of the class */
 	SIZE_T Size;
 	/** CRC of the generated code for this class */
@@ -308,29 +311,33 @@ struct TClassCompiledInDefer : public FFieldCompiledInInfo
 	{
 		return TClass::StaticClass();
 	}
+	virtual const TCHAR* ClassPackage() const override
+	{
+		return TClass::StaticPackage();
+	}
 };
 
 /**
  * Stashes the singleton function that builds a compiled in class. Later, this is executed.
  */
-COREUOBJECT_API void UObjectCompiledInDefer(class UClass *(*InRegister)(), class UClass *(*InStaticClass)(), const TCHAR* Name, bool bDynamic, const TCHAR* DynamicPathName, void (*InInitSearchableValues)(TMap<FName, FName>&));
+COREUOBJECT_API void UObjectCompiledInDefer(class UClass *(*InRegister)(), class UClass *(*InStaticClass)(), const TCHAR* Name, const TCHAR* PackageName, bool bDynamic, const TCHAR* DynamicPathName, void (*InInitSearchableValues)(TMap<FName, FName>&));
 
 struct FCompiledInDefer
 {
-	FCompiledInDefer(class UClass *(*InRegister)(), class UClass *(*InStaticClass)(), const TCHAR* Name, bool bDynamic, const TCHAR* DynamicPackageName = nullptr, const TCHAR* DynamicPathName = nullptr, void (*InInitSearchableValues)(TMap<FName, FName>&) = nullptr)
+	FCompiledInDefer(class UClass *(*InRegister)(), class UClass *(*InStaticClass)(), const TCHAR* PackageName, const TCHAR* Name, bool bDynamic, const TCHAR* DynamicPackageName = nullptr, const TCHAR* DynamicPathName = nullptr, void (*InInitSearchableValues)(TMap<FName, FName>&) = nullptr)
 	{
 		if (bDynamic)
 		{
 			GetConvertedDynamicPackageNameToTypeName().Add(FName(DynamicPackageName), FName(Name));
 		}
-		UObjectCompiledInDefer(InRegister, InStaticClass, Name, bDynamic, DynamicPathName, InInitSearchableValues);
+		UObjectCompiledInDefer(InRegister, InStaticClass, Name, PackageName, bDynamic, DynamicPathName, InInitSearchableValues);
 	}
 };
 
 /**
  * Stashes the singleton function that builds a compiled in struct (StaticStruct). Later, this is executed.
  */
-COREUOBJECT_API void UObjectCompiledInDeferStruct(class UScriptStruct *(*InRegister)(), const TCHAR* PackageName, const FName PathName, bool bDynamic);
+COREUOBJECT_API void UObjectCompiledInDeferStruct(class UScriptStruct *(*InRegister)(), const TCHAR* PackageName, const FName ObjectName, bool bDynamic, const TCHAR* DynamicPathName);
 
 struct FCompiledInDeferStruct
 {
@@ -340,7 +347,7 @@ struct FCompiledInDeferStruct
 		{
 			GetConvertedDynamicPackageNameToTypeName().Add(FName(DynamicPackageName), FName(Name));
 		}
-		UObjectCompiledInDeferStruct(InRegister, PackageName, DynamicPathName, bDynamic);
+		UObjectCompiledInDeferStruct(InRegister, PackageName, Name, bDynamic, DynamicPathName);
 	}
 };
 
@@ -352,7 +359,7 @@ COREUOBJECT_API class UScriptStruct *GetStaticStruct(class UScriptStruct *(*InRe
 /**
  * Stashes the singleton function that builds a compiled in enum. Later, this is executed.
  */
-COREUOBJECT_API void UObjectCompiledInDeferEnum(class UEnum *(*InRegister)(), const TCHAR* PackageName, const FName PathName, bool bDynamic);
+COREUOBJECT_API void UObjectCompiledInDeferEnum(class UEnum *(*InRegister)(), const TCHAR* PackageName, const FName ObjectName, bool bDynamic, const TCHAR* DynamicPathName);
 
 struct FCompiledInDeferEnum
 {
@@ -362,7 +369,7 @@ struct FCompiledInDeferEnum
 		{
 			GetConvertedDynamicPackageNameToTypeName().Add(FName(DynamicPackageName), FName(Name));
 		}
-		UObjectCompiledInDeferEnum(InRegister, PackageName, DynamicPathName, bDynamic);
+		UObjectCompiledInDeferEnum(InRegister, PackageName, Name, bDynamic, DynamicPathName);
 	}
 };
 
diff --git a/Engine/Source/Runtime/CoreUObject/Public/UObject/UnrealType.h b/Engine/Source/Runtime/CoreUObject/Public/UObject/UnrealType.h
--- a/Engine/Source/Runtime/CoreUObject/Public/UObject/UnrealType.h
+++ b/Engine/Source/Runtime/CoreUObject/Public/UObject/UnrealType.h
@@ -122,6 +122,9 @@ public:
 	static const TCHAR* ImportSingleProperty( const TCHAR* Str, void* DestData, class UStruct* ObjectStruct, UObject* SubobjectOuter, int32 PortFlags,
 											FOutputDevice* Warn, TArray<struct FDefinedProperty>& DefinedProperties );
 
+	/** Gets a redirected property name, will return NAME_None if no redirection was found */
+	static FName FindRedirectedPropertyName(UStruct* ObjectStruct, FName OldName);
+
 	// UHT interface
 	void ExportCppDeclaration(FOutputDevice& Out, EExportedDeclaration::Type DeclarationType, const TCHAR* ArrayDimOverride = NULL, uint32 AdditionalExportCPPFlags = 0
 		, bool bSkipParameterName = false, const FString* ActualCppType = nullptr, const FString* ActualExtendedType = nullptr, const FString* ActualParameterName = nullptr) const;
@@ -4600,24 +4603,6 @@ T* FindFieldChecked( const UStruct* Scope, FName FieldName )
 	return NULL;
 }
 
-/**
- * Dynamically cast a property to the specified type; if the type is a UArrayProperty, will return the UArrayProperty's Inner member, if it is of the correct type.
- */
-template<typename T>
-T* SmartCastProperty( UProperty* Src )
-{
-	T* Result = dynamic_cast<T*>(Src);
-	if ( Result == NULL )
-	{
-		UArrayProperty* ArrayProp = dynamic_cast<UArrayProperty*>(Src);
-		if ( ArrayProp != NULL )
-		{
-			Result = dynamic_cast<T*>(ArrayProp->Inner);
-		}
-	}
-	return Result;
-}
-
 /**
  * Determine if this object has SomeObject in its archetype chain.
  */
diff --git a/Engine/Source/Runtime/Engine/Classes/Components/ArrowComponent.h b/Engine/Source/Runtime/Engine/Classes/Components/ArrowComponent.h
--- a/Engine/Source/Runtime/Engine/Classes/Components/ArrowComponent.h
+++ b/Engine/Source/Runtime/Engine/Classes/Components/ArrowComponent.h
@@ -55,14 +55,9 @@ class ENGINE_API UArrowComponent : public UPrimitiveComponent
 	UPROPERTY(EditAnywhere, BlueprintReadWrite, Category=ArrowComponent)
 	bool bUseInEditorScaling;	
 #endif // WITH_EDITORONLY_DATA
-	// @NOTE!!!!! please remove _NEW when we can remove all _DEPRECATED functions and create redirect in BaseEngine.INI
-	/** Updates the arrow's colour, and tells it to refresh */
-	UFUNCTION(BlueprintCallable, Category="Components|Arrow", meta=(DeprecatedFunction, DeprecationMessage = "Use new SetArrowColor"))
-	virtual void SetArrowColor_DEPRECATED(FColor NewColor);
-
 	/** Updates the arrow's colour, and tells it to refresh */
 	UFUNCTION(BlueprintCallable, DisplayName="SetArrowColor", Category="Components|Arrow")
-	virtual void SetArrowColor_New(FLinearColor NewColor);
+	virtual void SetArrowColor(FLinearColor NewColor);
 
 	//~ Begin UPrimitiveComponent Interface.
 	virtual FPrimitiveSceneProxy* CreateSceneProxy() override;
diff --git a/Engine/Source/Runtime/Engine/Classes/Engine/MemberReference.h b/Engine/Source/Runtime/Engine/Classes/Engine/MemberReference.h
--- a/Engine/Source/Runtime/Engine/Classes/Engine/MemberReference.h
+++ b/Engine/Source/Runtime/Engine/Classes/Engine/MemberReference.h
@@ -340,7 +340,7 @@ public:
 #if WITH_EDITOR
 			if( TargetScope != NULL &&  !GIsSavingPackage )
 			{
-				ReturnField = Cast<TFieldType>(FindRemappedField(TargetScope, MemberName, true));
+				ReturnField = FindRemappedField<TFieldType>(TargetScope, MemberName, true);
 			}
 
 			if(ReturnField != NULL)
@@ -442,25 +442,22 @@ public:
 	/**
 	 * Searches the field redirect map for the specified named field in the scope, and returns the remapped field if found
 	 *
+	 * @param	FieldClass		UClass of field type we are looking for
 	 * @param	InitialScope	The scope the field was initially defined in.  The function will search up into parent scopes to attempt to find remappings
 	 * @param	InitialName		The name of the field to attempt to find a redirector for
 	 * @param	bInitialScopeMustBeOwnerOfField		if true the InitialScope must be Child of the field's owner
 	 * @return	The remapped field, if one exists
 	 */
-	ENGINE_API static UField* FindRemappedField(UClass* InitialScope, FName InitialName, bool bInitialScopeMustBeOwnerOfField = false);
+	ENGINE_API static UField* FindRemappedField(UClass *FieldClass, UClass* InitialScope, FName InitialName, bool bInitialScopeMustBeOwnerOfField = false);
 
-	ENGINE_API static const TMap<FFieldRemapInfo, FFieldRemapInfo>& GetFieldRedirectMap() { InitFieldRedirectMap(); return FieldRedirectMap; }
-	ENGINE_API static const TMultiMap<UClass*, FParamRemapInfo>& GetParamRedirectMap() { InitFieldRedirectMap(); return ParamRedirectMap; }
+	/** Templated version of above, extracts FieldClass and Casts result */
+	template<class TFieldType>
+	static TFieldType* FindRemappedField(UClass* InitialScope, FName InitialName, bool bInitialScopeMustBeOwnerOfField = false)
+	{
+		return Cast<TFieldType>(FindRemappedField(TFieldType::StaticClass(), InitialScope, InitialName, bInitialScopeMustBeOwnerOfField));
+	}
 
 protected:
-	/** 
-	 * A mapping from old property and function names to new ones.  Get primed from INI files, and should contain entries for properties, functions, and delegates that get moved, so they can be fixed up
-	 */
-	static TMap<FFieldRemapInfo, FFieldRemapInfo> FieldRedirectMap;
-	/** 
-	 * A mapping from old pin name to new pin name for each K2 node.  Get primed from INI files, and should contain entries for node class, and old param name and new param name
-	 */
-	ENGINE_API static TMultiMap<UClass*, FParamRemapInfo> ParamRedirectMap;
 
 	/** Has the field map been initialized this run */
 	static bool bFieldRedirectMapInitialized;
@@ -468,15 +465,6 @@ protected:
 	/** Init the field redirect map (if not already done) from .ini file entries */
 	ENGINE_API static void InitFieldRedirectMap();
 
-	/** 
-	 * Searches the field replacement map for an appropriately named field in the specified scope, and returns an updated name if this property has been listed in the remapping table
-	 *
-	 * @param	Class		The class scope to search for a field remap on
-	 * @param	FieldName	The field name (function name or property name) to search for
-	 * @param	RemapInfo	(out) Struct containing updated location info for the field
-	 * @return	Whether or not a remap was found in the specified scope
-	 */
-	static bool FindReplacementFieldName(UClass* Class, FName FieldName, FFieldRemapInfo& RemapInfo);
 #endif
 
 public:
diff --git a/Engine/Source/Runtime/Engine/Classes/Engine/Texture.h b/Engine/Source/Runtime/Engine/Classes/Engine/Texture.h
--- a/Engine/Source/Runtime/Engine/Classes/Engine/Texture.h
+++ b/Engine/Source/Runtime/Engine/Classes/Engine/Texture.h
@@ -726,6 +726,7 @@ public:
 	ENGINE_API virtual void BeginDestroy() override;
 	ENGINE_API virtual bool IsReadyForFinishDestroy() override;
 	ENGINE_API virtual void FinishDestroy() override;
+	ENGINE_API virtual void PostCDOContruct() override;
 #if WITH_EDITORONLY_DATA
 	ENGINE_API virtual void GetAssetRegistryTags(TArray<FAssetRegistryTag>& OutTags) const override;
 #endif
diff --git a/Engine/Source/Runtime/Engine/Classes/GameFramework/DefaultPawn.h b/Engine/Source/Runtime/Engine/Classes/GameFramework/DefaultPawn.h
--- a/Engine/Source/Runtime/Engine/Classes/GameFramework/DefaultPawn.h
+++ b/Engine/Source/Runtime/Engine/Classes/GameFramework/DefaultPawn.h
@@ -78,15 +78,6 @@ class ENGINE_API ADefaultPawn : public APawn
 	UPROPERTY(VisibleAnywhere, BlueprintReadOnly, Category="Pawn")
 	float BaseLookUpRate;
 
-private:
-	/** Input callback on pitch (look up) input. */
-	UFUNCTION(BlueprintCallable, Category="Pawn", meta=(DeprecatedFunction, DeprecationMessage="Use Pawn.AddControllerPitchInput instead."))
-	virtual void LookUp(float Val);
-
-	/** Input callback on yaw (turn) input. */
-	UFUNCTION(BlueprintCallable, Category="Pawn", meta=(DeprecatedFunction, DeprecationMessage="Use Pawn.AddControllerYawInput instead."))
-	virtual void Turn(float Val);
-
 public:
 	/** Name of the MovementComponent.  Use this name if you want to use a different class (with ObjectInitializer.SetDefaultSubobjectClass). */
 	static FName MovementComponentName;
diff --git a/Engine/Source/Runtime/Engine/Classes/Materials/MaterialInterface.h b/Engine/Source/Runtime/Engine/Classes/Materials/MaterialInterface.h
--- a/Engine/Source/Runtime/Engine/Classes/Materials/MaterialInterface.h
+++ b/Engine/Source/Runtime/Engine/Classes/Materials/MaterialInterface.h
@@ -255,6 +255,8 @@ public:
 	ENGINE_API virtual bool IsReadyForFinishDestroy() override;
 	ENGINE_API virtual void PostLoad() override;
 	ENGINE_API virtual void PostDuplicate(bool bDuplicateForPIE) override;
+	ENGINE_API virtual void PostCDOContruct() override;
+
 #if WITH_EDITOR
 	ENGINE_API virtual void PostEditChangeProperty(FPropertyChangedEvent& PropertyChangedEvent) override;
 #endif // WITH_EDITOR
diff --git a/Engine/Source/Runtime/Engine/Classes/MeshVertexPainter/MeshVertexPainter.h b/Engine/Source/Runtime/Engine/Classes/MeshVertexPainter/MeshVertexPainter.h
--- a/Engine/Source/Runtime/Engine/Classes/MeshVertexPainter/MeshVertexPainter.h
+++ b/Engine/Source/Runtime/Engine/Classes/MeshVertexPainter/MeshVertexPainter.h
@@ -4,6 +4,7 @@
 
 #include "CoreMinimal.h"
 #include "UObject/ObjectMacros.h"
+#include "MeshVertexPainter.generated.h"
 
 class UStaticMeshComponent;
 
diff --git a/Engine/Source/Runtime/Engine/Private/AudioDevice.cpp b/Engine/Source/Runtime/Engine/Private/AudioDevice.cpp
--- a/Engine/Source/Runtime/Engine/Private/AudioDevice.cpp
+++ b/Engine/Source/Runtime/Engine/Private/AudioDevice.cpp
@@ -546,7 +546,7 @@ bool FAudioDevice::HandleListSoundClassesCommand(const TCHAR* Cmd, FOutputDevice
 			TotalSounds += ACI->NumRealTime;
 		}
 
-		Ar.Logf(*Line);
+		Ar.Logf(TEXT("%s"), *Line);
 	}
 
 	Ar.Logf(TEXT("%d total sounds in %d classes"), TotalSounds, AudioClassInfos.Num());
@@ -581,7 +581,7 @@ void FAudioDevice::ShowSoundClassHierarchy(FOutputDevice& Ar, USoundClass* InSou
 		}
 		else
 		{
-			Ar.Logf(*SoundClass->GetName());
+			Ar.Logf(TEXT("%s"), *SoundClass->GetName());
 		}
 		for (int i = 0; i < SoundClass->ChildClasses.Num(); ++i)
 		{
diff --git a/Engine/Source/Runtime/Engine/Private/AudioStreaming.cpp b/Engine/Source/Runtime/Engine/Private/AudioStreaming.cpp
--- a/Engine/Source/Runtime/Engine/Private/AudioStreaming.cpp
+++ b/Engine/Source/Runtime/Engine/Private/AudioStreaming.cpp
@@ -73,34 +73,6 @@ FStreamingWaveData::FStreamingWaveData()
 	: SoundWave(NULL)
 	, IORequestHandle(nullptr)
 {
-	if (GNewAsyncIO)
-	{
-		AsyncFileCallBack =
-			[this](bool bWasCancelled, IAsyncReadRequest* Req)
-		{
-			uint8* Mem = Req->GetReadResults();
-			if (Mem)
-			{
-				bool bFound = false;
-				for (FLoadedAudioChunk& LoadedChunk : LoadedChunks)
-				{
-					if (LoadedChunk.IORequest == Req)
-					{
-						check(!LoadedChunk.Data);
-						LoadedChunk.Data = Mem;
-						DEC_MEMORY_STAT_BY(STAT_AsyncFileMemory, LoadedChunk.DataSize);
-						INC_DWORD_STAT_BY(STAT_AudioMemorySize, LoadedChunk.DataSize);
-						INC_DWORD_STAT_BY(STAT_AudioMemory, LoadedChunk.DataSize);
-						bFound = true;
-						break;
-					}
-				}
-				check(bFound);
-			}
-			// else was canceled
-			PendingChunkChangeRequestStatus.Decrement();
-		};
-	}
 }
 
 FStreamingWaveData::~FStreamingWaveData()
@@ -373,6 +345,24 @@ void FStreamingWaveData::BeginPendingRequests(const TArray<uint32>& IndicesToLoa
 						check(IORequestHandle); // this generally cannot fail because it is async
 					}
 					check(Chunk.BulkData.GetBulkDataSize() == ChunkStorage->DataSize);
+
+					FAsyncFileCallBack AsyncFileCallBack =
+						[this, ChunkStorage](bool bWasCancelled, IAsyncReadRequest* Req)
+					{
+						FLoadedAudioChunk& LoadedChunk = *ChunkStorage;
+						uint8* Mem = Req->GetReadResults();
+						if (Mem)
+						{
+							check(!LoadedChunk.Data);
+							LoadedChunk.Data = Mem;
+							DEC_MEMORY_STAT_BY(STAT_AsyncFileMemory, LoadedChunk.DataSize);
+							INC_DWORD_STAT_BY(STAT_AudioMemorySize, LoadedChunk.DataSize);
+							INC_DWORD_STAT_BY(STAT_AudioMemory, LoadedChunk.DataSize);
+						}
+						// else was canceled
+						PendingChunkChangeRequestStatus.Decrement();
+					};
+
 					ChunkStorage->IORequest = IORequestHandle->ReadRequest(Chunk.BulkData.GetBulkDataOffsetInFile(), ChunkStorage->DataSize, AsyncIOPriority, &AsyncFileCallBack);
 					if (!ChunkStorage->IORequest)
 					{
diff --git a/Engine/Source/Runtime/Engine/Private/BlueprintGeneratedClass.cpp b/Engine/Source/Runtime/Engine/Private/BlueprintGeneratedClass.cpp
--- a/Engine/Source/Runtime/Engine/Private/BlueprintGeneratedClass.cpp
+++ b/Engine/Source/Runtime/Engine/Private/BlueprintGeneratedClass.cpp
@@ -594,6 +594,13 @@ UObject* UBlueprintGeneratedClass::FindArchetype(UClass* ArchetypeClass, const F
 			}
 			else if(UInheritableComponentHandler* ICH = Class->GetInheritableComponentHandler())
 			{
+				if (GEventDrivenLoaderEnabled && EVENT_DRIVEN_ASYNC_LOAD_ACTIVE_AT_RUNTIME)
+				{
+					if (ICH->HasAnyFlags(RF_NeedLoad))
+					{
+						UE_LOG(LogClass, Fatal, TEXT("%s had RF_NeedLoad when searching for an archetype of %s named %s"), *GetFullNameSafe(ICH), *GetFullNameSafe(ArchetypeClass), *ArchetypeName.ToString());
+					}
+				}
 				// This would find either an SCS component template override (for which the archetype
 				// name will match the SCS variable name), or an old AddComponent node template override
 				// (for which the archetype name will match the override record's component template name).
diff --git a/Engine/Source/Runtime/Engine/Private/Components/ArrowComponent.cpp b/Engine/Source/Runtime/Engine/Private/Components/ArrowComponent.cpp
--- a/Engine/Source/Runtime/Engine/Private/Components/ArrowComponent.cpp
+++ b/Engine/Source/Runtime/Engine/Private/Components/ArrowComponent.cpp
@@ -340,13 +340,7 @@ FBoxSphereBounds UArrowComponent::CalcBounds(const FTransform& LocalToWorld) con
 	return FBoxSphereBounds(FBox(FVector(0,-ARROW_SCALE,-ARROW_SCALE),FVector(ArrowSize * ARROW_SCALE * 3.0f,ARROW_SCALE,ARROW_SCALE))).TransformBy(LocalToWorld);
 }
 
-void UArrowComponent::SetArrowColor_DEPRECATED(FColor NewColor)
-{
-	ArrowColor = NewColor;
-	MarkRenderStateDirty();
-}
-
-void UArrowComponent::SetArrowColor_New(FLinearColor NewColor)
+void UArrowComponent::SetArrowColor(FLinearColor NewColor)
 {
 	ArrowColor = NewColor.ToFColor(true);
 	MarkRenderStateDirty();
diff --git a/Engine/Source/Runtime/Engine/Private/DefaultPawn.cpp b/Engine/Source/Runtime/Engine/Private/DefaultPawn.cpp
--- a/Engine/Source/Runtime/Engine/Private/DefaultPawn.cpp
+++ b/Engine/Source/Runtime/Engine/Private/DefaultPawn.cpp
@@ -198,26 +198,6 @@ void ADefaultPawn::LookUpAtRate(float Rate)
 	AddControllerPitchInput(Rate * BaseLookUpRate * GetWorld()->GetDeltaSeconds() * CustomTimeDilation);
 }
 
-// @TODO: DEPRECATED, remove.
-void ADefaultPawn::LookUp(float Val)
-{
-	APlayerController* PC = Cast<APlayerController>(GetController());
-	if (PC)
-	{
-		PC->AddPitchInput(Val);
-	}
-}
-
-// @TODO: DEPRECATED, remove.
-void ADefaultPawn::Turn(float Val)
-{
-	APlayerController* PC = Cast<APlayerController>(GetController());
-	if (PC)
-	{
-		PC->AddYawInput(Val);
-	}
-}
-
 UPawnMovementComponent* ADefaultPawn::GetMovementComponent() const
 {
 	return MovementComponent;
diff --git a/Engine/Source/Runtime/Engine/Private/GameInstance.cpp b/Engine/Source/Runtime/Engine/Private/GameInstance.cpp
--- a/Engine/Source/Runtime/Engine/Private/GameInstance.cpp
+++ b/Engine/Source/Runtime/Engine/Private/GameInstance.cpp
@@ -367,6 +367,11 @@ UGameViewportClient* UGameInstance::GetGameViewportClient() const
 	return WC ? WC->GameViewport : nullptr;
 }
 
+// This can be defined in the target.cs file to allow map overrides in shipping builds
+#ifndef UE_ALLOW_MAP_OVERRIDE_IN_SHIPPING
+#define UE_ALLOW_MAP_OVERRIDE_IN_SHIPPING 0
+#endif
+
 void UGameInstance::StartGameInstance()
 {
 	UEngine* const Engine = GetEngine();
@@ -382,7 +387,7 @@ void UGameInstance::StartGameInstance()
 	
 	const TCHAR* Tmp = FCommandLine::Get();
 
-#if UE_BUILD_SHIPPING && !UE_SERVER
+#if UE_BUILD_SHIPPING && !UE_SERVER && !UE_ALLOW_MAP_OVERRIDE_IN_SHIPPING
 	// In shipping don't allow a map override unless on server
 	Tmp = TEXT("");
 #endif // UE_BUILD_SHIPPING && !UE_SERVER
diff --git a/Engine/Source/Runtime/Engine/Private/InterpolateComponentToAction.h b/Engine/Source/Runtime/Engine/Private/InterpolateComponentToAction.h
--- a/Engine/Source/Runtime/Engine/Private/InterpolateComponentToAction.h
+++ b/Engine/Source/Runtime/Engine/Private/InterpolateComponentToAction.h
@@ -153,7 +153,10 @@ public:
 	// Returns a human readable description of the latent operation's current state
 	virtual FString GetDescription() const override
 	{
-		return FString::Printf( *NSLOCTEXT("FInterpolateComponentToAction", "ActionTime", "Move (%.3f seconds left)").ToString(), TotalTime-TimeElapsed);
+		static const FNumberFormattingOptions DelayTimeFormatOptions = FNumberFormattingOptions()
+			.SetMinimumFractionalDigits(3)
+			.SetMaximumFractionalDigits(3);
+		return FText::Format(NSLOCTEXT("FInterpolateComponentToAction", "ActionTimeFmt", "Move ({0} seconds left)"), FText::AsNumber(TotalTime-TimeElapsed, &DelayTimeFormatOptions)).ToString();
 	}
 #endif
 };
diff --git a/Engine/Source/Runtime/Engine/Private/LatentActionManager.cpp b/Engine/Source/Runtime/Engine/Private/LatentActionManager.cpp
--- a/Engine/Source/Runtime/Engine/Private/LatentActionManager.cpp
+++ b/Engine/Source/Runtime/Engine/Private/LatentActionManager.cpp
@@ -234,7 +234,7 @@ FString FLatentActionManager::GetDescription(UObject* InObject, int32 UUID) cons
 			FString ActionDesc = PrimaryAction->GetDescription();
 
 			Description = (PendingActions > 1)
-				? FString::Printf( *NSLOCTEXT("LatentActionManager", "NumPendingActions", "%d Pending Actions: %s").ToString(), PendingActions, *ActionDesc)
+				? FText::Format(NSLOCTEXT("LatentActionManager", "NumPendingActionsFwd", "{0} Pending Actions: {1}"), PendingActions, FText::FromString(ActionDesc)).ToString()
 				: ActionDesc;
 		}
 	}
diff --git a/Engine/Source/Runtime/Engine/Private/LocalPlayer.cpp b/Engine/Source/Runtime/Engine/Private/LocalPlayer.cpp
--- a/Engine/Source/Runtime/Engine/Private/LocalPlayer.cpp
+++ b/Engine/Source/Runtime/Engine/Private/LocalPlayer.cpp
@@ -1454,7 +1454,7 @@ void ULocalPlayer::ExecMacro( const TCHAR* Filename, FOutputDevice& Ar )
 	}
 	else
 	{
-		UE_SUPPRESS(LogExec, Warning, Ar.Logf(*FString::Printf( TEXT("Can't find file '%s'"), Filename) ));
+		UE_SUPPRESS(LogExec, Warning, Ar.Logf(TEXT("Can't find file '%s'"), Filename));
 	}
 }
 
diff --git a/Engine/Source/Runtime/Engine/Private/Materials/Material.cpp b/Engine/Source/Runtime/Engine/Private/Materials/Material.cpp
--- a/Engine/Source/Runtime/Engine/Private/Materials/Material.cpp
+++ b/Engine/Source/Runtime/Engine/Private/Materials/Material.cpp
@@ -367,6 +367,9 @@ void UMaterialInterface::InitDefaultMaterials()
 	if (!bInitialized)
 	{		
 		check(IsInGameThread());
+		static int32 RecursionLevel = 0;
+		RecursionLevel++;
+
 		
 #if WITH_EDITOR
 		GPowerToRoughnessMaterialFunction = LoadObject< UMaterialFunction >( NULL, TEXT("/Engine/Functions/Engine_MaterialFunctions01/Shading/PowerToRoughness.PowerToRoughness"), NULL, LOAD_None, NULL );
@@ -383,18 +386,38 @@ void UMaterialInterface::InitDefaultMaterials()
 			if (GDefaultMaterials[Domain] == NULL)
 			{
 				GDefaultMaterials[Domain] = FindObject<UMaterial>(NULL,GDefaultMaterialNames[Domain]);
-				if (GDefaultMaterials[Domain] == NULL)
+				if (GDefaultMaterials[Domain] == NULL
+#if USE_EVENT_DRIVEN_ASYNC_LOAD_AT_BOOT_TIME
+					&& (RecursionLevel == 1 || !GEventDrivenLoaderEnabled)
+#endif
+					)
 				{
 					GDefaultMaterials[Domain] = LoadObject<UMaterial>(NULL, GDefaultMaterialNames[Domain], NULL, LOAD_DisableDependencyPreloading, NULL);
 					checkf(GDefaultMaterials[Domain] != NULL, TEXT("Cannot load default material '%s'"), GDefaultMaterialNames[Domain]);
 				}
-				GDefaultMaterials[Domain]->AddToRoot();
+				if (GDefaultMaterials[Domain])
+				{
+					GDefaultMaterials[Domain]->AddToRoot();
+				}
 			}
 		}
+		RecursionLevel--;
+#if USE_EVENT_DRIVEN_ASYNC_LOAD_AT_BOOT_TIME
+		bInitialized = !GEventDrivenLoaderEnabled || RecursionLevel == 0;
+#else
 		bInitialized = true;
+#endif
 	}
 }
 
+void UMaterialInterface::PostCDOContruct()
+{
+	UMaterial::StaticClass()->GetDefaultObject();
+	InitializeSharedSamplerStates();
+	UMaterialInterface::InitDefaultMaterials();
+}
+
+
 void UMaterialInterface::PostLoadDefaultMaterials()
 {
 	// Here we prevent this function from being called recursively. Mostly this
@@ -414,11 +437,24 @@ void UMaterialInterface::PostLoadDefaultMaterials()
 		for (int32 Domain = 0; Domain < MD_MAX; ++Domain)
 		{
 			UMaterial* Material = GDefaultMaterials[Domain];
+#if USE_EVENT_DRIVEN_ASYNC_LOAD_AT_BOOT_TIME
+			check(Material || (GIsInitialLoad && GEventDrivenLoaderEnabled));
+			check((GIsInitialLoad && GEventDrivenLoaderEnabled) || !Material->HasAnyFlags(RF_NeedLoad));
+			if (Material && !Material->HasAnyFlags(RF_NeedLoad))
+#else
 			check(Material);
-			Material->ConditionalPostLoad();
-			// Sometimes the above will get called before the material has been fully serialized
-			// in this case its NeedPostLoad flag will not be cleared.
-			if (Material->HasAnyFlags(RF_NeedPostLoad))
+			if (Material)
+#endif
+			{
+				Material->ConditionalPostLoad();
+				// Sometimes the above will get called before the material has been fully serialized
+				// in this case its NeedPostLoad flag will not be cleared.
+				if (Material->HasAnyFlags(RF_NeedPostLoad))
+				{
+					bPostLoaded = false;
+				}
+			}
+			else
 			{
 				bPostLoaded = false;
 			}
@@ -428,18 +464,28 @@ void UMaterialInterface::PostLoadDefaultMaterials()
 
 void UMaterialInterface::AssertDefaultMaterialsExist()
 {
-	for (int32 Domain = 0; Domain < MD_MAX; ++Domain)
+#if (USE_EVENT_DRIVEN_ASYNC_LOAD_AT_BOOT_TIME)
+	if (!GIsInitialLoad || !GEventDrivenLoaderEnabled)
+#endif
 	{
-		check(GDefaultMaterials[Domain] != NULL);
+		for (int32 Domain = 0; Domain < MD_MAX; ++Domain)
+		{
+			check(GDefaultMaterials[Domain] != NULL);
+		}
 	}
 }
 
 void UMaterialInterface::AssertDefaultMaterialsPostLoaded()
 {
-	for (int32 Domain = 0; Domain < MD_MAX; ++Domain)
+#if (USE_EVENT_DRIVEN_ASYNC_LOAD_AT_BOOT_TIME)
+	if (!GIsInitialLoad || !GEventDrivenLoaderEnabled)
+#endif
 	{
-		check(GDefaultMaterials[Domain] != NULL);
-		check(!GDefaultMaterials[Domain]->HasAnyFlags(RF_NeedPostLoad));
+		for (int32 Domain = 0; Domain < MD_MAX; ++Domain)
+		{
+			check(GDefaultMaterials[Domain] != NULL);
+			check(!GDefaultMaterials[Domain]->HasAnyFlags(RF_NeedPostLoad));
+		}
 	}
 }
 
diff --git a/Engine/Source/Runtime/Engine/Private/Materials/MaterialInterface.cpp b/Engine/Source/Runtime/Engine/Private/Materials/MaterialInterface.cpp
--- a/Engine/Source/Runtime/Engine/Private/Materials/MaterialInterface.cpp
+++ b/Engine/Source/Runtime/Engine/Private/Materials/MaterialInterface.cpp
@@ -47,8 +47,13 @@ UMaterialInterface::UMaterialInterface(const FObjectInitializer& ObjectInitializ
 {
 	if (!HasAnyFlags(RF_ClassDefaultObject))
 	{
-		InitDefaultMaterials();
-		AssertDefaultMaterialsExist();
+#if USE_EVENT_DRIVEN_ASYNC_LOAD_AT_BOOT_TIME
+		if (!GIsInitialLoad || !GEventDrivenLoaderEnabled)
+#endif
+		{
+			InitDefaultMaterials();
+			AssertDefaultMaterialsExist();
+		}
 
 		if (SamplerTypeEnum == nullptr)
 		{
@@ -63,7 +68,12 @@ UMaterialInterface::UMaterialInterface(const FObjectInitializer& ObjectInitializ
 void UMaterialInterface::PostLoad()
 {
 	Super::PostLoad();
-	PostLoadDefaultMaterials();
+#if USE_EVENT_DRIVEN_ASYNC_LOAD_AT_BOOT_TIME
+	if (!GEventDrivenLoaderEnabled)
+#endif
+	{
+		PostLoadDefaultMaterials();
+	}
 }
 
 void UMaterialInterface::GetUsedTexturesAndIndices(TArray<UTexture*>& OutTextures, TArray< TArray<int32> >& OutIndices, EMaterialQualityLevel::Type QualityLevel, ERHIFeatureLevel::Type FeatureLevel) const
diff --git a/Engine/Source/Runtime/Engine/Private/MatineeUtils.cpp b/Engine/Source/Runtime/Engine/Private/MatineeUtils.cpp
--- a/Engine/Source/Runtime/Engine/Private/MatineeUtils.cpp
+++ b/Engine/Source/Runtime/Engine/Private/MatineeUtils.cpp
@@ -353,8 +353,6 @@ namespace FMatineeUtils
 				{
 					// Now look for a component
 					const FName TrialCompName(*CompString);
-					const FName RemappedCompName = FLinkerLoad::FindSubobjectRedirectName(TrialCompName);
-					const FName CompName = (RemappedCompName != NAME_None) ? RemappedCompName : TrialCompName;
 
 					TArray<UObject*> Components;
 					InObject->CollectDefaultSubobjects(Components, false);
@@ -362,6 +360,9 @@ namespace FMatineeUtils
 					{
 						UObject* Component = Components[ComponentIndex];
 
+						const FName RemappedCompName = FLinkerLoad::FindSubobjectRedirectName(TrialCompName, Component->GetClass());
+						const FName CompName = (RemappedCompName != NAME_None) ? RemappedCompName : TrialCompName;
+
 						if (Component->GetFName() == CompName)
 						{
 							PerformQuery(Component, (void*)Component, Component->GetClass(), PropString);
diff --git a/Engine/Source/Runtime/Engine/Private/MemberReference.cpp b/Engine/Source/Runtime/Engine/Private/MemberReference.cpp
--- a/Engine/Source/Runtime/Engine/Private/MemberReference.cpp
+++ b/Engine/Source/Runtime/Engine/Private/MemberReference.cpp
@@ -2,6 +2,7 @@
 
 #include "Engine/MemberReference.h"
 #include "Misc/ConfigCacheIni.h"
+#include "UObject/CoreRedirects.h"
 
 #if WITH_EDITOR
 #include "Kismet2/BlueprintEditorUtils.h"
@@ -170,41 +171,17 @@ FName FMemberReference::RefreshLocalVariableName(UClass* InSelfScope) const
 	return RenamedMemberName;
 }
 
-TMap<FFieldRemapInfo, FFieldRemapInfo> FMemberReference::FieldRedirectMap;
-TMultiMap<UClass*, FParamRemapInfo> FMemberReference::ParamRedirectMap;
-
 bool FMemberReference::bFieldRedirectMapInitialized = false;
 
-bool FMemberReference::FindReplacementFieldName(UClass* Class, FName FieldName, FFieldRemapInfo& RemapInfo)
-{
-	check(Class != NULL);
-	InitFieldRedirectMap();
-
-	// Reset the property remap info
-	RemapInfo = FFieldRemapInfo();
-
-	FFieldRemapInfo OldField;
-	OldField.FieldClass = Class->GetFName();
-	OldField.FieldName = FieldName;
-
-	FFieldRemapInfo* NewFieldInfoPtr = FieldRedirectMap.Find(OldField);
-	if (NewFieldInfoPtr != NULL)
-	{
-		RemapInfo = *NewFieldInfoPtr;
-		return true;
-	}
-	else
-	{
-		return false;
-	}
-}
-
 void FMemberReference::InitFieldRedirectMap()
 {
+	// Soft deprecated, replaced by FCoreRedirects but will read the old ini format for the foreseeable future
 	if (!bFieldRedirectMapInitialized)
 	{
 		if (GConfig)
 		{
+			TArray<FCoreRedirect> NewRedirects;
+
 			FConfigSection* PackageRedirects = GConfig->GetSectionPrivate( TEXT("/Script/Engine.Engine"), false, true, GEngineIni );
 			for (FConfigSection::TIterator It(*PackageRedirects); It; ++It)
 			{
@@ -216,76 +193,26 @@ void FMemberReference::InitFieldRedirectMap()
 					FParse::Value( *It.Value().GetValue(), TEXT("OldFieldName="), OldFieldPathString );
 					FParse::Value( *It.Value().GetValue(), TEXT("NewFieldName="), NewFieldPathString );
 
-					// Handle both cases of just a field being renamed (just one FName), as well as a class and field name (ClassName.FieldName)
-					FFieldRemapInfo OldFieldRemap;
-					{
-						TArray<FString> OldFieldPath;
-						OldFieldPathString.ParseIntoArray(OldFieldPath, TEXT("."), true);
-
-						if (OldFieldPath.Num() == 1)
-						{
-							// Only the new property name is specified
-							OldFieldRemap.FieldName = FName(*OldFieldPath[0]);
-						}
-						else if (OldFieldPath.Num() == 2)
-						{
-							// Property name and new class are specified
-							OldFieldRemap.FieldClass = FName(*OldFieldPath[0]);
-							OldFieldRemap.FieldName = FName(*OldFieldPath[1]);
-						}
-					}
-
-					// Handle both cases of just a field being renamed (just one FName), as well as a class and field name (ClassName.FieldName)
-					FFieldRemapInfo NewFieldRemap;
-					{
-						TArray<FString> NewFieldPath;
-						NewFieldPathString.ParseIntoArray(NewFieldPath, TEXT("."), true);
-
-						if( NewFieldPath.Num() == 1 )
-						{
-							// Only the new property name is specified
-							NewFieldRemap.FieldName = FName(*NewFieldPath[0]);
-						}
-						else if( NewFieldPath.Num() == 2 )
-						{
-							// Property name and new class are specified
-							NewFieldRemap.FieldClass = FName(*NewFieldPath[0]);
-							NewFieldRemap.FieldName = FName(*NewFieldPath[1]);
-						}
-					}
-
-					FieldRedirectMap.Add(OldFieldRemap, NewFieldRemap);
+					// Add both a Property and Function redirect, as we don't know which it's trying to be with the old syntax
+					FCoreRedirect* PropertyRedirect = new (NewRedirects) FCoreRedirect(ECoreRedirectFlags::Type_Property, OldFieldPathString, NewFieldPathString);
+					FCoreRedirect* FunctionRedirect = new (NewRedirects) FCoreRedirect(ECoreRedirectFlags::Type_Function, OldFieldPathString, NewFieldPathString);
 				}			
 				if (It.Key() == TEXT("K2ParamRedirects"))
 				{
-					FName NodeName = NAME_None;
+					// Ignore NodeName/Title as it's not useful
 					FName OldParam = NAME_None;
 					FName NewParam = NAME_None;
-					FName NodeTitle = NAME_None;
 
 					FString OldParamValues;
 					FString NewParamValues;
 					FString CustomValueMapping;
 
-					FParse::Value( *It.Value().GetValue(), TEXT("NodeName="), NodeName );
-					FParse::Value( *It.Value().GetValue(), TEXT("NodeTitle="), NodeTitle );
 					FParse::Value( *It.Value().GetValue(), TEXT("OldParamName="), OldParam );
 					FParse::Value( *It.Value().GetValue(), TEXT("NewParamName="), NewParam );
 					FParse::Value( *It.Value().GetValue(), TEXT("OldParamValues="), OldParamValues );
 					FParse::Value( *It.Value().GetValue(), TEXT("NewParamValues="), NewParamValues );
 					FParse::Value( *It.Value().GetValue(), TEXT("CustomValueMapping="), CustomValueMapping );
 
-					FParamRemapInfo ParamRemap;
-					ParamRemap.NewParam = NewParam;
-					ParamRemap.OldParam = OldParam;
-					ParamRemap.NodeTitle = NodeTitle;
-					ParamRemap.bCustomValueMapping = (FCString::Stricmp(*CustomValueMapping,TEXT("true")) == 0);
-
-					if (CustomValueMapping.Len() > 0 && !ParamRemap.bCustomValueMapping)
-					{
-						UE_LOG(LogBlueprint, Warning, TEXT("Value other than true specified for CustomValueMapping for '%s' node param redirect '%s' to '%s'."), *(NodeName.ToString()), *(OldParam.ToString()), *(NewParam.ToString()));
-					}
-
 					TArray<FString> OldParamValuesList;
 					TArray<FString> NewParamValuesList;
 					OldParamValues.ParseIntoArray(OldParamValuesList, TEXT(";"), false);
@@ -293,98 +220,110 @@ void FMemberReference::InitFieldRedirectMap()
 
 					if (OldParamValuesList.Num() != NewParamValuesList.Num())
 					{
-						UE_LOG(LogBlueprint, Warning, TEXT("Unequal lengths for old and new param values for '%s' node param redirect '%s' to '%s'."), *(NodeName.ToString()), *(OldParam.ToString()), *(NewParam.ToString()));
+						UE_LOG(LogBlueprint, Warning, TEXT("Unequal lengths for old and new param values for  param redirect '%s' to '%s'."), *(OldParam.ToString()), *(NewParam.ToString()));
 					}
 
 					if (CustomValueMapping.Len() > 0 && (OldParamValuesList.Num() > 0 || NewParamValuesList.Num() > 0))
 					{
-						UE_LOG(LogBlueprint, Warning, TEXT("Both Custom and Automatic param value remapping specified for '%s' node param redirect '%s' to '%s'.  Only Custom will be applied."), *(NodeName.ToString()), *(OldParam.ToString()), *(NewParam.ToString()));
+						UE_LOG(LogBlueprint, Warning, TEXT("Both Custom and Automatic param value remapping specified for param redirect '%s' to '%s'.  Only Custom will be applied."), *(OldParam.ToString()), *(NewParam.ToString()));
 					}
 
-					for (int32 i = FMath::Min(OldParamValuesList.Num(), NewParamValuesList.Num()) - 1; i >= 0; --i)
-					{
-						int32 CurSize = ParamRemap.ParamValueMap.Num();
-						ParamRemap.ParamValueMap.Add(OldParamValuesList[i], NewParamValuesList[i]);
-						if (CurSize == ParamRemap.ParamValueMap.Num())
-						{
-							UE_LOG(LogBlueprint, Warning, TEXT("Duplicate old param value '%s' for '%s' node param redirect '%s' to '%s'."), *(OldParamValuesList[i]), *(NodeName.ToString()), *(OldParam.ToString()), *(NewParam.ToString()));
-						}
-					}
+					FCoreRedirect* Redirect = new (NewRedirects) FCoreRedirect(ECoreRedirectFlags::Type_Property, OldParam.ToString(), NewParam.ToString());
 
-					// load class
-					static UClass* K2NodeClass = StaticLoadClass(UObject::StaticClass(), nullptr, TEXT("BlueprintGraph.K2Node"));
-					if (K2NodeClass)
+					for (int32 i = FMath::Min(OldParamValuesList.Num(), NewParamValuesList.Num()) - 1; i >= 0; --i)
 					{
-						UClass* NodeClass = StaticLoadClass(K2NodeClass, nullptr, *NodeName.ToString());
-						if (NodeClass )
+						int32 CurSize = Redirect->ValueChanges.Num();
+						Redirect->ValueChanges.Add(OldParamValuesList[i], NewParamValuesList[i]);
+						if (CurSize == Redirect->ValueChanges.Num())
 						{
-							ParamRedirectMap.Add(NodeClass, ParamRemap);
+							UE_LOG(LogBlueprint, Warning, TEXT("Duplicate old param value '%s' for param redirect '%s' to '%s'."), *(OldParamValuesList[i]), *(OldParam.ToString()), *(NewParam.ToString()));
 						}
 					}
 				}			
 			}
 
+			FCoreRedirects::AddRedirectList(NewRedirects, GEngineIni);
 			bFieldRedirectMapInitialized = true;
 		}
 	}
 }
 
-UField* FMemberReference::FindRemappedField(UClass* InitialScope, FName InitialName, bool bInitialScopeMustBeOwnerOfField)
+UField* FMemberReference::FindRemappedField(UClass *FieldClass, UClass* InitialScope, FName InitialName, bool bInitialScopeMustBeOwnerOfField)
 {
-	FFieldRemapInfo NewFieldInfo;
+	DECLARE_SCOPE_CYCLE_COUNTER(TEXT("FMemberReference::FindRemappedField"), STAT_LinkerLoad_FindRemappedField, STATGROUP_LoadTimeVerbose);
+
+	InitFieldRedirectMap();
+
+	// In the case of a bifurcation of a variable (e.g. moved from a parent into certain children), verify that we don't also define the variable in the current scope first
+	if (FindField<UField>(InitialScope, InitialName) != nullptr)
+	{
+		return nullptr;
+	}
 
-	bool bFoundReplacement = false;
-	
 	// Step up the class chain to check if us or any of our parents specify a redirect
 	UClass* TestRemapClass = InitialScope;
-	while( TestRemapClass != NULL )
+	while (TestRemapClass != nullptr)
 	{
-		if( FindReplacementFieldName(TestRemapClass, InitialName, NewFieldInfo) )
+		UClass* SearchClass = TestRemapClass;
+
+		FName NewFieldName;
+
+		FCoreRedirectObjectName OldRedirectName = FCoreRedirectObjectName(InitialName, TestRemapClass->GetFName(), TestRemapClass->GetOutermost()->GetFName());
+		FCoreRedirectObjectName NewRedirectName = FCoreRedirects::GetRedirectedName(FCoreRedirects::GetFlagsForTypeClass(FieldClass), OldRedirectName);
+
+		if (NewRedirectName != OldRedirectName)
 		{
-			// Found it, stop our search
-			bFoundReplacement = true;
-			break;
-		}
+			NewFieldName = NewRedirectName.ObjectName;
 
-		TestRemapClass = TestRemapClass->GetSuperClass();
-	}
+			if (OldRedirectName.OuterName != NewRedirectName.OuterName)
+			{
+				// Try remapping class, this only works if class is in memory
+				FString ClassName = NewRedirectName.OuterName.ToString();
 
-	// In the case of a bifurcation of a variable (e.g. moved from a parent into certain children), verify that we don't also define the variable in the current scope first
-	if( bFoundReplacement && (FindField<UField>(InitialScope, InitialName) != nullptr))
-	{
-		bFoundReplacement = false;		
-	}
+				if (NewRedirectName.PackageName != NAME_Name)
+				{
+					// Use package if it's there
+					ClassName = FString::Printf(TEXT("%s.%s"), *NewRedirectName.PackageName.ToString(), *NewRedirectName.OuterName.ToString());
+				}
 
-	if( bFoundReplacement )
-	{
-		const FName NewFieldName = NewFieldInfo.FieldName;
-		UClass* SearchClass = (NewFieldInfo.FieldClass != NAME_None) ? (UClass*)StaticFindObject(UClass::StaticClass(), ANY_PACKAGE, *NewFieldInfo.FieldClass.ToString()) : (UClass*)TestRemapClass;
+				SearchClass = (UClass*)StaticFindObject(UClass::StaticClass(), ANY_PACKAGE, *ClassName);
 
-		// Find the actual field specified by the redirector, so we can return it and update the node that uses it
-		UField* NewField = FindField<UField>(SearchClass, NewFieldInfo.FieldName);
-		if( NewField != NULL )
+				if (!SearchClass)
+				{
+					UE_LOG(LogBlueprint, Log, TEXT("UK2Node:  Unable to find updated field name for '%s' on unknown class '%s'."), *InitialName.ToString(), *ClassName);
+					return nullptr;
+				}
+			}
+		}
+
+		if (NewFieldName != NAME_None)
 		{
-			if (bInitialScopeMustBeOwnerOfField && !InitialScope->IsChildOf(SearchClass))
+			// Find the actual field specified by the redirector, so we can return it and update the node that uses it
+			UField* NewField = FindField<UField>(SearchClass, NewFieldName);
+			if (NewField != nullptr)
 			{
-				UE_LOG(LogBlueprint, Log, TEXT("UK2Node:  Unable to update field. Remapped field '%s' in not owned by given scope. Scope: '%s', Owner: '%s'."), *InitialName.ToString(), *InitialScope->GetName(), *NewFieldInfo.FieldClass.ToString());
+				if (bInitialScopeMustBeOwnerOfField && !InitialScope->IsChildOf(SearchClass))
+				{
+					UE_LOG(LogBlueprint, Log, TEXT("UK2Node:  Unable to update field. Remapped field '%s' in not owned by given scope. Scope: '%s', Owner: '%s'."), *InitialName.ToString(), *InitialScope->GetName(), *NewFieldName.ToString());
+				}
+				else
+				{
+					UE_LOG(LogBlueprint, Log, TEXT("UK2Node:  Fixed up old field '%s' to new name '%s' on class '%s'."), *InitialName.ToString(), *NewFieldName.ToString(), *SearchClass->GetName());
+					return NewField;
+				}
 			}
 			else
 			{
-				UE_LOG(LogBlueprint, Log, TEXT("UK2Node:  Fixed up old field '%s' to new name '%s' on class '%s'."), *InitialName.ToString(), *NewFieldInfo.FieldName.ToString(), *SearchClass->GetName());
-				return NewField;
+				UE_LOG(LogBlueprint, Log, TEXT("UK2Node:  Unable to find updated field name for '%s' on class '%s'."), *InitialName.ToString(), *SearchClass->GetName());
 			}
+
+			return nullptr;
 		}
-		else if (SearchClass != NULL)
-		{
-			UE_LOG(LogBlueprint, Log, TEXT("UK2Node:  Unable to find updated field name for '%s' on class '%s'."), *InitialName.ToString(), *SearchClass->GetName());
-		}
-		else
-		{
-			UE_LOG(LogBlueprint, Log, TEXT("UK2Node:  Unable to find updated field name for '%s' on unknown class '%s'."), *InitialName.ToString(), *NewFieldInfo.FieldClass.ToString());
-		}
+
+		TestRemapClass = TestRemapClass->GetSuperClass();
 	}
 
-	return NULL;
+	return nullptr;
 }
 
 #endif
diff --git a/Engine/Source/Runtime/Engine/Private/Model.cpp b/Engine/Source/Runtime/Engine/Private/Model.cpp
--- a/Engine/Source/Runtime/Engine/Private/Model.cpp
+++ b/Engine/Source/Runtime/Engine/Private/Model.cpp
@@ -515,7 +515,7 @@ void UModel::GetResourceSizeEx(FResourceSizeEx& CumulativeResourceSize)
 
 #if WITH_EDITOR
 
-IMPLEMENT_INTRINSIC_CLASS(UModel, ENGINE_API, UObject, CORE_API, 
+IMPLEMENT_INTRINSIC_CLASS(UModel, ENGINE_API, UObject, CORE_API, "/Script/Engine",
 	{
 		Class->ClassAddReferencedObjects = &UModel::AddReferencedObjects;
 		Class->EmitObjectReference(STRUCT_OFFSET(UModel, Polys), TEXT("Polys"));
@@ -528,7 +528,7 @@ IMPLEMENT_INTRINSIC_CLASS(UModel, ENGINE_API, UObject, CORE_API,
 
 #else
 
-IMPLEMENT_INTRINSIC_CLASS(UModel, ENGINE_API, UObject, CORE_API, 
+IMPLEMENT_INTRINSIC_CLASS(UModel, ENGINE_API, UObject, CORE_API, "/Script/Engine",
 	{
 		Class->ClassAddReferencedObjects = &UModel::AddReferencedObjects;
 		const uint32 SkipIndexIndex = Class->EmitStructArrayBegin(STRUCT_OFFSET(UModel, Surfs), TEXT("Surfs"), sizeof(FBspSurf));
diff --git a/Engine/Source/Runtime/Engine/Private/PackageMapClient.cpp b/Engine/Source/Runtime/Engine/Private/PackageMapClient.cpp
--- a/Engine/Source/Runtime/Engine/Private/PackageMapClient.cpp
+++ b/Engine/Source/Runtime/Engine/Private/PackageMapClient.cpp
@@ -1624,7 +1624,7 @@ void UPackageMapClient::LogDebugInfo( FOutputDevice & Ar )
 
 		UObject *Obj = It.Key().Get();
 		FString Str = FString::Printf(TEXT("%s [%s] [%s] - %s"), *NetGUID.ToString(), *Status, NetGUID.IsDynamic() ? TEXT("Dynamic") : TEXT("Static") , Obj ? *Obj->GetPathName() : TEXT("NULL"));
-		Ar.Logf(*Str);
+		Ar.Logf(TEXT("%s"), *Str);
 		UE_LOG(LogNetPackageMap, Log, TEXT("%s"), *Str);
 	}
 }
diff --git a/Engine/Source/Runtime/Engine/Private/Rendering/SkinWeightVertexBuffer.cpp b/Engine/Source/Runtime/Engine/Private/Rendering/SkinWeightVertexBuffer.cpp
--- a/Engine/Source/Runtime/Engine/Private/Rendering/SkinWeightVertexBuffer.cpp
+++ b/Engine/Source/Runtime/Engine/Private/Rendering/SkinWeightVertexBuffer.cpp
@@ -106,9 +106,12 @@ FArchive& operator<<(FArchive& Ar, FSkinWeightVertexBuffer& VertexBuffer)
 		{
 			VertexBuffer.WeightData->Serialize(Ar);
 
-			// update cached buffer info
-			VertexBuffer.Data = (VertexBuffer.NumVertices > 0) ? VertexBuffer.WeightData->GetDataPointer() : nullptr;
-			VertexBuffer.Stride = VertexBuffer.WeightData->GetStride();
+			if (!Ar.IsCountingMemory())
+			{
+				// update cached buffer info
+				VertexBuffer.Data = (VertexBuffer.NumVertices > 0) ? VertexBuffer.WeightData->GetDataPointer() : nullptr;
+				VertexBuffer.Stride = VertexBuffer.WeightData->GetStride();
+			}
 		}
 	}
 
diff --git a/Engine/Source/Runtime/Engine/Private/SoundWave.cpp b/Engine/Source/Runtime/Engine/Private/SoundWave.cpp
--- a/Engine/Source/Runtime/Engine/Private/SoundWave.cpp
+++ b/Engine/Source/Runtime/Engine/Private/SoundWave.cpp
@@ -280,7 +280,7 @@ void USoundWave::LogSubtitle( FOutputDevice& Ar )
 #if WITH_EDITORONLY_DATA
 	Ar.Logf( TEXT( "Comment:   %s" ), *Comment );
 #endif // WITH_EDITORONLY_DATA
-	Ar.Logf( bMature ? TEXT( "Mature:    Yes" ) : TEXT( "Mature:    No" ) );
+	Ar.Logf( TEXT("Mature:    %s"), bMature ? TEXT( "Yes" ) : TEXT( "No" ) );
 }
 
 float USoundWave::GetSubtitlePriority() const
diff --git a/Engine/Source/Runtime/Engine/Private/StreamableManager.cpp b/Engine/Source/Runtime/Engine/Private/StreamableManager.cpp
--- a/Engine/Source/Runtime/Engine/Private/StreamableManager.cpp
+++ b/Engine/Source/Runtime/Engine/Private/StreamableManager.cpp
@@ -13,6 +13,7 @@ struct FStreamableRequest
 	FStreamableDelegate CompletionDelegate;
 	/** How many FStreamables are referenced by this request */
 	int32 StreamablesReferenced;
+	TArray<UObject*> FinishedTargets;
 
 	FStreamableRequest()
 		: StreamablesReferenced(0)
@@ -408,18 +409,27 @@ void FStreamableManager::CheckCompletedRequests(FStringAssetReference const& Tar
 	// If we're related to any requests we should start an unload after this
 	bool bShouldStartUnload = false;
 
-	for (int32 i = 0; i < Existing->RelatedRequests.Num(); i++)
+	// Make a copy of related requests to avoid possible recursion issues
+	// When this array goes out of scope, the FStreamableRequests may get deleted
+	TArray< TSharedRef< FStreamableRequest> > LocalRelatedRequests = Existing->RelatedRequests;
+	Existing->RelatedRequests.Empty();
+
+	for (int32 i = 0; i < LocalRelatedRequests.Num(); i++)
 	{
+		TSharedRef<FStreamableRequest> RelatedRequest = LocalRelatedRequests[i];
 		bShouldStartUnload = true;
 		// Decrement related requests, and call delegate if all are done
-		Existing->RelatedRequests[i]->StreamablesReferenced--;
-		if (Existing->RelatedRequests[i]->StreamablesReferenced == 0)
+		RelatedRequest->StreamablesReferenced--;
+		if (RelatedRequest->StreamablesReferenced == 0)
+		{
+			RelatedRequest->FinishedTargets.Empty();
+			RelatedRequest->CompletionDelegate.ExecuteIfBound();
+		}
+		else if (Existing->Target)
 		{
-			Existing->RelatedRequests[i]->CompletionDelegate.ExecuteIfBound();
+			RelatedRequest->FinishedTargets.Add(Existing->Target);
 		}
 	}
-	// This may result in request objects being deleted due to no references
-	Existing->RelatedRequests.Empty();
 
 	if (bShouldStartUnload)
 	{
@@ -477,6 +487,13 @@ void FStreamableManager::AddStructReferencedObjects(class FReferenceCollector& C
 		{
 			Collector.AddReferencedObject(Existing->Target);
 		}
+		for (int32 i = 0; i < Existing->RelatedRequests.Num(); i++)
+		{
+			for (UObject* Item : Existing->RelatedRequests[i]->FinishedTargets)
+			{
+				Collector.AddReferencedObject(Item);
+			}				
+		}
 	}
 }
 
diff --git a/Engine/Source/Runtime/Engine/Private/Texture.cpp b/Engine/Source/Runtime/Engine/Private/Texture.cpp
--- a/Engine/Source/Runtime/Engine/Private/Texture.cpp
+++ b/Engine/Source/Runtime/Engine/Private/Texture.cpp
@@ -500,6 +500,12 @@ UEnum* UTexture::GetPixelFormatEnum()
 	return PixelFormatEnum;
 }
 
+void UTexture::PostCDOContruct()
+{
+	GetPixelFormatEnum();
+}
+
+
 bool UTexture::ForceUpdateTextureStreaming()
 {
 	if (!IStreamingManager::HasShutdown())
diff --git a/Engine/Source/Runtime/Engine/Private/UnrealEngine.cpp b/Engine/Source/Runtime/Engine/Private/UnrealEngine.cpp
--- a/Engine/Source/Runtime/Engine/Private/UnrealEngine.cpp
+++ b/Engine/Source/Runtime/Engine/Private/UnrealEngine.cpp
@@ -3059,7 +3059,7 @@ bool UEngine::HandleGameVerCommand( const TCHAR* Cmd, FOutputDevice& Ar )
 	FString VersionString = FString::Printf( TEXT( "GameVersion Branch: %s, Configuration: %s, Build: %s, CommandLine: %s" ),
 											 *FApp::GetBranchName(), EBuildConfigurations::ToString( FApp::GetBuildConfiguration() ), FApp::GetBuildVersion(), FCommandLine::Get() );
 
-	Ar.Logf( *VersionString );
+	Ar.Logf( TEXT("%s"), *VersionString );
 	FPlatformMisc::ClipboardCopy( *VersionString );
 
 	if (FCString::Stristr(Cmd, TEXT("-display")))
@@ -4274,7 +4274,7 @@ bool UEngine::HandleMemReportDeferredCommand( const TCHAR* Cmd, FOutputDevice& A
 		UE_LOG(LogEngine, Log, TEXT("MemReportDeferred: saving to %s"), *FilenameFull);		
 	}
 
-	ReportAr->Logf( *FString::Printf( TEXT( "CommandLine Options: %s" ) LINE_TERMINATOR, FCommandLine::Get() ) );
+	ReportAr->Logf( TEXT( "CommandLine Options: %s" ) LINE_TERMINATOR, FCommandLine::Get() );
 
 	// Run commands from the ini
 	FConfigSection* CommandsToRun = GConfig->GetSectionPrivate(TEXT("MemReportCommands"), 0, 1, GEngineIni);
diff --git a/Engine/Source/Runtime/Engine/Private/VisualLogger/VisualLoggerKismetLibrary.cpp b/Engine/Source/Runtime/Engine/Private/VisualLogger/VisualLoggerKismetLibrary.cpp
--- a/Engine/Source/Runtime/Engine/Private/VisualLogger/VisualLoggerKismetLibrary.cpp
+++ b/Engine/Source/Runtime/Engine/Private/VisualLogger/VisualLoggerKismetLibrary.cpp
@@ -13,7 +13,7 @@ void UVisualLoggerKismetLibrary::LogText(UObject* WorldContextObject, FString Te
 {
 #if ENABLE_VISUAL_LOG
 	const ELogVerbosity::Type DefaultVerbosity = ELogVerbosity::Log;
-	FVisualLogger::CategorizedLogf(WorldContextObject, FLogCategoryBase(*CategoryName.ToString(), DefaultVerbosity, DefaultVerbosity), DefaultVerbosity, *Text);
+	FVisualLogger::CategorizedLogf(WorldContextObject, FLogCategoryBase(*CategoryName.ToString(), DefaultVerbosity, DefaultVerbosity), DefaultVerbosity, TEXT("%s"), *Text);
 #endif
 }
 
@@ -21,7 +21,7 @@ void UVisualLoggerKismetLibrary::LogLocation(UObject* WorldContextObject, FVecto
 {
 #if ENABLE_VISUAL_LOG
 	const ELogVerbosity::Type DefaultVerbosity = ELogVerbosity::Log;
-	FVisualLogger::GeometryShapeLogf(WorldContextObject, FLogCategoryBase(*CategoryName.ToString(), DefaultVerbosity, DefaultVerbosity), DefaultVerbosity, Location, Radius, Color.ToFColor(true), *Text);
+	FVisualLogger::GeometryShapeLogf(WorldContextObject, FLogCategoryBase(*CategoryName.ToString(), DefaultVerbosity, DefaultVerbosity), DefaultVerbosity, Location, Radius, Color.ToFColor(true), TEXT("%s"), *Text);
 #endif
 }
 
@@ -29,6 +29,6 @@ void UVisualLoggerKismetLibrary::LogBox(UObject* WorldContextObject, FBox Box, F
 {
 #if ENABLE_VISUAL_LOG
 	const ELogVerbosity::Type DefaultVerbosity = ELogVerbosity::Log;
-	FVisualLogger::GeometryShapeLogf(WorldContextObject, FLogCategoryBase(*CategoryName.ToString(), DefaultVerbosity, DefaultVerbosity), DefaultVerbosity, Box, FMatrix::Identity, ObjectColor.ToFColor(true), *Text);
+	FVisualLogger::GeometryShapeLogf(WorldContextObject, FLogCategoryBase(*CategoryName.ToString(), DefaultVerbosity, DefaultVerbosity), DefaultVerbosity, Box, FMatrix::Identity, ObjectColor.ToFColor(true), TEXT("%s"), *Text);
 #endif
 }
diff --git a/Engine/Source/Runtime/Engine/Public/AudioStreaming.h b/Engine/Source/Runtime/Engine/Public/AudioStreaming.h
--- a/Engine/Source/Runtime/Engine/Public/AudioStreaming.h
+++ b/Engine/Source/Runtime/Engine/Public/AudioStreaming.h
@@ -191,7 +191,6 @@ public:
 	TArray<FLoadedAudioChunk> LoadedChunks;
 
 	class IAsyncReadFileHandle* IORequestHandle;
-	FAsyncFileCallBack AsyncFileCallBack;
 
 	/** Potentially outstanding audio chunk I/O requests */
 	TArray<uint64>	IORequestIndices;
diff --git a/Engine/Source/Runtime/Engine/Public/DelayAction.h b/Engine/Source/Runtime/Engine/Public/DelayAction.h
--- a/Engine/Source/Runtime/Engine/Public/DelayAction.h
+++ b/Engine/Source/Runtime/Engine/Public/DelayAction.h
@@ -34,7 +34,10 @@ public:
 	// Returns a human readable description of the latent operation's current state
 	virtual FString GetDescription() const override
 	{
-		return FString::Printf( *NSLOCTEXT("DelayAction", "DelayActionTime", "Delay (%.3f seconds left)").ToString(), TimeRemaining);
+		static const FNumberFormattingOptions DelayTimeFormatOptions = FNumberFormattingOptions()
+			.SetMinimumFractionalDigits(3)
+			.SetMaximumFractionalDigits(3);
+		return FText::Format(NSLOCTEXT("DelayAction", "DelayActionTimeFmt", "Delay ({0} seconds left)"), FText::AsNumber(TimeRemaining, &DelayTimeFormatOptions)).ToString();
 	}
 #endif
 };
diff --git a/Engine/Source/Runtime/Launch/Private/LaunchEngineLoop.cpp b/Engine/Source/Runtime/Launch/Private/LaunchEngineLoop.cpp
--- a/Engine/Source/Runtime/Launch/Private/LaunchEngineLoop.cpp
+++ b/Engine/Source/Runtime/Launch/Private/LaunchEngineLoop.cpp
@@ -55,6 +55,7 @@
 	#include "UObject/UObjectHash.h"
 	#include "UObject/Package.h"
 	#include "UObject/Linker.h"
+	#include "UObject/LinkerLoad.h"
 #endif
 
 #if WITH_EDITOR
@@ -1732,6 +1733,15 @@ int32 FEngineLoop::PreInit( const TCHAR* CmdLine )
 
 		SlowTask.EnterProgressFrame(5);
 
+
+#if USE_EVENT_DRIVEN_ASYNC_LOAD_AT_BOOT_TIME
+		// If we don't do this now and the async loading thread is active, then we will attempt to load this module from a thread
+		if (GEventDrivenLoaderEnabled)
+		{
+			FModuleManager::Get().LoadModule("AssetRegistry");
+		}
+#endif
+
 		// Make sure all UObject classes are registered and default properties have been initialized
 		ProcessNewlyLoadedUObjects();
 
@@ -1851,6 +1861,7 @@ int32 FEngineLoop::PreInit( const TCHAR* CmdLine )
 	{
 		GUObjectArray.CloseDisregardForGC();
 	}
+	NotifyRegistrationComplete();
 #endif // WITH_COREUOBJECT
 
 #if WITH_ENGINE
diff --git a/Engine/Source/Runtime/MovieSceneCapture/Private/AutomatedLevelSequenceCapture.cpp b/Engine/Source/Runtime/MovieSceneCapture/Private/AutomatedLevelSequenceCapture.cpp
--- a/Engine/Source/Runtime/MovieSceneCapture/Private/AutomatedLevelSequenceCapture.cpp
+++ b/Engine/Source/Runtime/MovieSceneCapture/Private/AutomatedLevelSequenceCapture.cpp
@@ -52,7 +52,7 @@ void UAutomatedLevelSequenceCapture::AddFormatMappings(TMap<FString, FStringForm
 	OutFormatMappings.Add(TEXT("shot"), CachedState.CurrentShotName.ToString());
 
 	const int32 FrameNumber = FMath::RoundToInt(CachedState.CurrentShotLocalTime * CachedState.Settings.FrameRate);
-	OutFormatMappings.Add(TEXT("shot_frame"), FString::Printf(*FrameNumberFormat, FrameNumber));
+	OutFormatMappings.Add(TEXT("shot_frame"), FString::Printf(TEXT("%0*d"), Settings.ZeroPadFrameNumbers, FrameNumber));
 }
 
 void UAutomatedLevelSequenceCapture::Initialize(TSharedPtr<FSceneViewport> InViewport, int32 PIEInstance)
diff --git a/Engine/Source/Runtime/MovieSceneCapture/Private/MovieSceneCapture.cpp b/Engine/Source/Runtime/MovieSceneCapture/Private/MovieSceneCapture.cpp
--- a/Engine/Source/Runtime/MovieSceneCapture/Private/MovieSceneCapture.cpp
+++ b/Engine/Source/Runtime/MovieSceneCapture/Private/MovieSceneCapture.cpp
@@ -177,8 +177,6 @@ void UMovieSceneCapture::Initialize(TSharedPtr<FSceneViewport> InSceneViewport,
 	FormatMappings.Add(TEXT("height"), FString::Printf(TEXT("%d"), CachedMetrics.Height));
 	FormatMappings.Add(TEXT("world"), InSceneViewport->GetClient()->GetWorld()->GetName());
 
-	FrameNumberFormat = FString::Printf(TEXT("%%0%dd"), Settings.ZeroPadFrameNumbers);
-
 	if (!HasAnyFlags(RF_ClassDefaultObject))
 	{
 		FActiveMovieSceneCaptures::Get().Add(this);
@@ -292,7 +290,7 @@ FString UMovieSceneCapture::ResolveFileFormat(const FString& Format, const FFram
 {
 	TMap<FString, FStringFormatArg> AllArgs = FormatMappings;
 
-	AllArgs.Add(TEXT("frame"), FString::Printf(*FrameNumberFormat, Settings.bUseRelativeFrameNumbers ? FrameMetrics.FrameNumber : FrameMetrics.FrameNumber + FrameNumberOffset));
+	AllArgs.Add(TEXT("frame"), FString::Printf(TEXT("%0*d"), Settings.ZeroPadFrameNumbers, Settings.bUseRelativeFrameNumbers ? FrameMetrics.FrameNumber : FrameMetrics.FrameNumber + FrameNumberOffset));
 
 	AddFormatMappings(AllArgs, FrameMetrics);
 
diff --git a/Engine/Source/Runtime/MovieSceneCapture/Public/MovieSceneCapture.h b/Engine/Source/Runtime/MovieSceneCapture/Public/MovieSceneCapture.h
--- a/Engine/Source/Runtime/MovieSceneCapture/Public/MovieSceneCapture.h
+++ b/Engine/Source/Runtime/MovieSceneCapture/Public/MovieSceneCapture.h
@@ -183,8 +183,6 @@ protected:
 	int32 FrameNumberOffset;
 	/** Event that is triggered when capturing has finished */
 	FOnCaptureFinished OnCaptureFinishedDelegate;
-	/** Format string used for frame numbers */
-	FString FrameNumberFormat;
 	/** Cached quality levels */
 	Scalability::FQualityLevels CachedQualityLevels;
 };
