[NO NON-PII COMMIT MESSAGE]

diff --git a/Engine/Source/Runtime/Engine/Classes/Engine/StreamableManager.h b/Engine/Source/Runtime/Engine/Classes/Engine/StreamableManager.h
--- a/Engine/Source/Runtime/Engine/Classes/Engine/StreamableManager.h
+++ b/Engine/Source/Runtime/Engine/Classes/Engine/StreamableManager.h
@@ -298,7 +298,17 @@ private:
 	TStreamableMap StreamableItems;
 
 	/** Map of redirected paths */
-	typedef TMap<FStringAssetReference, FStringAssetReference> TStreamableRedirects;
+	struct FRedirectedPath
+	{
+		/** The path of the non-redirector object loaded */
+		FStringAssetReference NewPath;
+
+		/** The redirector that was loaded off disk, need to keep this around for path resolves until this redirect is freed */
+		UObjectRedirector* LoadedRedirector;
+
+		FRedirectedPath() : LoadedRedirector(LoadedRedirector) {}
+	};
+	typedef TMap<FStringAssetReference, FRedirectedPath> TStreamableRedirects;
 	TStreamableRedirects StreamableRedirects;
 
 	/** List of explicitly held handles */
diff --git a/Engine/Source/Runtime/Engine/Private/StreamableManager.cpp b/Engine/Source/Runtime/Engine/Private/StreamableManager.cpp
--- a/Engine/Source/Runtime/Engine/Private/StreamableManager.cpp
+++ b/Engine/Source/Runtime/Engine/Private/StreamableManager.cpp
@@ -627,7 +627,7 @@ void FStreamableManager::OnPostGarbageCollect()
 	{
 		for (TStreamableRedirects::TIterator It(StreamableRedirects); It; ++It)
 		{
-			if (RedirectsToRemove.Contains(It.Value()))
+			if (RedirectsToRemove.Contains(It.Value().NewPath))
 			{
 				It.RemoveCurrent();
 			}
@@ -638,7 +638,6 @@ void FStreamableManager::OnPostGarbageCollect()
 void FStreamableManager::AddReferencedObjects(FReferenceCollector& Collector)
 {
 	// If there are active streamable handles in the editor, this will cause the user to Force Delete, which is irritating but necessary because weak pointers cannot be used here
-
 	for (TStreamableMap::TConstIterator It(StreamableItems); It; ++It)
 	{
 		FStreamable* Existing = It.Value();
@@ -647,6 +646,15 @@ void FStreamableManager::AddReferencedObjects(FReferenceCollector& Collector)
 			Collector.AddReferencedObject(Existing->Target);
 		}
 	}
+
+	for (TStreamableRedirects::TIterator It(StreamableRedirects); It; ++It)
+	{
+		FRedirectedPath& Existing = It.Value();
+		if (Existing.LoadedRedirector)
+		{
+			Collector.AddReferencedObject(Existing.LoadedRedirector);
+		}
+	}
 }
 
 FStreamable* FStreamableManager::FindStreamable(const FStringAssetReference& Target) const
@@ -702,11 +710,16 @@ FStreamable* FStreamableManager::StreamInternal(const FStringAssetReference& InT
 		// If async loading isn't safe or it's forced on, we have to do a sync load which will flush all async loading
 		if (GIsInitialLoad || ThreadContext.IsInConstructor > 0 || bForceSynchronousLoads)
 		{
+			FRedirectedPath RedirectedPath;
 			UE_LOG(LogStreamableManager, Verbose, TEXT("     Static loading %s"), *TargetName.ToString());
 			Existing->Target = StaticLoadObject(UObject::StaticClass(), nullptr, *TargetName.ToString());
 			// need to manually detect redirectors because the above call only expects to load a UObject::StaticClass() type
 			while (UObjectRedirector* Redirector = Cast<UObjectRedirector>(Existing->Target))
 			{
+				if (!RedirectedPath.LoadedRedirector)
+				{
+					RedirectedPath.LoadedRedirector = Redirector;
+				}
 				Existing->Target = Redirector->DestinationObject;
 			}
 			if (Existing->Target)
@@ -716,7 +729,8 @@ FStreamable* FStreamableManager::StreamInternal(const FStringAssetReference& InT
 				if (PossiblyNewName != TargetName)
 				{
 					UE_LOG(LogStreamableManager, Verbose, TEXT("     Which redirected to %s"), *PossiblyNewName.ToString());
-					StreamableRedirects.Add(TargetName, PossiblyNewName);
+					RedirectedPath.NewPath = PossiblyNewName;
+					StreamableRedirects.Add(TargetName, RedirectedPath);
 					StreamableItems.Add(PossiblyNewName, Existing);
 					StreamableItems.Remove(TargetName);
 					TargetName = PossiblyNewName; // we are done with the old name
@@ -869,6 +883,10 @@ void FStreamableManager::FindInMemory( FStringAssetReference& InOutTargetName, s
 	Existing->Target = StaticFindObject(UObject::StaticClass(), nullptr, *InOutTargetName.ToString());
 
 	UObjectRedirector* Redir = Cast<UObjectRedirector>(Existing->Target);
+	
+	FRedirectedPath RedirectedPath;
+	RedirectedPath.LoadedRedirector = Redir;
+
 	while (Redir)
 	{
 		Existing->Target = Redir->DestinationObject;
@@ -890,7 +908,8 @@ void FStreamableManager::FindInMemory( FStringAssetReference& InOutTargetName, s
 		if (InOutTargetName != PossiblyNewName)
 		{
 			UE_LOG(LogStreamableManager, Verbose, TEXT("     Name changed to %s"), *PossiblyNewName.ToString());
-			StreamableRedirects.Add(InOutTargetName, PossiblyNewName);
+			RedirectedPath.NewPath = PossiblyNewName;
+			StreamableRedirects.Add(InOutTargetName, RedirectedPath);
 			StreamableItems.Add(PossiblyNewName, Existing);
 			StreamableItems.Remove(InOutTargetName);
 			InOutTargetName = PossiblyNewName; // we are done with the old name
@@ -1087,12 +1106,12 @@ bool FStreamableManager::GetActiveHandles(const FStringAssetReference& Target, T
 
 FStringAssetReference FStreamableManager::ResolveRedirects(const FStringAssetReference& Target) const
 {
-	FStringAssetReference const* Redir = StreamableRedirects.Find(Target);
+	FRedirectedPath const* Redir = StreamableRedirects.Find(Target);
 	if (Redir)
 	{
-		check(Target != *Redir);
-		UE_LOG(LogStreamableManager, Verbose, TEXT("Redirected %s -> %s"), *Target.ToString(), *Redir->ToString());
-		return *Redir;
+		check(Target != Redir->NewPath);
+		UE_LOG(LogStreamableManager, Verbose, TEXT("Redirected %s -> %s"), *Target.ToString(), *Redir->NewPath.ToString());
+		return Redir->NewPath;
 	}
 	return Target;
 }
