CIS fix, not sure why Windows compiles this

diff --git a/Engine/Source/Runtime/Engine/Classes/Engine/StreamableManager.h b/Engine/Source/Runtime/Engine/Classes/Engine/StreamableManager.h
--- a/Engine/Source/Runtime/Engine/Classes/Engine/StreamableManager.h
+++ b/Engine/Source/Runtime/Engine/Classes/Engine/StreamableManager.h
@@ -306,7 +306,7 @@ private:
 		/** The redirector that was loaded off disk, need to keep this around for path resolves until this redirect is freed */
 		UObjectRedirector* LoadedRedirector;
 
-		FRedirectedPath() : LoadedRedirector(LoadedRedirector) {}
+		FRedirectedPath() : LoadedRedirector(nullptr) {}
 	};
 	typedef TMap<FStringAssetReference, FRedirectedPath> TStreamableRedirects;
 	TStreamableRedirects StreamableRedirects;
