Fix VS2013 CIS errors.

diff --git a/Engine/Plugins/Runtime/ImmediatePhysics/Source/ImmediatePhysics/Public/ImmediatePhysicsLinearBlockAllocator.h b/Engine/Plugins/Runtime/ImmediatePhysics/Source/ImmediatePhysics/Public/ImmediatePhysicsLinearBlockAllocator.h
--- a/Engine/Plugins/Runtime/ImmediatePhysics/Source/ImmediatePhysics/Public/ImmediatePhysicsLinearBlockAllocator.h
+++ b/Engine/Plugins/Runtime/ImmediatePhysics/Source/ImmediatePhysics/Public/ImmediatePhysicsLinearBlockAllocator.h
@@ -4,14 +4,14 @@
 
 namespace ImmediatePhysics
 {
-
+	const int PageBufferSize = 1024*64;	//This has to be here because of VS2013
 
 struct FLinearBlockAllocator
 {
 	struct FPageStruct
 	{
 		//We assume FPageStruct is allocated with 16-byte alignment. Do not move Buffer. If we get full support for alignas(16) we could make this better
-		uint8 Buffer[1024 * 64];
+		uint8 Buffer[PageBufferSize];
 
 		FPageStruct* NextPage;
 		FPageStruct* PrevPage;
@@ -50,11 +50,11 @@ struct FLinearBlockAllocator
 
 	uint8* Alloc(int32 Bytes)
 	{
-		check(Bytes < sizeof(FPageStruct::Buffer));	//Page size needs to be increased since we don't allow spillover
+		check(Bytes < PageBufferSize);	//Page size needs to be increased since we don't allow spillover
 		if (Bytes)
 		{
 			//Assumes 16 byte alignment
-			int32 BytesLeft = sizeof(FPageStruct::Buffer) - FreePage->SeekPosition;	//don't switch to uint because negative implies we're out of 16 byte aligned space
+			int32 BytesLeft = PageBufferSize - FreePage->SeekPosition;	//don't switch to uint because negative implies we're out of 16 byte aligned space
 			if (BytesLeft < Bytes)
 			{
 				//no space so allocate new page if needed
diff --git a/Engine/Source/Runtime/Engine/Private/PhysicsEngine/BodySetup.cpp b/Engine/Source/Runtime/Engine/Private/PhysicsEngine/BodySetup.cpp
--- a/Engine/Source/Runtime/Engine/Private/PhysicsEngine/BodySetup.cpp
+++ b/Engine/Source/Runtime/Engine/Private/PhysicsEngine/BodySetup.cpp
@@ -636,7 +636,7 @@ void UBodySetup::AddShapesToRigidActor_AssumesLocked(FBodyInstance* OwningInstan
 		Scale3D = FVector(0.1f);
 	}
 
-	auto AttachShape_AssumesLocked = [bShapeSharing, NewShapes, PDestActor, ComplexMaterials, SimpleMaterial, SceneType, &ShapeData] (const PxGeometry& PGeom, const PxTransform& PLocalPose, const float ContactOffset, const FPhysxUserData* ShapeElemUserData, PxShapeFlags PShapeFlags = PxShapeFlag::eVISUALIZATION | PxShapeFlag::eSCENE_QUERY_SHAPE | PxShapeFlag::eSIMULATION_SHAPE)
+	auto AttachShape_AssumesLocked = [bShapeSharing, NewShapes, PDestActor, ComplexMaterials, SimpleMaterial, SceneType, &ShapeData] (const PxGeometry& PGeom, const PxTransform& PLocalPose, const float ContactOffset, const FPhysxUserData* ShapeElemUserData, PxShapeFlags PShapeFlags)
 	{
 		const PxMaterial* PMaterial = GetDefaultPhysMaterial();
 		PxShape* PNewShape = GPhysXSDK->createShape(PGeom, *PMaterial, !bShapeSharing, PShapeFlags);
@@ -671,7 +671,7 @@ void UBodySetup::AddShapesToRigidActor_AssumesLocked(FBodyInstance* OwningInstan
 
 	auto IterateSimpleShapes = [AttachShape_AssumesLocked](const FKShapeElem& Elem, const PxGeometry& Geom, const PxTransform& PLocalPose, float ContactOffset)
 	{
-		AttachShape_AssumesLocked(Geom, PLocalPose, ContactOffset, Elem.GetUserData());
+		AttachShape_AssumesLocked(Geom, PLocalPose, ContactOffset, Elem.GetUserData(), PxShapeFlag::eVISUALIZATION | PxShapeFlag::eSCENE_QUERY_SHAPE | PxShapeFlag::eSIMULATION_SHAPE);
 	};
 
 	auto IterateTrimeshes = [AttachShape_AssumesLocked](PxTriangleMesh*, const PxGeometry& Geom, const PxTransform& PLocalPose, float ContactOffset)
