Animation: Fixed non-unity build

diff --git a/Engine/Source/Runtime/Engine/Private/Animation/NodeMappingContainer.cpp b/Engine/Source/Runtime/Engine/Private/Animation/NodeMappingContainer.cpp
--- a/Engine/Source/Runtime/Engine/Private/Animation/NodeMappingContainer.cpp
+++ b/Engine/Source/Runtime/Engine/Private/Animation/NodeMappingContainer.cpp
@@ -1,9 +1,7 @@
 // Copyright 1998-2017 Epic Games, Inc. All Rights Reserved.
 
 #include "NodeMappingContainer.h"
-#if WITH_EDITOR
 #include "Engine/Blueprint.h"
-#endif
 
 ////////////////////////////////////////////////////////////////////////////////////////
 UNodeMappingContainer::UNodeMappingContainer(const FObjectInitializer& ObjectInitializer)
