Fixed crash launching resaved maps on platforms without clothing support

diff --git a/Engine/Source/Runtime/Engine/Private/Components/SkeletalMeshComponent.cpp b/Engine/Source/Runtime/Engine/Private/Components/SkeletalMeshComponent.cpp
--- a/Engine/Source/Runtime/Engine/Private/Components/SkeletalMeshComponent.cpp
+++ b/Engine/Source/Runtime/Engine/Private/Components/SkeletalMeshComponent.cpp
@@ -452,12 +452,16 @@ void USkeletalMeshComponent::OnRegister()
 	{
 		UClothingSimulationFactory* SimFactory = SimFactoryClass->GetDefaultObject<UClothingSimulationFactory>();
 		ClothingSimulation = SimFactory->CreateSimulation();
-		ClothingSimulation->Initialize();
-		ClothingSimulationContext = ClothingSimulation->CreateContext();
 
-		if(SkeletalMesh)
+		if(ClothingSimulation)
 		{
-			RecreateClothingActors();
+			ClothingSimulation->Initialize();
+			ClothingSimulationContext = ClothingSimulation->CreateContext();
+
+			if(SkeletalMesh)
+			{
+				RecreateClothingActors();
+			}
 		}
 	}
 #endif
