Copying //UE4/Orion-Staging to //UE4/Main (Source: //Orion/Dev-General @ 3316439)

diff --git a/Engine/Source/Runtime/Engine/Private/Animation/AnimInstanceProxy.cpp b/Engine/Source/Runtime/Engine/Private/Animation/AnimInstanceProxy.cpp
--- a/Engine/Source/Runtime/Engine/Private/Animation/AnimInstanceProxy.cpp
+++ b/Engine/Source/Runtime/Engine/Private/Animation/AnimInstanceProxy.cpp
@@ -305,6 +305,7 @@ void FAnimInstanceProxy::PostUpdate(UAnimInstance* InAnimInstance) const
 		case EDrawDebugItemType::DirectionalArrow: DrawDebugDirectionalArrow(InAnimInstance->GetSkelMeshComponent()->GetWorld(), DebugItem.StartLoc, DebugItem.EndLoc, DebugItem.Size, DebugItem.Color, DebugItem.bPersistentLines, DebugItem.LifeTime, 0, DebugItem.Thickness); break;
 		case EDrawDebugItemType::Sphere : DrawDebugSphere(InAnimInstance->GetSkelMeshComponent()->GetWorld(), DebugItem.Center, DebugItem.Radius, DebugItem.Segments, DebugItem.Color, DebugItem.bPersistentLines, DebugItem.LifeTime, 0, DebugItem.Thickness); break;
 		case EDrawDebugItemType::Line: DrawDebugLine(InAnimInstance->GetSkelMeshComponent()->GetWorld(), DebugItem.StartLoc, DebugItem.EndLoc, DebugItem.Color, DebugItem.bPersistentLines, DebugItem.LifeTime, 0, DebugItem.Thickness); break;
+		case EDrawDebugItemType::CoordinateSystem : DrawDebugCoordinateSystem(InAnimInstance->GetSkelMeshComponent()->GetWorld(), DebugItem.StartLoc, DebugItem.Rotation, DebugItem.Size, DebugItem.bPersistentLines, DebugItem.LifeTime, 0, DebugItem.Thickness); break;
 		}
 	}
 #endif
@@ -1180,6 +1181,21 @@ void FAnimInstanceProxy::AnimDrawDebugSphere(const FVector& Center, float Radius
 	QueuedDrawDebugItems.Add(DrawDebugItem);
 }
 
+void FAnimInstanceProxy::AnimDrawDebugCoordinateSystem(FVector const& AxisLoc, FRotator const& AxisRot, float Scale, bool bPersistentLines, float LifeTime, float Thickness)
+{
+	FQueuedDrawDebugItem DrawDebugItem;
+
+	DrawDebugItem.ItemType = EDrawDebugItemType::CoordinateSystem;
+	DrawDebugItem.StartLoc = AxisLoc;
+	DrawDebugItem.Rotation = AxisRot;
+	DrawDebugItem.Size = Scale;
+	DrawDebugItem.bPersistentLines = bPersistentLines;
+	DrawDebugItem.LifeTime = LifeTime;
+	DrawDebugItem.Thickness = Thickness;
+
+	QueuedDrawDebugItems.Add(DrawDebugItem);
+}
+
 void FAnimInstanceProxy::AnimDrawDebugLine(const FVector& StartLoc, const FVector& EndLoc, const FColor& Color, bool bPersistentLines, float LifeTime, float Thickness)
 {
 	FQueuedDrawDebugItem DrawDebugItem;
diff --git a/Engine/Source/Runtime/Engine/Private/Components/CharacterMovementComponent.cpp b/Engine/Source/Runtime/Engine/Private/Components/CharacterMovementComponent.cpp
--- a/Engine/Source/Runtime/Engine/Private/Components/CharacterMovementComponent.cpp
+++ b/Engine/Source/Runtime/Engine/Private/Components/CharacterMovementComponent.cpp
@@ -1613,7 +1613,10 @@ UPrimitiveComponent* UCharacterMovementComponent::GetMovementBase() const
 
 void UCharacterMovementComponent::SetBase( UPrimitiveComponent* NewBase, FName BoneName, bool bNotifyActor )
 {
-	if (CharacterOwner)
+	// prevent from changing Base while server is NavWalking (no Base in that mode), so both sides are in sync
+	// otherwise it will cause problems with position smoothing
+
+	if (CharacterOwner && !bIsNavWalkingOnServer)
 	{
 		CharacterOwner->SetBase(NewBase, NewBase ? BoneName : NAME_None, bNotifyActor);
 	}
@@ -6704,7 +6707,7 @@ void UCharacterMovementComponent::SmoothCorrection(const FVector& OldLocation, c
 
 		// Don't let the client fall too far behind or run ahead of new server time.
 		const double ServerDeltaTime = ClientData->SmoothingServerTimeStamp - OldServerTimeStamp;
-		const double MaxDelta = FMath::Max<double>(ClientData->MaxMoveDeltaTime * 2.0, ServerDeltaTime * 1.25);
+		const double MaxDelta = FMath::Clamp(ServerDeltaTime * 1.25, 0.0, ClientData->MaxMoveDeltaTime * 2.0);
 		ClientData->SmoothingClientTimeStamp = FMath::Clamp(ClientData->SmoothingClientTimeStamp, ClientData->SmoothingServerTimeStamp - MaxDelta, ClientData->SmoothingServerTimeStamp);
 
 		// Compute actual delta between new server timestamp and client simulation.
@@ -8230,6 +8233,15 @@ bool UCharacterMovementComponent::ServerCheckClientError(float ClientTimeStamp,
 	if (!bIgnoreClientMovementErrorChecksAndCorrection)
 	{
 		const FVector LocDiff = UpdatedComponent->GetComponentLocation() - ClientWorldLocation;
+
+#if ROOT_MOTION_DEBUG
+		if (RootMotionSourceDebug::CVarDebugRootMotionSources.GetValueOnAnyThread() == 1)
+		{
+			FString AdjustedDebugString = FString::Printf(TEXT("ServerCheckClientError LocDiff(%.1f) ExceedsAllowablePositionError(%d) TimeStamp(%f)"),
+				LocDiff.Size(), GetDefault<AGameNetworkManager>()->ExceedsAllowablePositionError(LocDiff), ClientTimeStamp);
+			RootMotionSourceDebug::PrintOnScreen(*CharacterOwner, AdjustedDebugString);
+		}
+#endif
 		if (GetDefault<AGameNetworkManager>()->ExceedsAllowablePositionError(LocDiff))
 		{
 			return true;
@@ -8554,12 +8566,9 @@ void UCharacterMovementComponent::ClientAdjustPosition_Implementation
 	if (RootMotionSourceDebug::CVarDebugRootMotionSources.GetValueOnAnyThread() == 1)
 	{
 		const FVector VelocityCorrection = NewVelocity - Velocity;
-		if (!VelocityCorrection.IsNearlyZero())
-		{
-			FString AdjustedDebugString = FString::Printf(TEXT("PerformMovement ClientAdjustPosition_Implementation Velocity(%s) OldVelocity(%s) Correction(%s)"),
-				*NewVelocity.ToCompactString(), *Velocity.ToCompactString(), *VelocityCorrection.ToCompactString());
-			RootMotionSourceDebug::PrintOnScreen(*CharacterOwner, AdjustedDebugString);
-		}
+		FString AdjustedDebugString = FString::Printf(TEXT("PerformMovement ClientAdjustPosition_Implementation Velocity(%s) OldVelocity(%s) Correction(%s) TimeStamp(%f)"),
+			*NewVelocity.ToCompactString(), *Velocity.ToCompactString(), *VelocityCorrection.ToCompactString(), TimeStamp);
+		RootMotionSourceDebug::PrintOnScreen(*CharacterOwner, AdjustedDebugString);
 	}
 #endif
 
@@ -8696,6 +8705,15 @@ void UCharacterMovementComponent::ClientAdjustRootMotionSourcePosition_Implement
 		return;
 	}
 
+#if ROOT_MOTION_DEBUG
+	if (RootMotionSourceDebug::CVarDebugRootMotionSources.GetValueOnAnyThread() == 1)
+	{
+		FString AdjustedDebugString = FString::Printf(TEXT("ClientAdjustRootMotionSourcePosition_Implementation TimeStamp(%f)"),
+			TimeStamp);
+		RootMotionSourceDebug::PrintOnScreen(*CharacterOwner, AdjustedDebugString);
+	}
+#endif
+
 	// Call ClientAdjustPosition first. This will Ack the move if it's not outdated.
 	ClientAdjustPosition(TimeStamp, ServerLoc, FVector(0.f, 0.f, ServerVelZ), ServerBase, ServerBaseBoneName, bHasBase, bBaseRelativePosition, ServerMovementMode);
 	
@@ -8775,6 +8793,14 @@ void UCharacterMovementComponent::ClientAckGoodMove_Implementation(float TimeSta
 	FNetworkPredictionData_Client_Character* ClientData = GetPredictionData_Client_Character();
 	check(ClientData);
 
+#if ROOT_MOTION_DEBUG
+	if (RootMotionSourceDebug::CVarDebugRootMotionSources.GetValueOnAnyThread() == 1)
+	{
+		FString AdjustedDebugString = FString::Printf(TEXT("ClientAckGoodMove_Implementation TimeStamp(%f)"),
+			TimeStamp);
+		RootMotionSourceDebug::PrintOnScreen(*CharacterOwner, AdjustedDebugString);
+	}
+#endif
 
 	// Ack move if it has not expired.
 	int32 MoveIndex = ClientData->GetSavedMoveIndex(TimeStamp);
diff --git a/Engine/Source/Runtime/Engine/Private/Components/PrimitiveComponent.cpp b/Engine/Source/Runtime/Engine/Private/Components/PrimitiveComponent.cpp
--- a/Engine/Source/Runtime/Engine/Private/Components/PrimitiveComponent.cpp
+++ b/Engine/Source/Runtime/Engine/Private/Components/PrimitiveComponent.cpp
@@ -1658,7 +1658,7 @@ bool UPrimitiveComponent::MoveComponentImpl( const FVector& Delta, const FQuat&
 	// Set up
 	const FVector TraceStart = GetComponentLocation();
 	const FVector TraceEnd = TraceStart + Delta;
-	float DeltaSizeSq = Delta.SizeSquared();
+	float DeltaSizeSq = (TraceEnd - TraceStart).SizeSquared();				// Recalc here to account for precision loss of float addition
 	const FQuat InitialRotationQuat = ComponentToWorld.GetRotation();
 
 	// ComponentSweepMulti does nothing if moving < KINDA_SMALL_NUMBER in distance, so it's important to not try to sweep distances smaller than that. 
diff --git a/Engine/Source/Runtime/Engine/Private/DataChannel.cpp b/Engine/Source/Runtime/Engine/Private/DataChannel.cpp
--- a/Engine/Source/Runtime/Engine/Private/DataChannel.cpp
+++ b/Engine/Source/Runtime/Engine/Private/DataChannel.cpp
@@ -49,6 +49,12 @@ static TAutoConsoleVariable<int> CVarNetProcessQueuedBunchesMillisecondLimit(
 	30,
 	TEXT("Time threshold for processing queued bunches. If it takes longer than this in a single frame, wait until the next frame to continue processing queued bunches. For unlimited time, set to 0."));
 
+static TAutoConsoleVariable<int32> CVarNetInstantReplayProcessQueuedBunchesMillisecondLimit(
+	TEXT("net.InstantReplayProcessQueuedBunchesMillisecondLimit"),
+	8,
+	TEXT("Time threshold for processing queued bunches during instant replays. If it takes longer than this in a single frame, wait until the next frame to continue processing queued bunches. For unlimited time, set to 0."));
+
+
 TAutoConsoleVariable<int32> CVarNetPartialBunchReliableThreshold(
 	TEXT("net.PartialBunchReliableThreshold"),
 	0,
@@ -307,7 +313,9 @@ bool UChannel::ReceivedSequencedBunch( FInBunch& Bunch )
 void UChannel::ReceivedRawBunch( FInBunch & Bunch, bool & bOutSkipAck )
 {
 	// Immediately consume the NetGUID portion of this bunch, regardless if it is partial or reliable.
-	if ( Bunch.bHasPackageMapExports )
+	// NOTE - For replays, we do this even earlier, to try and load this as soon as possible, in case there is an issue creating the channel
+	// If a replay fails to create a channel, we want to salvage as much as possible
+	if ( Bunch.bHasPackageMapExports && !Connection->InternalAck )
 	{
 		Cast<UPackageMapClient>( Connection->PackageMap )->ReceiveNetGUIDBunch( Bunch );
 
@@ -318,6 +326,11 @@ void UChannel::ReceivedRawBunch( FInBunch & Bunch, bool & bOutSkipAck )
 		}
 	}
 
+	if ( Connection->InternalAck && Broken )
+	{
+		return;
+	}
+
 	check(Connection->Channels[ChIndex]==this);
 
 	if ( Bunch.bReliable && Bunch.ChSequence != Connection->InReliable[ChIndex] + 1 )
@@ -618,7 +631,12 @@ bool UChannel::ReceivedNextBunch( FInBunch & Bunch, bool & bOutSkipAck )
 					return false;
 				}
 
-				check( !Connection->InternalAck );	// Shouldn't be possible for 100% reliable connections
+				if ( !ensure( !Connection->InternalAck ) )
+				{
+					// Shouldn't be possible for 100% reliable connections
+					Broken = 1;
+					return false;
+				}
 
 				// Don't ack this packet (since we won't process all of it)
 				bOutSkipAck = true;
@@ -1929,7 +1947,12 @@ bool UActorChannel::ProcessQueuedBunches()
 		}
 	}
 
-	const int BunchTimeLimit = CVarNetProcessQueuedBunchesMillisecondLimit.GetValueOnGameThread();
+	// Instant replays are played back in a duplicated level collection, so if this is instant replay
+	// playback, the driver's DuplicateLevelID will be something other than INDEX_NONE.
+	const int BunchTimeLimit = Connection->Driver->GetDuplicateLevelID() == INDEX_NONE ?
+		CVarNetProcessQueuedBunchesMillisecondLimit.GetValueOnGameThread() :
+		CVarNetInstantReplayProcessQueuedBunchesMillisecondLimit.GetValueOnGameThread();
+
 	const bool bHasTimeToProcess = BunchTimeLimit == 0 || Connection->Driver->ProcessQueuedBunchesCurrentFrameMilliseconds < BunchTimeLimit;
 
 	// We can process all of the queued up bunches if ALL of these are true:
@@ -2225,6 +2248,7 @@ void UActorChannel::ProcessBunch( FInBunch & Bunch )
 			if ( Connection->InternalAck )
 			{
 				UE_LOG( LogNet, Warning, TEXT( "UActorChannel::ProcessBunch: Replicator.ReceivedBunch failed (Ignoring because of InternalAck). RepObj: %s, Channel: %i" ), RepObj ? *RepObj->GetFullName() : TEXT( "NULL" ), ChIndex );
+				Broken = 1;
 				continue;		// Don't consider this catastrophic in replays
 			}
 
diff --git a/Engine/Source/Runtime/Engine/Private/DemoNetDriver.cpp b/Engine/Source/Runtime/Engine/Private/DemoNetDriver.cpp
--- a/Engine/Source/Runtime/Engine/Private/DemoNetDriver.cpp
+++ b/Engine/Source/Runtime/Engine/Private/DemoNetDriver.cpp
@@ -308,10 +308,16 @@ void UDemoNetDriver::ClearReplayTasks()
 
 bool UDemoNetDriver::ProcessReplayTasks()
 {
+	// Store a shared pointer to the current task in a local variable so that if
+	// the task itself causes tasks to be cleared (for example, if it calls StopDemo()
+	// in StartTask() or Tick()), the current task won't be destroyed immediately.
+	TSharedPtr<FQueuedReplayTask> LocalActiveTask;
+
 	if ( !ActiveReplayTask.IsValid() && QueuedReplayTasks.Num() > 0 )
 	{
 		// If we don't have an active task, pull one off now
 		ActiveReplayTask = QueuedReplayTasks[0];
+		LocalActiveTask = ActiveReplayTask;
 		QueuedReplayTasks.RemoveAt( 0 );
 
 		UE_LOG( LogDemo, Verbose, TEXT( "UDemoNetDriver::ProcessReplayTasks. Name: %s" ), *ActiveReplayTask->GetName() );
diff --git a/Engine/Source/Runtime/Engine/Private/GPUSkinCache.cpp b/Engine/Source/Runtime/Engine/Private/GPUSkinCache.cpp
--- a/Engine/Source/Runtime/Engine/Private/GPUSkinCache.cpp
+++ b/Engine/Source/Runtime/Engine/Private/GPUSkinCache.cpp
@@ -27,7 +27,7 @@ DEFINE_STAT(STAT_GPUSkinCache_SkippedForOutOfECS);
 
 DEFINE_LOG_CATEGORY_STATIC(LogSkinCache, Log, All);
 
-int32 GEnableGPUSkinCacheShaders = 0;
+static int32 GEnableGPUSkinCacheShaders = 0;
 static FAutoConsoleVariableRef CVarEnableGPUSkinCacheShaders(
 	TEXT("r.SkinCache.CompileShaders"),
 	GEnableGPUSkinCacheShaders,
@@ -73,6 +73,14 @@ TAutoConsoleVariable<int32> CVarGPUSkinCacheRecomputeTangents(
 	ECVF_RenderThreadSafe
 	);
 
+static int32 GForceRecomputeTangents = 0;
+FAutoConsoleVariableRef CVarGPUSkinCacheForceRecomputeTangents(
+	TEXT("r.SkinCache.ForceRecomputeTangents"),
+	GForceRecomputeTangents,
+	TEXT("Forces enabling/using the skincache and forces all skinned object to Recompute Tangents\n"),
+	ECVF_RenderThreadSafe | ECVF_ReadOnly
+);
+
 static int32 GMaxGPUSkinCacheElementsPerFrame = 1000;
 static FAutoConsoleVariableRef CVarMaxGPUSkinCacheElementsPerFrame(
 	TEXT("r.SkinCache.MaxGPUElementsPerFrame"),
@@ -98,19 +106,25 @@ static TAutoConsoleVariable<float> CVarGPUSkinCacheDebug(
 	ECVF_RenderThreadSafe
 	);
 
+
+bool IsGPUSkinCacheAvailable()
+{
+	return GEnableGPUSkinCacheShaders != 0 || GForceRecomputeTangents != 0;
+}
+
 // We don't have it always enabled as it's not clear if this has a performance cost
 // Call on render thread only!
 // Should only be called if SM5 (compute shaders, atomics) are supported.
 ENGINE_API bool DoSkeletalMeshIndexBuffersNeedSRV()
 {
 	// currently only implemented and tested on Window SM5 (needs Compute, Atomics, SRV for index buffers, UAV for VertexBuffers)
-	return (GMaxRHIShaderPlatform == SP_PCD3D_SM5 || GMaxRHIShaderPlatform == SP_METAL_SM5) && GEnableGPUSkinCacheShaders != 0;
+	return (GMaxRHIShaderPlatform == SP_PCD3D_SM5 || GMaxRHIShaderPlatform == SP_METAL_SM5) && IsGPUSkinCacheAvailable();
 }
 
 ENGINE_API bool DoRecomputeSkinTangentsOnGPU_RT()
 {
 	// currently only implemented and tested on Window SM5 (needs Compute, Atomics, SRV for index buffers, UAV for VertexBuffers)
-	return (GMaxRHIShaderPlatform == SP_PCD3D_SM5 || GMaxRHIShaderPlatform == SP_METAL_SM5) && GEnableGPUSkinCacheShaders && GEnableGPUSkinCache && CVarGPUSkinCacheRecomputeTangents.GetValueOnRenderThread() != 0;
+	return (GMaxRHIShaderPlatform == SP_PCD3D_SM5 || GMaxRHIShaderPlatform == SP_METAL_SM5) && GEnableGPUSkinCacheShaders != 0 && ((GEnableGPUSkinCache && CVarGPUSkinCacheRecomputeTangents.GetValueOnRenderThread() != 0) || GForceRecomputeTangents != 0);
 }
 
 TGlobalResource<FGPUSkinCache> GGPUSkinCache;
@@ -255,7 +269,7 @@ public:
 
 	static bool ShouldCache(EShaderPlatform Platform)
 	{
-		return GEnableGPUSkinCacheShaders && IsFeatureLevelSupported(Platform, ERHIFeatureLevel::SM5);
+		return IsGPUSkinCacheAvailable() && IsFeatureLevelSupported(Platform, ERHIFeatureLevel::SM5);
 	}
 
 	static void ModifyCompilationEnvironment(EShaderPlatform Platform, FShaderCompilerEnvironment& OutEnvironment)
@@ -326,7 +340,7 @@ void FGPUSkinCache::Initialize(FRHICommandListImmediate& RHICmdList)
 	static const auto CVar = IConsoleManager::Get().FindTConsoleVariableDataInt(TEXT("r.BasePassOutputsVelocity"));
 	checkf(!CVar || CVar->GetValueOnRenderThread() == 0, TEXT("GPU Skin caching is not allowed with outputting velocity on base pass (r.BasePassOutputsVelocity=1)"));
 
-	checkf(GEnableGPUSkinCacheShaders, TEXT("GPU Skin caching requires the shaders enabled (r.SkinCache.CompileShaders=1)"));
+	checkf(IsGPUSkinCacheAvailable(), TEXT("GPU Skin caching requires the shaders enabled (r.SkinCache.CompileShaders=1)"));
 
 	int32 NumFloatsPerBuffer = GGPUSkinCacheBufferSize / GPUSKINCACHE_FRAMES / sizeof(float);
 	for (int32 Index = 0; Index < GPUSKINCACHE_FRAMES; ++Index)
@@ -589,7 +603,7 @@ int32 FGPUSkinCache::StartCacheMesh(FRHICommandListImmediate& RHICmdList, uint32
 	
 	TargetVertexFactory->UpdateVertexDeclaration(VertexFactory, DispatchData.SkinCacheBuffer);
 
-	int32 RecomputeTangentsMode = CVarGPUSkinCacheRecomputeTangents.GetValueOnRenderThread();
+	int32 RecomputeTangentsMode = GForceRecomputeTangents > 0 ? 1 : CVarGPUSkinCacheRecomputeTangents.GetValueOnRenderThread();
 	if (RecomputeTangentsMode > 0)
 	{
 		if (BatchElement.bRecomputeTangent || RecomputeTangentsMode == 1)
@@ -1088,7 +1102,7 @@ void FGPUSkinCache::CVarSinkFunction()
 	// 0/1
 	int32 NewGPUSkinCacheValue = CVarEnableGPUSkinCache.GetValueOnAnyThread() != 0;
 
-	if (!GEnableGPUSkinCacheShaders)
+	if (!IsGPUSkinCacheAvailable())
 	{
 		NewGPUSkinCacheValue = 0;
 	}
diff --git a/Engine/Source/Runtime/Engine/Private/GPUSkinVertexFactory.cpp b/Engine/Source/Runtime/Engine/Private/GPUSkinVertexFactory.cpp
--- a/Engine/Source/Runtime/Engine/Private/GPUSkinVertexFactory.cpp
+++ b/Engine/Source/Runtime/Engine/Private/GPUSkinVertexFactory.cpp
@@ -641,7 +641,7 @@ void FGPUSkinPassthroughVertexFactory::ModifyCompilationEnvironment( EShaderPlat
 bool FGPUSkinPassthroughVertexFactory::ShouldCache(EShaderPlatform Platform, const class FMaterial* Material, const FShaderType* ShaderType)
 {
 	// Passthrough is only valid on platforms with Compute Shader support AND for (skeletal meshes or default materials)
-	return GEnableGPUSkinCacheShaders && IsFeatureLevelSupported(Platform, ERHIFeatureLevel::SM5) && Super::ShouldCache(Platform, Material, ShaderType) && (Material->IsUsedWithSkeletalMesh() || Material->IsSpecialEngineMaterial());
+	return IsGPUSkinCacheAvailable() && IsFeatureLevelSupported(Platform, ERHIFeatureLevel::SM5) && Super::ShouldCache(Platform, Material, ShaderType) && (Material->IsUsedWithSkeletalMesh() || Material->IsSpecialEngineMaterial());
 }
 
 void FGPUSkinPassthroughVertexFactory::InternalUpdateVertexDeclaration(FGPUBaseSkinVertexFactory* SourceVertexFactory, FRWBuffer* RWBuffer)
diff --git a/Engine/Source/Runtime/Engine/Private/GameUserSettings.cpp b/Engine/Source/Runtime/Engine/Private/GameUserSettings.cpp
--- a/Engine/Source/Runtime/Engine/Private/GameUserSettings.cpp
+++ b/Engine/Source/Runtime/Engine/Private/GameUserSettings.cpp
@@ -249,6 +249,13 @@ float UGameUserSettings::GetDefaultResolutionScale()
 	return FMath::Max(DesiredResQuality, MinResolutionScale);
 }
 
+float UGameUserSettings::GetRecommendedResolutionScale()
+{
+	const float RecommendedResQuality = FindResolutionQualityForScreenSize(LastRecommendedScreenWidth, LastRecommendedScreenHeight);
+
+	return FMath::Max(RecommendedResQuality, MinResolutionScale);
+}
+
 float UGameUserSettings::FindResolutionQualityForScreenSize(float Width, float Height)
 {
 	float ResolutionQuality = 100.0f;
@@ -732,7 +739,7 @@ void UGameUserSettings::ApplyHardwareBenchmarkResults()
 	SaveSettings();
 }
 
-bool UGameUserSettings::SupportsHDRDisplayOutput()
+bool UGameUserSettings::SupportsHDRDisplayOutput() const
 {
 	return GRHISupportsHDROutput;
 }
@@ -812,3 +819,13 @@ void UGameUserSettings::EnableHDRDisplayOutput(bool bEnable, int32 DisplayNits /
 		HDRDisplayOutputNits = DisplayNitLevel;
 	}
 }
+
+int32 UGameUserSettings::GetCurrentHDRDisplayNits() const
+{
+	return bUseHDRDisplayOutput ? HDRDisplayOutputNits : 0;
+}
+
+bool UGameUserSettings::IsHDREnabled() const
+{
+	return bUseHDRDisplayOutput;
+}
diff --git a/Engine/Source/Runtime/Engine/Private/Materials/Material.cpp b/Engine/Source/Runtime/Engine/Private/Materials/Material.cpp
--- a/Engine/Source/Runtime/Engine/Private/Materials/Material.cpp
+++ b/Engine/Source/Runtime/Engine/Private/Materials/Material.cpp
@@ -1030,23 +1030,25 @@ float UMaterial::GetScalarParameterDefault(FName ParameterName, ERHIFeatureLevel
 	if (FApp::CanEverRender())
 	{
 		FMaterialResource* Resource = GetMaterialResource(InFeatureLevel);
-		check(Resource);
-
-		// Iterate over both the 2D textures and cube texture expressions.
-		const TArray<TRefCountPtr<FMaterialUniformExpression> >& UniformExpressions = Resource->GetUniformScalarParameterExpressions();
-
-		// Iterate over each of the material's texture expressions.
-		for (FMaterialUniformExpression* UniformExpression : UniformExpressions)
+		
+		if (ensureAlways(Resource))
 		{
-			if (UniformExpression->GetType() == &FMaterialUniformExpressionScalarParameter::StaticType)
-			{
-				FMaterialUniformExpressionScalarParameter* ScalarExpression = static_cast<FMaterialUniformExpressionScalarParameter*>(UniformExpression);
+			// Iterate over both the 2D textures and cube texture expressions.
+			const TArray<TRefCountPtr<FMaterialUniformExpression> >& UniformExpressions = Resource->GetUniformScalarParameterExpressions();
 
-				if (ScalarExpression->GetParameterName() == ParameterName)
+			// Iterate over each of the material's texture expressions.
+			for (FMaterialUniformExpression* UniformExpression : UniformExpressions)
+			{
+				if (UniformExpression->GetType() == &FMaterialUniformExpressionScalarParameter::StaticType)
 				{
-					float Value = 0.f;
-					ScalarExpression->GetDefaultValue(Value);
-					return Value;
+					FMaterialUniformExpressionScalarParameter* ScalarExpression = static_cast<FMaterialUniformExpressionScalarParameter*>(UniformExpression);
+
+					if (ScalarExpression->GetParameterName() == ParameterName)
+					{
+						float Value = 0.f;
+						ScalarExpression->GetDefaultValue(Value);
+						return Value;
+					}
 				}
 			}
 		}
diff --git a/Engine/Source/Runtime/Engine/Private/Materials/MaterialInstance.cpp b/Engine/Source/Runtime/Engine/Private/Materials/MaterialInstance.cpp
--- a/Engine/Source/Runtime/Engine/Private/Materials/MaterialInstance.cpp
+++ b/Engine/Source/Runtime/Engine/Private/Materials/MaterialInstance.cpp
@@ -1032,22 +1032,24 @@ float UMaterialInstance::GetScalarParameterDefault(FName ParameterName, ERHIFeat
 		if (FApp::CanEverRender())
 		{
 			const FMaterialResource* SourceMaterialResource = GetMaterialResource(InFeatureLevel);
-			check(SourceMaterialResource);
-			const TArray<TRefCountPtr<FMaterialUniformExpression> >& UniformExpressions = SourceMaterialResource->GetUniformScalarParameterExpressions();
-
-			// Iterate over each of the material's texture expressions.
-			for (int32 ExpressionIndex = 0; ExpressionIndex < UniformExpressions.Num(); ExpressionIndex++)
+			if (ensureAlways(SourceMaterialResource))
 			{
-				FMaterialUniformExpression* UniformExpression = UniformExpressions[ExpressionIndex];
-				if (UniformExpression->GetType() == &FMaterialUniformExpressionScalarParameter::StaticType)
-				{
-					FMaterialUniformExpressionScalarParameter* ScalarExpression = static_cast<FMaterialUniformExpressionScalarParameter*>(UniformExpression);
+				const TArray<TRefCountPtr<FMaterialUniformExpression> >& UniformExpressions = SourceMaterialResource->GetUniformScalarParameterExpressions();
 
-					if (ScalarExpression->GetParameterName() == ParameterName)
+				// Iterate over each of the material's texture expressions.
+				for (int32 ExpressionIndex = 0; ExpressionIndex < UniformExpressions.Num(); ExpressionIndex++)
+				{
+					FMaterialUniformExpression* UniformExpression = UniformExpressions[ExpressionIndex];
+					if (UniformExpression->GetType() == &FMaterialUniformExpressionScalarParameter::StaticType)
 					{
-						float Value = 0.f;
-						ScalarExpression->GetDefaultValue(Value);
-						return Value;
+						FMaterialUniformExpressionScalarParameter* ScalarExpression = static_cast<FMaterialUniformExpressionScalarParameter*>(UniformExpression);
+
+						if (ScalarExpression->GetParameterName() == ParameterName)
+						{
+							float Value = 0.f;
+							ScalarExpression->GetDefaultValue(Value);
+							return Value;
+						}
 					}
 				}
 			}
@@ -2030,7 +2032,7 @@ void UMaterialInstance::PostLoad()
 	// Ensure that the instance's parent is PostLoaded before the instance.
 	if (Parent)
 	{
-		if (EVENT_DRIVEN_ASYNC_LOAD_ACTIVE_AT_RUNTIME)
+		if (GEventDrivenLoaderEnabled && EVENT_DRIVEN_ASYNC_LOAD_ACTIVE_AT_RUNTIME)
 		{
 			check(!Parent->HasAnyFlags(RF_NeedLoad));
 		}
diff --git a/Engine/Source/Runtime/Engine/Private/Materials/MaterialShared.cpp b/Engine/Source/Runtime/Engine/Private/Materials/MaterialShared.cpp
--- a/Engine/Source/Runtime/Engine/Private/Materials/MaterialShared.cpp
+++ b/Engine/Source/Runtime/Engine/Private/Materials/MaterialShared.cpp
@@ -1579,8 +1579,13 @@ bool FMaterial::CacheShaders(const FMaterialShaderMapId& ShaderMapId, EShaderPla
 		}
 	}
 
+	UMaterialInterface* MaterialInterface = GetMaterialInterface();
+	const bool bMaterialInstance = MaterialInterface && MaterialInterface->IsA(UMaterialInstance::StaticClass());
+	const bool bSpecialEngineMaterial = !bMaterialInstance && IsSpecialEngineMaterial();
+
 	// Log which shader, pipeline or factory is missing when about to have a fatal error
-	const bool bLogShaderMapFailInfo = IsSpecialEngineMaterial() && (bContainsInlineShaders || FPlatformProperties::RequiresCookedData());
+	const bool bLogShaderMapFailInfo = bSpecialEngineMaterial && (bContainsInlineShaders || FPlatformProperties::RequiresCookedData());
+
 
 	bool bAssumeShaderMapIsComplete = false;
 #if UE_BUILD_SHIPPING || UE_BUILD_TEST
@@ -1603,10 +1608,23 @@ bool FMaterial::CacheShaders(const FMaterialShaderMapId& ShaderMapId, EShaderPla
 	{
 		if (bContainsInlineShaders || FPlatformProperties::RequiresCookedData())
 		{
-			if (IsSpecialEngineMaterial())
+			if (bSpecialEngineMaterial)
 			{
+				UMaterialInterface* Interface = GetMaterialInterface();
+				FString Instance;
+				if (Interface)
+				{
+					Instance = Interface->GetPathName();
+				}
+
 				//assert if the default material's shader map was not found, since it will cause problems later
-				UE_LOG(LogMaterial, Fatal,TEXT("Failed to find shader map for default material %s!  Please make sure cooking was successful."), *GetFriendlyName());
+				UE_LOG(LogMaterial, Fatal,TEXT("Failed to find shader map for default material %s(%s)! Please make sure cooking was successful (%s inline shaders, %s GTSM%s)"),
+					*GetFriendlyName(),
+					*Instance,
+					bContainsInlineShaders ? TEXT("Contains") : TEXT("No"),
+					GameThreadShaderMap ? TEXT("has") : TEXT("null"),
+					bAssumeShaderMapIsComplete ? TEXT(" assumes map complete") : TEXT("")
+				);
 			}
 			else
 			{
diff --git a/Engine/Source/Runtime/Engine/Private/NetConnection.cpp b/Engine/Source/Runtime/Engine/Private/NetConnection.cpp
--- a/Engine/Source/Runtime/Engine/Private/NetConnection.cpp
+++ b/Engine/Source/Runtime/Engine/Private/NetConnection.cpp
@@ -1068,6 +1068,18 @@ void UNetConnection::ReceivedPacket( FBitReader& Reader )
 			if ( Bunch.bHasPackageMapExports )
 			{
 				Driver->NetGUIDInBytes += (BunchDataBits + (HeaderPos - IncomingStartPos)) >> 3;
+
+				if ( InternalAck )
+				{
+					// NOTE - For replays, we do this even earlier, to try and load this as soon as possible, in case there is an issue creating the channel
+					// If a replay fails to create a channel, we want to salvage as much as possible
+					Cast<UPackageMapClient>( PackageMap )->ReceiveNetGUIDBunch( Bunch );
+
+					if ( Bunch.IsError() )
+					{
+						UE_LOG( LogNetTraffic, Error, TEXT( "UNetConnection::ReceivedPacket: Bunch.IsError() after ReceiveNetGUIDBunch. ChIndex: %i" ), Bunch.ChIndex );
+					}
+				}
 			}
 
 			if( Bunch.bReliable )
@@ -1600,7 +1612,12 @@ void UNetConnection::Tick()
 			{
 				OutBunch->ReceivedAck = 1;
 			}
-			It->OpenAcked = 1;
+
+			if (Driver->IsServer() || It->OpenedLocally)
+			{
+				It->OpenAcked = 1;
+			}
+
 			It->ReceivedAcks();
 		}
 	}
diff --git a/Engine/Source/Runtime/Engine/Private/NetworkDriver.cpp b/Engine/Source/Runtime/Engine/Private/NetworkDriver.cpp
--- a/Engine/Source/Runtime/Engine/Private/NetworkDriver.cpp
+++ b/Engine/Source/Runtime/Engine/Private/NetworkDriver.cpp
@@ -2305,7 +2305,7 @@ void FPacketSimulationSettings::UnregisterCommands()
  */
 bool FPacketSimulationSettings::ParseSettings(const TCHAR* Cmd, const TCHAR* OptionalQualifier)
 {
-	UE_LOG(LogTemp, Warning, TEXT("ParseSettings for %s"), OptionalQualifier);
+	UE_LOG(LogTemp, Display, TEXT("ParseSettings for %s"), OptionalQualifier);
 	// note that each setting is tested.
 	// this is because the same function will be used to parse the command line as well
 	bool bParsed = false;
diff --git a/Engine/Source/Runtime/Engine/Private/RepLayout.cpp b/Engine/Source/Runtime/Engine/Private/RepLayout.cpp
--- a/Engine/Source/Runtime/Engine/Private/RepLayout.cpp
+++ b/Engine/Source/Runtime/Engine/Private/RepLayout.cpp
@@ -1927,7 +1927,14 @@ bool FRepLayout::ReceiveProperties( UActorChannel* OwningChannel, UClass * InObj
 
 	if ( OwningChannel->Connection->InternalAck )
 	{
-		TSharedPtr< FNetFieldExportGroup > NetFieldExportGroup = ( ( UPackageMapClient* )OwningChannel->Connection->PackageMap )->GetNetFieldExportGroupChecked( Owner->GetPathName() );
+		TSharedPtr< FNetFieldExportGroup > NetFieldExportGroup = ( ( UPackageMapClient* )OwningChannel->Connection->PackageMap )->GetNetFieldExportGroup( Owner->GetPathName() );
+
+		if ( !ensure( NetFieldExportGroup.IsValid() ) )
+		{
+			UE_LOG( LogRep, Warning, TEXT( "ReceiveProperties_BackwardsCompatible: Invalid path name: %s" ), *Owner->GetPathName() );
+			InBunch.SetError();
+			return false;
+		}
 
 		return ReceiveProperties_BackwardsCompatible_r( RepState, NetFieldExportGroup.Get(), InBunch, 0, Cmds.Num() - 1, bEnableRepNotifies ? RepState->StaticBuffer.GetData() : nullptr, ( uint8* )Data, ( uint8* )Data, &RepState->GuidReferencesMap, bOutHasUnmapped, bOutGuidsChanged );
 	}
diff --git a/Engine/Source/Runtime/Engine/Private/SkeletalMesh.cpp b/Engine/Source/Runtime/Engine/Private/SkeletalMesh.cpp
--- a/Engine/Source/Runtime/Engine/Private/SkeletalMesh.cpp
+++ b/Engine/Source/Runtime/Engine/Private/SkeletalMesh.cpp
@@ -4992,19 +4992,27 @@ public:
 	FORCEINLINE bool NotValidPreviewSection()
 	{
 #if WITH_EDITORONLY_DATA
-
-		int32 ActualPreviewSectionIdx = SectionIndexPreview;
-		if(ActualPreviewSectionIdx != INDEX_NONE && Sections.IsValidIndex(ActualPreviewSectionIdx))
+		int32 ActualPreviewMaterialIdx = SectionIndexPreview;
+		int32 ActualPreviewSectionIdx = INDEX_NONE;
+		if (ActualPreviewMaterialIdx != INDEX_NONE && Sections.IsValidIndex(SectionIndex))
 		{
-			const FSkelMeshSection& PreviewSection = Sections[ActualPreviewSectionIdx];
-			if(PreviewSection.CorrespondClothSectionIndex != INDEX_NONE)
+			const FSkeletalMeshSceneProxy::FSectionElementInfo& SectionInfo = LODSectionElements.SectionElements[SectionIndex];
+			if (SectionInfo.UseMaterialIndex == ActualPreviewMaterialIdx)
+			{
+				ActualPreviewSectionIdx = SectionIndex;
+			}
+			if (ActualPreviewSectionIdx != INDEX_NONE)
 			{
-				ActualPreviewSectionIdx = PreviewSection.CorrespondClothSectionIndex;
+				const FSkelMeshSection& PreviewSection = Sections[ActualPreviewSectionIdx];
+				if (PreviewSection.CorrespondClothSectionIndex != INDEX_NONE)
+				{
+					ActualPreviewSectionIdx = PreviewSection.CorrespondClothSectionIndex;
+				}
 			}
 		}
 
-		return	(SectionIndex < Sections.Num()) && 
-				((ActualPreviewSectionIdx >= 0) && (ActualPreviewSectionIdx != SectionIndex));
+		return	(SectionIndex < Sections.Num()) &&
+			((ActualPreviewMaterialIdx >= 0) && (ActualPreviewSectionIdx != SectionIndex));
 #else
 		return false;
 #endif
@@ -5111,7 +5119,7 @@ void FSkeletalMeshSceneProxy::GetMeshElementsConditionallySelectable(const TArra
 
 #if WITH_EDITORONLY_DATA
 			// TODO: This is not threadsafe! A render command should be used to propagate SelectedEditorSection to the scene proxy.
-			bSectionSelected = (SkeletalMeshForDebug->SelectedEditorSection == Iter.GetSectionElementIndex());
+			bSectionSelected = (SkeletalMeshForDebug->SelectedEditorSection == SectionElementInfo.UseMaterialIndex);
 #endif
 			// If hidden skip the draw
 			if (MeshObject->IsMaterialHidden(LODIndex, SectionElementInfo.UseMaterialIndex))
diff --git a/Engine/Source/Runtime/Engine/Private/SkeletalRenderGPUSkin.cpp b/Engine/Source/Runtime/Engine/Private/SkeletalRenderGPUSkin.cpp
--- a/Engine/Source/Runtime/Engine/Private/SkeletalRenderGPUSkin.cpp
+++ b/Engine/Source/Runtime/Engine/Private/SkeletalRenderGPUSkin.cpp
@@ -93,7 +93,7 @@ void FMorphVertexBuffer::InitDynamicRHI()
 	Flags = (EBufferUsageFlags)(Flags | BUF_ShaderResource);
 
 	VertexBufferRHI = RHICreateAndLockVertexBuffer(Size, Flags, CreateInfo, BufferData);
-	bool bUsesSkinCache = bSupportsComputeShaders && GEnableGPUSkinCacheShaders && GEnableGPUSkinCache;
+	bool bUsesSkinCache = bSupportsComputeShaders && IsGPUSkinCacheAvailable() && GEnableGPUSkinCache;
 	if (bUsesSkinCache)
 	{
 		SRVValue = RHICreateShaderResourceView(VertexBufferRHI, 4, PF_R32_FLOAT);
diff --git a/Engine/Source/Runtime/Engine/Public/Animation/AnimInstanceProxy.h b/Engine/Source/Runtime/Engine/Public/Animation/AnimInstanceProxy.h
--- a/Engine/Source/Runtime/Engine/Public/Animation/AnimInstanceProxy.h
+++ b/Engine/Source/Runtime/Engine/Public/Animation/AnimInstanceProxy.h
@@ -39,6 +39,7 @@ namespace EDrawDebugItemType
 		Sphere,
 		Line,
 		OnScreenMessage,
+		CoordinateSystem,
 	};
 }
 
@@ -59,6 +60,9 @@ struct FQueuedDrawDebugItem
 	UPROPERTY(Transient)
 	FVector Center;
 
+	UPROPERTY(Transient)
+	FRotator Rotation;
+
 	UPROPERTY(Transient)
 	float Radius;
 
@@ -330,11 +334,13 @@ public:
 	void AnimDrawDebugLine(const FVector& StartLoc, const FVector& EndLoc, const FColor& Color, bool bPersistentLines = false, float LifeTime = -1.f, float Thickness = 0.f);
 	void AnimDrawDebugDirectionalArrow(const FVector& LineStart, const FVector& LineEnd, float ArrowSize, const FColor& Color, bool bPersistentLines = false, float LifeTime = -1.f, float Thickness = 0.f);
 	void AnimDrawDebugSphere(const FVector& Center, float Radius, int32 Segments, const FColor& Color, bool bPersistentLines = false, float LifeTime = -1.f, float Thickness = 0.f);
+	void AnimDrawDebugCoordinateSystem(FVector const& AxisLoc, FRotator const& AxisRot, float Scale = 1.f, bool bPersistentLines = false, float LifeTime = -1.f, float Thickness = 0.f);
 #else
 	void AnimDrawDebugOnScreenMessage(const FString& DebugMessage, const FColor& Color, const FVector2D& TextScale = FVector2D::UnitVector) {}
 	void AnimDrawDebugLine(const FVector& StartLoc, const FVector& EndLoc, const FColor& Color, bool bPersistentLines = false, float LifeTime = -1.f, float Thickness = 0.f) {}
 	void AnimDrawDebugDirectionalArrow(const FVector& LineStart, const FVector& LineEnd, float ArrowSize, const FColor& Color, bool bPersistentLines = false, float LifeTime = -1.f, float Thickness = 0.f) {}
 	void AnimDrawDebugSphere(const FVector& Center, float Radius, int32 Segments, const FColor& Color, bool bPersistentLines = false, float LifeTime = -1.f, float Thickness = 0.f) {}
+	void AnimDrawDebugCoordinateSystem(FVector const& AxisLoc, FRotator const& AxisRot, float Scale = 1.f, bool bPersistentLines = false, float LifeTime = -1.f, float Thickness = 0.f) {}
 #endif // ENABLE_ANIM_DRAW_DEBUG
 
 #if !NO_LOGGING
diff --git a/Engine/Source/Runtime/Engine/Public/GPUSkinCache.h b/Engine/Source/Runtime/Engine/Public/GPUSkinCache.h
--- a/Engine/Source/Runtime/Engine/Public/GPUSkinCache.h
+++ b/Engine/Source/Runtime/Engine/Public/GPUSkinCache.h
@@ -52,7 +52,10 @@ struct FSkelMeshSection;
 
 typedef FRHIShaderResourceView* FShaderResourceViewRHIParamRef;
 
-extern ENGINE_API int32 GEnableGPUSkinCacheShaders;
+// Can the skin cache be used (ie shaders added, etc)
+extern ENGINE_API bool IsGPUSkinCacheAvailable();
+
+// Is it actually enabled?
 extern ENGINE_API int32 GEnableGPUSkinCache;
 
 class FSkeletalMeshObjectGPUSkin;
diff --git a/Engine/Source/Runtime/GameplayTags/Private/GameplayTagContainer.cpp b/Engine/Source/Runtime/GameplayTags/Private/GameplayTagContainer.cpp
--- a/Engine/Source/Runtime/GameplayTags/Private/GameplayTagContainer.cpp
+++ b/Engine/Source/Runtime/GameplayTags/Private/GameplayTagContainer.cpp
@@ -1084,7 +1084,7 @@ bool FGameplayTag::MatchesTag(const FGameplayTag& TagToCheck) const
 	}
 
 	// This should always be invalid if the node is missing
-	ensure(!IsValid());
+	ensureMsgf(!IsValid(), TEXT("Valid tag failed to conver to single tag container. %s"), *GetTagName().ToString());
 
 	return false;
 }
@@ -1103,8 +1103,7 @@ bool FGameplayTag::MatchesAny(const FGameplayTagContainer& ContainerToCheck) con
 	}
 
 	// This should always be invalid if the node is missing
-	ensure(!IsValid());
-
+	ensureMsgf(!IsValid(), TEXT("Valid tag failed to conver to single tag container. %s"), *GetTagName().ToString() );
 	return false;
 }
 
diff --git a/Engine/Source/Runtime/Launch/Private/LaunchEngineLoop.cpp b/Engine/Source/Runtime/Launch/Private/LaunchEngineLoop.cpp
--- a/Engine/Source/Runtime/Launch/Private/LaunchEngineLoop.cpp
+++ b/Engine/Source/Runtime/Launch/Private/LaunchEngineLoop.cpp
@@ -250,9 +250,9 @@ public:
 		if ( (bAllowLogVerbosity && Verbosity <= ELogVerbosity::Log) || (Verbosity <= ELogVerbosity::Display) )
 		{
 #if PLATFORM_USE_LS_SPEC_FOR_WIDECHAR
-			wprintf(TEXT("\n%ls"), *FOutputDeviceHelper::FormatLogLine(Verbosity, Category, V, GPrintLogTimes));
+			wprintf(TEXT("%ls\n"), *FOutputDeviceHelper::FormatLogLine(Verbosity, Category, V, GPrintLogTimes));
 #else
-			wprintf(TEXT("\n%s"), *FOutputDeviceHelper::FormatLogLine(Verbosity, Category, V, GPrintLogTimes));
+			wprintf(TEXT("%s\n"), *FOutputDeviceHelper::FormatLogLine(Verbosity, Category, V, GPrintLogTimes));
 #endif
 			fflush(stdout);
 		}
@@ -1730,6 +1730,11 @@ int32 FEngineLoop::PreInit( const TCHAR* CmdLine )
 			}
 
 		}
+		else if ( IsRunningCommandlet() )
+		{
+			// Create the engine font services now that the Slate renderer is ready
+			FEngineFontServices::Create();
+		}
 #endif
 
 		// In order to be able to use short script package names get all script
@@ -3078,24 +3083,26 @@ void FEngineLoop::Tick()
 		if (bDoConcurrentSlateTick)
 		{
 			const float DeltaSeconds = FApp::GetDeltaTime();
-			UGameViewportClient* const GameViewport = GEngine->GameViewport;
-			ConcurrentTask = TGraphTask<FExecuteConcurrentWithSlateTickTask>::CreateTask(nullptr, ENamedThreads::GameThread).ConstructAndDispatchWhenReady(
-				[GameViewport, DeltaSeconds]()
-				{
-					if (CVarDoAsyncEndOfFrameTasksRandomize.GetValueOnAnyThread(true) > 0)
-					{
-						FPlatformProcess::Sleep(FMath::RandRange(0.0f, .003f)); // this shakes up the threading to find race conditions
-					}
-					if (GameViewport != nullptr)
+			const UGameViewportClient* const GameViewport = GEngine->GameViewport;
+			const UWorld* const GameViewportWorld = GameViewport ? GameViewport->GetWorld() : nullptr;
+			UDemoNetDriver* const CurrentDemoNetDriver = GameViewportWorld ? GameViewportWorld->DemoNetDriver : nullptr;
+
+			if (CurrentDemoNetDriver && CurrentDemoNetDriver->ShouldTickFlushAsyncEndOfFrame())
+			{
+				ConcurrentTask = TGraphTask<FExecuteConcurrentWithSlateTickTask>::CreateTask(nullptr, ENamedThreads::GameThread).ConstructAndDispatchWhenReady(
+					[CurrentDemoNetDriver, DeltaSeconds]()
 					{
-						UWorld* const World = GameViewport->GetWorld();
-						if (World && World->DemoNetDriver)
+						if (CVarDoAsyncEndOfFrameTasksRandomize.GetValueOnAnyThread(true) > 0)
+						{
+							FPlatformProcess::Sleep(FMath::RandRange(0.0f, .003f)); // this shakes up the threading to find race conditions
+						}
+						if (CurrentDemoNetDriver != nullptr)
 						{
-							World->DemoNetDriver->TickFlushAsyncEndOfFrame(DeltaSeconds);
+							CurrentDemoNetDriver->TickFlushAsyncEndOfFrame(DeltaSeconds);
 						}
 					}
-				}
-			);
+				);
+			}
 		}
 #endif
 
diff --git a/Engine/Source/Runtime/LevelSequence/Private/LevelSequencePlayer.cpp b/Engine/Source/Runtime/LevelSequence/Private/LevelSequencePlayer.cpp
--- a/Engine/Source/Runtime/LevelSequence/Private/LevelSequencePlayer.cpp
+++ b/Engine/Source/Runtime/LevelSequence/Private/LevelSequencePlayer.cpp
@@ -16,6 +16,7 @@
 #include "LevelSequenceSpawnRegister.h"
 #include "Engine/Engine.h"
 #include "Engine/LevelStreaming.h"
+#include "Engine/LocalPlayer.h"
 #include "Tracks/MovieSceneCinematicShotTrack.h"
 #include "Sections/MovieSceneCinematicShotSection.h"
 #include "LevelSequenceActor.h"
@@ -185,6 +186,10 @@ void ULevelSequencePlayer::UpdateCameraCut(UObject* CameraObject, UObject* Unloc
 	if (!LastViewTarget.IsValid())
 	{
 		LastViewTarget = ViewTarget;
+		if (PC->GetLocalPlayer())
+		{
+			LastAspectRatioAxisConstraint = PC->GetLocalPlayer()->AspectRatioAxisConstraint;
+		}
 	}
 
 	UCameraComponent* CameraComponent = MovieSceneHelpers::CameraComponentFromRuntimeObject(CameraObject);
@@ -229,6 +234,11 @@ void ULevelSequencePlayer::UpdateCameraCut(UObject* CameraObject, UObject* Unloc
 	FViewTargetTransitionParams TransitionParams;
 	PC->SetViewTarget(CameraActor, TransitionParams);
 
+	if (PC->GetLocalPlayer())
+	{
+		PC->GetLocalPlayer()->AspectRatioAxisConstraint = EAspectRatioAxisConstraint::AspectRatio_MaintainXFOV;
+	}
+
 	if (CameraComponent)
 	{
 		CameraComponent->NotifyCameraCut();
diff --git a/Engine/Source/Runtime/LevelSequence/Public/LevelSequencePlayer.h b/Engine/Source/Runtime/LevelSequence/Public/LevelSequencePlayer.h
--- a/Engine/Source/Runtime/LevelSequence/Public/LevelSequencePlayer.h
+++ b/Engine/Source/Runtime/LevelSequence/Public/LevelSequencePlayer.h
@@ -3,6 +3,7 @@
 #pragma once
 
 #include "CoreMinimal.h"
+#include "Engine/EngineTypes.h"
 #include "UObject/ObjectMacros.h"
 #include "UObject/Object.h"
 #include "UObject/ScriptMacros.h"
@@ -154,6 +155,9 @@ private:
 	/** The last view target to reset to when updating camera cuts to null */
 	TWeakObjectPtr<AActor> LastViewTarget;
 
+	/** The last aspect ratio axis constraint to reset to when the camera cut is null */
+	EAspectRatioAxisConstraint LastAspectRatioAxisConstraint;
+
 protected:
 
 	/** How to take snapshots */
diff --git a/Engine/Source/Runtime/MovieScene/Private/Evaluation/MovieSceneEvaluationState.cpp b/Engine/Source/Runtime/MovieScene/Private/Evaluation/MovieSceneEvaluationState.cpp
--- a/Engine/Source/Runtime/MovieScene/Private/Evaluation/MovieSceneEvaluationState.cpp
+++ b/Engine/Source/Runtime/MovieScene/Private/Evaluation/MovieSceneEvaluationState.cpp
@@ -134,8 +134,8 @@ void FMovieSceneObjectCache::SetSequence(UMovieSceneSequence& InSequence, FMovie
 
 void FMovieSceneObjectCache::UpdateBindings(const FGuid& InGuid, IMovieScenePlayer& Player)
 {
-	FBoundObjects& Bindings = BoundObjects.FindOrAdd(InGuid);
-	Bindings.Objects.Reset();
+	FBoundObjects* Bindings = &BoundObjects.FindOrAdd(InGuid);
+	Bindings->Objects.Reset();
 
 	if (auto* Children = ChildBindings.Find(InGuid))
 	{
@@ -167,7 +167,11 @@ void FMovieSceneObjectCache::UpdateBindings(const FGuid& InGuid, IMovieScenePlay
 		{
 			ChildBindings.FindOrAdd(Possessable->GetParent()).AddUnique(InGuid);
 
-			for (TWeakObjectPtr<> Parent : FindBoundObjects(Possessable->GetParent(), Player))
+			TArrayView<TWeakObjectPtr<>> ParentBoundObjects = FindBoundObjects(Possessable->GetParent(), Player);
+
+			// Refresh bindings in case of map changes
+			Bindings = BoundObjects.Find(InGuid);
+			for (TWeakObjectPtr<> Parent : ParentBoundObjects)
 			{
 				if (bUseParentsAsContext)
 				{
@@ -182,7 +186,7 @@ void FMovieSceneObjectCache::UpdateBindings(const FGuid& InGuid, IMovieScenePlay
 				Player.ResolveBoundObjects(InGuid, SequenceID, *Sequence, ResolutionContext, FoundObjects);
 				for (UObject* Object : FoundObjects)
 				{
-					Bindings.Objects.Add(Object);
+					Bindings->Objects.Add(Object);
 				}
 			}
 		}
@@ -192,7 +196,7 @@ void FMovieSceneObjectCache::UpdateBindings(const FGuid& InGuid, IMovieScenePlay
 			Player.ResolveBoundObjects(InGuid, SequenceID, *Sequence, ResolutionContext, FoundObjects);
 			for (UObject* Object : FoundObjects)
 			{
-				Bindings.Objects.Add(Object);
+				Bindings->Objects.Add(Object);
 			}
 		}
 	}
@@ -202,13 +206,13 @@ void FMovieSceneObjectCache::UpdateBindings(const FGuid& InGuid, IMovieScenePlay
 		UObject* SpawnedObject = Player.GetSpawnRegister().FindSpawnedObject(InGuid, SequenceID);
 		if (SpawnedObject)
 		{
-			Bindings.Objects.Add(SpawnedObject);
+			Bindings->Objects.Add(SpawnedObject);
 		}
 	}
 
-	if (Bindings.Objects.Num())
+	if (Bindings->Objects.Num())
 	{
-		Bindings.bUpToDate = true;
+		Bindings->bUpToDate = true;
 	}
 }
 
diff --git a/Engine/Source/Runtime/MovieScene/Private/MovieSceneSignedObject.cpp b/Engine/Source/Runtime/MovieScene/Private/MovieSceneSignedObject.cpp
--- a/Engine/Source/Runtime/MovieScene/Private/MovieSceneSignedObject.cpp
+++ b/Engine/Source/Runtime/MovieScene/Private/MovieSceneSignedObject.cpp
@@ -33,7 +33,10 @@ void UMovieSceneSignedObject::MarkAsChanged()
 bool UMovieSceneSignedObject::Modify(bool bAlwaysMarkDirty)
 {
 	bool bModified = Super::Modify(bAlwaysMarkDirty);
-	MarkAsChanged();
+	if ( bAlwaysMarkDirty )
+	{
+		MarkAsChanged();
+	}
 	return bModified;
 }
 
diff --git a/Engine/Source/Runtime/NetworkReplayStreaming/HttpNetworkReplayStreaming/Private/HttpNetworkReplayStreaming.cpp b/Engine/Source/Runtime/NetworkReplayStreaming/HttpNetworkReplayStreaming/Private/HttpNetworkReplayStreaming.cpp
--- a/Engine/Source/Runtime/NetworkReplayStreaming/HttpNetworkReplayStreaming/Private/HttpNetworkReplayStreaming.cpp
+++ b/Engine/Source/Runtime/NetworkReplayStreaming/HttpNetworkReplayStreaming/Private/HttpNetworkReplayStreaming.cpp
@@ -1342,7 +1342,7 @@ void FHttpNetworkReplayStreamer::EnumerateStreams( const FNetworkReplayVersion&
 	// Add optional User parameter (filter replays by a user that was in the replay)
 	if ( !UserString.IsEmpty() )
 	{
-		URL += FString::Printf( TEXT( "&user=%s" ), *UserString );
+		URL += FString::Printf( TEXT( "&user=%s" ), *FGenericPlatformHttp::UrlEncode( UserString ) );
 	}
 
 	// Add any extra parms now
diff --git a/Engine/Source/Runtime/NetworkReplayStreaming/InMemoryNetworkReplayStreaming/Private/InMemoryNetworkReplayStreaming.cpp b/Engine/Source/Runtime/NetworkReplayStreaming/InMemoryNetworkReplayStreaming/Private/InMemoryNetworkReplayStreaming.cpp
--- a/Engine/Source/Runtime/NetworkReplayStreaming/InMemoryNetworkReplayStreaming/Private/InMemoryNetworkReplayStreaming.cpp
+++ b/Engine/Source/Runtime/NetworkReplayStreaming/InMemoryNetworkReplayStreaming/Private/InMemoryNetworkReplayStreaming.cpp
@@ -132,9 +132,7 @@ bool FInMemoryNetworkReplayStreamer::IsDataAvailable() const
 	// Assumptions:
 	// 1. All streamer instances run on the same thread, not simultaneously
 	// 2. A recording DemoNetDriver will write either no frames or entire frames each time it ticks
-	check(StreamerState == EStreamerState::Playback);
-
-	return FileAr.IsValid() && FileAr->Tell() < FileAr->TotalSize();
+	return StreamerState == EStreamerState::Playback && FileAr.IsValid() && FileAr->Tell() < FileAr->TotalSize();
 }
 
 bool FInMemoryNetworkReplayStreamer::IsLive() const
@@ -405,23 +403,27 @@ void FInMemoryNetworkReplayStreamer::GotoTimeInMS(const uint32 TimeInMS, const F
 
 	FInMemoryReplay* FoundReplay = GetCurrentReplayChecked();
 
-	if ( FoundReplay->Checkpoints.Num() > 0 && TimeInMS >= FoundReplay->Checkpoints.Last().TimeInMS )
+	// Checkpoints are sorted by time. Look backwards through the array
+	// to find the one immediately preceding the target time.
+	for (int32 i = FoundReplay->Checkpoints.Num() - 1; i >= 0; --i)
 	{
-		// If we're after the very last checkpoint, that's the one we want
-		CheckpointIndex = FoundReplay->Checkpoints.Num() - 1;
+		if (FoundReplay->Checkpoints[i].TimeInMS <= TimeInMS)
+		{
+			CheckpointIndex = i;
+			break;
+		}
 	}
-	else
+
+	if (CheckpointIndex == -1)
 	{
-		// Checkpoints should be sorted by time, return the checkpoint that exists right before the current time
-		// For fine scrubbing, we'll fast forward the rest of the way
-		// NOTE - If we're right before the very first checkpoint, we'll return -1, which is what we want when we want to start from the very beginning
-		for ( int i = 0; i < FoundReplay->Checkpoints.Num(); i++ )
+		// No checkpoint was found. We may be going the beginning of the stream without an explicit
+		// checkpoint, but if the target time is before the start time of the first stream chunk,
+		// the data was likely discarded due to the TimeBufferHintSeconds value and we can't do anything
+		// except report an error.
+		if (FoundReplay->StreamChunks.Num() == 0 || FoundReplay->StreamChunks[0].TimeInMS > TimeInMS)
 		{
-			if ( TimeInMS < FoundReplay->Checkpoints[i].TimeInMS )
-			{
-				CheckpointIndex = i - 1;
-				break;
-			}
+			Delegate.ExecuteIfBound(false, TimeInMS);
+			return;
 		}
 	}
 
diff --git a/Engine/Source/Runtime/Online/WebSockets/Private/Lws/LwsWebSocket.cpp b/Engine/Source/Runtime/Online/WebSockets/Private/Lws/LwsWebSocket.cpp
--- a/Engine/Source/Runtime/Online/WebSockets/Private/Lws/LwsWebSocket.cpp
+++ b/Engine/Source/Runtime/Online/WebSockets/Private/Lws/LwsWebSocket.cpp
@@ -24,11 +24,12 @@ namespace {
 	};
 }
 
-FLwsWebSocket::FLwsWebSocket(const FString& InUrl, const TArray<FString>& InProtocols, lws_context* Context)
+FLwsWebSocket::FLwsWebSocket(const FString& InUrl, const TArray<FString>& InProtocols, lws_context* Context, const FString& InUpgradeHeader)
 	: LwsContext(Context)
 	, LwsConnection(nullptr)
 	, Url(InUrl)
 	, Protocols(InProtocols)
+	, UpgradeHeader(InUpgradeHeader)
 	, ReceiveBuffer()
 	, CloseCode(0)
 	, CloseReason()
@@ -206,7 +207,7 @@ bool FLwsSendBuffer::Write(struct lws* LwsConnection)
 	return BytesWritten + (int32)(LWS_PRE) >= Payload.Num();
 }
 
-int FLwsWebSocket::LwsCallback(lws* Instance, lws_callback_reasons Reason, const void* Data, size_t Length)
+int FLwsWebSocket::LwsCallback(lws* Instance, lws_callback_reasons Reason, void* Data, size_t Length)
 {
 	switch (Reason)
 	{
@@ -303,6 +304,16 @@ int FLwsWebSocket::LwsCallback(lws* Instance, lws_callback_reasons Reason, const
 		}
 		break;
 	}
+	case LWS_CALLBACK_CLIENT_APPEND_HANDSHAKE_HEADER:
+	{
+		if (!UpgradeHeader.IsEmpty())
+		{
+			// FIXME: Should probably use strcat but libws suggests snprintf (https://libwebsockets.org/lws-api-doc-master/html/)
+			char** WriteableString = reinterpret_cast<char**>(Data);
+			*WriteableString += FCStringAnsi::Snprintf(*WriteableString, Length, TCHAR_TO_ANSI(*UpgradeHeader));
+		}
+		break;
+	}
 	default:
 		break;
 	}
diff --git a/Engine/Source/Runtime/Online/WebSockets/Private/Lws/LwsWebSocket.h b/Engine/Source/Runtime/Online/WebSockets/Private/Lws/LwsWebSocket.h
--- a/Engine/Source/Runtime/Online/WebSockets/Private/Lws/LwsWebSocket.h
+++ b/Engine/Source/Runtime/Online/WebSockets/Private/Lws/LwsWebSocket.h
@@ -113,10 +113,10 @@ public:
 		return RawMessageEvent;
 	}
 
-	int LwsCallback(lws* Instance, lws_callback_reasons Reason, const void* Data, size_t Length);
+	int LwsCallback(lws* Instance, lws_callback_reasons Reason, void* Data, size_t Length);
 private:
 
-	FLwsWebSocket(const FString& Url, const TArray<FString>& Protocols, lws_context* Context);
+	FLwsWebSocket(const FString& Url, const TArray<FString>& Protocols, lws_context* Context, const FString& UpgradeHeader);
 	void SendFromQueue();
 	void FlushQueues();
 	void DelayConnectionError(const FString& Error);
@@ -131,6 +131,7 @@ private:
 	struct lws *LwsConnection;
 	FString Url;
 	TArray<FString> Protocols;
+	FString UpgradeHeader;
 
 	FString ReceiveBuffer;
 	TQueue<TSharedPtr<FLwsSendBuffer>, EQueueMode::Spsc> SendQueue;
diff --git a/Engine/Source/Runtime/Online/WebSockets/Private/Lws/LwsWebSocketsManager.cpp b/Engine/Source/Runtime/Online/WebSockets/Private/Lws/LwsWebSocketsManager.cpp
--- a/Engine/Source/Runtime/Online/WebSockets/Private/Lws/LwsWebSocketsManager.cpp
+++ b/Engine/Source/Runtime/Online/WebSockets/Private/Lws/LwsWebSocketsManager.cpp
@@ -68,6 +68,7 @@ int FLwsWebSocketsManager::CallbackWrapper(lws* Connection, lws_callback_reasons
 	case LWS_CALLBACK_CLIENT_CONNECTION_ERROR:
 	case LWS_CALLBACK_CLIENT_WRITEABLE:
 	case LWS_CALLBACK_SERVER_WRITEABLE:
+	case LWS_CALLBACK_CLIENT_APPEND_HANDSHAKE_HEADER:
 		break;
 	}
 
@@ -101,10 +102,10 @@ bool FLwsWebSocketsManager::Tick(float DeltaTime)
 	return true;
 }
 
-TSharedRef<IWebSocket> FLwsWebSocketsManager::CreateWebSocket(const FString& Url, const TArray<FString>& Protocols)
+TSharedRef<IWebSocket> FLwsWebSocketsManager::CreateWebSocket(const FString& Url, const TArray<FString>& Protocols, const FString& UpgradeHeader)
 {
 	check(LwsContext);
-	TSharedRef<IWebSocket> Socket = MakeShareable(new FLwsWebSocket(Url, Protocols, LwsContext));
+	TSharedRef<IWebSocket> Socket = MakeShareable(new FLwsWebSocket(Url, Protocols, LwsContext, UpgradeHeader));
 	Sockets.Emplace(Socket);
 	return Socket;
 }
diff --git a/Engine/Source/Runtime/Online/WebSockets/Private/Lws/LwsWebSocketsManager.h b/Engine/Source/Runtime/Online/WebSockets/Private/Lws/LwsWebSocketsManager.h
--- a/Engine/Source/Runtime/Online/WebSockets/Private/Lws/LwsWebSocketsManager.h
+++ b/Engine/Source/Runtime/Online/WebSockets/Private/Lws/LwsWebSocketsManager.h
@@ -37,7 +37,7 @@ private:
 
 	virtual void ShutdownWebSockets() override;
 
-	virtual TSharedRef<IWebSocket> CreateWebSocket(const FString& Url, const TArray<FString>& Protocols) override;
+	virtual TSharedRef<IWebSocket> CreateWebSocket(const FString& Url, const TArray<FString>& Protocols, const FString& UpgradeHeader) override;
 	// end IWebSocketsManager overrides
 
 	// TickerObjectBase override
diff --git a/Engine/Source/Runtime/Online/WebSockets/Private/WebSocketsModule.cpp b/Engine/Source/Runtime/Online/WebSockets/Private/WebSocketsModule.cpp
--- a/Engine/Source/Runtime/Online/WebSockets/Private/WebSocketsModule.cpp
+++ b/Engine/Source/Runtime/Online/WebSockets/Private/WebSocketsModule.cpp
@@ -14,12 +14,22 @@ IMPLEMENT_MODULE(FWebSocketsModule, WebSockets);
 
 FWebSocketsModule* FWebSocketsModule::Singleton = nullptr;
 
+/*static*/ FString FWebSocketsModule::BuildUpgradeHeader(const TMap<FString, FString>& Headers)
+{
+	FString HeaderString;
+	for (const auto& OneHeader : Headers)
+	{
+		HeaderString += FString::Printf(TEXT("%s: %s\r\n"), *OneHeader.Key, *OneHeader.Value);
+	}
+	return HeaderString;
+}
+
 void FWebSocketsModule::StartupModule()
 {
 	Singleton = this;
 
 #if WITH_WEBSOCKETS
-	const FString Protocols[] = {TEXT("v10.stomp"), TEXT("v11.stomp"), TEXT("v12.stomp")};
+	const FString Protocols[] = {TEXT("ws"), TEXT("wss"), TEXT("v10.stomp"), TEXT("v11.stomp"), TEXT("v12.stomp")};
 
 	WebSocketsManager = new FPlatformWebSocketsManager;
 	WebSocketsManager->InitWebSockets(Protocols);
@@ -45,18 +55,18 @@ FWebSocketsModule& FWebSocketsModule::Get()
 }
 
 #if WITH_WEBSOCKETS
-TSharedRef<IWebSocket> FWebSocketsModule::CreateWebSocket(const FString& Url, const TArray<FString>& Protocols)
+TSharedRef<IWebSocket> FWebSocketsModule::CreateWebSocket(const FString& Url, const TArray<FString>& Protocols, const TMap<FString, FString>& UpgradeHeaders)
 {
 	check(WebSocketsManager);
-	return WebSocketsManager->CreateWebSocket(Url, Protocols);
+	return WebSocketsManager->CreateWebSocket(Url, Protocols, BuildUpgradeHeader(UpgradeHeaders));
 }
 
-TSharedRef<IWebSocket> FWebSocketsModule::CreateWebSocket(const FString& Url, const FString& Protocol)
+TSharedRef<IWebSocket> FWebSocketsModule::CreateWebSocket(const FString& Url, const FString& Protocol, const TMap<FString, FString>& UpgradeHeaders)
 {
 	check(WebSocketsManager);
 
 	TArray<FString> Protocols;
 	Protocols.Add(Protocol);
-	return WebSocketsManager->CreateWebSocket(Url, Protocols);
+	return WebSocketsManager->CreateWebSocket(Url, Protocols, BuildUpgradeHeader(UpgradeHeaders));
 }
 #endif // #if WITH_WEBSOCKETS
diff --git a/Engine/Source/Runtime/Online/WebSockets/Public/IWebSocketsManager.h b/Engine/Source/Runtime/Online/WebSockets/Public/IWebSocketsManager.h
--- a/Engine/Source/Runtime/Online/WebSockets/Public/IWebSocketsManager.h
+++ b/Engine/Source/Runtime/Online/WebSockets/Public/IWebSocketsManager.h
@@ -29,7 +29,7 @@ public:
 	 * @param Protocols a list of protocols the client will handle.
 	 * @return new IWebSocket instance
 	 */
-	virtual TSharedRef<IWebSocket> CreateWebSocket(const FString& Url, const TArray<FString>& Protocols) = 0;
+	virtual TSharedRef<IWebSocket> CreateWebSocket(const FString& Url, const TArray<FString>& Protocols, const FString& UpgradeHeader) = 0;
 
 	/**
 	 * Virtual destructor
diff --git a/Engine/Source/Runtime/Online/WebSockets/Public/WebSocketsModule.h b/Engine/Source/Runtime/Online/WebSockets/Public/WebSocketsModule.h
--- a/Engine/Source/Runtime/Online/WebSockets/Public/WebSocketsModule.h
+++ b/Engine/Source/Runtime/Online/WebSockets/Public/WebSocketsModule.h
@@ -38,7 +38,7 @@ public:
 	 * @param Protocols a list of protocols the client will handle.
 	 * @return new IWebSocket instance
 	 */
-	virtual TSharedRef<IWebSocket> CreateWebSocket(const FString& Url, const TArray<FString>& Protocols);
+	virtual TSharedRef<IWebSocket> CreateWebSocket(const FString& Url, const TArray<FString>& Protocols, const TMap<FString, FString>& UpgradeHeaders = TMap<FString, FString>());
 
 
 	/**
@@ -48,11 +48,13 @@ public:
 	 * @param Protocol an optional sub-protocol. If missing, an empty string is assumed.
 	 * @return new IWebSocket instance
 	 */
-	virtual TSharedRef<IWebSocket> CreateWebSocket(const FString& Url, const FString& Protocol = FString());
+	virtual TSharedRef<IWebSocket> CreateWebSocket(const FString& Url, const FString& Protocol = FString(), const TMap<FString, FString>& UpgradeHeaders = TMap<FString, FString>());
 #endif // #if WITH_WEBSOCKETS
 
 private:
 
+	static FString BuildUpgradeHeader(const TMap<FString, FString>& Headers);
+
 	// IModuleInterface
 
 	/**
diff --git a/Engine/Source/Runtime/SlateRHIRenderer/Private/SlateRHIRenderingPolicy.cpp b/Engine/Source/Runtime/SlateRHIRenderer/Private/SlateRHIRenderingPolicy.cpp
--- a/Engine/Source/Runtime/SlateRHIRenderer/Private/SlateRHIRenderingPolicy.cpp
+++ b/Engine/Source/Runtime/SlateRHIRenderer/Private/SlateRHIRenderingPolicy.cpp
@@ -517,8 +517,10 @@ void FSlateRHIRenderingPolicy::DrawElements(FRHICommandListImmediate& RHICmdList
 					// MaterialParameterCollections will not be correct for this scene, should they be
 					// used.
 					ActiveSceneIndex = NumScenes - 1;
-#if WITH_EDITOR
+#if UE_BUILD_DEBUG
+	#if WITH_EDITOR
 					UE_LOG(LogSlate, Error, TEXT("Invalid scene index in batch: %d of %d known scenes!"), RenderBatch.SceneIndex, ResourceManager->GetSceneCount());
+	#endif
 #endif
 				}
 
diff --git a/Engine/Source/Runtime/SynthBenchmark/Private/SynthBenchmarkPrivate.cpp b/Engine/Source/Runtime/SynthBenchmark/Private/SynthBenchmarkPrivate.cpp
--- a/Engine/Source/Runtime/SynthBenchmark/Private/SynthBenchmarkPrivate.cpp
+++ b/Engine/Source/Runtime/SynthBenchmark/Private/SynthBenchmarkPrivate.cpp
@@ -175,8 +175,8 @@ void FSynthBenchmark::Run(FSynthBenchmarkResults& InOut, bool bGPUBenchmark, flo
 		// The hardware is slow, we don't need a long test and risk driver TDR (driver recovery).
 		// We have seen this problem on very low end GPUs.
 		{
-			const float fFirstWorkScale = 0.01f;
-			const float fSecondWorkScale = 0.1f;
+			const float fFirstWorkScale = 0.01f * WorkScale;
+			const float fSecondWorkScale = 0.1f * WorkScale;
 
 			RendererModule.GPUBenchmark(InOut, fFirstWorkScale);
 			GPUTime = InOut.ComputeTotalGPUTime();
diff --git a/Engine/Source/Runtime/UMG/Private/Components/BackgroundBlur.cpp b/Engine/Source/Runtime/UMG/Private/Components/BackgroundBlur.cpp
--- a/Engine/Source/Runtime/UMG/Private/Components/BackgroundBlur.cpp
+++ b/Engine/Source/Runtime/UMG/Private/Components/BackgroundBlur.cpp
@@ -35,8 +35,7 @@ void UBackgroundBlur::ReleaseSlateResources(bool bReleaseChildren)
 
 UClass* UBackgroundBlur::GetSlotClass() const
 {
-	//return UBackgroundBlurSlot::StaticClass();
-	return Super::GetSlotClass();
+	return UBackgroundBlurSlot::StaticClass();
 }
 
 TSharedRef<SWidget> UBackgroundBlur::RebuildWidget()
@@ -45,8 +44,7 @@ TSharedRef<SWidget> UBackgroundBlur::RebuildWidget()
 
 	if ( GetChildrenCount() > 0 )
 	{
-		//Cast<UBackgroundBlurSlot>(GetContentSlot())->BuildSlot(MyBackgroundBlur.ToSharedRef());
-		MyBackgroundBlur->SetContent(GetContentSlot()->Content ? GetContentSlot()->Content->TakeWidget() : SNullWidget::NullWidget);
+		Cast<UBackgroundBlurSlot>(GetContentSlot())->BuildSlot(MyBackgroundBlur.ToSharedRef());
 	}
 	
 	return BuildDesignTimeWidget(MyBackgroundBlur.ToSharedRef());
@@ -76,18 +74,12 @@ void UBackgroundBlur::OnSlotAdded(UPanelSlot* InSlot)
 	BackgroundBlurSlot->HorizontalAlignment = HorizontalAlignment;
 	BackgroundBlurSlot->VerticalAlignment = VerticalAlignment;
 
-	//// Add the child to the live slot if it already exists
-	//if (MyBackgroundBlur.IsValid())
-	//{
-	//	// Construct the underlying slot
-	//	BackgroundBlurSlot->BuildSlot(MyBackgroundBlur.ToSharedRef());
-	//}
-
+	// Add the child to the live slot if it already exists
 	if (MyBackgroundBlur.IsValid())
 	{
-		MyBackgroundBlur->SetContent(InSlot->Content ? InSlot->Content->TakeWidget() : SNullWidget::NullWidget);
+		// Construct the underlying slot
+		BackgroundBlurSlot->BuildSlot(MyBackgroundBlur.ToSharedRef());
 	}
-
 }
 	
 void UBackgroundBlur::OnSlotRemoved(UPanelSlot* InSlot)
@@ -97,7 +89,8 @@ void UBackgroundBlur::OnSlotRemoved(UPanelSlot* InSlot)
 	{
 		MyBackgroundBlur->SetContent(SNullWidget::NullWidget);
 	}
-	}
+}
+
 void UBackgroundBlur::SetPadding(FMargin InPadding)
 {
 	Padding = InPadding;
diff --git a/Engine/Source/Runtime/UMG/Private/Components/EditableTextBox.cpp b/Engine/Source/Runtime/UMG/Private/Components/EditableTextBox.cpp
--- a/Engine/Source/Runtime/UMG/Private/Components/EditableTextBox.cpp
+++ b/Engine/Source/Runtime/UMG/Private/Components/EditableTextBox.cpp
@@ -107,6 +107,15 @@ void UEditableTextBox::SetText(FText InText)
 	}
 }
 
+void UEditableTextBox::SetHintText(FText InText)
+{
+	HintText = InText;
+	if (MyEditableTextBlock.IsValid())
+	{
+		MyEditableTextBlock->SetHintText(HintText);
+	}
+}
+
 void UEditableTextBox::SetError(FText InError)
 {
 	if ( MyEditableTextBlock.IsValid() )
diff --git a/Engine/Source/Runtime/UMG/Public/Components/EditableTextBox.h b/Engine/Source/Runtime/UMG/Public/Components/EditableTextBox.h
--- a/Engine/Source/Runtime/UMG/Public/Components/EditableTextBox.h
+++ b/Engine/Source/Runtime/UMG/Public/Components/EditableTextBox.h
@@ -145,6 +145,9 @@ public:
 	UFUNCTION(BlueprintCallable, Category="Widget", meta=(DisplayName="SetText (Text Box)"))
 	void SetText(FText InText);
 
+	UFUNCTION(BlueprintCallable, Category = "Widget", meta = (DisplayName = "Set Hint Text (Text Box)"))
+	void SetHintText(FText InText);
+
 	UFUNCTION(BlueprintCallable, Category="Widget",  meta=(DisplayName="SetError (Text Box)"))
 	void SetError(FText InError);
 
diff --git a/Engine/Source/Runtime/Windows/D3D11RHI/Private/D3D11Device.cpp b/Engine/Source/Runtime/Windows/D3D11RHI/Private/D3D11Device.cpp
--- a/Engine/Source/Runtime/Windows/D3D11RHI/Private/D3D11Device.cpp
+++ b/Engine/Source/Runtime/Windows/D3D11RHI/Private/D3D11Device.cpp
@@ -214,8 +214,9 @@ FD3D11DynamicRHI::~FD3D11DynamicRHI()
 {
 	UE_LOG(LogD3D11RHI, Log, TEXT("~FD3D11DynamicRHI"));
 
-	check(Direct3DDeviceIMContext == nullptr);
-	check(Direct3DDevice == nullptr);
+	// Removed until shutdown crashes in exception handler are fixed.
+	//check(Direct3DDeviceIMContext == nullptr);
+	//check(Direct3DDevice == nullptr);
 }
 
 void FD3D11DynamicRHI::Shutdown()
