Fix static meshes with auto-generated LODs no longer sharing lightmaps between LODs

diff --git a/Engine/Source/Runtime/Engine/Classes/Engine/StaticMesh.h b/Engine/Source/Runtime/Engine/Classes/Engine/StaticMesh.h
--- a/Engine/Source/Runtime/Engine/Classes/Engine/StaticMesh.h
+++ b/Engine/Source/Runtime/Engine/Classes/Engine/StaticMesh.h
@@ -805,6 +805,12 @@ public:
 	ENGINE_API void CalculateExtendedBounds();
 
 #if WITH_EDITOR
+
+	/**
+	 * Returns true if LODs of this static mesh may share texture lightmaps.
+	 */
+	bool CanLODsShareStaticLighting() const;
+
 	/**
 	 * Retrieves the names of all LOD groups.
 	 */
diff --git a/Engine/Source/Runtime/Engine/Private/Components/StaticMeshComponent.cpp b/Engine/Source/Runtime/Engine/Private/Components/StaticMeshComponent.cpp
--- a/Engine/Source/Runtime/Engine/Private/Components/StaticMeshComponent.cpp
+++ b/Engine/Source/Runtime/Engine/Private/Components/StaticMeshComponent.cpp
@@ -521,11 +521,9 @@ const FMeshMapBuildData* UStaticMeshComponent::GetMeshMapBuildData(const FStatic
 	{
 		// Check that the static-mesh hasn't been changed to be incompatible with the cached light-map.
 		int32 NumLODs = GetStaticMesh()->RenderData->LODResources.Num();
+		bool bLODsShareStaticLighting = GetStaticMesh()->RenderData->bLODsShareStaticLighting;
 
-		// SpeedTrees are set up for lighting to share between LODs
-		bool bCanLODsShareStaticLighting = GetStaticMesh()->SpeedTreeWind.IsValid();
-
-		if (!bCanLODsShareStaticLighting && NumLODs != LODData.Num())
+		if (!bLODsShareStaticLighting && NumLODs != LODData.Num())
 		{
 			return NULL;
 		}
@@ -560,8 +558,8 @@ const FMeshMapBuildData* UStaticMeshComponent::GetMeshMapBuildData(const FStatic
 			{
 				return MapBuildData->GetMeshBuildData(LODInfo.MapBuildDataId);
 			}
-			}
 		}
+	}
 	
 	return NULL;
 }
diff --git a/Engine/Source/Runtime/Engine/Private/InstancedStaticMesh.cpp b/Engine/Source/Runtime/Engine/Private/InstancedStaticMesh.cpp
--- a/Engine/Source/Runtime/Engine/Private/InstancedStaticMesh.cpp
+++ b/Engine/Source/Runtime/Engine/Private/InstancedStaticMesh.cpp
@@ -23,9 +23,6 @@
 #include "GameFramework/WorldSettings.h"
 #include "ComponentRecreateRenderStateContext.h"
 #include "SceneManagement.h"
-#if WITH_EDITOR
-#include "RawMesh.h"
-#endif
 
 const int32 InstancedStaticMeshMaxTexCoord = 8;
 
@@ -1267,17 +1264,9 @@ void UInstancedStaticMeshComponent::GetStaticLightingInfo(FStaticLightingPrimiti
 				->AddToken(FTextToken::Create(NSLOCTEXT("InstancedStaticMesh", "LargeStaticLightingWarning", "The total lightmap size for this InstancedStaticMeshComponent is large, consider reducing the component's lightmap resolution or number of mesh instances in this component")));
 		}
 
-		bool bHasMultipleLODModels = false;
-		for (int32 LODIndex = 1; LODIndex < GetStaticMesh()->SourceModels.Num(); ++LODIndex)
-		{
-			if (!GetStaticMesh()->SourceModels[LODIndex].RawMeshBulkData->IsEmpty())
-			{
-				bHasMultipleLODModels = true;
-				break;
-			}
-		}
+		// TODO: Support separate static lighting in LODs for instanced meshes.
 
-		if (bHasMultipleLODModels)
+		if (!GetStaticMesh()->CanLODsShareStaticLighting())
 		{
 			//TODO: Detect if the UVs for all sub-LODs overlap the base LOD UVs and omit this warning if they do.
 			FMessageLog("LightingResults").Message(EMessageSeverity::Warning)
@@ -1285,8 +1274,7 @@ void UInstancedStaticMeshComponent::GetStaticLightingInfo(FStaticLightingPrimiti
 				->AddToken(FTextToken::Create(NSLOCTEXT("InstancedStaticMesh", "UniqueStaticLightingForLODWarning", "Instanced meshes don't yet support unique static lighting for each LOD. Lighting on LOD 1+ may be incorrect unless lightmap UVs are the same for all LODs.")));
 		}
 
-		// TODO: We currently only support one LOD of static lighting in instanced meshes
-		// Need to create per-LOD instance data to fix that
+		// Force sharing LOD 0 lightmaps for now.
 		int32 NumLODs = 1;
 
 		CachedMappings.Reset(PerInstanceSMData.Num() * NumLODs);
diff --git a/Engine/Source/Runtime/Engine/Private/StaticMesh.cpp b/Engine/Source/Runtime/Engine/Private/StaticMesh.cpp
--- a/Engine/Source/Runtime/Engine/Private/StaticMesh.cpp
+++ b/Engine/Source/Runtime/Engine/Private/StaticMesh.cpp
@@ -491,6 +491,7 @@ void FStaticMeshLODResources::ReleaseResources()
 ------------------------------------------------------------------------------*/
 
 FStaticMeshRenderData::FStaticMeshRenderData()
+	: bLODsShareStaticLighting(false)
 {
 	for (int32 LODIndex = 0; LODIndex < MAX_STATIC_MESH_LODS; ++LODIndex)
 	{
@@ -554,6 +555,7 @@ void FStaticMeshRenderData::Serialize(FArchive& Ar, UStaticMesh* Owner, bool bCo
 	}
 
 	Ar << Bounds;
+	Ar << bLODsShareStaticLighting;
 
 	if (Ar.IsLoading() && Ar.CustomVer(FRenderingObjectVersion::GUID) < FRenderingObjectVersion::TextureStreamingMeshUVChannelData)
 	{
@@ -994,7 +996,7 @@ FArchive& operator<<(FArchive& Ar, FMeshBuildSettings& BuildSettings)
 // differences, etc.) replace the version GUID below with a new one.
 // In case of merge conflicts with DDC versions, you MUST generate a new GUID
 // and set this new GUID as the version.                                       
-#define STATICMESH_DERIVEDDATA_VER TEXT("71F138FFEE1B456C86761D2AD2F3ADB8")
+#define STATICMESH_DERIVEDDATA_VER TEXT("A1C08472697A4FC2B1E0A31E4B382B6E")
 
 static const FString& GetStaticMeshDerivedDataVersion()
 {
@@ -1203,6 +1205,7 @@ void FStaticMeshRenderData::Cache(UStaticMesh* Owner, const FStaticMeshLODSettin
 			IMeshUtilities& MeshUtilities = FModuleManager::Get().LoadModuleChecked<IMeshUtilities>(TEXT("MeshUtilities"));
 			MeshUtilities.BuildStaticMesh(*this, Owner->SourceModels, LODGroup, Owner->ImportVersion);
 			ComputeUVDensities();
+			bLODsShareStaticLighting = Owner->CanLODsShareStaticLighting();
 			FMemoryWriter Ar(DerivedData, /*bIsPersistent=*/ true);
 			Serialize(Ar, Owner, /*bCooked=*/ false);
 			GetDerivedDataCacheRef().Put(*DerivedDataKey, DerivedData);
@@ -3134,6 +3137,23 @@ const FStaticMeshLODResources& UStaticMesh::GetLODForExport(int32 LODIndex) cons
 }
 
 #if WITH_EDITOR
+bool UStaticMesh::CanLODsShareStaticLighting() const
+{
+	bool bCanShareData = true;
+	for (int32 LODIndex = 1; bCanShareData && LODIndex < SourceModels.Num(); ++LODIndex)
+	{
+		bCanShareData = bCanShareData && SourceModels[LODIndex].RawMeshBulkData->IsEmpty();
+	}
+
+	if (SpeedTreeWind.IsValid())
+	{
+		// SpeedTrees are set up for lighting to share between LODs
+		bCanShareData = true;
+	}
+
+	return bCanShareData;
+}
+
 void UStaticMesh::ConvertLegacyLODDistance()
 {
 	check(SourceModels.Num() > 0);
diff --git a/Engine/Source/Runtime/Engine/Private/StaticMeshLight.cpp b/Engine/Source/Runtime/Engine/Private/StaticMeshLight.cpp
--- a/Engine/Source/Runtime/Engine/Private/StaticMeshLight.cpp
+++ b/Engine/Source/Runtime/Engine/Private/StaticMeshLight.cpp
@@ -303,7 +303,7 @@ void UStaticMeshComponent::GetStaticLightingInfo(FStaticLightingPrimitiveInfo& O
 		GetLightMapResolution( BaseLightMapWidth, BaseLightMapHeight );
 
 		TArray<FStaticMeshStaticLightingMesh*> StaticLightingMeshes;
-		bool bCanLODsShareStaticLighting = GetStaticMesh()->SpeedTreeWind.IsValid();
+		bool bCanLODsShareStaticLighting = GetStaticMesh()->CanLODsShareStaticLighting();
 		int32 NumLODs = bCanLODsShareStaticLighting ? 1 : GetStaticMesh()->RenderData->LODResources.Num();
 		for(int32 LODIndex = 0;LODIndex < NumLODs;LODIndex++)
 		{
diff --git a/Engine/Source/Runtime/Engine/Private/StaticMeshRender.cpp b/Engine/Source/Runtime/Engine/Private/StaticMeshRender.cpp
--- a/Engine/Source/Runtime/Engine/Private/StaticMeshRender.cpp
+++ b/Engine/Source/Runtime/Engine/Private/StaticMeshRender.cpp
@@ -104,7 +104,7 @@ static FAutoConsoleCommand GToggleForceDefaultMaterialCmd(
 	);
 
 /** Initialization constructor. */
-FStaticMeshSceneProxy::FStaticMeshSceneProxy(UStaticMeshComponent* InComponent, bool bCanLODsShareStaticLighting):
+FStaticMeshSceneProxy::FStaticMeshSceneProxy(UStaticMeshComponent* InComponent, bool bForceLODsShareStaticLighting):
 	FPrimitiveSceneProxy(InComponent, InComponent->GetStaticMesh()->GetFName())
 	, Owner(InComponent->GetOwner())
 	, StaticMesh(InComponent->GetStaticMesh())
@@ -154,12 +154,10 @@ FStaticMeshSceneProxy::FStaticMeshSceneProxy(UStaticMeshComponent* InComponent,
 	// Build the proxy's LOD data.
 	bool bAnySectionCastsShadows = false;
 	LODs.Empty(RenderData->LODResources.Num());
-
-	// SpeedTrees are set up for lighting to share between LODs
-	bCanLODsShareStaticLighting |= InComponent->GetStaticMesh()->SpeedTreeWind.IsValid();
+	const bool bLODsShareStaticLighting = RenderData->bLODsShareStaticLighting || bForceLODsShareStaticLighting;
 	for(int32 LODIndex = 0;LODIndex < RenderData->LODResources.Num();LODIndex++)
 	{
-		FLODInfo* NewLODInfo = new(LODs) FLODInfo(InComponent,LODIndex,bCanLODsShareStaticLighting);
+		FLODInfo* NewLODInfo = new(LODs) FLODInfo(InComponent,LODIndex,bLODsShareStaticLighting);
 
 		// Under certain error conditions an LOD's material will be set to 
 		// DefaultMaterial. Ensure our material view relevance is set properly.
@@ -1368,7 +1366,7 @@ bool FStaticMeshSceneProxy::HasDynamicIndirectShadowCasterRepresentation() const
 }
 
 /** Initialization constructor. */
-FStaticMeshSceneProxy::FLODInfo::FLODInfo(const UStaticMeshComponent* InComponent, int32 LODIndex, bool bCanLODsShareStaticLighting)
+FStaticMeshSceneProxy::FLODInfo::FLODInfo(const UStaticMeshComponent* InComponent, int32 LODIndex, bool bLODsShareStaticLighting)
 	: FLightCacheInterface(nullptr, nullptr)
 	, OverrideColorVertexBuffer(0)
 	, PreCulledIndexBuffer(NULL)
@@ -1413,7 +1411,7 @@ FStaticMeshSceneProxy::FLODInfo::FLODInfo(const UStaticMeshComponent* InComponen
 		}
 	}
 
-	if (bCanLODsShareStaticLighting && InComponent->LODData.IsValidIndex(0))
+	if (LODIndex > 0 && bLODsShareStaticLighting && InComponent->LODData.IsValidIndex(0))
 	{
 		const FStaticMeshComponentLODInfo& ComponentLODInfo = InComponent->LODData[0];
 		const FMeshMapBuildData* MeshMapBuildData = InComponent->GetMeshMapBuildData(ComponentLODInfo);
diff --git a/Engine/Source/Runtime/Engine/Public/StaticMeshResources.h b/Engine/Source/Runtime/Engine/Public/StaticMeshResources.h
--- a/Engine/Source/Runtime/Engine/Public/StaticMeshResources.h
+++ b/Engine/Source/Runtime/Engine/Public/StaticMeshResources.h
@@ -290,6 +290,9 @@ public:
 	/** Bounds of the renderable mesh. */
 	FBoxSphereBounds Bounds;
 
+	/** True if LODs share static lighting data. */
+	bool bLODsShareStaticLighting;
+
 #if WITH_EDITORONLY_DATA
 	/** The derived data key associated with this render data. */
 	FString DerivedDataKey;
@@ -422,7 +425,7 @@ class ENGINE_API FStaticMeshSceneProxy : public FPrimitiveSceneProxy
 public:
 
 	/** Initialization constructor. */
-	FStaticMeshSceneProxy(UStaticMeshComponent* Component, bool bCanLODsShareStaticLighting);
+	FStaticMeshSceneProxy(UStaticMeshComponent* Component, bool bForceLODsShareStaticLighting);
 
 	virtual ~FStaticMeshSceneProxy() {}
 
@@ -540,7 +543,7 @@ protected:
 		const FRawStaticIndexBuffer* PreCulledIndexBuffer;
 
 		/** Initialization constructor. */
-		FLODInfo(const UStaticMeshComponent* InComponent, int32 InLODIndex, bool bCanLODsShareStaticLighting);
+		FLODInfo(const UStaticMeshComponent* InComponent, int32 InLODIndex, bool bLODsShareStaticLighting);
 
 		bool UsesMeshModifyingMaterials() const { return bUsesMeshModifyingMaterials; }
 
