CIS fix for Code Documentation warning.

diff --git a/Engine/Source/Runtime/Engine/Classes/PhysicsEngine/BodyInstance.h b/Engine/Source/Runtime/Engine/Classes/PhysicsEngine/BodyInstance.h
--- a/Engine/Source/Runtime/Engine/Classes/PhysicsEngine/BodyInstance.h
+++ b/Engine/Source/Runtime/Engine/Classes/PhysicsEngine/BodyInstance.h
@@ -1244,11 +1244,13 @@ private:
 //////////////////////////////////////////////////////////////////////////
 // BodyInstance inlines
 
+//~ APIDOCTOOL: Document=Off
 FORCEINLINE_DEBUGGABLE bool FBodyInstance::OverlapMulti(TArray<struct FOverlapResult>& InOutOverlaps, const class UWorld* World, const FTransform* pWorldToComponent, const FVector& Pos, const FRotator& Rot, ECollisionChannel TestChannel, const FComponentQueryParams& Params, const FCollisionResponseParams& ResponseParams, const FCollisionObjectQueryParams& ObjectQueryParams) const
 {
 	// Pass on to FQuat version
 	return OverlapMulti(InOutOverlaps, World, pWorldToComponent, Pos, Rot.Quaternion(), TestChannel, Params, ResponseParams, ObjectQueryParams);
 }
+//~ APIDOCTOOL: Document=On
 
 FORCEINLINE_DEBUGGABLE bool FBodyInstance::OverlapTestForBodies(const FVector& Position, const FQuat& Rotation, const TArray<FBodyInstance*>& Bodies) const
 {
