Copying //UE4/Dev-Core to //UE4/Dev-Main (Source: //UE4/Dev-Core @ 3345728)

diff --git a/Engine/Source/Runtime/Engine/Private/Materials/MaterialShader.cpp b/Engine/Source/Runtime/Engine/Private/Materials/MaterialShader.cpp
--- a/Engine/Source/Runtime/Engine/Private/Materials/MaterialShader.cpp
+++ b/Engine/Source/Runtime/Engine/Private/Materials/MaterialShader.cpp
@@ -2162,7 +2162,7 @@ void FMaterialShaderMap::RegisterSerializedShaders()
 	{
 		if (OrderedMeshShaderMaps[VFIndex] && OrderedMeshShaderMaps[VFIndex]->IsEmpty())
 		{
-			OrderedMeshShaderMaps[VFIndex] = NULL;
+			OrderedMeshShaderMaps[VFIndex] = nullptr;
 		}
 	}
 
@@ -2175,6 +2175,22 @@ void FMaterialShaderMap::RegisterSerializedShaders()
 	}
 }
 
+void FMaterialShaderMap::DiscardSerializedShaders()
+{
+	TShaderMap<FMaterialShaderType>::DiscardSerializedShaders();
+
+	for (int32 VFIndex = 0; VFIndex < OrderedMeshShaderMaps.Num(); VFIndex++)
+	{
+		OrderedMeshShaderMaps[VFIndex] = nullptr;
+	}
+
+	for (int32 Index = MeshShaderMaps.Num() - 1; Index >= 0; Index--)
+	{
+		MeshShaderMaps[Index].DiscardSerializedShaders();
+	}
+	MeshShaderMaps.Empty();
+}
+
 void FMaterialShaderMap::RemovePendingMaterial(FMaterial* Material)
 {
 	for (TMap<TRefCountPtr<FMaterialShaderMap>, TArray<FMaterial*> >::TIterator It(ShaderMapsBeingCompiled); It; ++It)
diff --git a/Engine/Source/Runtime/Engine/Private/Materials/MaterialShared.cpp b/Engine/Source/Runtime/Engine/Private/Materials/MaterialShared.cpp
--- a/Engine/Source/Runtime/Engine/Private/Materials/MaterialShared.cpp
+++ b/Engine/Source/Runtime/Engine/Private/Materials/MaterialShared.cpp
@@ -627,10 +627,7 @@ void FMaterial::SetRenderingThreadShaderMap(FMaterialShaderMap* InMaterialShader
 
 void FMaterial::AddReferencedObjects(FReferenceCollector& Collector)
 {
-	for (int32 ExpressionIndex = 0; ExpressionIndex < ErrorExpressions.Num(); ExpressionIndex++)
-	{
-		Collector.AddReferencedObject(ErrorExpressions[ExpressionIndex]);
-	}
+	Collector.AddReferencedObjects(ErrorExpressions);
 }
 
 struct FLegacyTextureLookup
@@ -732,6 +729,10 @@ void FMaterial::SerializeInlineShaderMap(FArchive& Ar)
 				{
 					GameThreadShaderMap = RenderingThreadShaderMap = LoadedShaderMap;
 				}
+				else
+				{
+					LoadedShaderMap->DiscardSerializedShaders();
+				}
 			}
 		}
 	}
diff --git a/Engine/Source/Runtime/Engine/Private/StaticMesh.cpp b/Engine/Source/Runtime/Engine/Private/StaticMesh.cpp
--- a/Engine/Source/Runtime/Engine/Private/StaticMesh.cpp
+++ b/Engine/Source/Runtime/Engine/Private/StaticMesh.cpp
@@ -2513,6 +2513,11 @@ void UStaticMesh::PostLoad()
 	CreateNavCollision();
 }
 
+bool UStaticMesh::CanBeClusterRoot() const
+{
+	return false;
+}
+
 //
 //	UStaticMesh::GetDesc
 //
diff --git a/Engine/Source/Runtime/Engine/Private/StaticMeshActor.cpp b/Engine/Source/Runtime/Engine/Private/StaticMeshActor.cpp
--- a/Engine/Source/Runtime/Engine/Private/StaticMeshActor.cpp
+++ b/Engine/Source/Runtime/Engine/Private/StaticMeshActor.cpp
@@ -32,6 +32,9 @@ AStaticMeshActor::AStaticMeshActor(const FObjectInitializer& ObjectInitializer)
 	StaticMeshComponent->bUseDefaultCollision = true;
 
 	RootComponent = StaticMeshComponent;
+
+	// By default all static mesh actors can be put inside of a GC cluster (see ULevelActorContainer and ULevel::CreateCluster())
+	bCanBeInCluster = true;
 }
 
 void AStaticMeshActor::BeginPlay()
diff --git a/Engine/Source/Runtime/Engine/Private/UnrealEngine.cpp b/Engine/Source/Runtime/Engine/Private/UnrealEngine.cpp
--- a/Engine/Source/Runtime/Engine/Private/UnrealEngine.cpp
+++ b/Engine/Source/Runtime/Engine/Private/UnrealEngine.cpp
@@ -5474,7 +5474,7 @@ bool UEngine::HandleObjCommand( const TCHAR* Cmd, FOutputDevice& Ar )
 						continue;
 					}
 					FUObjectItem* ObjectItem = GUObjectArray.ObjectToObjectItem(*It);
-					if (ObjectItem->GetOwnerIndex())
+					if (ObjectItem->GetOwnerIndex() > 0)
 					{
 						continue;
 					}
@@ -10934,6 +10934,7 @@ void UEngine::VerifyLoadMapWorldCleanup()
 				FString						ErrorString	= FArchiveTraceRoute::PrintRootPath( Route, World );
 				UE_LOG(LogLoad, Log, TEXT("%s"),*ErrorString);
 				// before asserting.
+
 				UE_LOG(LogLoad, Fatal, TEXT("%s not cleaned up by garbage collection!") LINE_TERMINATOR TEXT("%s") , *World->GetFullName(), *ErrorString );
 			}
 		}
diff --git a/Engine/Source/Runtime/Engine/Private/World.cpp b/Engine/Source/Runtime/Engine/Private/World.cpp
--- a/Engine/Source/Runtime/Engine/Private/World.cpp
+++ b/Engine/Source/Runtime/Engine/Private/World.cpp
@@ -427,9 +427,7 @@ void UWorld::AddReferencedObjects(UObject* InThis, FReferenceCollector& Collecto
 	if( GIsEditor )
 	{
 		Collector.AddReferencedObject( This->PersistentLevel, This );
-
-		Collector.AddReferencedObjects( This->Levels, This );
-
+		Collector.AddReferencedObjects(This->Levels, This);
 		Collector.AddReferencedObject( This->CurrentLevel, This );
 		Collector.AddReferencedObject( This->NetDriver, This );
 		Collector.AddReferencedObject( This->DemoNetDriver, This );
@@ -442,10 +440,6 @@ void UWorld::AddReferencedObjects(UObject* InThis, FReferenceCollector& Collecto
 		Collector.AddReferencedObject( This->NetworkManager, This );
 		Collector.AddReferencedObject( This->NavigationSystem, This );
 		Collector.AddReferencedObject( This->AvoidanceManager, This );
-
-		Collector.AddReferencedObjects( This->ExtraReferencedObjects, This );
-		Collector.AddReferencedObjects( This->StreamingLevels, This );
-		Collector.AddReferencedObjects(This->PerModuleDataObjects, This);
 	}
 #endif
 
@@ -5270,6 +5264,7 @@ UWorld* FSeamlessTravelHandler::Tick()
 					}
 					// add to new world's actor list and remove from old
 					LoadedWorld->PersistentLevel->Actors.Add(TheActor);
+					LoadedWorld->PersistentLevel->ActorsForGC.Add(TheActor);
 
 					TheActor->bActorSeamlessTraveled = true;
 				}
diff --git a/Engine/Source/Runtime/Engine/Public/GeneratedCodeHelpers.h b/Engine/Source/Runtime/Engine/Public/GeneratedCodeHelpers.h
--- a/Engine/Source/Runtime/Engine/Public/GeneratedCodeHelpers.h
+++ b/Engine/Source/Runtime/Engine/Public/GeneratedCodeHelpers.h
@@ -65,9 +65,9 @@ inline uint8 EnumToByte(TEnumAsByte<TEnum> Val)
 }
 
 template<class T>
-inline const T* GetDefaultValueSafe(UClass* Class)
+inline T* GetDefaultValueSafe(UClass* Class)
 {
-	return IsValid(Class) ? GetDefault<T>(Class) : nullptr;
+	return IsValid(Class) ? GetMutableDefault<T>(Class) : nullptr;
 }
 
 template<typename ValueType>
diff --git a/Engine/Source/Runtime/Engine/Public/GenericOctree.h b/Engine/Source/Runtime/Engine/Public/GenericOctree.h
--- a/Engine/Source/Runtime/Engine/Public/GenericOctree.h
+++ b/Engine/Source/Runtime/Engine/Public/GenericOctree.h
@@ -739,11 +739,9 @@ public:
 	/** Initialization constructor. */
 	TOctree(const FVector& InOrigin,float InExtent);
 
-#if WITH_HOT_RELOAD_CTORS
 	/** DO NOT USE. This constructor is for internal usage only for hot-reload purposes. */
 	TOctree();
-#endif // WITH_HOT_RELOAD_CTORS
-			
+
 private:
 
 	/** The octree's root node. */
diff --git a/Engine/Source/Runtime/Engine/Public/GenericQuadTree.h b/Engine/Source/Runtime/Engine/Public/GenericQuadTree.h
--- a/Engine/Source/Runtime/Engine/Public/GenericQuadTree.h
+++ b/Engine/Source/Runtime/Engine/Public/GenericQuadTree.h
@@ -12,10 +12,8 @@ class TQuadTree
 	typedef TQuadTree<ElementType, NodeCapacity> TreeType;
 public:
 
-#if WITH_HOT_RELOAD_CTORS
 	/** DO NOT USE. This constructor is for internal usage only for hot-reload purposes. */
 	TQuadTree();
-#endif // WITH_HOT_RELOAD_CTORS
 
 	TQuadTree(const FBox2D& InBox, float InMinimumQuadSize = 100.f);
 
@@ -154,13 +152,11 @@ TQuadTree<ElementType, NodeCapacity>::TQuadTree(const FBox2D& Box, float InMinim
 	SubTrees[0] = SubTrees[1] = SubTrees[2] = SubTrees[3] = nullptr;
 }
 
-#if WITH_HOT_RELOAD_CTORS
 template <typename ElementType, int32 NodeCapacity>
 TQuadTree<ElementType, NodeCapacity>::TQuadTree()
 {
 	EnsureRetrievingVTablePtrDuringCtor(TEXT("TQuadTree()"));
 }
-#endif // WITH_HOT_RELOAD_CTORS
 
 template <typename ElementType, int32 NodeCapacity>
 TQuadTree<ElementType, NodeCapacity>::~TQuadTree()
diff --git a/Engine/Source/Runtime/Engine/Public/InstancedReferenceSubobjectHelper.h b/Engine/Source/Runtime/Engine/Public/InstancedReferenceSubobjectHelper.h
--- a/Engine/Source/Runtime/Engine/Public/InstancedReferenceSubobjectHelper.h
+++ b/Engine/Source/Runtime/Engine/Public/InstancedReferenceSubobjectHelper.h
@@ -146,10 +146,9 @@ struct FInstancedSubObjRef
  * Contains a set of utility functions useful for searching out and identifying
  * instanced sub-objects contained within a specific outer object.
  */
-class FFindInstancedReferenceSubobjectHelper
+class ENGINE_API FFindInstancedReferenceSubobjectHelper
 {
 public:
-	//--------------------------------------------------------------------------
 	template<typename T>
 	static void GetInstancedSubObjects(const UObject* Container, T& OutObjects)
 	{
@@ -157,127 +156,12 @@ public:
 		for (UProperty* Prop = ContainerClass->RefLink; Prop; Prop = Prop->NextRef)
 		{
 			FInstancedPropertyPath RootPropertyPath(Prop);
-			GetInstancedSubObjects_Inner(RootPropertyPath, reinterpret_cast<const uint8*>(Container), OutObjects);
+			GetInstancedSubObjects_Inner(RootPropertyPath, reinterpret_cast<const uint8*>(Container), [&OutObjects](const FInstancedSubObjRef& Ref){ OutObjects.Add(Ref); });
 		}
 	}
 
-	//--------------------------------------------------------------------------
-	static void Duplicate(UObject* OldObject, UObject* NewObject, TMap<UObject*, UObject*>& ReferenceReplacementMap, TArray<UObject*>& DuplicatedObjects)
-	{
-		if (OldObject->GetClass()->HasAnyClassFlags(CLASS_HasInstancedReference) &&
-			NewObject->GetClass()->HasAnyClassFlags(CLASS_HasInstancedReference))
-		{
-			TArray<FInstancedSubObjRef> OldInstancedSubObjects;
-			GetInstancedSubObjects(OldObject, OldInstancedSubObjects);
-			if (OldInstancedSubObjects.Num() > 0)
-			{
-				TArray<FInstancedSubObjRef> NewInstancedSubObjects;
-				GetInstancedSubObjects(NewObject, NewInstancedSubObjects);
-				for (const FInstancedSubObjRef& Obj : NewInstancedSubObjects)
-				{
-					const bool bNewObjectHasOldOuter = (Obj->GetOuter() == OldObject);
-					if (bNewObjectHasOldOuter)
-					{
-						const bool bKeptByOld = OldInstancedSubObjects.Contains(Obj);
-						const bool bNotHandledYet = !ReferenceReplacementMap.Contains(Obj);
-						if (bKeptByOld && bNotHandledYet)
-						{
-							UObject* NewEditInlineSubobject = StaticDuplicateObject(Obj, NewObject);
-							ReferenceReplacementMap.Add(Obj, NewEditInlineSubobject);
-
-							// NOTE: we cannot patch OldObject's linker table here, since we don't 
-							//       know the relation between the  two objects (one could be of a 
-							//       super class, and the other a child)
-
-							// We also need to make sure to fixup any properties here
-							DuplicatedObjects.Add(NewEditInlineSubobject);
-						}
-					}
-				}
-			}
-		}
-	}
+	static void Duplicate(UObject* OldObject, UObject* NewObject, TMap<UObject*, UObject*>& ReferenceReplacementMap, TArray<UObject*>& DuplicatedObjects);
 
 private:
-	//--------------------------------------------------------------------------
-	template<typename T>
-	static void GetInstancedSubObjects_Inner(FInstancedPropertyPath& PropertyPath, const uint8* ContainerAddress, T& OutObjects)
-	{
-		check(ContainerAddress);
-		const UProperty* TargetProp = PropertyPath.Head();
-
-		const UArrayProperty* ArrayProperty = Cast<const UArrayProperty>(TargetProp);
-		if (ArrayProperty && ArrayProperty->HasAnyPropertyFlags(CPF_PersistentInstance | CPF_ContainsInstancedReference))
-		{
-			const UStructProperty* InnerStructProperty = Cast<const UStructProperty>(ArrayProperty->Inner);
-			const UObjectProperty* InnerObjectProperty = Cast<const UObjectProperty>(ArrayProperty->Inner);
-			if (InnerStructProperty && InnerStructProperty->Struct)
-			{
-				FScriptArrayHelper_InContainer ArrayHelper(ArrayProperty, ContainerAddress);
-				for (int32 ElementIndex = 0; ElementIndex < ArrayHelper.Num(); ++ElementIndex)
-				{
-					const uint8* ValueAddress = ArrayHelper.GetRawPtr(ElementIndex);
-					if (ValueAddress)
-					{
-						for (UProperty* StructProp = InnerStructProperty->Struct->RefLink; StructProp; StructProp = StructProp->NextRef)
-						{
-							PropertyPath.Push(InnerStructProperty, ElementIndex);
-							GetInstancedSubObjects_Inner(PropertyPath, ValueAddress, OutObjects);
-							PropertyPath.Pop();
-						}
-					}
-				}
-			}
-			else if (InnerObjectProperty && InnerObjectProperty->HasAllPropertyFlags(CPF_PersistentInstance))
-			{
-				ensure(InnerObjectProperty->HasAllPropertyFlags(CPF_InstancedReference));
-				FScriptArrayHelper_InContainer ArrayHelper(ArrayProperty, ContainerAddress);
-				for (int32 ElementIndex = 0; ElementIndex < ArrayHelper.Num(); ++ElementIndex)
-				{
-					UObject* ObjectValue = InnerObjectProperty->GetObjectPropertyValue(ArrayHelper.GetRawPtr(ElementIndex));
-					if (ObjectValue)
-					{
-						PropertyPath.Push(InnerObjectProperty, ElementIndex);
-						OutObjects.Add(FInstancedSubObjRef(ObjectValue, PropertyPath));
-						PropertyPath.Pop();
-					}
-					
-				}
-			}
-		}
-		else if (TargetProp->HasAllPropertyFlags(CPF_PersistentInstance))
-		{
-			ensure(TargetProp->HasAllPropertyFlags(CPF_InstancedReference));
-			const UObjectProperty* ObjectProperty = Cast<const UObjectProperty>(TargetProp);
-			if (ObjectProperty)
-			{
-				for (int32 ArrayIdx = 0; ArrayIdx < ObjectProperty->ArrayDim; ++ArrayIdx)
-				{
-					UObject* ObjectValue = ObjectProperty->GetObjectPropertyValue_InContainer(ContainerAddress, ArrayIdx);
-					if (ObjectValue)
-					{
-						// don't need to push to PropertyPath, since this property is already at its head
-						OutObjects.Add(FInstancedSubObjRef(ObjectValue, PropertyPath));	
-					}
-				}
-			}
-		}
-		else if (TargetProp->HasAnyPropertyFlags(CPF_ContainsInstancedReference))
-		{
-			const UStructProperty* StructProperty = Cast<const UStructProperty>(TargetProp);
-			if (StructProperty && StructProperty->Struct)
-			{
-				for (int32 ArrayIdx = 0; ArrayIdx < StructProperty->ArrayDim; ++ArrayIdx)
-				{
-					const uint8* ValueAddress = StructProperty->ContainerPtrToValuePtr<uint8>(ContainerAddress, ArrayIdx);
-					for (UProperty* StructProp = StructProperty->Struct->RefLink; StructProp; StructProp = StructProp->NextRef)
-					{
-						PropertyPath.Push(StructProp, ArrayIdx);
-						GetInstancedSubObjects_Inner(PropertyPath, ValueAddress, OutObjects);
-						PropertyPath.Pop();
-					}
-				}
-			}
-		}
-	}
+	static void GetInstancedSubObjects_Inner(FInstancedPropertyPath& PropertyPath, const uint8* ContainerAddress, TFunctionRef<void(const FInstancedSubObjRef& Ref)> OutObjects);
 };
diff --git a/Engine/Source/Runtime/Engine/Public/MaterialShared.h b/Engine/Source/Runtime/Engine/Public/MaterialShared.h
--- a/Engine/Source/Runtime/Engine/Public/MaterialShared.h
+++ b/Engine/Source/Runtime/Engine/Public/MaterialShared.h
@@ -693,6 +693,7 @@ public:
 
 	/** Registers all shaders that have been loaded in Serialize */
 	virtual void RegisterSerializedShaders() override;
+	virtual void DiscardSerializedShaders() override;
 
 	/** Backs up any FShaders in this shader map to memory through serialization and clears FShader references. */
 	TArray<uint8>* BackupShadersToMemory();
diff --git a/Engine/Source/Runtime/Engine/Public/Model.h b/Engine/Source/Runtime/Engine/Public/Model.h
--- a/Engine/Source/Runtime/Engine/Public/Model.h
+++ b/Engine/Source/Runtime/Engine/Public/Model.h
@@ -354,14 +354,10 @@ enum {MAX_NODES  = 65536};
 enum {MAX_POINTS = 128000};
 class UModel : public UObject
 {
-#if WITH_HOT_RELOAD_CTORS
 	DECLARE_CASTED_CLASS_INTRINSIC_NO_CTOR_NO_VTABLE_CTOR(UModel, UObject, 0, TEXT("/Script/Engine"), 0, ENGINE_API)
 
 	/** DO NOT USE. This constructor is for internal usage only for hot-reload purposes. */
 	UModel(FVTableHelper& Helper);
-#else // WITH_HOT_RELOAD_CTORS
-	DECLARE_CASTED_CLASS_INTRINSIC_NO_CTOR(UModel, UObject, 0, Engine, 0, ENGINE_API)
-#endif // WITH_HOT_RELOAD_CTORS
 
 #if WITH_EDITOR
 	// Arrays and subobjects.
diff --git a/Engine/Source/Runtime/Landscape/Private/Landscape.cpp b/Engine/Source/Runtime/Landscape/Private/Landscape.cpp
--- a/Engine/Source/Runtime/Landscape/Private/Landscape.cpp
+++ b/Engine/Source/Runtime/Landscape/Private/Landscape.cpp
@@ -2482,16 +2482,12 @@ void ALandscapeProxy::UpdateBakedTextures()
 	struct FBakedTextureSourceInfo
 	{
 		// pointer as FMemoryWriter caches the address of the FBufferArchive, and this struct could be relocated on a realloc.
-		FBufferArchive* ComponentStateAr;
+		TUniquePtr<FBufferArchive> ComponentStateAr;
 		TArray<ULandscapeComponent*> Components;
 
 		FBakedTextureSourceInfo()
 		{
-			ComponentStateAr = new FBufferArchive();
-		}
-		~FBakedTextureSourceInfo()
-		{
-			delete ComponentStateAr;
+			ComponentStateAr = MakeUnique<FBufferArchive>();
 		}
 	};
 
diff --git a/Engine/Source/Runtime/Launch/Private/LaunchEngineLoop.cpp b/Engine/Source/Runtime/Launch/Private/LaunchEngineLoop.cpp
--- a/Engine/Source/Runtime/Launch/Private/LaunchEngineLoop.cpp
+++ b/Engine/Source/Runtime/Launch/Private/LaunchEngineLoop.cpp
@@ -1847,6 +1847,11 @@ int32 FEngineLoop::PreInit( const TCHAR* CmdLine )
 	}
 #endif // !PLATFORM_SUPPORTS_EARLY_MOVIE_PLAYBACK
 	
+	{
+		FCoreUObjectDelegates::PreGarbageCollectConditionalBeginDestroy.AddStatic(StartRenderCommandFenceBundler);
+		FCoreUObjectDelegates::PostGarbageCollectConditionalBeginDestroy.AddStatic(StopRenderCommandFenceBundler);
+	}
+
 #if WITH_EDITOR
 	// We need to mount the shared resources for templates (if there are any) before we try and load and game classes
 	FUnrealEdMisc::Get().MountTemplateSharedPaths();
diff --git a/Engine/Source/Runtime/Online/BuildPatchServices/Private/Generation/CloudEnumeration.cpp b/Engine/Source/Runtime/Online/BuildPatchServices/Private/Generation/CloudEnumeration.cpp
--- a/Engine/Source/Runtime/Online/BuildPatchServices/Private/Generation/CloudEnumeration.cpp
+++ b/Engine/Source/Runtime/Online/BuildPatchServices/Private/Generation/CloudEnumeration.cpp
@@ -34,7 +34,6 @@ namespace BuildPatchServices
 		TMap<uint64, TSet<FGuid>> ChunkInventory;
 		TMap<FGuid, int64> ChunkFileSizes;
 		TMap<FGuid, FSHAHash> ChunkShaHashes;
-		TMap<FSHAHash, TSet<FGuid>> FileInventory;
 		FStatsCollectorRef StatsCollector;
 		TFuture<void> Future;
 		volatile FStatsCollector::FAtomicValue* StatManifestsLoaded;
diff --git a/Engine/Source/Runtime/PakFile/Private/IPlatformFilePak.cpp b/Engine/Source/Runtime/PakFile/Private/IPlatformFilePak.cpp
--- a/Engine/Source/Runtime/PakFile/Private/IPlatformFilePak.cpp
+++ b/Engine/Source/Runtime/PakFile/Private/IPlatformFilePak.cpp
@@ -2626,6 +2626,7 @@ struct FCachedAsyncBlock
 	int32 ProcessedSize;
 	int32 RefCount;
 	bool bInFlight;
+	bool bCPUWorkIsComplete;
 	FCachedAsyncBlock()
 		: RawRequest(0)
 		, Raw(nullptr)
@@ -2634,6 +2635,7 @@ struct FCachedAsyncBlock
 		, ProcessedSize(0)
 		, RefCount(0)
 		, bInFlight(false)
+		, bCPUWorkIsComplete(false)
 	{
 	}
 };
@@ -2920,6 +2922,11 @@ public:
 	{
 		check(!LiveRequests.Num()); // must delete all requests before you delete the handle
 		check(!NumLiveRawRequests); // must delete all requests before you delete the handle
+		for (FCachedAsyncBlock& Block : Blocks)
+		{
+			check(Block.RefCount == 0);
+			ClearBlock(Block, true);
+		}
 	}
 
 	virtual IAsyncReadRequest* SizeRequest(FAsyncFileCallBack* CompleteCallback = nullptr) override
@@ -2990,7 +2997,7 @@ public:
 	{
 		FCachedAsyncBlock& Block = Blocks[BlockIndex];
 		Block.bInFlight = true;
-		check(!Block.RawRequest && !Block.Processed && !Block.Raw && !Block.CPUWorkGraphEvent.GetReference() && !Block.ProcessedSize && !Block.RawSize);
+		check(!Block.RawRequest && !Block.Processed && !Block.Raw && !Block.CPUWorkGraphEvent.GetReference() && !Block.ProcessedSize && !Block.RawSize && !Block.bCPUWorkIsComplete);
 		Block.RawSize = FileEntry->CompressionBlocks[BlockIndex].CompressedEnd - FileEntry->CompressionBlocks[BlockIndex].CompressedStart;
 		if (FileEntry->bEncrypted)
 		{
@@ -3027,7 +3034,7 @@ public:
 					Block.ProcessedSize = FileEntry->CompressionBlockSize; // last block was a full block
 				}
 			}
-			check(Block.ProcessedSize);
+			check(Block.ProcessedSize && !Block.bCPUWorkIsComplete);
 			Block.CPUWorkGraphEvent = TGraphTask<FAsyncIOCPUWorkTask>::CreateTask().ConstructAndDispatchWhenReady(*this, BlockIndex);
 		}
 	}
@@ -3083,16 +3090,17 @@ public:
 				Block.CPUWorkGraphEvent = nullptr;
 				Block.bInFlight = false;
 			}
+			Block.bCPUWorkIsComplete = true;
 		}
 	}
-	void ClearBlock(FCachedAsyncBlock& Block)
+	void ClearBlock(FCachedAsyncBlock& Block, bool bForDestructorShouldAlreadyBeClear = false)
 	{
 		check(!Block.RawRequest);
 		Block.RawRequest = nullptr;
-		check(!Block.CPUWorkGraphEvent.GetReference() || Block.CPUWorkGraphEvent->IsComplete());
 		Block.CPUWorkGraphEvent = nullptr;
 		if (Block.Raw)
 		{
+			check(!bForDestructorShouldAlreadyBeClear);
 			// this was a cancel, clean it up now
 			FMemory::Free(Block.Raw);
 			Block.Raw = nullptr;
@@ -3102,6 +3110,7 @@ public:
 		Block.RawSize = 0;
 		if (Block.Processed)
 		{
+			check(bForDestructorShouldAlreadyBeClear == false);
 			FMemory::Free(Block.Processed);
 			Block.Processed = nullptr;
 			check(Block.ProcessedSize);
@@ -3109,6 +3118,7 @@ public:
 		}
 		Block.ProcessedSize = 0;
 		Block.bInFlight = false;
+		Block.bCPUWorkIsComplete = false;
 	}
 
 	void RemoveRequest(FPakProcessedReadRequest* Req, int64 Offset, int64 BytesToRead)
@@ -3134,7 +3144,7 @@ public:
 					Block.RawRequest = nullptr;
 					NumLiveRawRequests--;
 				}
-				if (!Block.CPUWorkGraphEvent.GetReference() || Block.CPUWorkGraphEvent->IsComplete())
+				if (Block.bCPUWorkIsComplete)
 				{
 					ClearBlock(Block);
 				}
diff --git a/Engine/Source/Runtime/RenderCore/Private/RenderResource.cpp b/Engine/Source/Runtime/RenderCore/Private/RenderResource.cpp
--- a/Engine/Source/Runtime/RenderCore/Private/RenderResource.cpp
+++ b/Engine/Source/Runtime/RenderCore/Private/RenderResource.cpp
@@ -141,8 +141,81 @@ void BeginUpdateResourceRHI(FRenderResource* Resource)
 		});
 }
 
+struct FBatchedReleaseResources
+{
+	enum 
+	{
+		NumPerBatch = 16
+	};
+	int32 NumBatch;
+	FRenderResource* Resources[NumPerBatch];
+	FBatchedReleaseResources()
+	{
+		Reset();
+	}
+	void Reset()
+	{
+		NumBatch = 0;
+	}
+	void Execute()
+	{
+		for (int32 Index = 0; Index < NumBatch; Index++)
+		{
+			Resources[Index]->ReleaseResource();
+		}
+		Reset();
+	}
+	void Flush()
+	{
+		if (NumBatch)
+		{
+			ENQUEUE_UNIQUE_RENDER_COMMAND_ONEPARAMETER(
+				BatchReleaseCommand,
+				FBatchedReleaseResources, BatchedReleaseResources, *this,
+				{
+					BatchedReleaseResources.Execute();
+				});
+			Reset();
+		}
+	}
+	void Add(FRenderResource* Resource)
+	{
+		if (NumBatch >= NumPerBatch)
+		{
+			Flush();
+		}
+		check(NumBatch < NumPerBatch);
+		Resources[NumBatch] = Resource;
+		NumBatch++;
+	}
+	bool IsEmpty()
+	{
+		return !NumBatch;
+	}
+};
+
+static bool GBatchedReleaseIsActive = false;
+static FBatchedReleaseResources GBatchedRelease;
+
+void StartBatchedRelease()
+{
+	check(IsInGameThread() && !GBatchedReleaseIsActive && GBatchedRelease.IsEmpty());
+	GBatchedReleaseIsActive = true;
+}
+void EndBatchedRelease()
+{
+	check(IsInGameThread() && GBatchedReleaseIsActive);
+	GBatchedRelease.Flush();
+	GBatchedReleaseIsActive = false;
+}
+
 void BeginReleaseResource(FRenderResource* Resource)
 {
+	if (GBatchedReleaseIsActive && IsInGameThread())
+	{
+		GBatchedRelease.Add(Resource);
+		return;
+	}
 	ENQUEUE_UNIQUE_RENDER_COMMAND_ONEPARAMETER(
 		ReleaseCommand,
 		FRenderResource*,Resource,Resource,
diff --git a/Engine/Source/Runtime/RenderCore/Private/RenderingThread.cpp b/Engine/Source/Runtime/RenderCore/Private/RenderingThread.cpp
--- a/Engine/Source/Runtime/RenderCore/Private/RenderingThread.cpp
+++ b/Engine/Source/Runtime/RenderCore/Private/RenderingThread.cpp
@@ -18,6 +18,8 @@
 #include "TickableObjectRenderThread.h"
 #include "Stats/StatsData.h"
 #include "HAL/ThreadHeartBeat.h"
+#include "RenderResource.h"
+
 //
 // Globals
 //
@@ -816,6 +818,46 @@ bool IsRenderingThreadHealthy()
 	return GIsRenderingThreadHealthy;
 }
 
+static FGraphEventRef BundledCompletionEvent;
+static FGraphEventRef BundledCompletionEventPrereq; // We fire this when we are done, which queues the actual fence
+
+void StartRenderCommandFenceBundler()
+{
+	if (!GIsThreadedRendering)
+	{
+		return;
+	}
+
+	check(IsInGameThread() && !BundledCompletionEvent.GetReference() && !BundledCompletionEventPrereq.GetReference()); // can't use this in a nested fashion
+	BundledCompletionEventPrereq = FGraphEvent::CreateGraphEvent();
+
+	FGraphEventArray Prereqs;
+	Prereqs.Add(BundledCompletionEventPrereq);
+
+	DECLARE_CYCLE_STAT(TEXT("FNullGraphTask.FenceRenderCommandBundled"),
+	STAT_FNullGraphTask_FenceRenderCommandBundled,
+		STATGROUP_TaskGraphTasks);
+
+	BundledCompletionEvent = TGraphTask<FNullGraphTask>::CreateTask(&Prereqs, ENamedThreads::GameThread).ConstructAndDispatchWhenReady(
+		GET_STATID(STAT_FNullGraphTask_FenceRenderCommandBundled), ENamedThreads::RenderThread);
+
+	StartBatchedRelease();
+}
+
+void StopRenderCommandFenceBundler()
+{
+	if (!GIsThreadedRendering || !BundledCompletionEvent.GetReference())
+	{
+		return;
+	}
+
+	EndBatchedRelease();
+	check(IsInGameThread() && BundledCompletionEvent.GetReference() && !BundledCompletionEvent->IsComplete() && BundledCompletionEventPrereq.GetReference() && !BundledCompletionEventPrereq->IsComplete()); // can't use this in a nested fashion
+	TArray<FBaseGraphTask*> NewTasks;
+	BundledCompletionEventPrereq->DispatchSubsequents(NewTasks);
+	BundledCompletionEventPrereq = nullptr;
+	BundledCompletionEvent = nullptr;
+}
 
 void FRenderCommandFence::BeginFence()
 {
@@ -825,6 +867,12 @@ void FRenderCommandFence::BeginFence()
 	}
 	else
 	{
+		if (BundledCompletionEvent.GetReference() && IsInGameThread())
+		{
+			CompletionEvent = BundledCompletionEvent;
+			return;
+		}
+
 		DECLARE_CYCLE_STAT(TEXT("FNullGraphTask.FenceRenderCommand"),
 			STAT_FNullGraphTask_FenceRenderCommand,
 			STATGROUP_TaskGraphTasks);
@@ -952,6 +1000,7 @@ void FRenderCommandFence::Wait(bool bProcessGameThreadTasks) const
 {
 	if (!IsFenceComplete())
 	{
+		StopRenderCommandFenceBundler();
 #if 0
 		// on most platforms this is a better solution because it doesn't spin
 		// windows needs to pump messages
diff --git a/Engine/Source/Runtime/RenderCore/Public/RenderResource.h b/Engine/Source/Runtime/RenderCore/Public/RenderResource.h
--- a/Engine/Source/Runtime/RenderCore/Public/RenderResource.h
+++ b/Engine/Source/Runtime/RenderCore/Public/RenderResource.h
@@ -133,6 +133,18 @@ extern RENDERCORE_API void BeginUpdateResourceRHI(FRenderResource* Resource);
  */
 extern RENDERCORE_API void BeginReleaseResource(FRenderResource* Resource);
 
+/**
+* Enables the batching of calls to BeginReleaseResource
+* This is called in the game thread.
+*/
+extern RENDERCORE_API void StartBatchedRelease();
+
+/**
+* Disables the batching of calls to BeginReleaseResource
+* This is called in the game thread.
+*/
+extern RENDERCORE_API void EndBatchedRelease();
+
 /**
  * Sends a message to the rendering thread to release a resource, and spins until the rendering thread has processed the message.
  * This is called in the game thread.
diff --git a/Engine/Source/Runtime/RenderCore/Public/RenderingThread.h b/Engine/Source/Runtime/RenderCore/Public/RenderingThread.h
--- a/Engine/Source/Runtime/RenderCore/Public/RenderingThread.h
+++ b/Engine/Source/Runtime/RenderCore/Public/RenderingThread.h
@@ -87,6 +87,9 @@ extern RENDERCORE_API void FlushPendingDeleteRHIResources_RenderThread();
 
 extern RENDERCORE_API void TickRenderingTickables();
 
+extern RENDERCORE_API void StartRenderCommandFenceBundler();
+extern RENDERCORE_API void StopRenderCommandFenceBundler();
+
 ////////////////////////////////////
 // Render thread suspension
 ////////////////////////////////////
diff --git a/Engine/Source/Runtime/Renderer/Private/ScenePrivate.h b/Engine/Source/Runtime/Renderer/Private/ScenePrivate.h
--- a/Engine/Source/Runtime/Renderer/Private/ScenePrivate.h
+++ b/Engine/Source/Runtime/Renderer/Private/ScenePrivate.h
@@ -1025,10 +1025,7 @@ public:
 
 	virtual void AddReferencedObjects(FReferenceCollector& Collector) override
 	{
-		for(UMaterialInstanceDynamic*& MID : MIDPool)
-		{
-			Collector.AddReferencedObject(MID);
-		}
+		Collector.AddReferencedObjects(MIDPool);
 	}
 
 	/** called in InitViews() */
diff --git a/Engine/Source/Runtime/ShaderCore/Public/Shader.h b/Engine/Source/Runtime/ShaderCore/Public/Shader.h
--- a/Engine/Source/Runtime/ShaderCore/Public/Shader.h
+++ b/Engine/Source/Runtime/ShaderCore/Public/Shader.h
@@ -711,6 +711,13 @@ public:
 		return this;
 	}
 
+	/** Discards the serialized resource, used when the engine is using NullRHI */
+	void DiscardSerializedResource()
+	{
+		delete SerializedResource;
+		SerializedResource = nullptr;
+	}
+
 protected:
 
 	/** Indexed the same as UniformBufferParameters.  Packed densely for coherent traversal. */
@@ -1749,6 +1756,30 @@ public:
 		SerializedShaderPipelines.Empty();
 	}
 
+	/** Discards serialized shaders when they are not going to be used for anything (NullRHI) */
+	virtual void DiscardSerializedShaders()
+	{
+		for (FShader* Shader : SerializedShaders)
+		{
+			if (Shader)
+			{
+				Shader->DiscardSerializedResource();
+			}
+			delete Shader;
+		}
+		SerializedShaders.Empty();
+
+		for (FSerializedShaderPipeline* SerializedPipeline : SerializedShaderPipelines)
+		{
+			for (TRefCountPtr<FShader> Shader : SerializedPipeline->ShaderStages)
+			{
+				Shader->DiscardSerializedResource();
+			}
+			delete SerializedPipeline;
+		}
+		SerializedShaderPipelines.Empty();
+	}
+
 	/** @return true if the map is empty */
 	inline bool IsEmpty() const
 	{
diff --git a/Engine/Source/Runtime/Slate/Public/Framework/Views/TableViewTypeTraits.h b/Engine/Source/Runtime/Slate/Public/Framework/Views/TableViewTypeTraits.h
--- a/Engine/Source/Runtime/Slate/Public/Framework/Views/TableViewTypeTraits.h
+++ b/Engine/Source/Runtime/Slate/Public/Framework/Views/TableViewTypeTraits.h
@@ -272,17 +272,10 @@ public:
 	static void AddReferencedObjects( FReferenceCollector& Collector, TArray<T*>& ItemsWithGeneratedWidgets, TSet<T*>& SelectedItems )
 	{
 		// Serialize generated items
-		const int32 NumGeneratedWidgets = ItemsWithGeneratedWidgets.Num();
-		for ( int32 ItemIndex = 0; ItemIndex < NumGeneratedWidgets; ++ItemIndex )
-		{
-			Collector.AddReferencedObject( ItemsWithGeneratedWidgets[ ItemIndex ] );
-		}
+		Collector.AddReferencedObjects(ItemsWithGeneratedWidgets);
 
 		// Serialize the selected items
-		for ( typename TSet<T*>::TIterator SelecteItemIt( SelectedItems ); SelecteItemIt; ++SelecteItemIt )
-		{
-			Collector.AddReferencedObject( *SelecteItemIt );
-		}
+		Collector.AddReferencedObjects(SelectedItems);
 	}
 
 	static bool IsPtrValid( T* InPtr ) { return InPtr != NULL; }
@@ -305,17 +298,10 @@ public:
 	static void AddReferencedObjects( FReferenceCollector& Collector, TArray<const T*>& ItemsWithGeneratedWidgets, TSet<const T*>& SelectedItems )
 	{
 		// Serialize generated items
-		const int32 NumGeneratedWidgets = ItemsWithGeneratedWidgets.Num();
-		for ( int32 ItemIndex = 0; ItemIndex < NumGeneratedWidgets; ++ItemIndex )
-		{
-			Collector.AddReferencedObject( ItemsWithGeneratedWidgets[ ItemIndex ] );
-		}
+		Collector.AddReferencedObjects(ItemsWithGeneratedWidgets);
 
 		// Serialize the selected items
-		for ( typename TSet<const T*>::TIterator SelecteItemIt( SelectedItems ); SelecteItemIt; ++SelecteItemIt )
-		{
-			Collector.AddReferencedObject( *SelecteItemIt );
-		}
+		Collector.AddReferencedObjects(SelectedItems);
 	}
 
 	static bool IsPtrValid( const T* InPtr ) { return InPtr != NULL; }
