Adding missing Actor.h include, fixing fallout from code relocation

Fixes  non-unity UE4Game Mac compilation

diff --git a/Engine/Source/Runtime/Engine/Private/KismetTraceUtils.cpp b/Engine/Source/Runtime/Engine/Private/KismetTraceUtils.cpp
--- a/Engine/Source/Runtime/Engine/Private/KismetTraceUtils.cpp
+++ b/Engine/Source/Runtime/Engine/Private/KismetTraceUtils.cpp
@@ -2,6 +2,7 @@
 
 #include "KismetTraceUtils.h"
 #include "PhysicsEngine/PhysicsSettings.h"
+#include "GameFramework/Actor.h"
 
 static const float KISMET_TRACE_DEBUG_IMPACTPOINT_SIZE = 16.f;
 
