Copying //UE4/UT-Staging to //UE4/Dev-Main

==========================
MAJOR FEATURES + CHANGES
==========================

3306447 Joe.Wilcox
Fix typo introduced in SWebBrowser

3306444 Joe.Wilcox
Added a way for games to restrict the context menu on Chromium Embedded Framework
-in fullscreen the window switch would cause a minimize

3306305 Peter.Knepley
Give gamestate a notification when async loading has succeeded for a package

3306275 Peter.Knepley
Fix netspeed in local replay situations

3267864 Peter.Knepley
Fix crash in AGameModeBase::GenericPlayerInitialization (may already be fixed in main)

3258890 Matt.Oelfke
Discard all navigation updates caused by octree construction needed to be moved up

3245388 Steve.Polge
DeltaTime parameter for HandleSwimmingWallHit

3245385 Steve.Polge
Added HandleSwimmingWallHit() to CharacterMovementComponent

3245011 Peter.Knepley
Extra safety for GetMoviePlayer() and slate loading thread

3242810 Joe.Wilcox
Made SSlider::CommitValue virtual

3235608 Peter.Knepley
Fix crash in APlayerController::Reset (may already be in main)

3224903 Joe.Wilcox
SC_CLOSE should translate to WM_CLOSE to fix corner cases with mouse focus and ALT+F4

3223132 Peter.Knepley
GetOutdatedFiles needs to be exported for non-monolithic builds to be able to link to it

diff --git a/Engine/Source/Runtime/Core/Private/Windows/WindowsApplication.cpp b/Engine/Source/Runtime/Core/Private/Windows/WindowsApplication.cpp
--- a/Engine/Source/Runtime/Core/Private/Windows/WindowsApplication.cpp
+++ b/Engine/Source/Runtime/Core/Private/Windows/WindowsApplication.cpp
@@ -1260,6 +1260,12 @@ int32 FWindowsApplication::ProcessMessage( HWND hwnd, uint32 msg, WPARAM wParam,
 						}
 					}
 					break;
+				case SC_CLOSE:
+					{
+						DeferMessage(CurrentNativeEventWindowPtr, hwnd, WM_CLOSE, 0, 0);
+						return 1;
+					}
+					break;
 				default:
 					if( !( MessageHandler->ShouldProcessUserInputMessages( CurrentNativeEventWindow ) && IsInputMessage( msg ) ) )
 					{
diff --git a/Engine/Source/Runtime/Engine/Classes/Engine/DemoNetDriver.h b/Engine/Source/Runtime/Engine/Classes/Engine/DemoNetDriver.h
--- a/Engine/Source/Runtime/Engine/Classes/Engine/DemoNetDriver.h
+++ b/Engine/Source/Runtime/Engine/Classes/Engine/DemoNetDriver.h
@@ -428,6 +428,9 @@ public:
 
 public:
 
+	UPROPERTY()
+	bool bIsLocalReplay;
+
 	/** @todo document */
 	bool UpdateDemoTime( float* DeltaTime, float TimeDilation );
 
diff --git a/Engine/Source/Runtime/Engine/Classes/GameFramework/CharacterMovementComponent.h b/Engine/Source/Runtime/Engine/Classes/GameFramework/CharacterMovementComponent.h
--- a/Engine/Source/Runtime/Engine/Classes/GameFramework/CharacterMovementComponent.h
+++ b/Engine/Source/Runtime/Engine/Classes/GameFramework/CharacterMovementComponent.h
@@ -1594,6 +1594,9 @@ protected:
 	/** @note Movement update functions should only be called through StartNewPhysics()*/
 	virtual void PhysCustom(float deltaTime, int32 Iterations);
 
+	/* Allow custom handling when character hits a wall while swimming. */
+	virtual void HandleSwimmingWallHit(const FHitResult& Hit, float DeltaTime);
+
 	/**
 	 * Compute a vector of movement, given a delta and a hit result of the surface we are on.
 	 *
diff --git a/Engine/Source/Runtime/Engine/Classes/GameFramework/GameStateBase.h b/Engine/Source/Runtime/Engine/Classes/GameFramework/GameStateBase.h
--- a/Engine/Source/Runtime/Engine/Classes/GameFramework/GameStateBase.h
+++ b/Engine/Source/Runtime/Engine/Classes/GameFramework/GameStateBase.h
@@ -49,6 +49,9 @@ class ENGINE_API AGameStateBase : public AInfo
 	UPROPERTY(Transient, BlueprintReadOnly, Category=GameState)
 	TArray<APlayerState*> PlayerArray;
 
+	/** Allow game states to react to asset packages being loaded asynchronously */
+	virtual void AsyncPackageLoaded(UObject* Package) {}
+
 	/** Helper to return the default object of the GameModeBase class corresponding to this GameState. This object is not safe to modify. */
 	const AGameModeBase* GetDefaultGameMode() const;
 
diff --git a/Engine/Source/Runtime/Engine/Private/AI/Navigation/NavigationSystem.cpp b/Engine/Source/Runtime/Engine/Private/AI/Navigation/NavigationSystem.cpp
--- a/Engine/Source/Runtime/Engine/Private/AI/Navigation/NavigationSystem.cpp
+++ b/Engine/Source/Runtime/Engine/Private/AI/Navigation/NavigationSystem.cpp
@@ -486,6 +486,9 @@ void UNavigationSystem::PostInitProperties()
 
 bool UNavigationSystem::ConditionalPopulateNavOctree()
 {
+	// Discard all navigation updates caused by octree construction
+	TGuardValue<TArray<FNavigationDirtyArea>> DirtyGuard(DirtyAreas, TArray<FNavigationDirtyArea>());
+
 	// We are going to fully re-populate NavOctree so all pending update request are outdated
 	PendingOctreeUpdates.Empty(32);
 	
@@ -538,17 +541,11 @@ bool UNavigationSystem::ConditionalPopulateNavOctree()
 	// Add all found elements to octree, this will not add new dirty areas to navigation
 	if (PendingOctreeUpdates.Num())
 	{
-		TArray<FNavigationDirtyArea> SavedDirtyAreas; 
-		Exchange(SavedDirtyAreas, DirtyAreas);
-
 		for (TSet<FNavigationDirtyElement>::TIterator It(PendingOctreeUpdates); It; ++It)
 		{
 			AddElementToNavOctree(*It);
 		}
 		PendingOctreeUpdates.Empty(32);
-
-		// Discard all navigation updates caused by octree construction
-		Exchange(SavedDirtyAreas, DirtyAreas);
 	}
 
 	return bSupportRebuilding;
diff --git a/Engine/Source/Runtime/Engine/Private/Components/CharacterMovementComponent.cpp b/Engine/Source/Runtime/Engine/Private/Components/CharacterMovementComponent.cpp
--- a/Engine/Source/Runtime/Engine/Private/Components/CharacterMovementComponent.cpp
+++ b/Engine/Source/Runtime/Engine/Private/Components/CharacterMovementComponent.cpp
@@ -3518,6 +3518,10 @@ void UCharacterMovementComponent::ApplyRootMotionToVelocity(float deltaTime)
 	}
 }
 
+void UCharacterMovementComponent::HandleSwimmingWallHit(const FHitResult& Hit, float DeltaTime)
+{
+}
+
 void UCharacterMovementComponent::PhysSwimming(float deltaTime, int32 Iterations)
 {
 	if (deltaTime < MIN_TICK_TIME)
@@ -3569,6 +3573,7 @@ void UCharacterMovementComponent::PhysSwimming(float deltaTime, int32 Iterations
 
 	if ( Hit.Time < 1.f && CharacterOwner)
 	{
+		HandleSwimmingWallHit(Hit, deltaTime);
 		if (bLimitedUpAccel && (Velocity.Z >= 0.f))
 		{
 			// allow upward velocity at surface if against obstacle
diff --git a/Engine/Source/Runtime/Engine/Private/DemoNetDriver.cpp b/Engine/Source/Runtime/Engine/Private/DemoNetDriver.cpp
--- a/Engine/Source/Runtime/Engine/Private/DemoNetDriver.cpp
+++ b/Engine/Source/Runtime/Engine/Private/DemoNetDriver.cpp
@@ -27,6 +27,7 @@
 #include "Engine/LevelStreamingKismet.h"
 #include "GameFramework/SpectatorPawnMovement.h"
 #include "Net/UnrealNetwork.h"
+#include "UnrealEngine.h"
 #include "Net/NetworkProfiler.h"
 #include "GameFramework/GameStateBase.h"
 #include "GameFramework/PlayerState.h"
@@ -3203,7 +3204,21 @@ void UDemoNetConnection::HandleClientPlayer( APlayerController* PC, UNetConnecti
 		return;
 	}
 
+	ULocalPlayer* LocalPlayer = NULL;
+	for (FLocalPlayerIterator It(GEngine, Driver->GetWorld()); It; ++It)
+	{
+		LocalPlayer = *It;
+		break;
+	}
+	int32 SavedNetSpeed = LocalPlayer ? LocalPlayer->CurrentNetSpeed : 0;
+
 	Super::HandleClientPlayer( PC, NetConnection );
+	
+	// Restore the netspeed if we're a local replay
+	if (GetDriver()->bIsLocalReplay && LocalPlayer)
+	{
+		LocalPlayer->CurrentNetSpeed = SavedNetSpeed;
+	}
 
 	// Assume this is our special spectator controller
 	GetDriver()->SpectatorController = PC;
diff --git a/Engine/Source/Runtime/Engine/Private/GameModeBase.cpp b/Engine/Source/Runtime/Engine/Private/GameModeBase.cpp
--- a/Engine/Source/Runtime/Engine/Private/GameModeBase.cpp
+++ b/Engine/Source/Runtime/Engine/Private/GameModeBase.cpp
@@ -887,8 +887,11 @@ void AGameModeBase::GenericPlayerInitialization(AController* C)
 		// Notify the game that we can now be muted and mute others
 		UpdateGameplayMuteList(PC);
 
-		// Tell the player to enable voice by default or use the push to talk method
-		PC->ClientEnableNetworkVoice(!GameSession->RequiresPushToTalk());
+		if (GameSession != nullptr)
+		{
+			// Tell the player to enable voice by default or use the push to talk method
+			PC->ClientEnableNetworkVoice(!GameSession->RequiresPushToTalk());
+		}
 
 		ReplicateStreamingStatus(PC);
 
diff --git a/Engine/Source/Runtime/Engine/Private/PackageMapClient.cpp b/Engine/Source/Runtime/Engine/Private/PackageMapClient.cpp
--- a/Engine/Source/Runtime/Engine/Private/PackageMapClient.cpp
+++ b/Engine/Source/Runtime/Engine/Private/PackageMapClient.cpp
@@ -19,6 +19,7 @@
 #include "Engine/ActorChannel.h"
 #include "Net/RepLayout.h"
 #include "ProfilingDebugging/ScopedTimers.h"
+#include "GameFramework/GameStateBase.h"
 
 // ( OutPacketId == GUID_PACKET_NOT_ACKED ) == NAK'd		(this GUID is not acked, and is not pending either, so sort of waiting)
 // ( OutPacketId == GUID_PACKET_ACKED )		== FULLY ACK'd	(this GUID is fully acked, and we no longer need to send full path)
@@ -2228,7 +2229,7 @@ void FNetGUIDCache::AsyncPackageCallback(const FName& PackageName, UPackage * Pa
 	check( Package == NULL || Package->IsFullyLoaded() );
 
 	FNetworkGUID NetGUID = PendingAsyncPackages.FindRef(PackageName);
-
+	
 	PendingAsyncPackages.Remove(PackageName);
 
 	if ( !NetGUID.IsValid() )
@@ -2257,6 +2258,15 @@ void FNetGUIDCache::AsyncPackageCallback(const FName& PackageName, UPackage * Pa
 		CacheObject->bIsBroken = true;
 		UE_LOG( LogNetPackageMap, Error, TEXT( "AsyncPackageCallback: Package FAILED to load. Path: %s, NetGUID: %s" ), *PackageName.ToString(), *NetGUID.ToString() );
 	}
+
+	if (CacheObject->Object.IsValid() && CacheObject->Object->GetWorld())
+	{
+		AGameStateBase* GS = CacheObject->Object->GetWorld()->GetGameState();
+		if (GS)
+		{
+			GS->AsyncPackageLoaded(CacheObject->Object.Get());
+		}
+	}
 }
 
 static bool ObjectLevelHasFinishedLoading( UObject* Object, UNetDriver* Driver )
diff --git a/Engine/Source/Runtime/Engine/Private/PlayerController.cpp b/Engine/Source/Runtime/Engine/Private/PlayerController.cpp
--- a/Engine/Source/Runtime/Engine/Private/PlayerController.cpp
+++ b/Engine/Source/Runtime/Engine/Private/PlayerController.cpp
@@ -1151,7 +1151,7 @@ void APlayerController::ClientReset_Implementation()
 	ResetCameraMode();
 	SetViewTarget(this);
 
-	bPlayerIsWaiting = !PlayerState->bOnlySpectator;
+	bPlayerIsWaiting = (PlayerState == nullptr) || !PlayerState->bOnlySpectator;
 	ChangeState(NAME_Spectating);
 }
 
diff --git a/Engine/Source/Runtime/MoviePlayer/Private/MoviePlayer.cpp b/Engine/Source/Runtime/MoviePlayer/Private/MoviePlayer.cpp
--- a/Engine/Source/Runtime/MoviePlayer/Private/MoviePlayer.cpp
+++ b/Engine/Source/Runtime/MoviePlayer/Private/MoviePlayer.cpp
@@ -68,6 +68,8 @@ TSharedRef<SWidget> FLoadingScreenAttributes::NewTestLoadingScreenWidget()
 
 TSharedPtr<IGameMoviePlayer> GetMoviePlayer()
 {
+	ensureMsgf(!IsInSlateThread(), TEXT("GetMoviePlayer() is not safe to be called from the Slate Thread"));
+
 	if (!IsMoviePlayerEnabled() || GUsingNullRHI)
 	{
 		return FNullGameMoviePlayer::Get();
diff --git a/Engine/Source/Runtime/Slate/Public/Widgets/Input/SSlider.h b/Engine/Source/Runtime/Slate/Public/Widgets/Input/SSlider.h
--- a/Engine/Source/Runtime/Slate/Public/Widgets/Input/SSlider.h
+++ b/Engine/Source/Runtime/Slate/Public/Widgets/Input/SSlider.h
@@ -140,7 +140,7 @@ protected:
 	 *
 	 * @param NewValue The value to commit.
 	 */
-	void CommitValue(float NewValue);
+	virtual void CommitValue(float NewValue);
 
 	/**
 	 * Calculates the new value based on the given absolute coordinates.
@@ -151,7 +151,7 @@ protected:
 	 */
 	float PositionToValue( const FGeometry& MyGeometry, const FVector2D& AbsolutePosition );
 
-private:
+protected:
 
 	// Holds the style passed to the widget upon construction.
 	const FSliderStyle* Style;
diff --git a/Engine/Source/Runtime/WebBrowser/Private/Android/AndroidWebBrowserWindow.h b/Engine/Source/Runtime/WebBrowser/Private/Android/AndroidWebBrowserWindow.h
--- a/Engine/Source/Runtime/WebBrowser/Private/Android/AndroidWebBrowserWindow.h
+++ b/Engine/Source/Runtime/WebBrowser/Private/Android/AndroidWebBrowserWindow.h
@@ -168,6 +168,11 @@ public:
 		return DismissAllDialogsDelegate;
 	}
 
+	virtual FOnSuppressContextMenu& OnSuppressContextMenu() override
+	{
+		return SuppressContextMenuDelgate;
+	}
+
 public:
 
 	/**
@@ -241,6 +246,9 @@ private:
 	/** Delegate for dismissing all dialogs. */
 	FOnDismissAllDialogs DismissAllDialogsDelegate;
 
+	/** Delegate for suppressing context menu */
+	FOnSuppressContextMenu SuppressContextMenuDelgate;
+
 	/** Current state of the document being loaded. */
 	EWebBrowserDocumentState DocumentState;
 	int ErrorCode;
diff --git a/Engine/Source/Runtime/WebBrowser/Private/CEF/CEFBrowserHandler.cpp b/Engine/Source/Runtime/WebBrowser/Private/CEF/CEFBrowserHandler.cpp
--- a/Engine/Source/Runtime/WebBrowser/Private/CEF/CEFBrowserHandler.cpp
+++ b/Engine/Source/Runtime/WebBrowser/Private/CEF/CEFBrowserHandler.cpp
@@ -586,6 +586,15 @@ void FCEFBrowserHandler::OnResetDialogState(CefRefPtr<CefBrowser> Browser)
 }
 
 
+void FCEFBrowserHandler::OnBeforeContextMenu(CefRefPtr<CefBrowser> Browser, CefRefPtr<CefFrame> Frame, CefRefPtr<CefContextMenuParams> Params, CefRefPtr<CefMenuModel> Model)
+{
+	TSharedPtr<FCEFWebBrowserWindow> BrowserWindow = BrowserWindowPtr.Pin();
+	if ( BrowserWindow.IsValid() && BrowserWindow->OnSuppressContextMenu().IsBound() && BrowserWindow->OnSuppressContextMenu().Execute() )
+	{
+		Model->Clear();
+	}
+}
+
 #undef LOCTEXT_NAMESPACE
 
 #endif // WITH_CEF
diff --git a/Engine/Source/Runtime/WebBrowser/Private/CEF/CEFBrowserHandler.h b/Engine/Source/Runtime/WebBrowser/Private/CEF/CEFBrowserHandler.h
--- a/Engine/Source/Runtime/WebBrowser/Private/CEF/CEFBrowserHandler.h
+++ b/Engine/Source/Runtime/WebBrowser/Private/CEF/CEFBrowserHandler.h
@@ -47,6 +47,7 @@ class FCEFBrowserHandler
 	, public CefRequestHandler
 	, public CefKeyboardHandler
 	, public CefJSDialogHandler
+	, public CefContextMenuHandler
 {
 public:
 
@@ -111,6 +112,12 @@ public:
 		return this;
 	}
 
+	virtual CefRefPtr<CefContextMenuHandler> GetContextMenuHandler() override
+	{
+		return this;
+	}
+
+
 	virtual bool OnProcessMessageReceived(CefRefPtr<CefBrowser> Browser,
 		CefProcessId SourceProcess,
 		CefRefPtr<CefProcessMessage> Message) override;
@@ -246,6 +253,12 @@ public:
 		return CreateWindowDelegate;
 	}
 
+public:
+	virtual void OnBeforeContextMenu(CefRefPtr<CefBrowser> Browser, 
+ 		CefRefPtr<CefFrame> Frame,
+		CefRefPtr<CefContextMenuParams> Params,
+		CefRefPtr<CefMenuModel> Model) override;
+
 private:
 
 	bool ShowDevTools(const CefRefPtr<CefBrowser>& Browser);
diff --git a/Engine/Source/Runtime/WebBrowser/Private/CEF/CEFWebBrowserWindow.h b/Engine/Source/Runtime/WebBrowser/Private/CEF/CEFWebBrowserWindow.h
--- a/Engine/Source/Runtime/WebBrowser/Private/CEF/CEFWebBrowserWindow.h
+++ b/Engine/Source/Runtime/WebBrowser/Private/CEF/CEFWebBrowserWindow.h
@@ -235,6 +235,11 @@ public:
 		return DismissAllDialogsDelegate;
 	}
 
+	virtual FOnSuppressContextMenu& OnSuppressContextMenu() override
+	{
+		return SuppressContextMenuDelgate;
+	}
+
 private:
 
 	/**
@@ -526,6 +531,10 @@ private:
 	/** Delegate for dismissing all dialogs. */
 	FOnDismissAllDialogs DismissAllDialogsDelegate;
 
+	/** Delegate for suppressing context menu */
+	FOnSuppressContextMenu SuppressContextMenuDelgate;
+
+
 	/** Tracks the current mouse cursor */
 	EMouseCursor::Type Cursor;
 
diff --git a/Engine/Source/Runtime/WebBrowser/Private/IOS/IOSPlatformWebBrowser.h b/Engine/Source/Runtime/WebBrowser/Private/IOS/IOSPlatformWebBrowser.h
--- a/Engine/Source/Runtime/WebBrowser/Private/IOS/IOSPlatformWebBrowser.h
+++ b/Engine/Source/Runtime/WebBrowser/Private/IOS/IOSPlatformWebBrowser.h
@@ -206,6 +206,12 @@ public:
 		return DismissAllDialogsDelegate;
 	}
 
+	virtual FOnSuppressContextMenu& OnSuppressContextMenu() override
+	{
+		return SuppressContextMenuDelgate;
+	}
+
+
 public:
 
 private:
@@ -262,6 +268,9 @@ private:
 	/** Delegate for dismissing all dialogs. */
 	FOnDismissAllDialogs DismissAllDialogsDelegate;
 
+	/** Delegate for suppressing context menu */
+	FOnSuppressContextMenu SuppressContextMenuDelgate;
+
 	TSharedPtr<SWindow> ParentWindow;
 
 };
diff --git a/Engine/Source/Runtime/WebBrowser/Private/SWebBrowser.cpp b/Engine/Source/Runtime/WebBrowser/Private/SWebBrowser.cpp
--- a/Engine/Source/Runtime/WebBrowser/Private/SWebBrowser.cpp
+++ b/Engine/Source/Runtime/WebBrowser/Private/SWebBrowser.cpp
@@ -31,7 +31,7 @@ void SWebBrowser::Construct(const FArguments& InArgs, const TSharedPtr<IWebBrows
 	OnCreateWindow = InArgs._OnCreateWindow;
 	OnCloseWindow = InArgs._OnCloseWindow;
 	bShowInitialThrobber = InArgs._ShowInitialThrobber;
-
+	
 	ChildSlot
 	[
 		SNew(SVerticalBox)
@@ -124,6 +124,7 @@ void SWebBrowser::Construct(const FArguments& InArgs, const TSharedPtr<IWebBrows
 				.OnShowDialog(OnShowDialog)
 				.OnDismissAllDialogs(OnDismissAllDialogs)
 				.Visibility(this, &SWebBrowser::GetViewportVisibility)
+				.OnSuppressContextMenu(InArgs._OnSuppressContextMenu)
 			]
 			+ SOverlay::Slot()
 			.HAlign(HAlign_Center)
diff --git a/Engine/Source/Runtime/WebBrowser/Private/SWebBrowserView.cpp b/Engine/Source/Runtime/WebBrowser/Private/SWebBrowserView.cpp
--- a/Engine/Source/Runtime/WebBrowser/Private/SWebBrowserView.cpp
+++ b/Engine/Source/Runtime/WebBrowser/Private/SWebBrowserView.cpp
@@ -173,6 +173,10 @@ void SWebBrowserView::Construct(const FArguments& InArgs, const TSharedPtr<IWebB
 		BrowserWindow->OnDismissAllDialogs().BindSP(this, &SWebBrowserView::HandleDismissAllDialogs);
 		BrowserWindow->OnShowPopup().AddSP(this, &SWebBrowserView::HandleShowPopup);
 		BrowserWindow->OnDismissPopup().AddSP(this, &SWebBrowserView::HandleDismissPopup);
+
+		BrowserWindow->OnSuppressContextMenu().BindSP(this, &SWebBrowserView::HandleSuppressContextMenu);
+		OnSuppressContextMenu = InArgs._OnSuppressContextMenu;
+
 		BrowserViewport = MakeShareable(new FWebBrowserViewport(BrowserWindow));
 #if WITH_CEF3
 		BrowserWidget->SetViewportInterface(BrowserViewport.ToSharedRef());
@@ -569,5 +573,14 @@ void SWebBrowserView::HandleDismissPopup()
 	}
 }
 
+bool SWebBrowserView::HandleSuppressContextMenu()
+{
+	if (OnSuppressContextMenu.IsBound())
+	{
+		return OnSuppressContextMenu.Execute();
+	}
+
+	return false;
+}
 
 #undef LOCTEXT_NAMESPACE
diff --git a/Engine/Source/Runtime/WebBrowser/Public/IWebBrowserWindow.h b/Engine/Source/Runtime/WebBrowser/Public/IWebBrowserWindow.h
--- a/Engine/Source/Runtime/WebBrowser/Public/IWebBrowserWindow.h
+++ b/Engine/Source/Runtime/WebBrowser/Public/IWebBrowserWindow.h
@@ -343,6 +343,10 @@ public:
 	DECLARE_DELEGATE(FOnDismissAllDialogs)
 	virtual FOnDismissAllDialogs& OnDismissAllDialogs() = 0;
 
+	/** Should return true if this dialog wants to suppress the context menu */
+	DECLARE_DELEGATE_RetVal(bool, FOnSuppressContextMenu);
+	virtual FOnSuppressContextMenu& OnSuppressContextMenu() = 0;
+
 protected:
 
 	/** Virtual Destructor. */
diff --git a/Engine/Source/Runtime/WebBrowser/Public/SWebBrowser.h b/Engine/Source/Runtime/WebBrowser/Public/SWebBrowser.h
--- a/Engine/Source/Runtime/WebBrowser/Public/SWebBrowser.h
+++ b/Engine/Source/Runtime/WebBrowser/Public/SWebBrowser.h
@@ -16,6 +16,8 @@ class SEditableTextBox;
 struct FWebNavigationRequest;
 enum class EWebBrowserDialogEventResponse;
 
+DECLARE_DELEGATE_RetVal(bool, FOnSuppressContextMenu);
+
 class WEBBROWSER_API SWebBrowser
 	: public SCompoundWidget
 {
@@ -109,6 +111,9 @@ public:
 		/** Called to dismiss any dialogs shown via OnShowDialog. */
 		SLATE_EVENT(FSimpleDelegate, OnDismissAllDialogs)
 
+		SLATE_EVENT(FOnSuppressContextMenu, OnSuppressContextMenu);
+
+
 	SLATE_END_ARGS()
 
 
@@ -287,4 +292,5 @@ private:
 
 	/** The initial throbber setting */
 	bool bShowInitialThrobber;
+
 };
diff --git a/Engine/Source/Runtime/WebBrowser/Public/SWebBrowserView.h b/Engine/Source/Runtime/WebBrowser/Public/SWebBrowserView.h
--- a/Engine/Source/Runtime/WebBrowser/Public/SWebBrowserView.h
+++ b/Engine/Source/Runtime/WebBrowser/Public/SWebBrowserView.h
@@ -37,6 +37,7 @@ public:
 	DECLARE_DELEGATE_RetVal_TwoParams(bool, FOnBeforeBrowse, const FString& /*Url*/, const FWebNavigationRequest& /*Request*/)
 	DECLARE_DELEGATE_RetVal_ThreeParams(bool, FOnLoadUrl, const FString& /*Method*/, const FString& /*Url*/, FString& /* Response */)
 	DECLARE_DELEGATE_RetVal_OneParam(EWebBrowserDialogEventResponse, FOnShowDialog, const TWeakPtr<IWebBrowserDialog>&);
+	DECLARE_DELEGATE_RetVal(bool, FOnSuppressContextMenu);
 
 	SLATE_BEGIN_ARGS(SWebBrowserView)
 		: _InitialURL(TEXT("https://www.google.com"))
@@ -115,6 +116,8 @@ public:
 		/** Called to dismiss any dialogs shown via OnShowDialog. */
 		SLATE_EVENT(FSimpleDelegate, OnDismissAllDialogs)
 
+		SLATE_EVENT(FOnSuppressContextMenu, OnSuppressContextMenu);
+
 	SLATE_END_ARGS()
 
 
@@ -361,4 +364,10 @@ private:
 	FSimpleDelegate OnDismissAllDialogs;
 
 	FDelegateHandle SlateParentWindowSetupTickHandle;
+
+	FOnSuppressContextMenu OnSuppressContextMenu;
+
+protected:
+	bool HandleSuppressContextMenu();
+
 };
