Copying //UE4/Dev-Automation to //UE4/Dev-Main (Source: //UE4/Dev-Automation @ 3376875)

==========================
MAJOR FEATURES + CHANGES
==========================

Change 3340659 on 2017/03/09 by nick.bullard

	Submitting Template map with update to lighting. The direction of the light was not optimal for looking down X+ for screenshots, which is what the SSF actor does be default.

Change 3340686 on 2017/03/09 by [REDACTED_PERSON]

	Updating Animaton/SetVertexColor map to use Template map setup.

Change 3364139 on 2017/03/24 by [REDACTED_PERSON]

	Automation - Removing the option to disable taking screenshots, doesn't make sense to have that any more.  Screenshot comparison in an intergal part of the tests now.
	Automation - The system now waits for the asset registry to finish loading assets before running tests.

Change 3364149 on 2017/03/24 by [REDACTED_PERSON]

	Automation - Fixing a bug with the source control not reverting the state of screenshots before marking them for delete.

Change 3364588 on 2017/03/24 by [REDACTED_PERSON]

	Removing the allow screenshots option.

Change 3364591 on 2017/03/24 by [REDACTED_PERSON]

	Automation - Making the camera functional test actor respect the delay.

Change 3364598 on 2017/03/24 by [REDACTED_PERSON]

	Automation - Disabling Noisy rendering features now actually disables TAA, instead of enabling FXAA.

Change 3364723 on 2017/03/25 by [REDACTED_PERSON]

	Automation - Fixing the client functional test runner to pass object and package path so that in the editor we can reliably load a particular map with object path, but for actually opening the map we can use the package path when running the test since that is what works at editor and in a cooked game.

Change 3366600 on 2017/03/27 by [REDACTED_PERSON]

	Adding FuntionalTest map back as selection in New Map window

Change 3367590 on 2017/03/28 by [REDACTED_PERSON]

	Automation - Fixing a bug with initial tolerence levels.  They showed as 'Low' in the UI, but the values were in fact all set to 0.  So new screenshot tests had zero tolerence for pixel differences.

	Automation - Adding some documentation to Functional Test.

Change 3367602 on 2017/03/28 by [REDACTED_PERSON]

	Automation - Forcing more things to stream in before the screenshots are taken.

Change 3367604 on 2017/03/28 by [REDACTED_PERSON]

	Automation - Adding some flare to how disabled functional test actors appear in the editor.

Change 3368024 on 2017/03/28 by mason.seay

	Added new test to map.  Disabled since it hasn't ran at all

Change 3368109 on 2017/03/28 by mason.seay

	Updating authorship on test actors

Change 3369701 on 2017/03/29 by [REDACTED_PERSON]

	Removing old Ground Truth images. Investigating why new copies were made rather than replacing original.

Change 3373253 on 2017/03/30 by [REDACTED_PERSON]

	Test updates for Containers. Adding coverage for native containers

Change 3373294 on 2017/03/30 by [REDACTED_PERSON]

	Automation - Fixing a bug with Tolerence levels set via blueprints, now properly configuring the tolerence amounts before the screenshot is taken.

Change 3374355 on 2017/03/31 by [REDACTED_PERSON]

	Added testing scenarios to the test actors and fixed a few typos.

diff --git a/Engine/Plugins/Tests/EditorTests/Source/EditorTests/Private/UnrealEd/EditorAutomationTests.cpp b/Engine/Plugins/Tests/EditorTests/Source/EditorTests/Private/UnrealEd/EditorAutomationTests.cpp
--- a/Engine/Plugins/Tests/EditorTests/Source/EditorTests/Private/UnrealEd/EditorAutomationTests.cpp
+++ b/Engine/Plugins/Tests/EditorTests/Source/EditorTests/Private/UnrealEd/EditorAutomationTests.cpp
@@ -263,8 +263,6 @@ bool FLoadAllMapsInEditorTest::RunTest(const FString& Parameters)
 	}
 	
 
-	const bool bTakeScreenshots = FAutomationTestFramework::Get().IsScreenshotAllowed();
-	if( bTakeScreenshots )
 	{
 		//Find the main editor window
 		TArray<TSharedRef<SWindow> > AllWindows;
@@ -300,15 +298,6 @@ bool FLoadAllMapsInEditorTest::RunTest(const FString& Parameters)
 			ADD_LATENT_AUTOMATION_COMMAND(FWaitLatentCommand(1.5f));
 		}
 	}
-	else
-	{
-		//Get the current number of seconds.  This will be used to track how long it took to load the map.
-		MapLoadStartTime = FPlatformTime::Seconds();
-				//Load the map
-		FAutomationEditorCommonUtils::LoadMap(MapName);
-		//Log how long it took to launch the map.
-		UE_LOG(LogEditorAutomationTests, Display, TEXT("Map '%s' took %.3f to load"), *MapName, FPlatformTime::Seconds() - MapLoadStartTime);
-	}
 
 	return true;
 }
diff --git a/Engine/Source/Developer/AutomationController/AutomationController.Build.cs b/Engine/Source/Developer/AutomationController/AutomationController.Build.cs
--- a/Engine/Source/Developer/AutomationController/AutomationController.Build.cs
+++ b/Engine/Source/Developer/AutomationController/AutomationController.Build.cs
@@ -23,7 +23,8 @@ namespace UnrealBuildTool.Rules
                     "Json",
                     "JsonUtilities",
 					"ScreenShotComparisonTools",
-					"HTTP"
+					"HTTP",
+                    "AssetRegistry"
                 }
 			);
 
diff --git a/Engine/Source/Developer/AutomationController/Private/AutomationCommandline.cpp b/Engine/Source/Developer/AutomationController/Private/AutomationCommandline.cpp
--- a/Engine/Source/Developer/AutomationController/Private/AutomationCommandline.cpp
+++ b/Engine/Source/Developer/AutomationController/Private/AutomationCommandline.cpp
@@ -13,33 +13,28 @@
 #include "FileManager.h"
 #include "Paths.h"
 #include "FileHelper.h"
+#include "AssetRegistryModule.h"
 
 
 /** States for running the automation process */
-namespace EAutomationTestState
+enum class EAutomationTestState : uint8
 {
-	enum Type
-	{
-		Idle,				// Automation process is not running
-		FindWorkers,		// Find workers to run the tests
-		RequestTests,		// Find the tests that can be run on the workers
-		DoingRequestedWork,	// Do whatever was requested from the commandline
-		Complete			// The process is finished
-	};
+	Initializing,		// 
+	Idle,				// Automation process is not running
+	FindWorkers,		// Find workers to run the tests
+	RequestTests,		// Find the tests that can be run on the workers
+	DoingRequestedWork,	// Do whatever was requested from the commandline
+	Complete			// The process is finished
 };
 
-namespace EAutomationCommand
+enum class EAutomationCommand : uint8
 {
-	enum Type
-	{
-		ListAllTests,				//List all tests for the session
-		//RunSingleTest,			//Run one test specified by the commandline
-		RunCommandLineTests,		//Run only tests that are listed on the commandline
-		RunCheckpointTests,			// Run only tests listed on the commandline with checkpoints in case of a crash.
-		RunAll,						//Run all the tests that are supported
-		RunFilter,                  //
-		Quit						//quit the app when tests are done
-	};
+	ListAllTests,			//List all tests for the session
+	RunCommandLineTests,	//Run only tests that are listed on the commandline
+	RunCheckpointTests,		// Run only tests listed on the commandline with checkpoints in case of a crash.
+	RunAll,					//Run all the tests that are supported
+	RunFilter,              //
+	Quit					//quit the app when tests are done
 };
 
 
@@ -56,7 +51,7 @@ public:
 		SessionID = FApp::GetSessionId();
 
 		// Set state to FindWorkers to kick off the testing process
-		AutomationTestState = EAutomationTestState::Idle;
+		AutomationTestState = EAutomationTestState::Initializing;
 		DelayTimer = 5.0f;
 
 		// Load the automation controller
@@ -65,11 +60,9 @@ public:
 
 		AutomationController->Init();
 
-		const bool bSkipScreenshots = FParse::Param(FCommandLine::Get(), TEXT("NoScreenshots"));
 		//TODO AUTOMATION Always use fullsize screenshots.
 		const bool bFullSizeScreenshots = FParse::Param(FCommandLine::Get(), TEXT("FullSizeScreenshots"));
 		const bool bSendAnalytics = FParse::Param(FCommandLine::Get(), TEXT("SendAutomationAnalytics"));
-		AutomationController->SetScreenshotsEnabled(!bSkipScreenshots);
 
 		// Register for the callback that tells us there are tests available
 		AutomationController->OnTestsRefreshed().AddRaw(this, &FAutomationExecCmd::HandleRefreshTestCallback);
@@ -301,6 +294,16 @@ public:
 		// Update the automation process
 		switch (AutomationTestState)
 		{
+			case EAutomationTestState::Initializing:
+			{
+				FAssetRegistryModule& AssetRegistryModule = FModuleManager::LoadModuleChecked<FAssetRegistryModule>("AssetRegistry");
+				if ( AssetRegistryModule.Get().IsLoadingAssets() == false )
+				{
+					AutomationTestState = EAutomationTestState::Idle;
+				}
+
+				break;
+			}
 			case EAutomationTestState::FindWorkers:
 			{
 				FindWorkers(DeltaTime);
@@ -520,16 +523,16 @@ private:
 	IAutomationControllerManagerPtr AutomationController;
 
 	/** The current state of the automation process */
-	EAutomationTestState::Type AutomationTestState;
+	EAutomationTestState AutomationTestState;
 
 	/** The priority flags we would like to run */
 	EAutomationTestFlags::Type AutomationPriority;
 
 	/** What work was requested */
-	TArray<EAutomationCommand::Type> AutomationCommandQueue;
+	TArray<EAutomationCommand> AutomationCommandQueue;
 
 	/** What work was requested */
-	EAutomationCommand::Type AutomationCommand;
+	EAutomationCommand AutomationCommand;
 
 	/** Delay used before finding workers on game instances. Just to ensure they have started up */
 	float DelayTimer;
diff --git a/Engine/Source/Developer/AutomationController/Private/AutomationControllerManager.h b/Engine/Source/Developer/AutomationController/Private/AutomationControllerManager.h
--- a/Engine/Source/Developer/AutomationController/Private/AutomationControllerManager.h
+++ b/Engine/Source/Developer/AutomationController/Private/AutomationControllerManager.h
@@ -160,16 +160,6 @@ public:
 		bSendAnalytics = bNewValue;
 	}
 
-	virtual bool IsScreenshotAllowed() const override
-	{
-		return bScreenshotsEnabled;
-	}
-
-	virtual void SetScreenshotsEnabled( const bool bNewValue ) override
-	{
-		bScreenshotsEnabled = bNewValue;
-	}
-
 	virtual void SetFilter( TSharedPtr< AutomationFilterCollection > InFilter ) override
 	{
 		ReportManager.SetFilter( InFilter );
@@ -513,9 +503,6 @@ private:
 	/** The current test pass we are on. */
 	int32 CurrentTestPass;
 
-	/** If screenshots are enabled. */
-	bool bScreenshotsEnabled;
-
 	/** If we should send result to analytics */
 	bool bSendAnalytics;
 
diff --git a/Engine/Source/Developer/AutomationController/Private/AutomationControllerManger.cpp b/Engine/Source/Developer/AutomationController/Private/AutomationControllerManger.cpp
--- a/Engine/Source/Developer/AutomationController/Private/AutomationControllerManger.cpp
+++ b/Engine/Source/Developer/AutomationController/Private/AutomationControllerManger.cpp
@@ -175,7 +175,6 @@ void FAutomationControllerManager::Init()
 
 	AutomationTestState = EAutomationControllerModuleState::Disabled;
 	bTestResultsAvailable = false;
-	bScreenshotsEnabled = true;
 	bSendAnalytics = FParse::Param(FCommandLine::Get(), TEXT("SendAutomationAnalytics"));
 }
 
@@ -442,7 +441,7 @@ void FAutomationControllerManager::ExecuteNextTask( int32 ClusterIndex, OUT bool
 							FMessageAddress DeviceAddress = DeviceAddresses[AddressIndex];
 
 							// Send the test to the device for execution!
-							MessageEndpoint->Send(new FAutomationWorkerRunTests(ExecutionCount, AddressIndex, NextTest->GetCommand(), NextTest->GetDisplayName(), bScreenshotsEnabled, bSendAnalytics), DeviceAddress);
+							MessageEndpoint->Send(new FAutomationWorkerRunTests(ExecutionCount, AddressIndex, NextTest->GetCommand(), NextTest->GetDisplayName(), bSendAnalytics), DeviceAddress);
 
 							// Add a test so we can check later if the device is still active
 							TestRunningArray.Add(FTestRunningInfo(DeviceAddress));
diff --git a/Engine/Source/Developer/AutomationController/Public/Interfaces/IAutomationControllerManager.h b/Engine/Source/Developer/AutomationController/Public/Interfaces/IAutomationControllerManager.h
--- a/Engine/Source/Developer/AutomationController/Public/Interfaces/IAutomationControllerManager.h
+++ b/Engine/Source/Developer/AutomationController/Public/Interfaces/IAutomationControllerManager.h
@@ -217,16 +217,6 @@ public:
 	*/
 	virtual void SetSendAnalytics(const bool bNewValue) = 0;
 
-	/**
-	 * Returns if screenshots are allowed.
-	 */
-	virtual bool IsScreenshotAllowed() const = 0;
-
-	/**
-	 * Sets if screenshots are enabled.
-	 */
-	virtual void SetScreenshotsEnabled( const bool bNewValue ) = 0;
-
 	/**
 	 * Filters the visible list of tests.
 	 */
diff --git a/Engine/Source/Developer/AutomationWindow/Private/SAutomationWindow.cpp b/Engine/Source/Developer/AutomationWindow/Private/SAutomationWindow.cpp
--- a/Engine/Source/Developer/AutomationWindow/Private/SAutomationWindow.cpp
+++ b/Engine/Source/Developer/AutomationWindow/Private/SAutomationWindow.cpp
@@ -1183,19 +1183,6 @@ TSharedRef< SWidget > SAutomationWindow::GenerateTestsOptionsMenuContent( )
 			.Text(LOCTEXT("AutomationSendAnalyticsText", "Enable analytics"))
 		];
 
-	TSharedRef<SWidget> EnableScreenshotsWidget =
-		SNew(SCheckBox)
-		.IsChecked(this, &SAutomationWindow::IsEnableScreenshotsCheckBoxChecked)
-		.OnCheckStateChanged(this, &SAutomationWindow::HandleEnableScreenshotsBoxCheckStateChanged)
-		.Padding(FMargin(4.0f, 0.0f))
-		.ToolTipText(LOCTEXT("AutomationEnableScreenshotsTip", "If checked, tests that support it will save screenshots"))
-		.IsEnabled( this, &SAutomationWindow::IsAutomationControllerIdle )
-		.Content()
-		[
-			SNew(STextBlock)
-			.Text(LOCTEXT("AutomationNoScreenshotText", "Enable screenshots"))
-		];
-
 
 	MenuBuilder.BeginSection("AutomationWindowRunTest", LOCTEXT("RunTestOptions", "Advanced Settings"));
 	{
@@ -1203,11 +1190,6 @@ TSharedRef< SWidget > SAutomationWindow::GenerateTestsOptionsMenuContent( )
 		MenuBuilder.AddWidget(SendAnalyticsWidget, FText::GetEmpty());
 	}
 	MenuBuilder.EndSection();
-	MenuBuilder.BeginSection("AutomationWindowScreenshots", LOCTEXT("ScreenshotOptions", "Screenshot Settings"));
-	{
-		MenuBuilder.AddWidget(EnableScreenshotsWidget, FText::GetEmpty());
-	}
-	MenuBuilder.EndSection();
 
 	return MenuBuilder.MakeWidget();
 }
@@ -1222,21 +1204,6 @@ void SAutomationWindow::HandleSendAnalyticsBoxCheckStateChanged(ECheckBoxState C
 	AutomationController->SetSendAnalytics(CheckBoxState == ECheckBoxState::Checked);
 }
 
-ECheckBoxState SAutomationWindow::IsEnableScreenshotsCheckBoxChecked() const
-{
-	return AutomationController->IsScreenshotAllowed() ? ECheckBoxState::Checked : ECheckBoxState::Unchecked;
-}
-
-void SAutomationWindow::HandleEnableScreenshotsBoxCheckStateChanged(ECheckBoxState CheckBoxState)
-{
-	AutomationController->SetScreenshotsEnabled(CheckBoxState == ECheckBoxState::Checked);
-}
-
-bool SAutomationWindow::IsFullSizeScreenshotsOptionEnabled() const
-{
-	return AutomationController->IsScreenshotAllowed() && IsAutomationControllerIdle();
-}
-
 TArray<FString> SAutomationWindow::SaveExpandedTestNames(TSet<TSharedPtr<IAutomationReport>> ExpandedItems)
 {
 	TArray<FString> ExpandedItemsNames;
diff --git a/Engine/Source/Developer/AutomationWindow/Private/SAutomationWindow.h b/Engine/Source/Developer/AutomationWindow/Private/SAutomationWindow.h
--- a/Engine/Source/Developer/AutomationWindow/Private/SAutomationWindow.h
+++ b/Engine/Source/Developer/AutomationWindow/Private/SAutomationWindow.h
@@ -307,17 +307,10 @@ private:
 	
 	/** Returns if analytics should be sent to the back end*/
 	ECheckBoxState IsSendAnalyticsCheckBoxChecked() const;
+
 	/** Toggles if we are sending analytics results from the tests*/
 	void HandleSendAnalyticsBoxCheckStateChanged(ECheckBoxState CheckBoxState);
 
-	/** Returns if screen shots are enabled */
-	ECheckBoxState IsEnableScreenshotsCheckBoxChecked() const;
-	/** Toggles if we are taking screenshots */
-	void HandleEnableScreenshotsBoxCheckStateChanged(ECheckBoxState CheckBoxState);
-
-	/** Returns if the full size screenshots option is enabled */
-	bool IsFullSizeScreenshotsOptionEnabled() const;
-
 	/** Returns if a device group is enabled */
 	ECheckBoxState IsDeviceGroupCheckBoxIsChecked(const int32 DeviceGroupFlag) const;
 	
diff --git a/Engine/Source/Developer/FunctionalTesting/Classes/FunctionalTest.h b/Engine/Source/Developer/FunctionalTesting/Classes/FunctionalTest.h
--- a/Engine/Source/Developer/FunctionalTesting/Classes/FunctionalTest.h
+++ b/Engine/Source/Developer/FunctionalTesting/Classes/FunctionalTest.h
@@ -256,23 +256,50 @@ public:
 	AFunctionalTest(const FObjectInitializer& ObjectInitializer = FObjectInitializer::Get());
 
 private_subobject:
-	DEPRECATED_FORGAME(4.6, "SpriteComponent should not be accessed directly, please use GetSpriteComponent() function instead. SpriteComponent will soon be private and your code will not compile.")
 	UPROPERTY()
 	UBillboardComponent* SpriteComponent;
 
 protected:
+	/**
+	 * Allows a test to be disabled.  If a test is disabled, it will not appear in the set of
+	 * runnable tests (after saving the map).
+	 */
 	UPROPERTY(EditAnywhere, BlueprintReadOnly, Category="Functional Testing")
 	uint32 bIsEnabled:1;
 
+	/**
+	 * If this is enabled, any warning logged while this functional test is running is treated as
+	 * an error.
+	 */
 	UPROPERTY(EditAnywhere, BlueprintReadOnly, Category="Functional Testing")
 	uint32 bWarningsAsErrors:1;
 
+	/**
+	 * The author is the group or person responsible for the test.  Generally you should use a group name
+	 * like 'Editor Team' or 'Rendering Team'.  When a test fails it may not be obvious who should investigate
+	 * so this provides a associate responsible groups with tests.
+	 */
+	UPROPERTY(EditAnywhere, BlueprintReadOnly, Category="Functional Testing", meta=( MultiLine="true" ))
+	FString Author;
+
+	/**
+	 * A description of the test, like what is this test trying to determine.
+	 */
 	UPROPERTY(EditAnywhere, BlueprintReadOnly, Category="Functional Testing", meta=( MultiLine="true" ))
 	FString Description;
 
+	/**
+	 * Allows you to specify another actor to view the test from.  Usually this is a camera you place
+	 * in the map to observe the test.  Not useful when running on a build farm, but provides a handy
+	 * way to observe the test from a different location than you place the functional test actor.
+	 */
 	UPROPERTY(EditAnywhere, BlueprintReadWrite, Category="Functional Testing")
 	AActor* ObservationPoint;
 
+	/**
+	 * A random number stream that you can use during testing.  This number stream will be consistent
+	 * every time the test is run.
+	 */
 	UPROPERTY(EditAnywhere, BlueprintReadWrite, Category="Functional Testing", AdvancedDisplay)
 	FRandomStream RandomNumbersStream;
 
@@ -577,7 +604,9 @@ public:
 	float TotalTime;
 
 	uint32 RunFrame;
+
 	uint32 StartFrame;
+	float StartTime;
 
 private:
 	bool bIsReady;
diff --git a/Engine/Source/Developer/FunctionalTesting/Private/AutomationBlueprintFunctionLibrary.cpp b/Engine/Source/Developer/FunctionalTesting/Private/AutomationBlueprintFunctionLibrary.cpp
--- a/Engine/Source/Developer/FunctionalTesting/Private/AutomationBlueprintFunctionLibrary.cpp
+++ b/Engine/Source/Developer/FunctionalTesting/Private/AutomationBlueprintFunctionLibrary.cpp
@@ -24,6 +24,8 @@
 #include "ShaderCompiler.h"
 #include "AutomationBlueprintFunctionLibrary.h"
 #include "BufferVisualizationData.h"
+#include "Engine/LocalPlayer.h"
+#include "ContentStreaming.h"
 
 #define LOCTEXT_NAMESPACE "Automation"
 
@@ -105,8 +107,9 @@ class FAutomationScreenshotTaker
 	FConsoleVariableSwapper ContactShadows;
 
 public:
-	FAutomationScreenshotTaker(const FString& InName, FAutomationScreenshotOptions InOptions)
-		: Name(InName)
+	FAutomationScreenshotTaker(UWorld* InWorld, const FString& InName, FAutomationScreenshotOptions InOptions)
+		: World(InWorld)
+		, Name(InName)
 		, Options(InOptions)
 		, DefaultFeature_AntiAliasing(TEXT("r.DefaultFeature.AntiAliasing"))
 		, DefaultFeature_AutoExposure(TEXT("r.DefaultFeature.AutoExposure"))
@@ -117,22 +120,42 @@ public:
 		, EyeAdaptationQuality(TEXT("r.EyeAdaptationQuality"))
 		, ContactShadows(TEXT("r.ContactShadows"))
 	{
+		// Force all mip maps to load before taking the screenshot.
+		UTexture::ForceUpdateTextureStreaming();
+
+		// Force all shader compiling to finish.
+		GShaderCompilingManager->FinishAllCompilation();
+
+		IStreamingManager::Get().StreamAllResources(0.0f);
+
 		GEngine->GameViewport->OnScreenshotCaptured().AddRaw(this, &FAutomationScreenshotTaker::GrabScreenShot);
 
 		check(IsInGameThread());
 
 		if ( Options.bDisableNoisyRenderingFeatures )
 		{
-			DefaultFeature_AntiAliasing.Set(1);
+			DefaultFeature_AntiAliasing.Set(0);
 			DefaultFeature_AutoExposure.Set(0);
 			DefaultFeature_MotionBlur.Set(0);
-			PostProcessAAQuality.Set(1);
+			PostProcessAAQuality.Set(0);
 			MotionBlurQuality.Set(0);
 			ScreenSpaceReflectionQuality.Set(0);
 			EyeAdaptationQuality.Set(0);
 			ContactShadows.Set(0);
 		}
 
+		Options.SetToleranceAmounts(Options.Tolerance);
+
+		//if ( World.IsValid() )
+		//{
+		//	if ( ULocalPlayer* LocalPlayer = World->GetFirstLocalPlayerFromController() )
+		//	{
+		//		LocalPlayer->ResetViewState();
+		//	}
+		//}
+
+		//ADD_LATENT_AUTOMATION_COMMAND(FExecWorldStringLatentCommand(TEXT("r.ResetViewState")));
+
 		if ( UGameViewportClient* ViewportClient = GEngine->GameViewport )
 		{
 			static IConsoleVariable* ICVar = IConsoleManager::Get().FindConsoleVariable(FBufferVisualizationData::GetVisualizationTargetConsoleCommandName());
@@ -148,7 +171,7 @@ public:
 		}
 	}
 
-	~FAutomationScreenshotTaker()
+	virtual ~FAutomationScreenshotTaker()
 	{
 		check(IsInGameThread());
 
@@ -242,6 +265,10 @@ public:
 
 		delete this;
 	}
+
+private:
+
+	TWeakObjectPtr<UWorld> World;
 };
 
 #endif
@@ -251,14 +278,8 @@ UAutomationBlueprintFunctionLibrary::UAutomationBlueprintFunctionLibrary(const c
 {
 }
 
-bool UAutomationBlueprintFunctionLibrary::TakeAutomationScreenshotInternal(const FString& Name, FAutomationScreenshotOptions Options)
+bool UAutomationBlueprintFunctionLibrary::TakeAutomationScreenshotInternal(UObject* WorldContextObject, const FString& Name, FAutomationScreenshotOptions Options)
 {
-	if ( !FAutomationTestFramework::Get().IsScreenshotAllowed() )
-	{
-		UE_LOG(AutomationFunctionLibrary, Log, TEXT("Attempted to capture screenshot (%s) but screenshots are not enabled."), *Name);
-		return false;
-	}
-
 	// Fallback resolution if all else fails for screenshots.
 	uint32 ResolutionX = 1280;
 	uint32 ResolutionY = 720;
@@ -295,14 +316,8 @@ bool UAutomationBlueprintFunctionLibrary::TakeAutomationScreenshotInternal(const
 		}
 	}
 
-	// Force all mip maps to load before taking the screenshot.
-	UTexture::ForceUpdateTextureStreaming();
-
-	// Force all shader compiling to finish.
-	GShaderCompilingManager->FinishAllCompilation();
-
 #if (WITH_DEV_AUTOMATION_TESTS || WITH_PERF_AUTOMATION_TESTS)
-	FAutomationScreenshotTaker* TempObject = new FAutomationScreenshotTaker(Name, Options);
+	FAutomationScreenshotTaker* TempObject = new FAutomationScreenshotTaker(WorldContextObject ? WorldContextObject->GetWorld() : nullptr, Name, Options);
 #endif
 
 	//static IConsoleVariable* HighResScreenshotDelay = IConsoleManager::Get().FindConsoleVariable(TEXT("r.HighResScreenshotDelay"));
@@ -395,25 +410,6 @@ void UAutomationBlueprintFunctionLibrary::TakeAutomationScreenshotAtCamera(UObje
 
 void UAutomationBlueprintFunctionLibrary::TakeAutomationScreenshotOfUI(UObject* WorldContextObject, FLatentActionInfo LatentInfo, const FString& Name, const FAutomationScreenshotOptions& Options)
 {
-	if ( !FAutomationTestFramework::Get().IsScreenshotAllowed() )
-	{
-		UE_LOG(AutomationFunctionLibrary, Log, TEXT("Attempted to capture screenshot (%s) but screenshots are not enabled."), *Name);
-
-		if ( UWorld* World = WorldContextObject->GetWorld() )
-		{
-			FLatentActionManager& LatentActionManager = World->GetLatentActionManager();
-			if ( LatentActionManager.FindExistingAction<FDelayAction>(LatentInfo.CallbackTarget, LatentInfo.UUID) == nullptr )
-			{
-				LatentActionManager.AddNewAction(LatentInfo.CallbackTarget, LatentInfo.UUID, new FDelayAction(0, LatentInfo));
-			}
-		}
-
-		return;
-	}
-
-	// Force all mip maps to load before taking the screenshot.
-	UTexture::ForceUpdateTextureStreaming();
-
 	if ( UWorld* World = WorldContextObject->GetWorld() )
 	{
 		if ( UGameViewportClient* GameViewport = WorldContextObject->GetWorld()->GetGameViewport() )
@@ -426,7 +422,7 @@ void UAutomationBlueprintFunctionLibrary::TakeAutomationScreenshotOfUI(UObject*
 				if ( FSlateApplication::Get().TakeScreenshot(Viewport.ToSharedRef(), OutColorData, OutSize) )
 				{
 #if (WITH_DEV_AUTOMATION_TESTS || WITH_PERF_AUTOMATION_TESTS)
-					FAutomationScreenshotTaker* TempObject = new FAutomationScreenshotTaker(Name, Options);
+					FAutomationScreenshotTaker* TempObject = new FAutomationScreenshotTaker(WorldContextObject->GetWorld(), Name, Options);
 
 					FAutomationScreenshotData Data = AutomationCommon::BuildScreenshotData(GWorld->GetName(), Name, OutSize.X, OutSize.Y);
 
diff --git a/Engine/Source/Developer/FunctionalTesting/Private/FunctionalTest.cpp b/Engine/Source/Developer/FunctionalTesting/Private/FunctionalTest.cpp
--- a/Engine/Source/Developer/FunctionalTesting/Private/FunctionalTest.cpp
+++ b/Engine/Source/Developer/FunctionalTesting/Private/FunctionalTest.cpp
@@ -94,6 +94,7 @@ AFunctionalTest::AFunctionalTest( const FObjectInitializer& ObjectInitializer )
 	, TotalTime(0.f)
 	, RunFrame(0)
 	, StartFrame(0)
+	, StartTime(0.0f)
 	, bIsReady(false)
 {
 	PrimaryActorTick.bCanEverTick = true;
@@ -154,7 +155,6 @@ AFunctionalTest::AFunctionalTest( const FObjectInitializer& ObjectInitializer )
 	{
 		TestName->bHiddenInGame = true;
 		TestName->SetHorizontalAlignment(EHTA_Center);
-		TestName->SetTextRenderColor(FColor(11, 255, 0));
 		TestName->SetRelativeLocation(FVector(0, 0, 80));
 		TestName->SetRelativeRotation(FRotator(0, 0, 0));
 		TestName->PostPhysicsComponentTick.bCanEverTick = false;
@@ -170,7 +170,16 @@ void AFunctionalTest::OnConstruction(const FTransform& Transform)
 #if WITH_EDITOR
 	if ( TestName )
 	{
-		TestName->SetText(FText::FromString(GetActorLabel()));
+		if ( bIsEnabled )
+		{
+			TestName->SetTextRenderColor(FColor(11, 255, 0));
+			TestName->SetText(FText::FromString(GetActorLabel()));
+		}
+		else
+		{
+			TestName->SetTextRenderColor(FColor(55, 55, 55));
+			TestName->SetText(FText::FromString(GetActorLabel() + TEXT("\n") + TEXT("# Disabled #")));
+		}
 	}
 #endif
 }
@@ -214,6 +223,7 @@ void AFunctionalTest::StartTest()
 {
 	TotalTime = 0.f;
 	StartFrame = GFrameNumber;
+	StartTime = GetWorld()->GetTimeSeconds();
 
 	ReceiveStartTest();
 	OnTestStart.Broadcast();
diff --git a/Engine/Source/Developer/FunctionalTesting/Private/ScreenshotFunctionalTest.cpp b/Engine/Source/Developer/FunctionalTesting/Private/ScreenshotFunctionalTest.cpp
--- a/Engine/Source/Developer/FunctionalTesting/Private/ScreenshotFunctionalTest.cpp
+++ b/Engine/Source/Developer/FunctionalTesting/Private/ScreenshotFunctionalTest.cpp
@@ -32,14 +32,19 @@ void AScreenshotFunctionalTest::PrepareTest()
 
 bool AScreenshotFunctionalTest::IsReady_Implementation()
 {
-	return (GFrameNumber - RunFrame) > 2;
+	if ( (GetWorld()->GetTimeSeconds() - StartTime) > ScreenshotOptions.Delay )
+	{
+		return ( GFrameNumber - RunFrame ) > 5;
+	}
+	
+	return false;
 }
 
 void AScreenshotFunctionalTest::StartTest()
 {
 	Super::StartTest();
 
-	UAutomationBlueprintFunctionLibrary::TakeAutomationScreenshotInternal(GetName(), ScreenshotOptions);
+	UAutomationBlueprintFunctionLibrary::TakeAutomationScreenshotInternal(this, GetName(), ScreenshotOptions);
 
 	FAutomationTestFramework::Get().OnScreenshotTakenAndCompared.AddUObject(this, &AScreenshotFunctionalTest::OnScreenshotTakenAndCompared);
 }
diff --git a/Engine/Source/Developer/FunctionalTesting/Private/TakeScreenshotAfterTimeLatentAction.cpp b/Engine/Source/Developer/FunctionalTesting/Private/TakeScreenshotAfterTimeLatentAction.cpp
--- a/Engine/Source/Developer/FunctionalTesting/Private/TakeScreenshotAfterTimeLatentAction.cpp
+++ b/Engine/Source/Developer/FunctionalTesting/Private/TakeScreenshotAfterTimeLatentAction.cpp
@@ -40,7 +40,7 @@ void FTakeScreenshotAfterTimeLatentAction::UpdateOperation(FLatentResponse& Resp
 			{
 				FAutomationTestFramework::Get().OnScreenshotTakenAndCompared.AddRaw(this, &FTakeScreenshotAfterTimeLatentAction::OnScreenshotTakenAndCompared);
 
-				if ( UAutomationBlueprintFunctionLibrary::TakeAutomationScreenshotInternal(ScreenshotName, Options) )
+				if ( UAutomationBlueprintFunctionLibrary::TakeAutomationScreenshotInternal(nullptr, ScreenshotName, Options) )
 				{
 					IssuedScreenshotCapture = true;
 				}
diff --git a/Engine/Source/Developer/FunctionalTesting/Private/Tests/ClientFuncTestPerforming.cpp b/Engine/Source/Developer/FunctionalTesting/Private/Tests/ClientFuncTestPerforming.cpp
--- a/Engine/Source/Developer/FunctionalTesting/Private/Tests/ClientFuncTestPerforming.cpp
+++ b/Engine/Source/Developer/FunctionalTesting/Private/Tests/ClientFuncTestPerforming.cpp
@@ -95,6 +95,16 @@ bool FStartFTestOnMap::Update()
 // Project.Maps.Client Functional Testing
 // Project.Maps.Functional Tests
 
+void ParseTestMapInfo(const FString& Parameters, FString& MapObjectPath, FString& MapPackageName, FString& MapTestName)
+{
+	TArray<FString> ParamArray;
+	Parameters.ParseIntoArray(ParamArray, TEXT(";"), true);
+
+	MapObjectPath = ParamArray[0];
+	MapPackageName = ParamArray[1];
+	MapTestName = ( ParamArray.Num() > 2 ) ? ParamArray[2] : TEXT("");
+}
+
 
 class FClientFunctionalTestingMapsBase : public FAutomationTestBase
 {
@@ -104,22 +114,20 @@ public:
 	{
 	}
 
-	virtual FString GetTestOpenCommand(const FString& Parameter) const override
+	virtual FString GetTestOpenCommand(const FString& Parameters) const override
 	{
-		FString MapPath;
-		FString MapTestName;
-		Parameter.Split(TEXT(";"), &MapPath, &MapTestName);
+		FString MapObjectPath, MapPackageName, MapTestName;
+		ParseTestMapInfo(Parameters, MapObjectPath, MapPackageName, MapTestName);
 
-		return FString::Printf(TEXT("Automate.OpenMapAndFocusActor %s %s"), *MapPath, *MapTestName);
+		return FString::Printf(TEXT("Automate.OpenMapAndFocusActor %s %s"), *MapObjectPath, *MapTestName);
 	}
 
-	virtual FString GetTestAssetPath(const FString& Parameter) const override
+	virtual FString GetTestAssetPath(const FString& Parameters) const override
 	{
-		FString MapPath;
-		FString MapTestName;
-		Parameter.Split(TEXT(";"), &MapPath, &MapTestName);
+		FString MapObjectPath, MapPackageName, MapTestName;
+		ParseTestMapInfo(Parameters, MapObjectPath, MapPackageName, MapTestName);
 
-		return MapPath;
+		return MapObjectPath;
 	}
 };
 
@@ -165,7 +173,7 @@ void FClientFunctionalTestingMaps::GetTests(TArray<FString>& OutBeautifiedNames,
 							if ( MapTest.Split(TEXT("|"), &BeautifulTestName, &RealTestName) )
 							{
 								OutBeautifiedNames.Add(MapAsset.PackageName.ToString() + TEXT(".") + *BeautifulTestName);
-								OutTestCommands.Add(MapAsset.ObjectPath.ToString() + TEXT(";") + *RealTestName);
+								OutTestCommands.Add(MapAsset.ObjectPath.ToString() + TEXT(";") + MapAsset.PackageName.ToString() + TEXT(";") + *RealTestName);
 							}
 						}
 					}
@@ -173,7 +181,7 @@ void FClientFunctionalTestingMaps::GetTests(TArray<FString>& OutBeautifiedNames,
 				else if ( MapAsset.AssetName.ToString().Find(TEXT("FTEST_")) == 0 )
 				{
 					OutBeautifiedNames.Add(MapAsset.AssetName.ToString());
-					OutTestCommands.Add(MapAsset.PackageName.ToString());
+					OutTestCommands.Add(MapAsset.ObjectPath.ToString() + TEXT(";") + MapAsset.PackageName.ToString());
 				}
 			}
 		}
@@ -208,25 +216,22 @@ static UWorld* GetAnyGameWorld()
  */
 bool FClientFunctionalTestingMaps::RunTest(const FString& Parameters)
 {
-	TArray<FString> ParamArray;
-	Parameters.ParseIntoArray(ParamArray, TEXT(";"), true);
-
-	FString MapName = ParamArray[0];
-	FString MapTestName = ( ParamArray.Num() > 1 ) ? ParamArray[1] : TEXT("");
+	FString MapObjectPath, MapPackageName, MapTestName;
+	ParseTestMapInfo(Parameters, MapObjectPath, MapPackageName, MapTestName);
 
 	bool bCanProceed = false;
 
 	FFunctionalTestingModule::Get()->MarkPendingActivation();
 
 	UWorld* TestWorld = GetAnyGameWorld();
-	if ( TestWorld && TestWorld->GetMapName() == MapName )
+	if ( TestWorld && TestWorld->GetMapName() == MapPackageName )
 	{
 		// Map is already loaded.
 		bCanProceed = true;
 	}
 	else
 	{
-		bCanProceed = AutomationOpenMap(MapName);
+		bCanProceed = AutomationOpenMap(MapPackageName);
 	}
 
 	if (bCanProceed)
@@ -248,7 +253,7 @@ bool FClientFunctionalTestingMaps::RunTest(const FString& Parameters)
 	//	ADD_LATENT_AUTOMATION_COMMAND(FWaitLatentCommand(1.f));
 	//  ADD_LATENT_AUTOMATION_COMMAND(FExitGameCommand);
 
-	UE_LOG(LogFunctionalTesting, Error, TEXT("Failed to start the %s map (possibly due to BP compilation issues)"), *MapName);
+	UE_LOG(LogFunctionalTesting, Error, TEXT("Failed to start the %s map (possibly due to BP compilation issues)"), *MapPackageName);
 	return false;
 }
 
diff --git a/Engine/Source/Developer/FunctionalTesting/Public/AutomationBlueprintFunctionLibrary.h b/Engine/Source/Developer/FunctionalTesting/Public/AutomationBlueprintFunctionLibrary.h
--- a/Engine/Source/Developer/FunctionalTesting/Public/AutomationBlueprintFunctionLibrary.h
+++ b/Engine/Source/Developer/FunctionalTesting/Public/AutomationBlueprintFunctionLibrary.h
@@ -17,7 +17,7 @@ class FUNCTIONALTESTING_API UAutomationBlueprintFunctionLibrary : public UBluepr
 	GENERATED_UCLASS_BODY()
 	
 public:
-	static bool TakeAutomationScreenshotInternal(const FString& Name, FAutomationScreenshotOptions Options);
+	static bool TakeAutomationScreenshotInternal(UObject* WorldContextObject, const FString& Name, FAutomationScreenshotOptions Options);
 
 	UFUNCTION(BlueprintCallable, Category = "Automation", meta = (Latent, HidePin = "WorldContextObject", DefaultToSelf = "WorldContextObject", LatentInfo = "LatentInfo", Name = "" ))
 	static void TakeAutomationScreenshot(UObject* WorldContextObject, FLatentActionInfo LatentInfo, const FString& Name, const FAutomationScreenshotOptions& Options);
diff --git a/Engine/Source/Developer/FunctionalTesting/Public/AutomationScreenshotOptions.h b/Engine/Source/Developer/FunctionalTesting/Public/AutomationScreenshotOptions.h
--- a/Engine/Source/Developer/FunctionalTesting/Public/AutomationScreenshotOptions.h
+++ b/Engine/Source/Developer/FunctionalTesting/Public/AutomationScreenshotOptions.h
@@ -6,7 +6,7 @@
 #include "UObject/ObjectMacros.h"
 #include "AutomationScreenshotOptions.generated.h"
 
-UENUM()
+UENUM(BlueprintType)
 enum class EComparisonTolerance : uint8
 {
 	Zero,
@@ -16,7 +16,7 @@ enum class EComparisonTolerance : uint8
 	Custom
 };
 
-USTRUCT()
+USTRUCT(BlueprintType)
 struct FComparisonToleranceAmount
 {
 	GENERATED_BODY()
diff --git a/Engine/Source/Developer/ScreenShotComparison/Private/Models/ScreenComparisonModel.cpp b/Engine/Source/Developer/ScreenShotComparison/Private/Models/ScreenComparisonModel.cpp
--- a/Engine/Source/Developer/ScreenShotComparison/Private/Models/ScreenComparisonModel.cpp
+++ b/Engine/Source/Developer/ScreenShotComparison/Private/Models/ScreenComparisonModel.cpp
@@ -165,11 +165,19 @@ bool FScreenComparisonModel::RemoveExistingApproved(IScreenShotManagerPtr Screen
 
 		FString DestFilePath = LocalApprovedFolder / RelativeFile;
 		SourceControlFiles.Add(DestFilePath);
-
-		IFileManager::Get().Delete(*DestFilePath, false, true, false);
 	}
 
 	ISourceControlProvider& SourceControlProvider = ISourceControlModule::Get().GetProvider();
+	if ( SourceControlProvider.Execute(ISourceControlOperation::Create<FRevert>(), SourceControlFiles) == ECommandResult::Failed )
+	{
+		//TODO Error
+	}
+
+	for ( const FString& File : SourceControlFiles )
+	{
+		IFileManager::Get().Delete(*File, false, true, false);
+	}
+
 	if ( SourceControlProvider.Execute(ISourceControlOperation::Create<FDelete>(), SourceControlFiles) == ECommandResult::Failed )
 	{
 		//TODO Error
diff --git a/Engine/Source/Editor/UnrealEd/Private/Tests/AutomationEditorPromotionCommon.cpp b/Engine/Source/Editor/UnrealEd/Private/Tests/AutomationEditorPromotionCommon.cpp
--- a/Engine/Source/Editor/UnrealEd/Private/Tests/AutomationEditorPromotionCommon.cpp
+++ b/Engine/Source/Editor/UnrealEd/Private/Tests/AutomationEditorPromotionCommon.cpp
@@ -246,65 +246,61 @@ void FEditorPromotionTestUtilities::EndPIE()
 */
 void FEditorPromotionTestUtilities::TakeScreenshot(const FString& ScreenshotName, const FAutomationScreenshotOptions& Options, bool bUseTopWindow)
 {
-	//Update the screenshot name, then take a screenshot.
-	if (FAutomationTestFramework::Get().IsScreenshotAllowed())
-	{
-		TSharedPtr<SWindow> Window;
+	TSharedPtr<SWindow> Window;
 
-		if (bUseTopWindow)
+	if (bUseTopWindow)
+	{
+		Window = FSlateApplication::Get().GetActiveTopLevelWindow();
+	}
+	else
+	{
+		//Find the main editor window
+		TArray<TSharedRef<SWindow> > AllWindows;
+		FSlateApplication::Get().GetAllVisibleWindowsOrdered(AllWindows);
+		if (AllWindows.Num() == 0)
 		{
-			Window = FSlateApplication::Get().GetActiveTopLevelWindow();
+			UE_LOG(LogEditorPromotionTests, Error, TEXT("ERROR: Could not find the main editor window."));
+			return;
 		}
-		else
-		{
-			//Find the main editor window
-			TArray<TSharedRef<SWindow> > AllWindows;
-			FSlateApplication::Get().GetAllVisibleWindowsOrdered(AllWindows);
-			if (AllWindows.Num() == 0)
-			{
-				UE_LOG(LogEditorPromotionTests, Error, TEXT("ERROR: Could not find the main editor window."));
-				return;
-			}
 
-			Window = AllWindows[0];
-		}
+		Window = AllWindows[0];
+	}
 
-		if (Window.IsValid())
-		{
-			TSharedRef<SWidget> WindowRef = Window.ToSharedRef();
+	if (Window.IsValid())
+	{
+		TSharedRef<SWidget> WindowRef = Window.ToSharedRef();
 
-			TArray<FColor> OutImageData;
-			FIntVector OutImageSize;
-			if (FSlateApplication::Get().TakeScreenshot(WindowRef, OutImageData, OutImageSize))
+		TArray<FColor> OutImageData;
+		FIntVector OutImageSize;
+		if (FSlateApplication::Get().TakeScreenshot(WindowRef, OutImageData, OutImageSize))
+		{
+			FAutomationScreenshotData Data = AutomationCommon::BuildScreenshotData(TEXT("Editor"), ScreenshotName, OutImageSize.X, OutImageSize.Y);
+
+			// Copy the relevant data into the metadata for the screenshot.
+			Data.bHasComparisonRules = true;
+			Data.ToleranceRed = Options.ToleranceAmount.Red;
+			Data.ToleranceGreen = Options.ToleranceAmount.Green;
+			Data.ToleranceBlue = Options.ToleranceAmount.Blue;
+			Data.ToleranceAlpha = Options.ToleranceAmount.Alpha;
+			Data.ToleranceMinBrightness = Options.ToleranceAmount.MinBrightness;
+			Data.ToleranceMaxBrightness = Options.ToleranceAmount.MaxBrightness;
+			Data.bIgnoreAntiAliasing = Options.bIgnoreAntiAliasing;
+			Data.bIgnoreColors = Options.bIgnoreColors;
+			Data.MaximumLocalError = Options.MaximumLocalError;
+			Data.MaximumGlobalError = Options.MaximumGlobalError;
+
+			FAutomationTestFramework::Get().OnScreenshotCaptured().ExecuteIfBound(OutImageData, Data);
+
+			// TODO Add comparison support.
+			if ( GIsAutomationTesting )
 			{
-				FAutomationScreenshotData Data = AutomationCommon::BuildScreenshotData(TEXT("Editor"), ScreenshotName, OutImageSize.X, OutImageSize.Y);
-
-				// Copy the relevant data into the metadata for the screenshot.
-				Data.bHasComparisonRules = true;
-				Data.ToleranceRed = Options.ToleranceAmount.Red;
-				Data.ToleranceGreen = Options.ToleranceAmount.Green;
-				Data.ToleranceBlue = Options.ToleranceAmount.Blue;
-				Data.ToleranceAlpha = Options.ToleranceAmount.Alpha;
-				Data.ToleranceMinBrightness = Options.ToleranceAmount.MinBrightness;
-				Data.ToleranceMaxBrightness = Options.ToleranceAmount.MaxBrightness;
-				Data.bIgnoreAntiAliasing = Options.bIgnoreAntiAliasing;
-				Data.bIgnoreColors = Options.bIgnoreColors;
-				Data.MaximumLocalError = Options.MaximumLocalError;
-				Data.MaximumGlobalError = Options.MaximumGlobalError;
-
-				FAutomationTestFramework::Get().OnScreenshotCaptured().ExecuteIfBound(OutImageData, Data);
-
-				// TODO Add comparison support.
-				if ( GIsAutomationTesting )
-				{
-					//FAutomationTestFramework::Get().OnScreenshotCompared.AddRaw(this, &FAutomationScreenshotTaker::OnComparisonComplete);
-				}
+				//FAutomationTestFramework::Get().OnScreenshotCompared.AddRaw(this, &FAutomationScreenshotTaker::OnComparisonComplete);
 			}
 		}
-		else
-		{
-			UE_LOG(LogEditorPromotionTests, Error, TEXT("Failed to find editor window for screenshot (%s)"), *ScreenshotName);
-		}
+	}
+	else
+	{
+		UE_LOG(LogEditorPromotionTests, Error, TEXT("Failed to find editor window for screenshot (%s)"), *ScreenshotName);
 	}
 }
 
diff --git a/Engine/Source/Runtime/AutomationMessages/Classes/AutomationWorkerMessages.h b/Engine/Source/Runtime/AutomationMessages/Classes/AutomationWorkerMessages.h
--- a/Engine/Source/Runtime/AutomationMessages/Classes/AutomationWorkerMessages.h
+++ b/Engine/Source/Runtime/AutomationMessages/Classes/AutomationWorkerMessages.h
@@ -277,10 +277,6 @@ struct FAutomationWorkerRunTests
 	UPROPERTY()
 	FString BeautifiedTestName;
 
-	/** If true, we will save out screenshots for tests that support them. */
-	UPROPERTY(EditAnywhere, Category="Message")
-	bool bScreenshotsEnabled;
-
 	/** If true, send results to analytics when complete */
 	UPROPERTY()
 	bool bSendAnalytics;
@@ -289,12 +285,11 @@ struct FAutomationWorkerRunTests
 	FAutomationWorkerRunTests( ) { }
 
 	/** Creates and initializes a new instance. */
-	FAutomationWorkerRunTests( uint32 InExecutionCount, int32 InRoleIndex, FString InTestName, FString InBeautifiedTestName, bool InScreenshotsEnabled, bool InSendAnalytics)
+	FAutomationWorkerRunTests( uint32 InExecutionCount, int32 InRoleIndex, FString InTestName, FString InBeautifiedTestName, bool InSendAnalytics)
 		: ExecutionCount(InExecutionCount)
 		, RoleIndex(InRoleIndex)
 		, TestName(InTestName)
 		, BeautifiedTestName(InBeautifiedTestName)
-		, bScreenshotsEnabled(InScreenshotsEnabled)
 		, bSendAnalytics(InSendAnalytics)
 	{ }
 };
diff --git a/Engine/Source/Runtime/AutomationWorker/Private/AutomationWorkerModule.cpp b/Engine/Source/Runtime/AutomationWorker/Private/AutomationWorkerModule.cpp
--- a/Engine/Source/Runtime/AutomationWorker/Private/AutomationWorkerModule.cpp
+++ b/Engine/Source/Runtime/AutomationWorker/Private/AutomationWorkerModule.cpp
@@ -405,40 +405,37 @@ void FAutomationWorkerModule::HandleScreenShotCaptured(int32 Width, int32 Height
 
 void FAutomationWorkerModule::HandleScreenShotCapturedWithName(const TArray<FColor>& RawImageData, const FAutomationScreenshotData& Data)
 {
-	if( FAutomationTestFramework::Get().IsScreenshotAllowed() )
-	{
-		int32 NewHeight = Data.Height;
-		int32 NewWidth = Data.Width;
+	int32 NewHeight = Data.Height;
+	int32 NewWidth = Data.Width;
 
-		TArray<uint8> CompressedBitmap;
-		FImageUtils::CompressImageArray(NewWidth, NewHeight, RawImageData, CompressedBitmap);
+	TArray<uint8> CompressedBitmap;
+	FImageUtils::CompressImageArray(NewWidth, NewHeight, RawImageData, CompressedBitmap);
 
-		FAutomationScreenshotMetadata Metadata(Data);
+	FAutomationScreenshotMetadata Metadata(Data);
 		
-		// Send the screen shot if we have a target
-		if( TestRequesterAddress.IsValid() )
-		{
-			FAutomationWorkerScreenImage* Message = new FAutomationWorkerScreenImage();
+	// Send the screen shot if we have a target
+	if( TestRequesterAddress.IsValid() )
+	{
+		FAutomationWorkerScreenImage* Message = new FAutomationWorkerScreenImage();
 
-			Message->ScreenShotName = FPaths::RootDir() / Data.Path;
-			FPaths::MakePathRelativeTo(Message->ScreenShotName, *FPaths::AutomationDir());
-			Message->ScreenImage = CompressedBitmap;
-			Message->Metadata = Metadata;
-			MessageEndpoint->Send(Message, TestRequesterAddress);
-		}
-		else
-		{
-			//Save locally
-			const bool bTree = true;
-			IFileManager::Get().MakeDirectory(*FPaths::GetPath(Data.Path), bTree);
-			FFileHelper::SaveArrayToFile(CompressedBitmap, *Data.Path);
+		Message->ScreenShotName = FPaths::RootDir() / Data.Path;
+		FPaths::MakePathRelativeTo(Message->ScreenShotName, *FPaths::AutomationDir());
+		Message->ScreenImage = CompressedBitmap;
+		Message->Metadata = Metadata;
+		MessageEndpoint->Send(Message, TestRequesterAddress);
+	}
+	else
+	{
+		//Save locally
+		const bool bTree = true;
+		IFileManager::Get().MakeDirectory(*FPaths::GetPath(Data.Path), bTree);
+		FFileHelper::SaveArrayToFile(CompressedBitmap, *Data.Path);
 
-			FString Json;
-			if ( FJsonObjectConverter::UStructToJsonObjectString(Metadata, Json) )
-			{
-				FString MetadataPath = FPaths::ChangeExtension(Data.Path, TEXT("json"));
-				FFileHelper::SaveStringToFile(Json, *MetadataPath, FFileHelper::EEncodingOptions::ForceUTF8WithoutBOM);
-			}
+		FString Json;
+		if ( FJsonObjectConverter::UStructToJsonObjectString(Metadata, Json) )
+		{
+			FString MetadataPath = FPaths::ChangeExtension(Data.Path, TEXT("json"));
+			FFileHelper::SaveStringToFile(Json, *MetadataPath, FFileHelper::EEncodingOptions::ForceUTF8WithoutBOM);
 		}
 	}
 }
@@ -452,7 +449,6 @@ void FAutomationWorkerModule::HandleRunTestsMessage( const FAutomationWorkerRunT
 	BeautifiedTestName = Message.BeautifiedTestName;
 	bSendAnalytics = Message.bSendAnalytics;
 	TestRequesterAddress = Context->GetSender();
-	FAutomationTestFramework::Get().SetScreenshotOptions(Message.bScreenshotsEnabled);
 
 	// Always allow the first network command to execute
 	bExecuteNextNetworkCommand = true;
diff --git a/Engine/Source/Runtime/Core/Private/Misc/AutomationTest.cpp b/Engine/Source/Runtime/Core/Private/Misc/AutomationTest.cpp
--- a/Engine/Source/Runtime/Core/Private/Misc/AutomationTest.cpp
+++ b/Engine/Source/Runtime/Core/Private/Misc/AutomationTest.cpp
@@ -489,16 +489,6 @@ FOnTestScreenshotCaptured& FAutomationTestFramework::OnScreenshotCaptured()
 	return TestScreenshotCapturedDelegate;
 }
 
-void FAutomationTestFramework::SetScreenshotOptions( const bool bInScreenshotsEnabled )
-{
-	bScreenshotsEnabled = bInScreenshotsEnabled;
-}
-
-bool FAutomationTestFramework::IsScreenshotAllowed() const
-{
-	return bScreenshotsEnabled;
-}
-
 void FAutomationTestFramework::PrepForAutomationTests()
 {
 	check(!GIsAutomationTesting);
@@ -674,7 +664,6 @@ FAutomationTestFramework::FAutomationTestFramework()
 	, bTestSuccessful(false)
 	, CurrentTest(nullptr)
 	, bDeveloperDirectoryIncluded(false)
-	, bScreenshotsEnabled(true)
 	, NetworkRoleIndex(0)
 	, bForceSmokeTests(false)
 	, bCaptureStack(true)
diff --git a/Engine/Source/Runtime/Core/Public/Misc/AutomationTest.h b/Engine/Source/Runtime/Core/Public/Misc/AutomationTest.h
--- a/Engine/Source/Runtime/Core/Public/Misc/AutomationTest.h
+++ b/Engine/Source/Runtime/Core/Public/Misc/AutomationTest.h
@@ -722,17 +722,6 @@ public:
 	 */
 	FOnTestScreenshotCaptured& OnScreenshotCaptured();
 
-	/**
-	 * Sets screenshot options
-	 * @param bInScreenshotsEnabled - If screenshots are enabled
-	 */
-	void SetScreenshotOptions( const bool bInScreenshotsEnabled );
-
-	/**
-	 * Gets if screenshots are allowed
-	 */
-	bool IsScreenshotAllowed() const;
-
 	/**
 	 * Sets forcing smoke tests.
 	 */
@@ -887,9 +876,6 @@ private:
 	/** Whether we want to run automation tests on content within the Developer Directories */
 	bool bDeveloperDirectoryIncluded;
 
-	/** Wheather screenshots are enabled */
-	bool bScreenshotsEnabled;
-
 	/** Participation role as given by the automation controller */
 	uint32 NetworkRoleIndex;
 
diff --git a/Engine/Source/Runtime/Engine/Private/Tests/EngineAutomationTests.cpp b/Engine/Source/Runtime/Engine/Private/Tests/EngineAutomationTests.cpp
--- a/Engine/Source/Runtime/Engine/Private/Tests/EngineAutomationTests.cpp
+++ b/Engine/Source/Runtime/Engine/Private/Tests/EngineAutomationTests.cpp
@@ -220,7 +220,7 @@ bool FLoadAllMapsInGameTest::RunTest(const FString& Parameters)
 	//Open the map
 	GEngine->Exec(GetSimpleEngineAutomationTestGameWorld(GetTestFlags()), *FString::Printf(TEXT("Open %s"), *MapName));
 
-	if( FAutomationTestFramework::Get().IsScreenshotAllowed() )
+	//if( FAutomationTestFramework::Get().IsScreenshotAllowed() )
 	{
 		//Generate the screen shot name and path
 		FString ScreenshotFileName;
@@ -1090,103 +1090,103 @@ bool FWaitForMatineeToCompleteAndDoScreenshotsLatentCommand::Update()
 	return bTestComplete;
 }
 
-/**
- * FRenderOutputValidation - render deterministic and compare results of multiple runs
- */
-IMPLEMENT_SIMPLE_AUTOMATION_TEST(FRenderOutputValidation, "Rendering.RenderOutputValidation", EAutomationTestFlags::ClientContext | EAutomationTestFlags::EngineFilter)
-
-/** 
- * render deterministic and compare results of multiple runs
- *
- * @param Parameters - Unused for this test
- * @return	TRUE if the test was successful, FALSE otherwise
- */
-bool FRenderOutputValidation::RunTest(const FString& Parameters)
-{
-	const bool bChangeResolution = false;
-	const bool bLoadMap = false;			// true doesn't work at the moment (iterate matinees needs to happen after loading the map)
-
-	//Gets the default map that the game uses.
-	const UGameMapsSettings* GameMapsSettings = GetDefault<UGameMapsSettings>();
-//	const FString& MapName = GameMapsSettings->GetGameDefaultMap();
-	const FString MapName = TEXT("ScreenshotsTest");
-
-	// request full res screenshots
-	FAutomationTestFramework::Get().SetScreenshotOptions(true);
-
-	if(bLoadMap)
-	{
-		// Opens the actual default map in game.
-		GEngine->Exec(GetSimpleEngineAutomationTestGameWorld(GetTestFlags()), *FString::Printf(TEXT("Open %s"), *MapName));
-		ADD_LATENT_AUTOMATION_COMMAND(FEngineWaitLatentCommand(2.0f));
-	}
-
-	{
-		UConsole* ViewportConsole = (GEngine->GameViewport != nullptr) ? GEngine->GameViewport->ViewportConsole : nullptr;
-		FConsoleOutputDevice StrOut(ViewportConsole);
-
-		StrOut.Logf(TEXT("FRenderOutputValidation:Runtest()"));
-	}
-
-	// Gets the current resolution.
-	FString RestoreResolutionString = FString::Printf(TEXT("setres %dx%d"), GSystemResolution.ResX, GSystemResolution.ResY);
-
-	if(bChangeResolution)
-	{
-		// Change the resolution
-		ADD_LATENT_AUTOMATION_COMMAND(FExecStringLatentCommand(TEXT("setres 640x480")));
-		ADD_LATENT_AUTOMATION_COMMAND(FEngineWaitLatentCommand(2.0f));
-	}
-	// 5 sec should be enough, if more is needed - we should investigate
-	ADD_LATENT_AUTOMATION_COMMAND(FStreamAllResourcesLatentCommand(5.0f));
-
-//later	ADD_LATENT_AUTOMATION_COMMAND(FBuildLightingCommand);
-
-	// streaming should not be distributing it's work over multiple frames - to be more deterministic
-	ADD_LATENT_AUTOMATION_COMMAND(FExecWorldStringLatentCommand(TEXT("r.Streaming.FramesForFullUpdate 0")));
-
-	// reset all data in the ViewState e.g. TemporalAA, SSR cyclic counter, TemporalAA images
-	ADD_LATENT_AUTOMATION_COMMAND(FExecWorldStringLatentCommand(TEXT("r.ResetViewState")));
-
-	// this needs to be improved: better thank hanging would be to render frames (showing the user that it waits on something)
-	ADD_LATENT_AUTOMATION_COMMAND(FWaitForShadersToFinishCompilingInGame);
-
-	if( FAutomationTestFramework::Get().IsScreenshotAllowed() )
-	{
-		for (TObjectIterator<AMatineeActor> It; It; ++It)
-		{
-			AMatineeActor* MatineeActor = *It;
-
-			FString Name;
-			MatineeActor->GetName(Name);
-
-			if(Name.Contains(TEXT("Automation")))
-			{
-				UE_LOG(LogEngineAutomationTests, Log, TEXT("Iterated matinee: %p '%s'"), MatineeActor, *MatineeActor->GetName())
-
-				//add latent action to execute this matinee
-				ADD_LATENT_AUTOMATION_COMMAND(FPlayMatineeLatentCommand(MatineeActor));
-
-				//Run the Stat FPS Chart command
-		//			ADD_LATENT_AUTOMATION_COMMAND(FExecWorldStringLatentCommand(TEXT("StartFPSChart")));
-
-				//add action to wait until matinee is complete
-				ADD_LATENT_AUTOMATION_COMMAND(FWaitForMatineeToCompleteAndDoScreenshotsLatentCommand(MatineeActor));
-
-				//Stop the Stat FPS Chart command
-		//			ADD_LATENT_AUTOMATION_COMMAND(FExecWorldStringLatentCommand(TEXT("StopFPSChart")));
-			}
-		}
-	}
-
-	if(bChangeResolution)
-	{
-		// restore the resolution
-		ADD_LATENT_AUTOMATION_COMMAND(FExecStringLatentCommand(RestoreResolutionString));
-	}
-
-	return true;
-}
+///**
+// * FRenderOutputValidation - render deterministic and compare results of multiple runs
+// */
+//IMPLEMENT_SIMPLE_AUTOMATION_TEST(FRenderOutputValidation, "Rendering.RenderOutputValidation", EAutomationTestFlags::ClientContext | EAutomationTestFlags::EngineFilter)
+//
+///** 
+// * render deterministic and compare results of multiple runs
+// *
+// * @param Parameters - Unused for this test
+// * @return	TRUE if the test was successful, FALSE otherwise
+// */
+//bool FRenderOutputValidation::RunTest(const FString& Parameters)
+//{
+//	const bool bChangeResolution = false;
+//	const bool bLoadMap = false;			// true doesn't work at the moment (iterate matinees needs to happen after loading the map)
+//
+//	//Gets the default map that the game uses.
+//	const UGameMapsSettings* GameMapsSettings = GetDefault<UGameMapsSettings>();
+////	const FString& MapName = GameMapsSettings->GetGameDefaultMap();
+//	const FString MapName = TEXT("ScreenshotsTest");
+//
+//	// request full res screenshots
+//	//FAutomationTestFramework::Get().SetScreenshotOptions(true);
+//
+//	if(bLoadMap)
+//	{
+//		// Opens the actual default map in game.
+//		GEngine->Exec(GetSimpleEngineAutomationTestGameWorld(GetTestFlags()), *FString::Printf(TEXT("Open %s"), *MapName));
+//		ADD_LATENT_AUTOMATION_COMMAND(FEngineWaitLatentCommand(2.0f));
+//	}
+//
+//	{
+//		UConsole* ViewportConsole = (GEngine->GameViewport != nullptr) ? GEngine->GameViewport->ViewportConsole : nullptr;
+//		FConsoleOutputDevice StrOut(ViewportConsole);
+//
+//		StrOut.Logf(TEXT("FRenderOutputValidation:Runtest()"));
+//	}
+//
+//	// Gets the current resolution.
+//	FString RestoreResolutionString = FString::Printf(TEXT("setres %dx%d"), GSystemResolution.ResX, GSystemResolution.ResY);
+//
+//	if(bChangeResolution)
+//	{
+//		// Change the resolution
+//		ADD_LATENT_AUTOMATION_COMMAND(FExecStringLatentCommand(TEXT("setres 640x480")));
+//		ADD_LATENT_AUTOMATION_COMMAND(FEngineWaitLatentCommand(2.0f));
+//	}
+//	// 5 sec should be enough, if more is needed - we should investigate
+//	ADD_LATENT_AUTOMATION_COMMAND(FStreamAllResourcesLatentCommand(5.0f));
+//
+////later	ADD_LATENT_AUTOMATION_COMMAND(FBuildLightingCommand);
+//
+//	// streaming should not be distributing it's work over multiple frames - to be more deterministic
+//	ADD_LATENT_AUTOMATION_COMMAND(FExecWorldStringLatentCommand(TEXT("r.Streaming.FramesForFullUpdate 0")));
+//
+//	// reset all data in the ViewState e.g. TemporalAA, SSR cyclic counter, TemporalAA images
+//	ADD_LATENT_AUTOMATION_COMMAND(FExecWorldStringLatentCommand(TEXT("r.ResetViewState")));
+//
+//	// this needs to be improved: better thank hanging would be to render frames (showing the user that it waits on something)
+//	ADD_LATENT_AUTOMATION_COMMAND(FWaitForShadersToFinishCompilingInGame);
+//
+//	//if( FAutomationTestFramework::Get().IsScreenshotAllowed() )
+//	{
+//		for (TObjectIterator<AMatineeActor> It; It; ++It)
+//		{
+//			AMatineeActor* MatineeActor = *It;
+//
+//			FString Name;
+//			MatineeActor->GetName(Name);
+//
+//			if(Name.Contains(TEXT("Automation")))
+//			{
+//				UE_LOG(LogEngineAutomationTests, Log, TEXT("Iterated matinee: %p '%s'"), MatineeActor, *MatineeActor->GetName())
+//
+//				//add latent action to execute this matinee
+//				ADD_LATENT_AUTOMATION_COMMAND(FPlayMatineeLatentCommand(MatineeActor));
+//
+//				//Run the Stat FPS Chart command
+//		//			ADD_LATENT_AUTOMATION_COMMAND(FExecWorldStringLatentCommand(TEXT("StartFPSChart")));
+//
+//				//add action to wait until matinee is complete
+//				ADD_LATENT_AUTOMATION_COMMAND(FWaitForMatineeToCompleteAndDoScreenshotsLatentCommand(MatineeActor));
+//
+//				//Stop the Stat FPS Chart command
+//		//			ADD_LATENT_AUTOMATION_COMMAND(FExecWorldStringLatentCommand(TEXT("StopFPSChart")));
+//			}
+//		}
+//	}
+//
+//	if(bChangeResolution)
+//	{
+//		// restore the resolution
+//		ADD_LATENT_AUTOMATION_COMMAND(FExecStringLatentCommand(RestoreResolutionString));
+//	}
+//
+//	return true;
+//}
 
 #endif //WITH_DEV_AUTOMATION_TESTS
 
