Fix crash when post loading old staticmesh

diff --git a/Engine/Source/Runtime/Engine/Private/StaticMesh.cpp b/Engine/Source/Runtime/Engine/Private/StaticMesh.cpp
--- a/Engine/Source/Runtime/Engine/Private/StaticMesh.cpp
+++ b/Engine/Source/Runtime/Engine/Private/StaticMesh.cpp
@@ -2539,25 +2539,7 @@ void UStaticMesh::PostLoad()
 		}
 
 		CacheDerivedData();
-		
-		// Only required in an editor build as other builds process this in a different place
-		if (bRequiresLODDistanceConversion)
-		{
-			// Convert distances to Display Factors
-			ConvertLegacyLODDistance();
-		}
-
-		if (bRequiresLODScreenSizeConversion)
-		{
-			// Convert screen area to screen size
-			ConvertLegacyLODScreenArea();
-		}
 
-		if (RenderData && GStaticMeshesThatNeedMaterialFixup.Get(this))
-		{
-			FixupZeroTriangleSections();
-		}
-		
 		//Fix up the material to remove redundant material, this is needed since the material refactor where we do not have anymore copy of the materials
 		//in the materials list
 		if (RenderData && CleanUpRedondantMaterialPostLoad.Num() > 1)
@@ -2602,13 +2584,31 @@ void UStaticMesh::PostLoad()
 			ToRemoveMaterials.Sort();
 			for (int32 RemoveIndex = ToRemoveMaterials.Num() - 1; RemoveIndex >= 0; --RemoveIndex)
 			{
-				if (StaticMaterials.IsValidIndex(RemoveIndex))
+				if (StaticMaterials.IsValidIndex(ToRemoveMaterials[RemoveIndex]))
 				{
 					StaticMaterials.RemoveAt(ToRemoveMaterials[RemoveIndex]);
 				}
 			}
 			CleanUpRedondantMaterialPostLoad.Reset();
 		}
+
+		// Only required in an editor build as other builds process this in a different place
+		if (bRequiresLODDistanceConversion)
+		{
+			// Convert distances to Display Factors
+			ConvertLegacyLODDistance();
+		}
+
+		if (bRequiresLODScreenSizeConversion)
+		{
+			// Convert screen area to screen size
+			ConvertLegacyLODScreenArea();
+		}
+
+		if (RenderData && GStaticMeshesThatNeedMaterialFixup.Get(this))
+		{
+			FixupZeroTriangleSections();
+		}
 	}
 #endif // #if WITH_EDITOR
 
