Fix Shipping/Test build errors with SendRenderDebugPhysics

diff --git a/Engine/Source/Runtime/Engine/Private/Components/PrimitiveComponent.cpp b/Engine/Source/Runtime/Engine/Private/Components/PrimitiveComponent.cpp
--- a/Engine/Source/Runtime/Engine/Private/Components/PrimitiveComponent.cpp
+++ b/Engine/Source/Runtime/Engine/Private/Components/PrimitiveComponent.cpp
@@ -609,7 +609,9 @@ void UPrimitiveComponent::OnCreatePhysicsState()
 #if UE_WITH_PHYSICS
 			// Create the body.
 			BodyInstance.InitBody(BodySetup, BodyTransform, this, GetWorld()->GetPhysicsScene());		
+#if !(UE_BUILD_SHIPPING || UE_BUILD_TEST)
 			SendRenderDebugPhysics();
+#endif // !(UE_BUILD_SHIPPING || UE_BUILD_TEST)
 #endif //UE_WITH_PHYSICS
 
 #if WITH_EDITOR
@@ -684,14 +686,14 @@ void UPrimitiveComponent::OnDestroyPhysicsState()
 		BodyInstance.TermBody();
 	}
 
-#if !UE_BUILD_SHIPPING
+#if !(UE_BUILD_SHIPPING || UE_BUILD_TEST)
 	SendRenderDebugPhysics();
 #endif
 
 	Super::OnDestroyPhysicsState();
 }
 
-#if !UE_BUILD_SHIPPING
+#if !(UE_BUILD_SHIPPING || UE_BUILD_TEST)
 void UPrimitiveComponent::SendRenderDebugPhysics(FPrimitiveSceneProxy* OverrideSceneProxy)
 {
 	FPrimitiveSceneProxy* UseSceneProxy = OverrideSceneProxy ? OverrideSceneProxy : SceneProxy;
diff --git a/Engine/Source/Runtime/Engine/Private/SkeletalMeshComponentPhysics.cpp b/Engine/Source/Runtime/Engine/Private/SkeletalMeshComponentPhysics.cpp
--- a/Engine/Source/Runtime/Engine/Private/SkeletalMeshComponentPhysics.cpp
+++ b/Engine/Source/Runtime/Engine/Private/SkeletalMeshComponentPhysics.cpp
@@ -1261,7 +1261,9 @@ void USkeletalMeshComponent::OnCreatePhysicsState()
 	{
 		InitArticulated(GetWorld()->GetPhysicsScene());
 		USceneComponent::OnCreatePhysicsState(); // Need to route CreatePhysicsState, skip PrimitiveComponent
+#if !(UE_BUILD_SHIPPING || UE_BUILD_TEST)
 		SendRenderDebugPhysics();
+#endif
 	}
 	else
 	{
@@ -1716,8 +1718,9 @@ void USkeletalMeshComponent::SetPhysicsAsset(UPhysicsAsset* InPhysicsAsset, bool
 		// Indicate that 'required bones' array will need to be recalculated.
 		bRequiredBonesUpToDate = false;
 
+#if !(UE_BUILD_SHIPPING || UE_BUILD_TEST)
 		SendRenderDebugPhysics();
-
+#endif
 	}
 }
 
