One more fix for Test builds using SendRenderDebugPhysics

diff --git a/Engine/Source/Runtime/Engine/Private/PhysicsEngine/BodyInstance.cpp b/Engine/Source/Runtime/Engine/Private/PhysicsEngine/BodyInstance.cpp
--- a/Engine/Source/Runtime/Engine/Private/PhysicsEngine/BodyInstance.cpp
+++ b/Engine/Source/Runtime/Engine/Private/PhysicsEngine/BodyInstance.cpp
@@ -3826,7 +3826,7 @@ void FBodyInstance::UpdateMassProperties()
 
 void FBodyInstance::UpdateDebugRendering()
 {
-#if !(UE_BUILD_SHIPPING)
+#if !(UE_BUILD_SHIPPING || UE_BUILD_TEST)
 	//After we update the mass properties, we should update any debug rendering
 	if (UPrimitiveComponent* OwnerPrim = OwnerComponent.Get())
 	{
