fix 4.13 staticmesh conversion by making sure the LOD 0 material duplicate are not cleanup. We must clean only the sub LODs

diff --git a/Engine/Source/Runtime/Engine/Classes/Engine/StaticMesh.h b/Engine/Source/Runtime/Engine/Classes/Engine/StaticMesh.h
--- a/Engine/Source/Runtime/Engine/Classes/Engine/StaticMesh.h
+++ b/Engine/Source/Runtime/Engine/Classes/Engine/StaticMesh.h
@@ -870,7 +870,7 @@ private:
 	/**
 	 * Fixes up the material when it was converted to the new staticmesh build process
 	 */
-	TArray<int32> CleanUpRedondantMaterialPostLoad;
+	bool CleanUpRedondantMaterialPostLoad;
 
 #endif // #if WITH_EDITOR
 };
diff --git a/Engine/Source/Runtime/Engine/Private/StaticMesh.cpp b/Engine/Source/Runtime/Engine/Private/StaticMesh.cpp
--- a/Engine/Source/Runtime/Engine/Private/StaticMesh.cpp
+++ b/Engine/Source/Runtime/Engine/Private/StaticMesh.cpp
@@ -2429,13 +2429,23 @@ void UStaticMesh::Serialize(FArchive& Ar)
 	else if (Ar.IsLoading())
 	{
 		TArray<UMaterialInterface*> Unique_Materials_DEPRECATED;
+		TArray<FName> MaterialSlotNames;
 		for (UMaterialInterface *MaterialInterface : Materials_DEPRECATED)
 		{
-			StaticMaterials.Add(FStaticMaterial(MaterialInterface, MaterialInterface != nullptr ? MaterialInterface->GetFName() : NAME_None));
+			FName MaterialSlotName = MaterialInterface != nullptr ? MaterialInterface->GetFName() : NAME_None;
+			int32 NameCounter = 1;
+			while (MaterialSlotNames.Find(MaterialSlotName) != INDEX_NONE && MaterialSlotName != NAME_None)
+			{
+				FString MaterialSlotNameStr = MaterialInterface->GetName() + TEXT("_") + FString::FromInt(NameCounter);
+				MaterialSlotName = FName(*MaterialSlotNameStr);
+				NameCounter++;
+			}
+			MaterialSlotNames.Add(MaterialSlotName);
+			StaticMaterials.Add(FStaticMaterial(MaterialInterface, MaterialSlotName));
 			int32 UniqueIndex = Unique_Materials_DEPRECATED.AddUnique(MaterialInterface);
 #if WITH_EDITOR
 			//We must cleanup the material list since we have a new way to build static mesh
-			CleanUpRedondantMaterialPostLoad.Add(UniqueIndex);
+			CleanUpRedondantMaterialPostLoad = StaticMaterials.Num() > 1;
 #endif
 		}
 		Materials_DEPRECATED.Empty();
@@ -2543,8 +2553,10 @@ void UStaticMesh::PostLoad()
 
 		//Fix up the material to remove redundant material, this is needed since the material refactor where we do not have anymore copy of the materials
 		//in the materials list
-		if (RenderData && CleanUpRedondantMaterialPostLoad.Num() > 1)
+		if (RenderData && CleanUpRedondantMaterialPostLoad)
 		{
+			bool bMaterialChange = false;
+			TArray<FStaticMaterial> CompactedMaterial;
 			for (int32 LODIndex = 0; LODIndex < RenderData->LODResources.Num(); ++LODIndex)
 			{
 				if (RenderData->LODResources.IsValidIndex(LODIndex))
@@ -2556,48 +2568,58 @@ void UStaticMesh::PostLoad()
 						const int32 MaterialIndex = LOD.Sections[SectionIndex].MaterialIndex;
 						if (StaticMaterials.IsValidIndex(MaterialIndex))
 						{
-							FMeshSectionInfo MeshSectionInfo = SectionInfoMap.Get(LODIndex, SectionIndex);
-							if (CleanUpRedondantMaterialPostLoad.IsValidIndex(MeshSectionInfo.MaterialIndex))
+							if (LODIndex == 0)
+							{
+								//We do not compact LOD 0 material
+								CompactedMaterial.Add(StaticMaterials[MaterialIndex]);
+							}
+							else
 							{
-								MeshSectionInfo.MaterialIndex = CleanUpRedondantMaterialPostLoad[MeshSectionInfo.MaterialIndex];
-								SectionInfoMap.Set(LODIndex, SectionIndex, MeshSectionInfo);
+								FMeshSectionInfo MeshSectionInfo = SectionInfoMap.Get(LODIndex, SectionIndex);
+								int32 CompactedIndex = INDEX_NONE;
+								if (StaticMaterials.IsValidIndex(MeshSectionInfo.MaterialIndex))
+								{
+									for (int32 CompactedMaterialIndex = 0; CompactedMaterialIndex < CompactedMaterial.Num(); ++CompactedMaterialIndex)
+									{
+										const FStaticMaterial& StaticMaterial = CompactedMaterial[CompactedMaterialIndex];
+										if (StaticMaterials[MeshSectionInfo.MaterialIndex].MaterialInterface == StaticMaterial.MaterialInterface)
+										{
+											CompactedIndex = CompactedMaterialIndex;
+											break;
+										}
+									}
+								}
+
+								if (CompactedIndex == INDEX_NONE)
+								{
+									CompactedIndex = CompactedMaterial.Add(StaticMaterials[MaterialIndex]);
+								}
+								if (MeshSectionInfo.MaterialIndex != CompactedIndex)
+								{
+									MeshSectionInfo.MaterialIndex = CompactedIndex;
+									SectionInfoMap.Set(LODIndex, SectionIndex, MeshSectionInfo);
+									bMaterialChange = true;
+								}
 							}
 						}
 					}
 				}
 			}
-
-			//Remove the redundant materials from the array
-			TArray<int32> UniqueMaterials;
-			TArray<int32> ToRemoveMaterials;
-			for (int32 StaticMaterialIndex = 0; StaticMaterialIndex < CleanUpRedondantMaterialPostLoad.Num(); ++StaticMaterialIndex)
+			//If we change some section material index or there is unused material, we must use the new compacted material list.
+			if (bMaterialChange || CompactedMaterial.Num() < StaticMaterials.Num())
 			{
-				if (UniqueMaterials.Contains(CleanUpRedondantMaterialPostLoad[StaticMaterialIndex]))
-				{
-					ToRemoveMaterials.Add(StaticMaterialIndex);
-				}
-				else
+				StaticMaterials.Empty(CompactedMaterial.Num());
+				for (const FStaticMaterial &Material : CompactedMaterial)
 				{
-					UniqueMaterials.AddUnique(CleanUpRedondantMaterialPostLoad[StaticMaterialIndex]);
+					StaticMaterials.Add(Material);
 				}
-			}
-			//Sort the remove material list so we can remove from the end to ensure index are valid during remove operation
-			ToRemoveMaterials.Sort();
-			bool bRemovedMaterial = false;
-			for (int32 RemoveIndex = ToRemoveMaterials.Num() - 1; RemoveIndex >= 0; --RemoveIndex)
-			{
-				if (StaticMaterials.IsValidIndex(ToRemoveMaterials[RemoveIndex]))
+				//Make sure the physic data is recompute
+				if (BodySetup)
 				{
-					StaticMaterials.RemoveAt(ToRemoveMaterials[RemoveIndex]);
-					bRemovedMaterial = true;
+					BodySetup->InvalidatePhysicsData();
 				}
 			}
-			CleanUpRedondantMaterialPostLoad.Reset();
-
-			if(bRemovedMaterial && BodySetup)
-			{
-				BodySetup->InvalidatePhysicsData();
-			}
+			CleanUpRedondantMaterialPostLoad = false;
 		}
 
 		// Only required in an editor build as other builds process this in a different place
