Fix another non-unity compile error in UE4Game.

diff --git a/Engine/Source/Runtime/Engine/Private/Components/StaticMeshComponent.cpp b/Engine/Source/Runtime/Engine/Private/Components/StaticMeshComponent.cpp
--- a/Engine/Source/Runtime/Engine/Private/Components/StaticMeshComponent.cpp
+++ b/Engine/Source/Runtime/Engine/Private/Components/StaticMeshComponent.cpp
@@ -36,6 +36,7 @@
 #include "PhysicsEngine/BodySetup.h"
 #include "EngineGlobals.h"
 #include "ComponentRecreateRenderStateContext.h"
+#include "Engine/StaticMesh.h"
 
 #define LOCTEXT_NAMESPACE "StaticMeshComponent"
 
