Fixed MakeMA expression warnings.

diff --git a/Engine/Shaders/ShaderVersion.usf b/Engine/Shaders/ShaderVersion.usf
--- a/Engine/Shaders/ShaderVersion.usf
+++ b/Engine/Shaders/ShaderVersion.usf
@@ -2,4 +2,4 @@
 // Each time the console command is executed it generates a new GUID. As this file is included
 // in common.usf (which should be included in any shader) it allows to invalidate the shader DDC.
 // 
-// GUID = B0FBC0FCE55F4DBA83537F6D7C18A31B
\ No newline at end of file
+// GUID = A5D32FD08DAE4754B7F8105EA046ED4A
\ No newline at end of file
diff --git a/Engine/Source/Runtime/Engine/Private/Materials/MaterialExpressions.cpp b/Engine/Source/Runtime/Engine/Private/Materials/MaterialExpressions.cpp
--- a/Engine/Source/Runtime/Engine/Private/Materials/MaterialExpressions.cpp
+++ b/Engine/Source/Runtime/Engine/Private/Materials/MaterialExpressions.cpp
@@ -4106,6 +4106,12 @@ void UMaterialExpressionMakeMaterialAttributes::Serialize(FArchive& Ar)
 	Super::Serialize(Ar);
 
 	Ar.UsingCustomVersion(FRenderingObjectVersion::GUID);
+	
+	if (Ar.CustomVer(FRenderingObjectVersion::GUID) < FRenderingObjectVersion::FixedLegacyMaterialAttributeNodeTypes)
+	{
+		// Update the legacy masks else fail on vec3 to vec2 conversion
+		Refraction.SetMask(1, 1, 1, 0, 0);
+	}
 }
 
 #if WITH_EDITOR
