Merging CL 3408287 from Release-4.16

diff --git a/Engine/Source/Runtime/Engine/Private/ComponentInstanceDataCache.cpp b/Engine/Source/Runtime/Engine/Private/ComponentInstanceDataCache.cpp
--- a/Engine/Source/Runtime/Engine/Private/ComponentInstanceDataCache.cpp
+++ b/Engine/Source/Runtime/Engine/Private/ComponentInstanceDataCache.cpp
@@ -236,6 +236,7 @@ FComponentInstanceDataCache::FComponentInstanceDataCache(const AActor* Actor)
 				{
 					if (SceneComponent->GetAttachParent() && SceneComponent->GetAttachParent()->IsCreatedByConstructionScript())
 					{
+						SceneComponent->ConditionalUpdateComponentToWorld();
 						InstanceComponentTransformToRootMap.Add(SceneComponent, SceneComponent->GetComponentTransform().GetRelativeTransform(Actor->GetRootComponent()->GetComponentTransform()));
 					}
 				}
