Merging CL 3410397 from Release-4.16

diff --git a/Engine/Source/Runtime/Engine/Private/ComponentInstanceDataCache.cpp b/Engine/Source/Runtime/Engine/Private/ComponentInstanceDataCache.cpp
--- a/Engine/Source/Runtime/Engine/Private/ComponentInstanceDataCache.cpp
+++ b/Engine/Source/Runtime/Engine/Private/ComponentInstanceDataCache.cpp
@@ -218,11 +218,10 @@ FComponentInstanceDataCache::FComponentInstanceDataCache(const AActor* Actor)
 		{
 			if (bIsChildActor || Component->IsCreatedByConstructionScript()) // Only cache data from 'created by construction script' components
 			{
-				if (!Component->IsRegistered())
+				if (USceneComponent* SceneComponent = Cast<USceneComponent>(Component))
 				{
-					Component->RegisterComponent();
+					SceneComponent->ConditionalUpdateComponentToWorld();
 				}
-
 				FActorComponentInstanceData* ComponentInstanceData = Component->GetComponentInstanceData();
 				if (ComponentInstanceData)
 				{
