Merging CL 3409529 and CL 3409529 from Release-4.16.

diff --git a/Engine/Source/Runtime/Engine/Private/HierarchicalInstancedStaticMesh.cpp b/Engine/Source/Runtime/Engine/Private/HierarchicalInstancedStaticMesh.cpp
--- a/Engine/Source/Runtime/Engine/Private/HierarchicalInstancedStaticMesh.cpp
+++ b/Engine/Source/Runtime/Engine/Private/HierarchicalInstancedStaticMesh.cpp
@@ -2376,7 +2376,10 @@ void UHierarchicalInstancedStaticMeshComponent::BuildTreeIfOutdated(bool Async,
 		|| UnbuiltInstanceBoundsList.Num() > 0
 		|| GetLinkerUE4Version() < VER_UE4_REBUILD_HIERARCHICAL_INSTANCE_TREES)
 	{
-		UE_LOG(LogStaticMesh, Warning, TEXT("Rebuilding hierarchical instanced mesh component, please resave map %s."), *GetFullName());
+		if (GIsEditor && !FApp::IsUnattended() && !IsRunningCommandlet()) // only when user is in control (i.e normal editor)
+		{
+			UE_LOG(LogStaticMesh, Warning, TEXT("Rebuilding hierarchical instanced mesh component, please resave map %s."), *GetFullName());
+		}
 
 		if (GetStaticMesh())
 		{
