Copying //UE4/Dev-Framework to //UE4/Dev-Main (Source: //UE4/Dev-Framework @ 3431384)

diff --git a/Engine/Source/Runtime/Engine/Classes/Engine/EngineTypes.h b/Engine/Source/Runtime/Engine/Classes/Engine/EngineTypes.h
--- a/Engine/Source/Runtime/Engine/Classes/Engine/EngineTypes.h
+++ b/Engine/Source/Runtime/Engine/Classes/Engine/EngineTypes.h
@@ -357,7 +357,7 @@ enum ESceneCaptureCompositeMode
 
 #define NUM_LIGHTING_CHANNELS 3
 
-USTRUCT()
+USTRUCT(BlueprintType)
 struct FLightingChannels
 {
 	GENERATED_USTRUCT_BODY()
@@ -869,7 +869,7 @@ struct FResponseChannel
 /**
  *	Container for indicating a set of collision channels that this object will collide with.
  */
-USTRUCT()
+USTRUCT(BlueprintType)
 struct ENGINE_API FCollisionResponseContainer
 {
 	GENERATED_USTRUCT_BODY()
@@ -1315,7 +1315,7 @@ struct FCollisionImpactData
 
 
 /** Struct used to hold effects for destructible damage events */
-USTRUCT()
+USTRUCT(BlueprintType)
 struct FFractureEffect
 {
 	GENERATED_USTRUCT_BODY()
@@ -1336,7 +1336,7 @@ struct FFractureEffect
 
 
 /**	Struct for handling positions relative to a base actor, which is potentially moving */
-USTRUCT()
+USTRUCT(BlueprintType)
 struct ENGINE_API FBasedPosition
 {
 	GENERATED_USTRUCT_BODY()
@@ -1465,11 +1465,11 @@ struct FSubtitleCue
 	GENERATED_USTRUCT_BODY()
 
 	/** The text to appear in the subtitle. */
-	UPROPERTY(EditAnywhere, BlueprintReadWrite, Category=SubtitleCue)
+	UPROPERTY(EditAnywhere, Category=SubtitleCue)
 	FText Text;
 
 	/** The time at which the subtitle is to be displayed, in seconds relative to the beginning of the line. */
-	UPROPERTY(EditAnywhere, BlueprintReadWrite, Category=SubtitleCue)
+	UPROPERTY(EditAnywhere, Category=SubtitleCue)
 	float Time;
 
 	FSubtitleCue()
@@ -1522,18 +1522,18 @@ struct FLightmassLightSettings
 	GENERATED_USTRUCT_BODY()
 
 	/** 0 will be completely desaturated, 1 will be unchanged */
-	UPROPERTY(EditAnywhere, BlueprintReadWrite, Category=Lightmass, meta=(UIMin = "0.0", UIMax = "4.0"))
+	UPROPERTY(EditAnywhere, Category=Lightmass, meta=(UIMin = "0.0", UIMax = "4.0"))
 	float IndirectLightingSaturation;
 
 	/** Controls the falloff of shadow penumbras */
-	UPROPERTY(EditAnywhere, BlueprintReadWrite, Category=Lightmass, meta=(UIMin = "0.1", UIMax = "4.0"))
+	UPROPERTY(EditAnywhere, Category=Lightmass, meta=(UIMin = "0.1", UIMax = "4.0"))
 	float ShadowExponent;
 
 	/** 
 	 * Whether to use area shadows for stationary light precomputed shadowmaps.  
 	 * Area shadows get softer the further they are from shadow casters, but require higher lightmap resolution to get the same quality where the shadow is sharp.
 	 */
-	UPROPERTY(EditAnywhere, BlueprintReadWrite, Category=Lightmass)
+	UPROPERTY(EditAnywhere, Category=Lightmass)
 	bool bUseAreaShadowsForStationaryLight;
 
 	FLightmassLightSettings()
@@ -1559,7 +1559,7 @@ struct FLightmassDirectionalLightSettings : public FLightmassLightSettings
 	GENERATED_USTRUCT_BODY()
 
 	/** Angle that the directional light's emissive surface extends relative to a receiver, affects penumbra sizes. */
-	UPROPERTY(EditAnywhere, BlueprintReadWrite, Category=Lightmass, meta=(UIMin = ".0001", UIMax = "5"))
+	UPROPERTY(EditAnywhere, Category=Lightmass, meta=(UIMin = ".0001", UIMax = "5"))
 	float LightSourceAngle;
 
 	FLightmassDirectionalLightSettings()
@@ -1576,22 +1576,22 @@ struct FLightmassPrimitiveSettings
 	GENERATED_USTRUCT_BODY()
 
 	/** If true, this object will be lit as if it receives light from both sides of its polygons. */
-	UPROPERTY(EditAnywhere, BlueprintReadOnly, Category=Lightmass)
+	UPROPERTY(EditAnywhere, Category=Lightmass)
 	uint32 bUseTwoSidedLighting:1;
 
 	/** If true, this object will only shadow indirect lighting.  					*/
-	UPROPERTY(EditAnywhere, BlueprintReadOnly, Category=Lightmass)
+	UPROPERTY(EditAnywhere, Category=Lightmass)
 	uint32 bShadowIndirectOnly:1;
 
 	/** If true, allow using the emissive for static lighting.						*/
-	UPROPERTY(EditAnywhere, BlueprintReadOnly, Category=Lightmass)
+	UPROPERTY(EditAnywhere, Category=Lightmass)
 	uint32 bUseEmissiveForStaticLighting:1;
 
 	/** 
 	 * Typically the triangle normal is used for hemisphere gathering which prevents incorrect self-shadowing from artist-tweaked vertex normals. 
 	 * However in the case of foliage whose vertex normal has been setup to match the underlying terrain, gathering in the direction of the vertex normal is desired.
 	 */
-	UPROPERTY(EditAnywhere, BlueprintReadOnly, Category=Lightmass)
+	UPROPERTY(EditAnywhere, Category=Lightmass)
 	uint32 bUseVertexNormalForHemisphereGather:1;
 
 	/** Direct lighting falloff exponent for mesh area lights created from emissive areas on this primitive. */
@@ -1611,11 +1611,11 @@ struct FLightmassPrimitiveSettings
 	float EmissiveBoost;
 
 	/** Scales the diffuse contribution of all materials applied to this object.	*/
-	UPROPERTY(EditAnywhere, BlueprintReadWrite, Category=Lightmass)
+	UPROPERTY(EditAnywhere, Category=Lightmass)
 	float DiffuseBoost;
 
 	/** Fraction of samples taken that must be occluded in order to reach full occlusion. */
-	UPROPERTY(EditAnywhere, BlueprintReadWrite, Category=Lightmass)
+	UPROPERTY(EditAnywhere, Category=Lightmass)
 	float FullyOccludedSamplesFraction;
 
 	FLightmassPrimitiveSettings()
@@ -1664,85 +1664,85 @@ struct FLightmassDebugOptions
 	 *	If false, UnrealLightmass.exe is launched automatically (default)
 	 *	If true, it must be launched manually (e.g. through a debugger) with the -debug command line parameter.
 	 */
-	UPROPERTY(EditAnywhere, BlueprintReadWrite, Category=LightmassDebugOptions)
+	UPROPERTY(EditAnywhere, Category=LightmassDebugOptions)
 	uint32 bDebugMode:1;
 
 	/**	If true, all participating Lightmass agents will report back detailed stats to the log.	 */
-	UPROPERTY(EditAnywhere, BlueprintReadWrite, Category=LightmassDebugOptions)
+	UPROPERTY(EditAnywhere, Category=LightmassDebugOptions)
 	uint32 bStatsEnabled:1;
 
 	/**	If true, BSP surfaces split across model components are joined into 1 mapping	 */
-	UPROPERTY(EditAnywhere, BlueprintReadWrite, Category=LightmassDebugOptions)
+	UPROPERTY(EditAnywhere, Category=LightmassDebugOptions)
 	uint32 bGatherBSPSurfacesAcrossComponents:1;
 
 	/**	The tolerance level used when gathering BSP surfaces.	 */
-	UPROPERTY(EditAnywhere, BlueprintReadWrite, Category=LightmassDebugOptions)
+	UPROPERTY(EditAnywhere, Category=LightmassDebugOptions)
 	float CoplanarTolerance;
 
 	/**
 	 *	If true, Lightmass will import mappings immediately as they complete.
 	 *	It will not process them, however.
 	 */
-	UPROPERTY(EditAnywhere, BlueprintReadWrite, Category=LightmassDebugOptions)
+	UPROPERTY(EditAnywhere, Category=LightmassDebugOptions)
 	uint32 bUseImmediateImport:1;
 
 	/**
 	 *	If true, Lightmass will process appropriate mappings as they are imported.
 	 *	NOTE: Requires ImmediateMode be enabled to actually work.
 	 */
-	UPROPERTY(EditAnywhere, BlueprintReadWrite, Category=LightmassDebugOptions)
+	UPROPERTY(EditAnywhere, Category=LightmassDebugOptions)
 	uint32 bImmediateProcessMappings:1;
 
 	/**	If true, Lightmass will sort mappings by texel cost. */
-	UPROPERTY(EditAnywhere, BlueprintReadWrite, Category=LightmassDebugOptions)
+	UPROPERTY(EditAnywhere, Category=LightmassDebugOptions)
 	uint32 bSortMappings:1;
 
 	/**	If true, the generate coefficients will be dumped to binary files. */
-	UPROPERTY(EditAnywhere, BlueprintReadWrite, Category=LightmassDebugOptions)
+	UPROPERTY(EditAnywhere, Category=LightmassDebugOptions)
 	uint32 bDumpBinaryFiles:1;
 
 	/**
 	 *	If true, Lightmass will write out BMPs for each generated material property
 	 *	sample to <GAME>\ScreenShots\Materials.
 	 */
-	UPROPERTY(EditAnywhere, BlueprintReadWrite, Category=LightmassDebugOptions)
+	UPROPERTY(EditAnywhere, Category=LightmassDebugOptions)
 	uint32 bDebugMaterials:1;
 
 	/**	If true, Lightmass will pad the calculated mappings to reduce/eliminate seams. */
-	UPROPERTY(EditAnywhere, BlueprintReadWrite, Category=LightmassDebugOptions)
+	UPROPERTY(EditAnywhere, Category=LightmassDebugOptions)
 	uint32 bPadMappings:1;
 
 	/**
 	 *	If true, will fill padding of mappings with a color rather than the sampled edges.
 	 *	Means nothing if bPadMappings is not enabled...
 	 */
-	UPROPERTY(EditAnywhere, BlueprintReadWrite, Category=LightmassDebugOptions)
+	UPROPERTY(EditAnywhere, Category=LightmassDebugOptions)
 	uint32 bDebugPaddings:1;
 
 	/**
 	 * If true, only the mapping containing a debug texel will be calculated, all others
 	 * will be set to white
 	 */
-	UPROPERTY(EditAnywhere, BlueprintReadWrite, Category=LightmassDebugOptions)
+	UPROPERTY(EditAnywhere, Category=LightmassDebugOptions)
 	uint32 bOnlyCalcDebugTexelMappings:1;
 
 	/** If true, color lightmaps a random color */
-	UPROPERTY(EditAnywhere, BlueprintReadWrite, Category=LightmassDebugOptions)
+	UPROPERTY(EditAnywhere, Category=LightmassDebugOptions)
 	uint32 bUseRandomColors:1;
 
 	/** If true, a green border will be placed around the edges of mappings */
-	UPROPERTY(EditAnywhere, BlueprintReadWrite, Category=LightmassDebugOptions)
+	UPROPERTY(EditAnywhere, Category=LightmassDebugOptions)
 	uint32 bColorBordersGreen:1;
 
 	/**
 	 * If true, Lightmass will overwrite lightmap data with a shade of red relating to
 	 * how long it took to calculate the mapping (Red = Time / ExecutionTimeDivisor)
 	 */
-	UPROPERTY(EditAnywhere, BlueprintReadWrite, Category=LightmassDebugOptions)
+	UPROPERTY(EditAnywhere, Category=LightmassDebugOptions)
 	uint32 bColorByExecutionTime:1;
 
 	/** The amount of time that will be count as full red when bColorByExecutionTime is enabled */
-	UPROPERTY(EditAnywhere, BlueprintReadWrite, Category=LightmassDebugOptions)
+	UPROPERTY(EditAnywhere, Category=LightmassDebugOptions)
 	float ExecutionTimeDivisor;
 
 	ENGINE_API FLightmassDebugOptions();
@@ -1761,14 +1761,14 @@ struct FSwarmDebugOptions
 	 *	If true, Swarm will distribute jobs.
 	 *	If false, only the local machine will execute the jobs.
 	 */
-	UPROPERTY(EditAnywhere, BlueprintReadWrite, Category=SwarmDebugOptions)
+	UPROPERTY(EditAnywhere, Category=SwarmDebugOptions)
 	uint32 bDistributionEnabled:1;
 
 	/**
 	 *	If true, Swarm will force content to re-export rather than using the cached version.
 	 *	If false, Swarm will attempt to use the cached version.
 	 */
-	UPROPERTY(EditAnywhere, BlueprintReadWrite, Category=SwarmDebugOptions)
+	UPROPERTY(EditAnywhere, Category=SwarmDebugOptions)
 	uint32 bForceContentExport:1;
 
 	UPROPERTY()
@@ -2399,7 +2399,7 @@ public:
 /**
  * Point Of View type.
  */
-USTRUCT()
+USTRUCT(BlueprintType)
 struct FPOV
 {
 	GENERATED_USTRUCT_BODY()
@@ -2564,7 +2564,7 @@ struct FMeshBuildSettings
 
 
 
-USTRUCT()
+USTRUCT(BlueprintType)
 struct ENGINE_API FDamageEvent
 {
 	GENERATED_USTRUCT_BODY()
@@ -2631,27 +2631,27 @@ struct ENGINE_API FPointDamageEvent : public FDamageEvent
 };
 
 
-USTRUCT()
+USTRUCT(BlueprintType)
 struct ENGINE_API FRadialDamageParams
 {
 	GENERATED_USTRUCT_BODY()
 
-	UPROPERTY(EditAnywhere, Category=RadialDamageParams)
+	UPROPERTY(EditAnywhere, BlueprintReadWrite, Category=RadialDamageParams)
 	float BaseDamage;
 
-	UPROPERTY(EditAnywhere, Category=RadialDamageParams)
+	UPROPERTY(EditAnywhere, BlueprintReadWrite, Category=RadialDamageParams)
 	float MinimumDamage;
 	
-	UPROPERTY(EditAnywhere, Category=RadialDamageParams)
+	UPROPERTY(EditAnywhere, BlueprintReadWrite, Category=RadialDamageParams)
 	float InnerRadius;
 		
-	UPROPERTY(EditAnywhere, Category=RadialDamageParams)
+	UPROPERTY(EditAnywhere, BlueprintReadWrite, Category=RadialDamageParams)
 	float OuterRadius;
 		
-	UPROPERTY(EditAnywhere, Category=RadialDamageParams)
+	UPROPERTY(EditAnywhere, BlueprintReadWrite, Category=RadialDamageParams)
 	float DamageFalloff;
 
-// 	UPROPERTY(EditAnywhere, Category=RadiusDamageParams)
+// 	UPROPERTY(EditAnywhere, BlueprintReadWrite, Category=RadiusDamageParams)
 // 	float BaseImpulseMag;
 
 	FRadialDamageParams()
@@ -3748,7 +3748,7 @@ struct FFontRenderInfo
 
 
 /** Simple 2d triangle with UVs */
-USTRUCT()
+USTRUCT(BlueprintType)
 struct FCanvasUVTri
 {
 	GENERATED_USTRUCT_BODY()
diff --git a/Engine/Source/Runtime/Engine/Classes/Engine/Scene.h b/Engine/Source/Runtime/Engine/Classes/Engine/Scene.h
--- a/Engine/Source/Runtime/Engine/Classes/Engine/Scene.h
+++ b/Engine/Source/Runtime/Engine/Classes/Engine/Scene.h
@@ -56,17 +56,17 @@ enum EBloomMethod
 	BM_MAX,
 };
 
-USTRUCT()
+USTRUCT(BlueprintType)
 struct FWeightedBlendable
 {
 	GENERATED_USTRUCT_BODY()
 
 	/** 0:no effect .. 1:full effect */
-	UPROPERTY(interp, BlueprintReadWrite, Category=FWeightedBlendable, meta=(ClampMin = "0.0", ClampMax = "1.0", Delta = "0.01"))
+	UPROPERTY(interp, Category=FWeightedBlendable, meta=(ClampMin = "0.0", ClampMax = "1.0", Delta = "0.01"))
 	float Weight;
 
 	/** should be of the IBlendableInterface* type but UProperties cannot express that */
-	UPROPERTY(EditAnywhere, BlueprintReadWrite, Category=FWeightedBlendable, meta=( AllowedClasses="BlendableInterface", Keywords="PostProcess" ))
+	UPROPERTY(EditAnywhere, Category=FWeightedBlendable, meta=( AllowedClasses="BlendableInterface", Keywords="PostProcess" ))
 	UObject* Object;
 
 	// default constructor
@@ -91,7 +91,7 @@ struct FWeightedBlendables
 {
 	GENERATED_USTRUCT_BODY()
 
-	UPROPERTY(EditAnywhere, BlueprintReadWrite, Category="PostProcessSettings", meta=( Keywords="PostProcess" ))
+	UPROPERTY(EditAnywhere, Category="PostProcessSettings", meta=( Keywords="PostProcess" ))
 	TArray<FWeightedBlendable> Array;
 };
 
@@ -269,7 +269,7 @@ struct FPostProcessSettings
 	UPROPERTY(EditAnywhere, BlueprintReadWrite, Category=Overrides, meta=(PinHiddenByDefault, InlineEditConditionToggle))
 	uint32 bOverride_BloomSizeScale:1;
 
-	UPROPERTY(EditAnywhere, BlueprintReadWrite, Category = Overrides, meta = (PinHiddenByDefault, InlineEditConditionToggle))
+	UPROPERTY(EditAnywhere, BlueprintReadWrite, Category=Overrides, meta=(PinHiddenByDefault, InlineEditConditionToggle))
 	uint32 bOverride_BloomConvolutionTexture : 1;
 
 	UPROPERTY(EditAnywhere, BlueprintReadWrite, Category = Overrides, meta = (PinHiddenByDefault, InlineEditConditionToggle))
@@ -428,7 +428,7 @@ struct FPostProcessSettings
 	UPROPERTY(EditAnywhere, BlueprintReadWrite, Category=Overrides, meta=(PinHiddenByDefault, InlineEditConditionToggle))
 	uint32 bOverride_LPVEmissiveInjectionIntensity:1;
 
-	UPROPERTY(EditAnywhere, BlueprintReadWrite, Category=Overrides, meta = (PinHiddenByDefault, InlineEditConditionToggle))
+	UPROPERTY(EditAnywhere, BlueprintReadWrite, Category=Overrides, meta=(PinHiddenByDefault, InlineEditConditionToggle))
 	uint32 bOverride_LPVFadeRange : 1;
 
 	UPROPERTY(EditAnywhere, BlueprintReadWrite, Category=Overrides, meta = (PinHiddenByDefault, InlineEditConditionToggle))
@@ -745,11 +745,11 @@ struct FPostProcessSettings
 	class UTexture* BloomDirtMask;	
 	
 	/** BloomDirtMask intensity */
-	UPROPERTY(interp, BlueprintReadWrite, Category = "Lens|Dirt Mask", meta = (ClampMin = "0.0", UIMax = "8.0", editcondition = "bOverride_BloomDirtMaskIntensity", DisplayName = "Dirt Mask Intensity"))
+	UPROPERTY(interp, BlueprintReadWrite, Category="Lens|Dirt Mask", meta=(ClampMin = "0.0", UIMax = "8.0", editcondition = "bOverride_BloomDirtMaskIntensity", DisplayName = "Dirt Mask Intensity"))
 	float BloomDirtMaskIntensity;
 
 	/** BloomDirtMask tint color */
-	UPROPERTY(interp, BlueprintReadWrite, Category = "Lens|Dirt Mask", meta = (editcondition = "bOverride_BloomDirtMaskTint", DisplayName = "Dirt Mask Tint", HideAlphaChannel))
+	UPROPERTY(interp, BlueprintReadWrite, Category="Lens|Dirt Mask", meta=(editcondition = "bOverride_BloomDirtMaskTint", DisplayName = "Dirt Mask Tint", HideAlphaChannel))
 	FLinearColor BloomDirtMaskTint;
 
 	/** AmbientCubemap tint color */
diff --git a/Engine/Source/Runtime/Engine/Classes/Engine/UserInterfaceSettings.h b/Engine/Source/Runtime/Engine/Classes/Engine/UserInterfaceSettings.h
--- a/Engine/Source/Runtime/Engine/Classes/Engine/UserInterfaceSettings.h
+++ b/Engine/Source/Runtime/Engine/Classes/Engine/UserInterfaceSettings.h
@@ -7,6 +7,8 @@
 #include "Misc/StringClassReference.h"
 #include "Curves/CurveFloat.h"
 #include "Engine/DeveloperSettings.h"
+#include "Widgets/SWidget.h"
+
 #include "UserInterfaceSettings.generated.h"
 
 class UDPICustomScalingRule;
diff --git a/Engine/Source/Runtime/Engine/Classes/GameFramework/Actor.h b/Engine/Source/Runtime/Engine/Classes/GameFramework/Actor.h
--- a/Engine/Source/Runtime/Engine/Classes/GameFramework/Actor.h
+++ b/Engine/Source/Runtime/Engine/Classes/GameFramework/Actor.h
@@ -748,11 +748,7 @@ public:
 	/** Get the local-to-world transform of the RootComponent. Identical to GetTransform(). */
 	FORCEINLINE FTransform ActorToWorld() const
 	{
-		if( RootComponent != NULL )
-		{
-			return RootComponent->ComponentToWorld;
-		}
-		return FTransform::Identity;
+		return (RootComponent ? RootComponent->GetComponentTransform() : FTransform::Identity);
 	}
 
 
@@ -1123,6 +1119,7 @@ public:
 	class UActorComponent* AddComponent(FName TemplateName, bool bManualAttachment, const FTransform& RelativeTransform, const UObject* ComponentTemplateContext);
 
 	/** DEPRECATED - Use Component::DestroyComponent */
+	DEPRECATED(4.17, "Use Component.DestroyComponent instead")
 	UFUNCTION(BlueprintCallable, meta=(DeprecatedFunction, DeprecationMessage = "Use Component.DestroyComponent instead", BlueprintProtected = "true", DisplayName = "DestroyComponent"))
 	void K2_DestroyComponent(UActorComponent* Component);
 
@@ -1137,6 +1134,7 @@ public:
 	 *  Attaches the RootComponent of this Actor to the supplied component, optionally at a named socket. It is not valid to call this on components that are not Registered.
 	 *   @param AttachLocationType	Type of attachment, AbsoluteWorld to keep its world position, RelativeOffset to keep the object's relative offset and SnapTo to snap to the new parent.
 	 */
+	DEPRECATED(4.17, "Use AttachToComponent instead.")
 	UFUNCTION(BlueprintCallable, meta = (DisplayName = "AttachActorToComponent (Deprecated)", AttachLocationType = "KeepRelativeOffset"), Category = "Utilities|Transformation")
 	void K2_AttachRootComponentTo(USceneComponent* InParent, FName InSocketName = NAME_None, EAttachLocation::Type AttachLocationType = EAttachLocation::KeepRelativeOffset, bool bWeldSimulatedBodies = true);
 
@@ -1172,7 +1170,7 @@ public:
 	*  Attaches the RootComponent of this Actor to the supplied component, optionally at a named socket. It is not valid to call this on components that are not Registered.
 	*   @param AttachLocationType	Type of attachment, AbsoluteWorld to keep its world position, RelativeOffset to keep the object's relative offset and SnapTo to snap to the new parent.
 	*/
-
+	DEPRECATED(4.17, "Use AttachToActor instead.")
 	UFUNCTION(BlueprintCallable, meta = (DisplayName = "AttachActorToActor (Deprecated)", AttachLocationType = "KeepRelativeOffset"), Category = "Utilities|Transformation")
 	void K2_AttachRootComponentToActor(AActor* InParentActor, FName InSocketName = NAME_None, EAttachLocation::Type AttachLocationType = EAttachLocation::KeepRelativeOffset, bool bWeldSimulatedBodies = true);
 
@@ -1200,6 +1198,7 @@ public:
 	 *  Snap the RootComponent of this Actor to the supplied Actor's root component, optionally at a named socket. It is not valid to call this on components that are not Registered. 
 	 *  If InSocketName == NAME_None, it will attach to origin of the InParentActor. 
 	 */
+	DEPRECATED(4.17, "Use AttachRootComponentTo with EAttachLocation::SnapToTarget option instead")
 	UFUNCTION(BlueprintCallable, meta=(DeprecatedFunction, DeprecationMessage = "Use AttachRootComponentTo with EAttachLocation::SnapToTarget option instead", DisplayName = "SnapActorTo"), Category="Utilities|Transformation")
 	void SnapRootComponentTo(AActor* InParentActor, FName InSocketName);
 
@@ -1207,6 +1206,7 @@ public:
 	 *  Detaches the RootComponent of this Actor from any SceneComponent it is currently attached to. 
 	 *   @param bMaintainWorldTransform	If true, update the relative location/rotation of this component to keep its world position the same
 	 */
+	DEPRECATED(4.17, "Use DetachFromActor instead")
 	UFUNCTION(BlueprintCallable, meta=(DisplayName = "DetachActorFromActor (Deprecated)"), Category="Utilities|Transformation")
 	void DetachRootComponentFromParent(bool bMaintainWorldPosition = true);
 
@@ -1286,6 +1286,7 @@ public:
 	void SetTickableWhenPaused(bool bTickableWhenPaused);
 
 	/** Allocate a MID for a given parent material. */
+	DEPRECATED(4.17, "Use PrimitiveComponent.CreateAndSetMaterialInstanceDynamic instead.")
 	UFUNCTION(BlueprintCallable, meta=(DeprecatedFunction, DeprecationMessage="Use PrimitiveComponent.CreateAndSetMaterialInstanceDynamic instead.", BlueprintProtected = "true"), Category="Rendering|Material")
 	class UMaterialInstanceDynamic* MakeMIDForMaterial(class UMaterialInterface* Parent);
 
diff --git a/Engine/Source/Runtime/Engine/Classes/GameFramework/CharacterMovementComponent.h b/Engine/Source/Runtime/Engine/Classes/GameFramework/CharacterMovementComponent.h
--- a/Engine/Source/Runtime/Engine/Classes/GameFramework/CharacterMovementComponent.h
+++ b/Engine/Source/Runtime/Engine/Classes/GameFramework/CharacterMovementComponent.h
@@ -1968,9 +1968,9 @@ public:
 	 */
 	virtual void SmoothCorrection(const FVector& OldLocation, const FQuat& OldRotation, const FVector& NewLocation, const FQuat& NewRotation) override;
 
-	/** Get prediction data for a client game. Should not be used if not running as a client. Allocates the data on demand and can be overridden to allocate a custom override if desired. */
+	/** Get prediction data for a client game. Should not be used if not running as a client. Allocates the data on demand and can be overridden to allocate a custom override if desired. Result must be a FNetworkPredictionData_Client_Character. */
 	virtual class FNetworkPredictionData_Client* GetPredictionData_Client() const override;
-	/** Get prediction data for a server game. Should not be used if not running as a server. Allocates the data on demand and can be overridden to allocate a custom override if desired. */
+	/** Get prediction data for a server game. Should not be used if not running as a server. Allocates the data on demand and can be overridden to allocate a custom override if desired. Result must be a FNetworkPredictionData_Server_Character. */
 	virtual class FNetworkPredictionData_Server* GetPredictionData_Server() const override;
 
 	class FNetworkPredictionData_Client_Character* GetPredictionData_Client_Character() const;
diff --git a/Engine/Source/Runtime/Engine/Classes/GameFramework/InputSettings.h b/Engine/Source/Runtime/Engine/Classes/GameFramework/InputSettings.h
--- a/Engine/Source/Runtime/Engine/Classes/GameFramework/InputSettings.h
+++ b/Engine/Source/Runtime/Engine/Classes/GameFramework/InputSettings.h
@@ -105,30 +105,49 @@ class ENGINE_API UInputSettings
 	virtual void PostInitProperties() override;
 	// End of UObject interface
 
+	/** Returns the game local machine settings (resolution, windowing mode, scalability settings, etc...) */
+	UFUNCTION(BlueprintPure, Category = Settings)
+	static UInputSettings* GetInputSettings();
+
 	/** Programmatically add an action mapping to the project defaults */
-	void AddActionMapping(const FInputActionKeyMapping& KeyMapping);
+	UFUNCTION(BlueprintCallable, Category = Settings)
+	void AddActionMapping(const FInputActionKeyMapping& KeyMapping, bool bForceRebuildKeymaps = true);
+
+	UFUNCTION(BlueprintPure, Category = Settings)
+	void GetActionMappingByName(const FName InActionName, TArray<FInputActionKeyMapping>& OutMappings) const;
 
 	/** Programmatically remove an action mapping to the project defaults */
-	void RemoveActionMapping(const FInputActionKeyMapping& KeyMapping);
+	UFUNCTION(BlueprintCallable, Category = Settings)
+	void RemoveActionMapping(const FInputActionKeyMapping& KeyMapping, bool bForceRebuildKeymaps = true);
 
 	/** Programmatically add an axis mapping to the project defaults */
-	void AddAxisMapping(const FInputAxisKeyMapping& KeyMapping);
+	UFUNCTION(BlueprintCallable, Category = Settings)
+	void AddAxisMapping(const FInputAxisKeyMapping& KeyMapping, bool bForceRebuildKeymaps = true);
+
+	/** Retrieve all axis mappings by a certain name. */
+	UFUNCTION(BlueprintPure, Category = Settings)
+	void GetAxisMappingByName(const FName InAxisName, TArray<FInputAxisKeyMapping>& OutMappings) const;
 
 	/** Programmatically remove an axis mapping to the project defaults */
-	void RemoveAxisMapping(const FInputAxisKeyMapping& KeyMapping);
+	UFUNCTION(BlueprintCallable, Category = Settings)
+	void RemoveAxisMapping(const FInputAxisKeyMapping& KeyMapping, bool bForceRebuildKeymaps = true);
 
 	/** Flush the current mapping values to the config file */
+	UFUNCTION(BlueprintCallable, Category=Settings)
 	void SaveKeyMappings();
 
 	/** Populate a list of all defined action names */
+	UFUNCTION(BlueprintCallable, Category = Settings)
 	void GetActionNames(TArray<FName>& ActionNames) const;
 
 	/** Populate a list of all defined axis names */
+	UFUNCTION(BlueprintCallable, Category = Settings)
 	void GetAxisNames(TArray<FName>& AxisNames) const;
 
-private:
 	/** When changes are made to the default mappings, push those changes out to PlayerInput key maps */
+	UFUNCTION(BlueprintCallable, Category = Settings)
 	void ForceRebuildKeymaps();
 
+private:
 	void PopulateAxisConfigs();
 };
diff --git a/Engine/Source/Runtime/Engine/Classes/GameFramework/PlayerInput.h b/Engine/Source/Runtime/Engine/Classes/GameFramework/PlayerInput.h
--- a/Engine/Source/Runtime/Engine/Classes/GameFramework/PlayerInput.h
+++ b/Engine/Source/Runtime/Engine/Classes/GameFramework/PlayerInput.h
@@ -136,7 +136,7 @@ struct FInputActionKeyMapping
 	GENERATED_USTRUCT_BODY()
 
 	/** Friendly name of action, e.g "jump" */
-	UPROPERTY(EditAnywhere, Category="Input")
+	UPROPERTY(EditAnywhere, BlueprintReadWrite, Category="Input")
 	FName ActionName;
 
 	/** Key to bind it to. */
@@ -204,15 +204,15 @@ struct FInputAxisKeyMapping
 	GENERATED_USTRUCT_BODY()
 
 	/** Friendly name of axis, e.g "MoveForward" */
-	UPROPERTY(EditAnywhere, Category="Input")
+	UPROPERTY(EditAnywhere, BlueprintReadWrite, Category="Input")
 	FName AxisName;
 
 	/** Key to bind it to. */
-	UPROPERTY(EditAnywhere, Category="Input")
+	UPROPERTY(EditAnywhere, BlueprintReadWrite, Category="Input")
 	FKey Key;
 
 	/** Multiplier to use for the mapping when accumulating the axis value */
-	UPROPERTY(EditAnywhere, Category="Input")
+	UPROPERTY(EditAnywhere, BlueprintReadWrite, Category="Input")
 	float Scale;
 
 	bool operator==(const FInputAxisKeyMapping& Other) const
diff --git a/Engine/Source/Runtime/Engine/Classes/GameFramework/SpringArmComponent.h b/Engine/Source/Runtime/Engine/Classes/GameFramework/SpringArmComponent.h
--- a/Engine/Source/Runtime/Engine/Classes/GameFramework/SpringArmComponent.h
+++ b/Engine/Source/Runtime/Engine/Classes/GameFramework/SpringArmComponent.h
@@ -46,6 +46,7 @@ class ENGINE_API USpringArmComponent : public USceneComponent
 
 	/**
 	 * If this component is placed on a pawn, should it use the view/control rotation of the pawn where possible?
+	 * When disabled, the component will revert to using the stored RelativeRotation of the component.
 	 * @see APawn::GetViewRotation()
 	 */
 	UPROPERTY(EditAnywhere, BlueprintReadWrite, Category=CameraSettings)
diff --git a/Engine/Source/Runtime/Engine/Classes/Kismet/KismetInputLibrary.h b/Engine/Source/Runtime/Engine/Classes/Kismet/KismetInputLibrary.h
--- a/Engine/Source/Runtime/Engine/Classes/Kismet/KismetInputLibrary.h
+++ b/Engine/Source/Runtime/Engine/Classes/Kismet/KismetInputLibrary.h
@@ -40,43 +40,43 @@ class ENGINE_API UKismetInputLibrary : public UBlueprintFunctionLibrary
 	/**
 	 * @returns True if the key is a modifier key: Ctrl, Command, Alt, Shift
 	 */
-	UFUNCTION(BlueprintPure, Category="Utilities|Key")
+	UFUNCTION(BlueprintPure, meta = (DisplayName = "Is Modifier Key"), Category="Utilities|Key")
 	static bool Key_IsModifierKey(const FKey& Key);
 	
 	/**
 	 * @returns True if the key is a gamepad button
 	 */
-	UFUNCTION(BlueprintPure, Category="Utilities|Key")
+	UFUNCTION(BlueprintPure, meta = (DisplayName = "Is Gamepad Key"), Category="Utilities|Key")
 	static bool Key_IsGamepadKey(const FKey& Key);
 	
 	/**
 	 * @returns True if the key is a mouse button
 	 */
-	UFUNCTION(BlueprintPure, Category="Utilities|Key")
+	UFUNCTION(BlueprintPure, meta = (DisplayName = "Is Mouse Button"), Category="Utilities|Key")
 	static bool Key_IsMouseButton(const FKey& Key);
 	
 	/**
 	 * @returns True if the key is a keyboard button
 	 */
-	UFUNCTION(BlueprintPure, Category="Utilities|Key")
+	UFUNCTION(BlueprintPure, meta = (DisplayName = "Is Keyboard Key"), Category="Utilities|Key")
 	static bool Key_IsKeyboardKey(const FKey& Key);
 	
 	/**
 	 * @returns True if the key is a float axis
 	 */
-	UFUNCTION(BlueprintPure, Category="Utilities|Key")
+	UFUNCTION(BlueprintPure, meta = (DisplayName = "Is Float Axis"), Category="Utilities|Key")
 	static bool Key_IsFloatAxis(const FKey& Key);
 	
 	/**
 	 * @returns True if the key is a vector axis
 	 */
-	UFUNCTION(BlueprintPure, Category="Utilities|Key")
+	UFUNCTION(BlueprintPure, meta = (DisplayName = "Is Vector Axis"), Category="Utilities|Key")
 	static bool Key_IsVectorAxis(const FKey& Key);
 	
 	/**
 	 * @returns The display name of the key.
 	 */
-	UFUNCTION(BlueprintPure, Category="Utilities|Key")
+	UFUNCTION(BlueprintPure, meta = (DisplayName = "Get Key Display Name"), Category="Utilities|Key")
 	static FText Key_GetDisplayName(const FKey& Key);
 
 	/**
@@ -84,7 +84,7 @@ class ENGINE_API UKismetInputLibrary : public UBlueprintFunctionLibrary
 	 *
 	 * @return  True if this character is a repeat
 	 */
-	UFUNCTION(BlueprintPure, meta=( DisplayName = "IsRepeat" ), Category="Utilities|InputEvent")
+	UFUNCTION(BlueprintPure, meta=( DisplayName = "Is Repeat" ), Category="Utilities|InputEvent")
 	static bool InputEvent_IsRepeat(const FInputEvent& Input);
 
 	/**
@@ -92,7 +92,7 @@ class ENGINE_API UKismetInputLibrary : public UBlueprintFunctionLibrary
 	 *
 	 * @return  True if shift is pressed
 	 */
-	UFUNCTION(BlueprintPure, meta=( DisplayName = "IsShiftDown" ), Category="Utilities|InputEvent")
+	UFUNCTION(BlueprintPure, meta=( DisplayName = "Is Shift Down" ), Category="Utilities|InputEvent")
 	static bool InputEvent_IsShiftDown(const FInputEvent& Input);
 
 	/**
@@ -100,7 +100,7 @@ class ENGINE_API UKismetInputLibrary : public UBlueprintFunctionLibrary
 	 *
 	 * @return True if left shift is pressed.
 	 */
-	UFUNCTION(BlueprintPure, meta=( DisplayName = "IsLeftShiftDown" ), Category="Utilities|InputEvent")
+	UFUNCTION(BlueprintPure, meta=( DisplayName = "Is Left Shift Down" ), Category="Utilities|InputEvent")
 	static bool InputEvent_IsLeftShiftDown(const FInputEvent& Input);
 
 	/**
@@ -108,7 +108,7 @@ class ENGINE_API UKismetInputLibrary : public UBlueprintFunctionLibrary
 	 *
 	 * @return True if right shift is pressed.
 	 */
-	UFUNCTION(BlueprintPure, meta=( DisplayName = "IsRightShiftDown" ), Category="Utilities|InputEvent")
+	UFUNCTION(BlueprintPure, meta=( DisplayName = "Is Right Shift Down" ), Category="Utilities|InputEvent")
 	static bool InputEvent_IsRightShiftDown(const FInputEvent& Input);
 
 	/**
@@ -116,7 +116,7 @@ class ENGINE_API UKismetInputLibrary : public UBlueprintFunctionLibrary
 	 *
 	 * @return  True if control is pressed
 	 */
-	UFUNCTION(BlueprintPure, meta=( DisplayName = "IsControlDown" ), Category="Utilities|InputEvent")
+	UFUNCTION(BlueprintPure, meta=( DisplayName = "Is Control Down" ), Category="Utilities|InputEvent")
 	static bool InputEvent_IsControlDown(const FInputEvent& Input);
 
 	/**
@@ -124,7 +124,7 @@ class ENGINE_API UKismetInputLibrary : public UBlueprintFunctionLibrary
 	 *
 	 * @return  True if left control is pressed
 	 */
-	UFUNCTION(BlueprintPure, meta=( DisplayName = "IsLeftControlDown" ), Category="Utilities|InputEvent")
+	UFUNCTION(BlueprintPure, meta=( DisplayName = "Is Left Control Down" ), Category="Utilities|InputEvent")
 	static bool InputEvent_IsLeftControlDown(const FInputEvent& Input);
 
 	/**
@@ -132,7 +132,7 @@ class ENGINE_API UKismetInputLibrary : public UBlueprintFunctionLibrary
 	 *
 	 * @return  True if left control is pressed
 	 */
-	UFUNCTION(BlueprintPure, meta=( DisplayName = "IsRightControlDown" ), Category="Utilities|InputEvent")
+	UFUNCTION(BlueprintPure, meta=( DisplayName = "Is Right Control Down" ), Category="Utilities|InputEvent")
 	static bool InputEvent_IsRightControlDown(const FInputEvent& Input);
 
 	/**
@@ -140,7 +140,7 @@ class ENGINE_API UKismetInputLibrary : public UBlueprintFunctionLibrary
 	 *
 	 * @return  True if alt is pressed
 	 */
-	UFUNCTION(BlueprintPure, meta=( DisplayName = "IsAltDown" ), Category="Utilities|InputEvent")
+	UFUNCTION(BlueprintPure, meta=( DisplayName = "Is Alt Down" ), Category="Utilities|InputEvent")
 	static bool InputEvent_IsAltDown(const FInputEvent& Input);
 
 	/**
@@ -148,7 +148,7 @@ class ENGINE_API UKismetInputLibrary : public UBlueprintFunctionLibrary
 	 *
 	 * @return  True if left alt is pressed
 	 */
-	UFUNCTION(BlueprintPure, meta=( DisplayName = "IsLeftAltDown" ), Category="Utilities|InputEvent")
+	UFUNCTION(BlueprintPure, meta=( DisplayName = "Is Left Alt Down" ), Category="Utilities|InputEvent")
 	static bool InputEvent_IsLeftAltDown(const FInputEvent& Input);
 
 	/**
@@ -156,7 +156,7 @@ class ENGINE_API UKismetInputLibrary : public UBlueprintFunctionLibrary
 	 *
 	 * @return  True if right alt is pressed
 	 */
-	UFUNCTION(BlueprintPure, meta=( DisplayName = "IsRightAltDown" ), Category="Utilities|InputEvent")
+	UFUNCTION(BlueprintPure, meta=( DisplayName = "Is Right Alt Down" ), Category="Utilities|InputEvent")
 	static bool InputEvent_IsRightAltDown(const FInputEvent& Input);
 
 	/**
@@ -164,7 +164,7 @@ class ENGINE_API UKismetInputLibrary : public UBlueprintFunctionLibrary
 	 *
 	 * @return  True if command is pressed
 	 */
-	UFUNCTION(BlueprintPure, meta=( DisplayName = "IsCommandDown" ), Category="Utilities|InputEvent")
+	UFUNCTION(BlueprintPure, meta=( DisplayName = "Is Command Down" ), Category="Utilities|InputEvent")
 	static bool InputEvent_IsCommandDown(const FInputEvent& Input);
 
 	/**
@@ -172,7 +172,7 @@ class ENGINE_API UKismetInputLibrary : public UBlueprintFunctionLibrary
 	 *
 	 * @return  True if left command is pressed
 	 */
-	UFUNCTION(BlueprintPure, meta=( DisplayName = "IsLeftCommandDown" ), Category="Utilities|InputEvent")
+	UFUNCTION(BlueprintPure, meta=( DisplayName = "Is Left Command Down" ), Category="Utilities|InputEvent")
 	static bool InputEvent_IsLeftCommandDown(const FInputEvent& Input);
 
 	/**
@@ -180,7 +180,7 @@ class ENGINE_API UKismetInputLibrary : public UBlueprintFunctionLibrary
 	 *
 	 * @return  True if right command is pressed
 	 */
-	UFUNCTION(BlueprintPure, meta=( DisplayName = "IsRightCommandDown" ), Category="Utilities|InputEvent")
+	UFUNCTION(BlueprintPure, meta=( DisplayName = "Is Right Command Down" ), Category="Utilities|InputEvent")
 	static bool InputEvent_IsRightCommandDown(const FInputEvent& Input);
 
 
@@ -189,68 +189,68 @@ class ENGINE_API UKismetInputLibrary : public UBlueprintFunctionLibrary
 	 *
 	 * @return  Key name
 	 */
-	UFUNCTION(BlueprintCallable, Category="Utilities|KeyEvent")
+	UFUNCTION(BlueprintPure, Category="Utilities|KeyEvent")
 	static FKey GetKey(const FKeyEvent& Input);
 
-	UFUNCTION(BlueprintCallable, Category = "Utilities|KeyEvent")
+	UFUNCTION(BlueprintPure, Category = "Utilities|KeyEvent")
 	static int32 GetUserIndex(const FKeyEvent& Input);
 
-	UFUNCTION(BlueprintCallable, Category = "Utilities|FAnalogInputEvent")
+	UFUNCTION(BlueprintPure, Category = "Utilities|FAnalogInputEvent")
 	static float GetAnalogValue(const FAnalogInputEvent& Input);
 
 	/** @return The position of the cursor in screen space */
-	UFUNCTION(BlueprintPure, meta=( DisplayName = "GetScreenSpacePosition" ), Category="Utilities|PointerEvent")
+	UFUNCTION(BlueprintPure, meta=( DisplayName = "Get Screen Space Position" ), Category="Utilities|PointerEvent")
 	static FVector2D PointerEvent_GetScreenSpacePosition(const FPointerEvent& Input);
 
 	/** @return The position of the cursor in screen space last time we handled an input event */
-	UFUNCTION(BlueprintPure, meta=( DisplayName = "GetLastScreenSpacePosition" ), Category="Utilities|PointerEvent")
+	UFUNCTION(BlueprintPure, meta=( DisplayName = "Get Last Screen Space Position" ), Category="Utilities|PointerEvent")
 	static FVector2D PointerEvent_GetLastScreenSpacePosition(const FPointerEvent& Input);
 
 	/** @return the distance the mouse traveled since the last event was handled. */
-	UFUNCTION(BlueprintPure, meta=( DisplayName = "GetCursorDelta" ), Category="Utilities|PointerEvent")
+	UFUNCTION(BlueprintPure, meta=( DisplayName = "Get Cursor Delta" ), Category="Utilities|PointerEvent")
 	static FVector2D PointerEvent_GetCursorDelta(const FPointerEvent& Input);
 
 	/** Mouse buttons that are currently pressed */
-	UFUNCTION(BlueprintPure, meta=( DisplayName = "IsMouseButtonDown" ), Category="Utilities|PointerEvent")
+	UFUNCTION(BlueprintPure, meta=( DisplayName = "Is Mouse Button Down" ), Category="Utilities|PointerEvent")
 	static bool PointerEvent_IsMouseButtonDown(const FPointerEvent& Input, FKey MouseButton);
 
 	/** Mouse button that caused this event to be raised (possibly EB_None) */
-	UFUNCTION(BlueprintPure, meta=( DisplayName = "GetEffectingButton" ), Category="Utilities|PointerEvent")
+	UFUNCTION(BlueprintPure, meta=( DisplayName = "Get Effecting Button" ), Category="Utilities|PointerEvent")
 	static FKey PointerEvent_GetEffectingButton(const FPointerEvent& Input);
 
 	/** How much did the mouse wheel turn since the last mouse event */
-	UFUNCTION(BlueprintPure, meta=( DisplayName = "GetWheelDelta" ), Category="Utilities|PointerEvent")
+	UFUNCTION(BlueprintPure, meta=( DisplayName = "Get Wheel Delta" ), Category="Utilities|PointerEvent")
 	static float PointerEvent_GetWheelDelta(const FPointerEvent& Input);
 
 	/** @return The index of the user that caused the event */
-	UFUNCTION(BlueprintPure, meta=( DisplayName = "GetUserIndex" ), Category="Utilities|PointerEvent")
+	UFUNCTION(BlueprintPure, meta=( DisplayName = "Get User Index" ), Category="Utilities|PointerEvent")
 	static int32 PointerEvent_GetUserIndex(const FPointerEvent& Input);
 
 	/** @return The unique identifier of the pointer (e.g., finger index) */
-	UFUNCTION(BlueprintPure, meta=( DisplayName = "GetPointerIndex" ), Category="Utilities|PointerEvent")
+	UFUNCTION(BlueprintPure, meta=( DisplayName = "Get Pointer Index" ), Category="Utilities|PointerEvent")
 	static int32 PointerEvent_GetPointerIndex(const FPointerEvent& Input);
 
 	/** @return The index of the touch pad that generated this event (for platforms with multiple touch pads per user) */
-	UFUNCTION(BlueprintPure, meta=( DisplayName = "GetTouchpadIndex" ), Category="Utilities|PointerEvent")
+	UFUNCTION(BlueprintPure, meta=( DisplayName = "Get Touchpad Index" ), Category="Utilities|PointerEvent")
 	static int32 PointerEvent_GetTouchpadIndex(const FPointerEvent& Input);
 
 	/** @return Is this event a result from a touch (as opposed to a mouse) */
-	UFUNCTION(BlueprintPure, meta=( DisplayName = "IsTouchEvent" ), Category="Utilities|PointerEvent")
+	UFUNCTION(BlueprintPure, meta=( DisplayName = "Is Touch Event" ), Category="Utilities|PointerEvent")
 	static bool PointerEvent_IsTouchEvent(const FPointerEvent& Input);
 
 	//TODO UMG Support GetGestureType()
 
 	///** @return The type of touch gesture */
-	//UFUNCTION(BlueprintPure, meta=( DisplayName = "GetGestureType" ), Category="Utilities|PointerEvent")
+	//UFUNCTION(BlueprintPure, meta=( DisplayName = "Get Gesture Type" ), Category="Utilities|PointerEvent")
 	//static EGestureEvent::Type PointerEvent_GetGestureType(const FPointerEvent& Input);
 
 	/** @return The change in gesture value since the last gesture event of the same type. */
-	UFUNCTION(BlueprintPure, meta=( DisplayName = "GetGestureDelta" ), Category="Utilities|PointerEvent")
+	UFUNCTION(BlueprintPure, meta=( DisplayName = "Get Gesture Delta" ), Category="Utilities|PointerEvent")
 	static FVector2D PointerEvent_GetGestureDelta(const FPointerEvent& Input);
 
 
 	/** @return The controller button that caused this event */
-	UFUNCTION(BlueprintPure, meta = (DeprecatedFunction, DeprecationMessage = "Use GetKey() for KeyEvent instead", DisplayName = "GetEffectingButton"), Category = "Utilities|ControllerEvent")
+	UFUNCTION(BlueprintPure, meta = (DeprecatedFunction, DeprecationMessage = "Use GetKey() for KeyEvent instead", DisplayName = "Get Effecting Button"), Category = "Utilities|ControllerEvent")
 	static FKey ControllerEvent_GetEffectingButton(const FControllerEvent& Input);
 
 	/** @return The index of the user that caused the event */
diff --git a/Engine/Source/Runtime/Engine/Classes/Kismet/KismetTextLibrary.h b/Engine/Source/Runtime/Engine/Classes/Kismet/KismetTextLibrary.h
--- a/Engine/Source/Runtime/Engine/Classes/Kismet/KismetTextLibrary.h
+++ b/Engine/Source/Runtime/Engine/Classes/Kismet/KismetTextLibrary.h
@@ -57,7 +57,7 @@ namespace EFormatArgumentType
  * Used to pass argument/value pairs into FText::Format.
  * The full C++ struct is located here: Engine\Source\Runtime\Core\Public\Internationalization\Text.h
  */
-USTRUCT(noexport)
+USTRUCT(noexport, BlueprintInternalUseOnly)
 struct FFormatArgumentData
 {
 	UPROPERTY(EditInstanceOnly, BlueprintReadWrite, Category=ArgumentName)
@@ -85,27 +85,27 @@ class ENGINE_API UKismetTextLibrary : public UBlueprintFunctionLibrary
 {
 	GENERATED_UCLASS_BODY()
 
-	/** Converts a vector value to a localizable text, in the form 'X= Y= Z=' */
+	/** Converts a vector value to localized formatted text, in the form 'X= Y= Z=' */
 	UFUNCTION(BlueprintPure, meta = (DisplayName = "ToText (Vector)", CompactNodeTitle = "->", BlueprintAutocast), Category = "Utilities|Text")
 	static FText Conv_VectorToText(FVector InVec);
 
-	/** Converts a vector2d value to a localizable text, in the form 'X= Y=' */
+	/** Converts a vector2d value to localized formatted text, in the form 'X= Y=' */
 	UFUNCTION(BlueprintPure, meta = (DisplayName = "ToText (vector2d)", CompactNodeTitle = "->", BlueprintAutocast), Category = "Utilities|Text")
 	static FText Conv_Vector2dToText(FVector2D InVec);
 
-	/** Converts a rotator value to a localizable text, in the form 'P= Y= R=' */
+	/** Converts a rotator value to localized formatted text, in the form 'P= Y= R=' */
 	UFUNCTION(BlueprintPure, meta = (DisplayName = "ToText (rotator)", CompactNodeTitle = "->", BlueprintAutocast), Category = "Utilities|Text")
 	static FText Conv_RotatorToText(FRotator InRot);
 
-	/** Converts a transform value to a localizable text, in the form 'Translation: X= Y= Z= Rotation: P= Y= R= Scale: X= Y= Z=' */
+	/** Converts a transform value to localized formatted text, in the form 'Translation: X= Y= Z= Rotation: P= Y= R= Scale: X= Y= Z=' */
 	UFUNCTION(BlueprintPure, meta = (DisplayName = "ToText (transform)", CompactNodeTitle = "->", BlueprintAutocast), Category = "Utilities|Text")
 	static FText Conv_TransformToText(const FTransform& InTrans);
 
-	/** Converts a UObject value to a localizable text by calling the object's GetName method */
-	UFUNCTION(BlueprintPure, meta = (DisplayName = "ToText (object)", CompactNodeTitle = "->", BlueprintAutocast), Category = "Utilities|Text")
+	/** Converts a UObject value to culture invariant text by calling the object's GetName method */
+	UFUNCTION(BlueprintPure, meta = (DisplayName = "ToText (object)", BlueprintAutocast), Category = "Utilities|Text")
 	static FText Conv_ObjectToText(class UObject* InObj);
 
-	/** Converts a linear color value to a localizable text, in the form '(R=,G=,B=,A=)' */
+	/** Converts a linear color value to localized formatted text, in the form '(R=,G=,B=,A=)' */
 	UFUNCTION(BlueprintPure, meta = (DisplayName = "ToText (linear color)", CompactNodeTitle = "->", BlueprintAutocast), Category = "Utilities|Text")
 	static FText Conv_ColorToText(FLinearColor InColor);
 
@@ -113,12 +113,12 @@ class ENGINE_API UKismetTextLibrary : public UBlueprintFunctionLibrary
 	UFUNCTION(BlueprintPure, meta = (DisplayName = "ToString (text)", CompactNodeTitle = "->", BlueprintAutocast), Category = "Utilities|String")
 	static FString Conv_TextToString(const FText& InText);
 
-	/** Converts string to localizable text */
-	UFUNCTION(BlueprintPure, meta=(DisplayName = "ToText (string)", CompactNodeTitle = "->", BlueprintAutocast), Category="Utilities|Text")
+	/** Converts string to culture invariant text. Use Format or Make Literal Text to create localizable text */
+	UFUNCTION(BlueprintPure, meta=(DisplayName = "ToText (string)", BlueprintAutocast), Category="Utilities|Text")
 	static FText Conv_StringToText(const FString& InString);
 
-	/** Converts string to localizable text */
-	UFUNCTION(BlueprintPure, meta=(DisplayName = "ToText (name)", CompactNodeTitle = "->", BlueprintAutocast), Category="Utilities|Text")
+	/** Converts Name to culture invariant text */
+	UFUNCTION(BlueprintPure, meta=(DisplayName = "ToText (name)", BlueprintAutocast), Category="Utilities|Text")
 	static FText Conv_NameToText(FName InName);
 
 	/* Returns true if text is empty. */
@@ -183,21 +183,21 @@ class ENGINE_API UKismetTextLibrary : public UBlueprintFunctionLibrary
 	UFUNCTION(BlueprintPure, meta=(DisplayName = "NotEqual, Case Insensitive (text)", CompactNodeTitle = "!="), Category="Utilities|Text")
 	static bool NotEqual_IgnoreCase_TextText(const FText& A, const FText& B);
 
-	/** Converts a boolean value to text, either 'true' or 'false' */
+	/** Converts a boolean value to formatted text, either 'true' or 'false' */
 	UFUNCTION(BlueprintPure, meta=(DisplayName = "ToText (boolean)", CompactNodeTitle = "->", BlueprintAutocast), Category="Utilities|Text")
 	static FText Conv_BoolToText(bool InBool);
 
-	/** Converts a byte value to text */
+	/** Converts a byte value to formatted text */
 	UFUNCTION(BlueprintPure, meta=(DisplayName = "ToText (byte)", CompactNodeTitle = "->", BlueprintAutocast), Category="Utilities|Text")
 	static FText Conv_ByteToText(uint8 Value);
 
 	// Default values are duplicated from FNumberFormattingOptions and should be replicated in all functions and in the struct when changed!
-	/* Converts a passed in integer to a text based on formatting options */
+	/* Converts a passed in integer to text based on formatting options */
 	UFUNCTION(BlueprintPure, meta=(DisplayName = "ToText (int)", AdvancedDisplay = "1", BlueprintAutocast), Category="Utilities|Text")
 	static FText Conv_IntToText(int32 Value, bool bUseGrouping = true, int32 MinimumIntegralDigits = 1, int32 MaximumIntegralDigits = 324);
 
 	// Default values are duplicated from FNumberFormattingOptions and should be replicated in all functions and in the struct when changed!
-	/* Converts a passed in float to a text based on formatting options */
+	/* Converts a passed in float to text based on formatting options */
 	UFUNCTION(BlueprintPure, meta=(DisplayName = "ToText (float)", AdvancedDisplay = "1", BlueprintAutocast), Category="Utilities|Text")
 	static FText Conv_FloatToText(float Value, TEnumAsByte<ERoundingMode> RoundingMode, bool bUseGrouping = true, int32 MinimumIntegralDigits = 1, int32 MaximumIntegralDigits = 324, int32 MinimumFractionalDigits = 0, int32 MaximumFractionalDigits = 3);
 	
diff --git a/Engine/Source/Runtime/Engine/Classes/Materials/MaterialInstance.h b/Engine/Source/Runtime/Engine/Classes/Materials/MaterialInstance.h
--- a/Engine/Source/Runtime/Engine/Classes/Materials/MaterialInstance.h
+++ b/Engine/Source/Runtime/Engine/Classes/Materials/MaterialInstance.h
@@ -28,7 +28,7 @@ class FMaterialShaderMapId;
 class FSHAHash;
 
 /** Editable font parameter. */
-USTRUCT()
+USTRUCT(BlueprintType)
 struct FFontParameterValue
 {
 	GENERATED_USTRUCT_BODY()
@@ -58,7 +58,7 @@ struct FFontParameterValue
 };
 
 /** Editable scalar parameter. */
-USTRUCT()
+USTRUCT(BlueprintType)
 struct FScalarParameterValue
 {
 	GENERATED_USTRUCT_BODY()
@@ -84,7 +84,7 @@ struct FScalarParameterValue
 };
 
 /** Editable texture parameter. */
-USTRUCT()
+USTRUCT(BlueprintType)
 struct FTextureParameterValue
 {
 	GENERATED_USTRUCT_BODY()
@@ -110,7 +110,7 @@ struct FTextureParameterValue
 };
 
 /** Editable vector parameter. */
-USTRUCT()
+USTRUCT(BlueprintType)
 struct FVectorParameterValue
 {
 	GENERATED_USTRUCT_BODY()
@@ -213,7 +213,7 @@ class UMaterialInstance : public UMaterialInterface
 	UPROPERTY()
 	bool bOverrideBaseProperties_DEPRECATED;
 
-	UPROPERTY(EditAnywhere, BlueprintReadOnly, Category=MaterialInstance)
+	UPROPERTY(EditAnywhere, Category=MaterialInstance)
 	struct FMaterialInstanceBasePropertyOverrides BasePropertyOverrides;
 
 	//Cached copies of the base property overrides or the value from the parent to avoid traversing the parent chain for each access.
diff --git a/Engine/Source/Runtime/Engine/Classes/Materials/MaterialInterface.h b/Engine/Source/Runtime/Engine/Classes/Materials/MaterialInterface.h
--- a/Engine/Source/Runtime/Engine/Classes/Materials/MaterialInterface.h
+++ b/Engine/Source/Runtime/Engine/Classes/Materials/MaterialInterface.h
@@ -102,7 +102,7 @@ struct FLightmassMaterialInterfaceSettings
 	GENERATED_USTRUCT_BODY()
 
 	/** If true, forces translucency to cast static shadows as if the material were masked. */
-	UPROPERTY(EditAnywhere, BlueprintReadWrite, Category=Material)
+	UPROPERTY(EditAnywhere, Category=Material)
 	uint32 bCastShadowAsMasked:1;
 
 	/** Scales the emissive contribution of this material to static lighting. */
@@ -110,14 +110,14 @@ struct FLightmassMaterialInterfaceSettings
 	float EmissiveBoost;
 
 	/** Scales the diffuse contribution of this material to static lighting. */
-	UPROPERTY(EditAnywhere, BlueprintReadWrite, Category=Material)
+	UPROPERTY(EditAnywhere, Category=Material)
 	float DiffuseBoost;
 
 	/** 
 	 * Scales the resolution that this material's attributes were exported at. 
 	 * This is useful for increasing material resolution when details are needed.
 	 */
-	UPROPERTY(EditAnywhere, BlueprintReadWrite, Category=Material)
+	UPROPERTY(EditAnywhere, Category=Material)
 	float ExportResolutionScale;
 
 	/** Boolean override flags - only used in MaterialInstance* cases. */
@@ -211,7 +211,7 @@ class UMaterialInterface : public UObject, public IBlendableInterface
 
 protected:
 	/** The Lightmass settings for this object. */
-	UPROPERTY(EditAnywhere, BlueprintReadWrite, Category=Lightmass)
+	UPROPERTY(EditAnywhere, Category=Lightmass)
 	struct FLightmassMaterialInterfaceSettings LightmassSettings;
 
 #if WITH_EDITORONLY_DATA
diff --git a/Engine/Source/Runtime/Engine/Classes/Matinee/InterpTrackParticleReplay.h b/Engine/Source/Runtime/Engine/Classes/Matinee/InterpTrackParticleReplay.h
--- a/Engine/Source/Runtime/Engine/Classes/Matinee/InterpTrackParticleReplay.h
+++ b/Engine/Source/Runtime/Engine/Classes/Matinee/InterpTrackParticleReplay.h
@@ -26,11 +26,11 @@ struct FParticleReplayTrackKey
 	float Time;
 
 	/** Time length this clip should be captured/played for */
-	UPROPERTY(EditAnywhere, BlueprintReadWrite, Category=ParticleReplayTrackKey)
+	UPROPERTY(EditAnywhere, Category=ParticleReplayTrackKey)
 	float Duration;
 
 	/** Replay clip ID number that identifies the clip we should capture to or playback from */
-	UPROPERTY(EditAnywhere, BlueprintReadWrite, Category=ParticleReplayTrackKey)
+	UPROPERTY(EditAnywhere, Category=ParticleReplayTrackKey)
 	int32 ClipIDNumber;
 
 
diff --git a/Engine/Source/Runtime/Engine/Classes/Matinee/InterpTrackVisibility.h b/Engine/Source/Runtime/Engine/Classes/Matinee/InterpTrackVisibility.h
--- a/Engine/Source/Runtime/Engine/Classes/Matinee/InterpTrackVisibility.h
+++ b/Engine/Source/Runtime/Engine/Classes/Matinee/InterpTrackVisibility.h
@@ -53,7 +53,7 @@ struct FVisibilityTrackKey
 	UPROPERTY()
 	float Time;
 
-	UPROPERTY(EditAnywhere, BlueprintReadWrite, Category=VisibilityTrackKey)
+	UPROPERTY(EditAnywhere, Category=VisibilityTrackKey)
 	TEnumAsByte<enum EVisibilityTrackAction> Action;
 
 	/** Condition that must be satisfied for this key event to fire */
diff --git a/Engine/Source/Runtime/Engine/Classes/Matinee/MatineeActor.h b/Engine/Source/Runtime/Engine/Classes/Matinee/MatineeActor.h
--- a/Engine/Source/Runtime/Engine/Classes/Matinee/MatineeActor.h
+++ b/Engine/Source/Runtime/Engine/Classes/Matinee/MatineeActor.h
@@ -47,10 +47,10 @@ struct FInterpGroupActorInfo
 {
 	GENERATED_USTRUCT_BODY()
 
-	UPROPERTY(EditAnywhere, BlueprintReadWrite, Category=InterpGroupActorInfo)
+	UPROPERTY(EditAnywhere, Category=InterpGroupActorInfo)
 	FName ObjectName;
 
-	UPROPERTY(EditAnywhere, BlueprintReadWrite, Category=InterpGroupActorInfo)
+	UPROPERTY(EditAnywhere, Category=InterpGroupActorInfo)
 	TArray<class AActor*> Actors;
 
 };
diff --git a/Engine/Source/Runtime/Engine/Classes/Particles/Parameter/ParticleModuleParameterDynamic.h b/Engine/Source/Runtime/Engine/Classes/Particles/Parameter/ParticleModuleParameterDynamic.h
--- a/Engine/Source/Runtime/Engine/Classes/Particles/Parameter/ParticleModuleParameterDynamic.h
+++ b/Engine/Source/Runtime/Engine/Classes/Particles/Parameter/ParticleModuleParameterDynamic.h
@@ -50,7 +50,7 @@ struct FEmitterDynamicParameter
 	GENERATED_USTRUCT_BODY()
 
 	/** The parameter name - from the material DynamicParameter expression. READ-ONLY */
-	UPROPERTY(Category=EmitterDynamicParameter, VisibleAnywhere, BlueprintReadWrite)
+	UPROPERTY(Category=EmitterDynamicParameter, VisibleAnywhere)
 	FName ParamName;
 
 	/** If true, use the EmitterTime to retrieve the value, otherwise use Particle RelativeTime. */
diff --git a/Engine/Source/Runtime/Engine/Classes/Particles/ParticleSystemComponent.h b/Engine/Source/Runtime/Engine/Classes/Particles/ParticleSystemComponent.h
--- a/Engine/Source/Runtime/Engine/Classes/Particles/ParticleSystemComponent.h
+++ b/Engine/Source/Runtime/Engine/Classes/Particles/ParticleSystemComponent.h
@@ -694,7 +694,7 @@ public:
 			check(bAsyncDataCopyIsValid); 
 			return AsyncComponentToWorld;
 		}
-		return ComponentToWorld;
+		return GetComponentTransform();
 	}
 
 	FORCEINLINE const TArray<struct FParticleSysParam>& GetAsyncInstanceParameters()
diff --git a/Engine/Source/Runtime/Engine/Classes/PhysicsEngine/ConstraintDrives.h b/Engine/Source/Runtime/Engine/Classes/PhysicsEngine/ConstraintDrives.h
--- a/Engine/Source/Runtime/Engine/Classes/PhysicsEngine/ConstraintDrives.h
+++ b/Engine/Source/Runtime/Engine/Classes/PhysicsEngine/ConstraintDrives.h
@@ -77,11 +77,11 @@ struct ENGINE_API FLinearDriveConstraint
 	
 
 	/** Target position the linear drive.*/
-	UPROPERTY(EditAnywhere, BlueprintReadOnly, Category = LinearMotor)
+	UPROPERTY(EditAnywhere, Category = LinearMotor)
 	FVector PositionTarget;
 
 	/** Target velocity the linear drive. */
-	UPROPERTY(EditAnywhere, BlueprintReadOnly, Category = LinearMotor)
+	UPROPERTY(EditAnywhere, Category = LinearMotor)
 	FVector VelocityTarget;
 
 	UPROPERTY(EditAnywhere, Category = LinearMotor)
@@ -142,15 +142,15 @@ struct ENGINE_API FAngularDriveConstraint
 	FConstraintDrive SlerpDrive;
 	
 	/** Target orientation relative to the the body reference frame.*/
-	UPROPERTY(EditAnywhere, BlueprintReadOnly, Category = AngularMotor)
+	UPROPERTY(EditAnywhere, Category = AngularMotor)
 	FRotator OrientationTarget;
 
 	/** Target angular velocity relative to the body reference frame. */
-	UPROPERTY(EditAnywhere, BlueprintReadOnly, Category = AngularMotor)
+	UPROPERTY(EditAnywhere, Category = AngularMotor)
 	FVector AngularVelocityTarget;
 
 	/** Whether motors use SLERP (spherical lerp) or decompose into a Swing motor (cone constraints) and Twist motor (roll constraints). NOTE: SLERP will NOT work if any of the angular constraints are locked. */
-	UPROPERTY(EditAnywhere, BlueprintReadOnly, Category = AngularMotor)
+	UPROPERTY(EditAnywhere, Category = AngularMotor)
 	TEnumAsByte<enum EAngularDriveMode::Type> AngularDriveMode;
 
 	FAngularDriveConstraint();
diff --git a/Engine/Source/Runtime/Engine/Classes/Sound/SoundNodeDistanceCrossFade.h b/Engine/Source/Runtime/Engine/Classes/Sound/SoundNodeDistanceCrossFade.h
--- a/Engine/Source/Runtime/Engine/Classes/Sound/SoundNodeDistanceCrossFade.h
+++ b/Engine/Source/Runtime/Engine/Classes/Sound/SoundNodeDistanceCrossFade.h
@@ -21,23 +21,23 @@ struct FDistanceDatum
 	/* The FadeInDistance at which to start hearing this sound.
 	 * If you want to hear the sound up close then setting this to 0 might be a good option. 
 	 */
-	UPROPERTY(EditAnywhere, BlueprintReadWrite, Category=DistanceDatum )
+	UPROPERTY(EditAnywhere, Category=DistanceDatum )
 	float FadeInDistanceStart;
 
 	/* The distance at which this sound has faded in completely. */
-	UPROPERTY(EditAnywhere, BlueprintReadWrite, Category=DistanceDatum )
+	UPROPERTY(EditAnywhere, Category=DistanceDatum )
 	float FadeInDistanceEnd;
 
 	/* The distance at which this sound starts fading out. */
-	UPROPERTY(EditAnywhere, BlueprintReadWrite, Category=DistanceDatum )
+	UPROPERTY(EditAnywhere, Category=DistanceDatum )
 	float FadeOutDistanceStart;
 
 	/* The distance at which this sound is no longer audible. */
-	UPROPERTY(EditAnywhere, BlueprintReadWrite, Category=DistanceDatum )
+	UPROPERTY(EditAnywhere, Category=DistanceDatum )
 	float FadeOutDistanceEnd;
 
 	/* The volume for which this Input should be played. */
-	UPROPERTY(EditAnywhere, BlueprintReadWrite, Category=DistanceDatum)
+	UPROPERTY(EditAnywhere, Category=DistanceDatum)
 	float Volume;
 
 
diff --git a/Engine/Source/Runtime/Engine/Private/AI/Navigation/NavLinkComponent.cpp b/Engine/Source/Runtime/Engine/Private/AI/Navigation/NavLinkComponent.cpp
--- a/Engine/Source/Runtime/Engine/Private/AI/Navigation/NavLinkComponent.cpp
+++ b/Engine/Source/Runtime/Engine/Private/AI/Navigation/NavLinkComponent.cpp
@@ -38,7 +38,7 @@ FBoxSphereBounds UNavLinkComponent::CalcBounds(const FTransform &LocalToWorld) c
 
 void UNavLinkComponent::GetNavigationData(FNavigationRelevantData& Data) const
 {
-	NavigationHelper::ProcessNavLinkAndAppend(&Data.Modifiers, GetOwner(), Links);
+	NavigationHelper::ProcessNavLinkAndAppend(&Data.Modifiers, NavigationHelper::FNavLinkOwnerData(*this), Links);
 }
 
 bool UNavLinkComponent::IsNavigationRelevant() const
diff --git a/Engine/Source/Runtime/Engine/Private/AI/Navigation/NavLinkProxy.cpp b/Engine/Source/Runtime/Engine/Private/AI/Navigation/NavLinkProxy.cpp
--- a/Engine/Source/Runtime/Engine/Private/AI/Navigation/NavLinkProxy.cpp
+++ b/Engine/Source/Runtime/Engine/Private/AI/Navigation/NavLinkProxy.cpp
@@ -219,7 +219,7 @@ FBox ANavLinkProxy::GetComponentsBoundingBox(bool bNonColliding) const
 		LinksBB += SegmentLink.RightEnd;
 	}
 
-	LinksBB = LinksBB.TransformBy(RootComponent->ComponentToWorld);
+	LinksBB = LinksBB.TransformBy(RootComponent->GetComponentTransform());
 
 	if (SmartLinkComp && SmartLinkComp->IsNavigationRelevant())
 	{
diff --git a/Engine/Source/Runtime/Engine/Private/AI/Navigation/NavModifierComponent.cpp b/Engine/Source/Runtime/Engine/Private/AI/Navigation/NavModifierComponent.cpp
--- a/Engine/Source/Runtime/Engine/Private/AI/Navigation/NavModifierComponent.cpp
+++ b/Engine/Source/Runtime/Engine/Private/AI/Navigation/NavModifierComponent.cpp
@@ -32,7 +32,7 @@ void UNavModifierComponent::CalcAndCacheBounds() const
 				UBodySetup* BodySetup = PrimComp->GetBodySetup();
 				if (BodySetup)
 				{
-					FTransform ParentTM = PrimComp->ComponentToWorld;
+					FTransform ParentTM = PrimComp->GetComponentTransform();
 					const FVector Scale3D = ParentTM.GetScale3D();
 					ParentTM.RemoveScaling();
 					Bounds += PrimComp->Bounds.GetBox();
diff --git a/Engine/Source/Runtime/Engine/Private/AI/Navigation/NavigationModifier.cpp b/Engine/Source/Runtime/Engine/Private/AI/Navigation/NavigationModifier.cpp
--- a/Engine/Source/Runtime/Engine/Private/AI/Navigation/NavigationModifier.cpp
+++ b/Engine/Source/Runtime/Engine/Private/AI/Navigation/NavigationModifier.cpp
@@ -355,7 +355,7 @@ FAreaNavModifier::FAreaNavModifier(const UBrushComponent* BrushComponent, const
 	}
 
 	Init(InAreaClass);
-	SetConvex(Verts.GetData(), 0, Verts.Num(), ENavigationCoordSystem::Unreal, BrushComponent->ComponentToWorld);
+	SetConvex(Verts.GetData(), 0, Verts.Num(), ENavigationCoordSystem::Unreal, BrushComponent->GetComponentTransform());
 }
 
 void FAreaNavModifier::GetCylinder(FCylinderNavAreaData& Data) const
@@ -789,7 +789,7 @@ void FCompositeNavModifier::CreateAreaModifiers(const UPrimitiveComponent* PrimC
 		const FKBoxElem& BoxElem = BodySetup->AggGeom.BoxElems[Idx];
 		const FBox BoxSize = BoxElem.CalcAABB(FTransform::Identity, 1.0f);
 
-		FAreaNavModifier AreaMod(BoxSize, PrimComp->ComponentToWorld, AreaClass);
+		FAreaNavModifier AreaMod(BoxSize, PrimComp->GetComponentTransform(), AreaClass);
 		Add(AreaMod);
 	}
 
@@ -799,7 +799,7 @@ void FCompositeNavModifier::CreateAreaModifiers(const UPrimitiveComponent* PrimC
 		const float CapsuleHeight = SphylElem.Length + SphylElem.Radius;
 		const FTransform AreaOffset(FVector(0, 0, -CapsuleHeight));
 
-		FAreaNavModifier AreaMod(SphylElem.Radius, CapsuleHeight, AreaOffset * PrimComp->ComponentToWorld, AreaClass);
+		FAreaNavModifier AreaMod(SphylElem.Radius, CapsuleHeight, AreaOffset * PrimComp->GetComponentTransform(), AreaClass);
 		Add(AreaMod);
 	}
 
@@ -807,7 +807,7 @@ void FCompositeNavModifier::CreateAreaModifiers(const UPrimitiveComponent* PrimC
 	{
 		const FKConvexElem& ConvexElem = BodySetup->AggGeom.ConvexElems[Idx];
 		
-		FAreaNavModifier AreaMod(ConvexElem.VertexData, 0, ConvexElem.VertexData.Num(), ENavigationCoordSystem::Unreal, PrimComp->ComponentToWorld, AreaClass);
+		FAreaNavModifier AreaMod(ConvexElem.VertexData, 0, ConvexElem.VertexData.Num(), ENavigationCoordSystem::Unreal, PrimComp->GetComponentTransform(), AreaClass);
 		Add(AreaMod);
 	}
 	
@@ -815,7 +815,7 @@ void FCompositeNavModifier::CreateAreaModifiers(const UPrimitiveComponent* PrimC
 	{
 		const FKSphereElem& SphereElem = BodySetup->AggGeom.SphereElems[Idx];
 		
-		FAreaNavModifier AreaMod(SphereElem.Radius, SphereElem.Radius, PrimComp->ComponentToWorld, AreaClass);
+		FAreaNavModifier AreaMod(SphereElem.Radius, SphereElem.Radius, PrimComp->GetComponentTransform(), AreaClass);
 		Add(AreaMod);
 	}
 }
diff --git a/Engine/Source/Runtime/Engine/Private/AI/Navigation/NavigationSystem.cpp b/Engine/Source/Runtime/Engine/Private/AI/Navigation/NavigationSystem.cpp
--- a/Engine/Source/Runtime/Engine/Private/AI/Navigation/NavigationSystem.cpp
+++ b/Engine/Source/Runtime/Engine/Private/AI/Navigation/NavigationSystem.cpp
@@ -1325,7 +1325,7 @@ void UNavigationSystem::SimpleMoveToLocation(AController* Controller, const FVec
 	}
 }
 
-UNavigationPath* UNavigationSystem::FindPathToActorSynchronously(UObject* WorldContext, const FVector& PathStart, AActor* GoalActor, float TetherDistance, AActor* PathfindingContext, TSubclassOf<UNavigationQueryFilter> FilterClass)
+UNavigationPath* UNavigationSystem::FindPathToActorSynchronously(UObject* WorldContextObject, const FVector& PathStart, AActor* GoalActor, float TetherDistance, AActor* PathfindingContext, TSubclassOf<UNavigationQueryFilter> FilterClass)
 {
 	if (GoalActor == NULL)
 	{
@@ -1333,7 +1333,7 @@ UNavigationPath* UNavigationSystem::FindPathToActorSynchronously(UObject* WorldC
 	}
 
 	INavAgentInterface* NavAgent = Cast<INavAgentInterface>(GoalActor);
-	UNavigationPath* GeneratedPath = FindPathToLocationSynchronously(WorldContext, PathStart, NavAgent ? NavAgent->GetNavAgentLocation() : GoalActor->GetActorLocation(), PathfindingContext, FilterClass);
+	UNavigationPath* GeneratedPath = FindPathToLocationSynchronously(WorldContextObject, PathStart, NavAgent ? NavAgent->GetNavAgentLocation() : GoalActor->GetActorLocation(), PathfindingContext, FilterClass);
 	if (GeneratedPath != NULL && GeneratedPath->GetPath().IsValid() == true)
 	{
 		GeneratedPath->GetPath()->SetGoalActorObservation(*GoalActor, TetherDistance);
@@ -1342,13 +1342,13 @@ UNavigationPath* UNavigationSystem::FindPathToActorSynchronously(UObject* WorldC
 	return GeneratedPath;
 }
 
-UNavigationPath* UNavigationSystem::FindPathToLocationSynchronously(UObject* WorldContext, const FVector& PathStart, const FVector& PathEnd, AActor* PathfindingContext, TSubclassOf<UNavigationQueryFilter> FilterClass)
+UNavigationPath* UNavigationSystem::FindPathToLocationSynchronously(UObject* WorldContextObject, const FVector& PathStart, const FVector& PathEnd, AActor* PathfindingContext, TSubclassOf<UNavigationQueryFilter> FilterClass)
 {
 	UWorld* World = NULL;
 
-	if (WorldContext != NULL)
+	if (WorldContextObject != NULL)
 	{
-		World = GEngine->GetWorldFromContextObject(WorldContext);
+		World = GEngine->GetWorldFromContextObject(WorldContextObject);
 	}
 	if (World == NULL && PathfindingContext != NULL)
 	{
@@ -3712,14 +3712,14 @@ void UNavigationSystem::DiscardNavigationDataChunks(UWorld* InWorld)
 //----------------------------------------------------------------------//
 // Blueprint functions
 //----------------------------------------------------------------------//
-UNavigationSystem* UNavigationSystem::GetNavigationSystem(UObject* WorldContext)
+UNavigationSystem* UNavigationSystem::GetNavigationSystem(UObject* WorldContextObject)
 {
-	return GetCurrent(WorldContext);
+	return GetCurrent(WorldContextObject);
 }
 
-bool UNavigationSystem::K2_ProjectPointToNavigation(UObject* WorldContext, const FVector& Point, FVector& ProjectedLocation, ANavigationData* NavData, TSubclassOf<UNavigationQueryFilter> FilterClass, const FVector QueryExtent)
+bool UNavigationSystem::K2_ProjectPointToNavigation(UObject* WorldContextObject, const FVector& Point, FVector& ProjectedLocation, ANavigationData* NavData, TSubclassOf<UNavigationQueryFilter> FilterClass, const FVector QueryExtent)
 {
-	UWorld* World = GEngine->GetWorldFromContextObject(WorldContext);
+	UWorld* World = GEngine->GetWorldFromContextObject(WorldContextObject);
 	UNavigationSystem* NavSys = UNavigationSystem::GetCurrent(World);
 
 	ProjectedLocation = Point;
@@ -3732,7 +3732,7 @@ bool UNavigationSystem::K2_ProjectPointToNavigation(UObject* WorldContext, const
 		if (UseNavData)
 		{
 			bResult = NavSys->ProjectPointToNavigation(Point, OutNavLocation, QueryExtent, NavData
-				, UNavigationQueryFilter::GetQueryFilter(*UseNavData, WorldContext, FilterClass));
+				, UNavigationQueryFilter::GetQueryFilter(*UseNavData, WorldContextObject, FilterClass));
 			ProjectedLocation = OutNavLocation.Location;
 		}
 	}
@@ -3740,19 +3740,19 @@ bool UNavigationSystem::K2_ProjectPointToNavigation(UObject* WorldContext, const
 	return bResult;
 }
 
-bool UNavigationSystem::K2_GetRandomReachablePointInRadius(UObject* WorldContext, const FVector& Origin, FVector& RandomLocation, float Radius, ANavigationData* NavData, TSubclassOf<UNavigationQueryFilter> FilterClass)
+bool UNavigationSystem::K2_GetRandomReachablePointInRadius(UObject* WorldContextObject, const FVector& Origin, FVector& RandomLocation, float Radius, ANavigationData* NavData, TSubclassOf<UNavigationQueryFilter> FilterClass)
 {
 	FNavLocation RandomPoint(Origin);
 	bool bResult = false;
 
-	UWorld* World = GEngine->GetWorldFromContextObject(WorldContext);
+	UWorld* World = GEngine->GetWorldFromContextObject(WorldContextObject);
 	UNavigationSystem* NavSys = UNavigationSystem::GetCurrent(World);
 	if (NavSys)
 	{
 		ANavigationData* UseNavData = NavData ? NavData : NavSys->GetMainNavData(FNavigationSystem::DontCreate);
 		if (UseNavData)
 		{
-			bResult = NavSys->GetRandomReachablePointInRadius(Origin, Radius, RandomPoint, UseNavData, UNavigationQueryFilter::GetQueryFilter(*UseNavData, WorldContext, FilterClass));
+			bResult = NavSys->GetRandomReachablePointInRadius(Origin, Radius, RandomPoint, UseNavData, UNavigationQueryFilter::GetQueryFilter(*UseNavData, WorldContextObject, FilterClass));
 			RandomLocation = RandomPoint.Location;
 		}
 	}
@@ -3760,19 +3760,19 @@ bool UNavigationSystem::K2_GetRandomReachablePointInRadius(UObject* WorldContext
 	return bResult;
 }
 
-bool UNavigationSystem::K2_GetRandomPointInNavigableRadius(UObject* WorldContext, const FVector& Origin, FVector& RandomLocation, float Radius, ANavigationData* NavData, TSubclassOf<UNavigationQueryFilter> FilterClass)
+bool UNavigationSystem::K2_GetRandomPointInNavigableRadius(UObject* WorldContextObject, const FVector& Origin, FVector& RandomLocation, float Radius, ANavigationData* NavData, TSubclassOf<UNavigationQueryFilter> FilterClass)
 {
 	FNavLocation RandomPoint(Origin);
 	bool bResult = false;
 
-	UWorld* World = GEngine->GetWorldFromContextObject(WorldContext);
+	UWorld* World = GEngine->GetWorldFromContextObject(WorldContextObject);
 	UNavigationSystem* NavSys = UNavigationSystem::GetCurrent(World);
 	if (NavSys)
 	{
 		ANavigationData* UseNavData = NavData ? NavData : NavSys->GetMainNavData(FNavigationSystem::DontCreate);
 		if (UseNavData)
 		{
-			bResult = NavSys->GetRandomPointInNavigableRadius(Origin, Radius, RandomPoint, UseNavData, UNavigationQueryFilter::GetQueryFilter(*UseNavData, WorldContext, FilterClass));
+			bResult = NavSys->GetRandomPointInNavigableRadius(Origin, Radius, RandomPoint, UseNavData, UNavigationQueryFilter::GetQueryFilter(*UseNavData, WorldContextObject, FilterClass));
 			RandomLocation = RandomPoint.Location;
 		}
 	}
@@ -4183,36 +4183,36 @@ FVector UNavigationSystem::ProjectPointToNavigation(UObject* WorldContextObject,
 	return ProjectedPoint.Location;
 }
 
-FVector UNavigationSystem::GetRandomReachablePointInRadius(UObject* WorldContext, const FVector& Origin, float Radius, ANavigationData* NavData, TSubclassOf<UNavigationQueryFilter> FilterClass)
+FVector UNavigationSystem::GetRandomReachablePointInRadius(UObject* WorldContextObject, const FVector& Origin, float Radius, ANavigationData* NavData, TSubclassOf<UNavigationQueryFilter> FilterClass)
 {
 	FNavLocation RandomPoint;
 
-	UWorld* World = GEngine->GetWorldFromContextObject(WorldContext);
+	UWorld* World = GEngine->GetWorldFromContextObject(WorldContextObject);
 	UNavigationSystem* NavSys = UNavigationSystem::GetCurrent(World);
 	if (NavSys)
 	{
 		ANavigationData* UseNavData = NavData ? NavData : NavSys->GetMainNavData(FNavigationSystem::DontCreate);
 		if (UseNavData)
 		{
-			NavSys->GetRandomReachablePointInRadius(Origin, Radius, RandomPoint, UseNavData, UNavigationQueryFilter::GetQueryFilter(*UseNavData, WorldContext, FilterClass));
+			NavSys->GetRandomReachablePointInRadius(Origin, Radius, RandomPoint, UseNavData, UNavigationQueryFilter::GetQueryFilter(*UseNavData, WorldContextObject, FilterClass));
 		}
 	}
 
 	return RandomPoint.Location;
 }
 
-FVector UNavigationSystem::GetRandomPointInNavigableRadius(UObject* WorldContext, const FVector& Origin, float Radius, ANavigationData* NavData, TSubclassOf<UNavigationQueryFilter> FilterClass)
+FVector UNavigationSystem::GetRandomPointInNavigableRadius(UObject* WorldContextObject, const FVector& Origin, float Radius, ANavigationData* NavData, TSubclassOf<UNavigationQueryFilter> FilterClass)
 {
 	FNavLocation RandomPoint;
 
-	UWorld* World = GEngine->GetWorldFromContextObject(WorldContext);
+	UWorld* World = GEngine->GetWorldFromContextObject(WorldContextObject);
 	UNavigationSystem* NavSys = UNavigationSystem::GetCurrent(World);
 	if (NavSys)
 	{
 		ANavigationData* UseNavData = NavData ? NavData : NavSys->GetMainNavData(FNavigationSystem::DontCreate);
 		if (UseNavData)
 		{
-			NavSys->GetRandomPointInNavigableRadius(Origin, Radius, RandomPoint, UseNavData, UNavigationQueryFilter::GetQueryFilter(*UseNavData, WorldContext, FilterClass));
+			NavSys->GetRandomPointInNavigableRadius(Origin, Radius, RandomPoint, UseNavData, UNavigationQueryFilter::GetQueryFilter(*UseNavData, WorldContextObject, FilterClass));
 		}
 	}
 
diff --git a/Engine/Source/Runtime/Engine/Private/AI/Navigation/NavigationSystemHelpers.cpp b/Engine/Source/Runtime/Engine/Private/AI/Navigation/NavigationSystemHelpers.cpp
--- a/Engine/Source/Runtime/Engine/Private/AI/Navigation/NavigationSystemHelpers.cpp
+++ b/Engine/Source/Runtime/Engine/Private/AI/Navigation/NavigationSystemHelpers.cpp
@@ -15,6 +15,18 @@
 
 namespace NavigationHelper
 {
+	FNavLinkOwnerData::FNavLinkOwnerData(const AActor& InActor)
+	{
+		Actor = &InActor;
+		LinkToWorld = InActor.GetActorTransform();
+	}
+
+	FNavLinkOwnerData::FNavLinkOwnerData(const USceneComponent& InComponent)
+	{
+		Actor = InComponent.GetOwner();
+		LinkToWorld = InComponent.GetComponentTransform();
+	}
+
 	void GatherCollision(UBodySetup* RigidBody, TNavStatArray<FVector>& OutVertexBuffer, TNavStatArray<int32>& OutIndexBuffer, const FTransform& LocalToWorld)
 	{
 		if (RigidBody == NULL)
@@ -65,10 +77,9 @@ namespace NavigationHelper
 		return FallDownHeight;
 	}
 
-	void DefaultNavLinkProcessorImpl(FCompositeNavModifier* OUT CompositeModifier, const AActor* Actor, const TArray<FNavigationLink>& IN NavLinks)
+	void DefaultNavLinkProcessorImpl(FCompositeNavModifier* OUT CompositeModifier, const FNavLinkOwnerData& OwnerData, const TArray<FNavigationLink>& IN NavLinks)
 	{
-		const FTransform LocalToWorld = Actor->ActorToWorld();
-		FSimpleLinkNavModifier SimpleLink(NavLinks, LocalToWorld);
+		FSimpleLinkNavModifier SimpleLink(NavLinks, OwnerData.LinkToWorld);
 
 		// adjust links
 		for (int32 LinkIndex = 0; LinkIndex < SimpleLink.Links.Num(); ++LinkIndex)
@@ -83,14 +94,13 @@ namespace NavigationHelper
 
 			if (Link.MaxFallDownLength > 0)
 			{
-				const FVector WorldRight = LocalToWorld.TransformPosition(Link.Right);
-				const float FallDownHeight = RawGeometryFall(Actor, WorldRight, Link.MaxFallDownLength);
+				const FVector WorldRight = OwnerData.LinkToWorld.TransformPosition(Link.Right);
+				const float FallDownHeight = RawGeometryFall(OwnerData.Actor, WorldRight, Link.MaxFallDownLength);
 
 				if (FallDownHeight > 0.f)
 				{
 					// @todo maybe it's a good idea to clear ModifiedLink.MaxFallDownLength here
-					UE_VLOG_SEGMENT(Actor, LogNavigation, Log, WorldRight, WorldRight + FVector(0, 0, -FallDownHeight)
-						, FColor::Green, TEXT("FallDownHeight %d"), LinkIndex);
+					UE_VLOG_SEGMENT(OwnerData.Actor, LogNavigation, Log, WorldRight, WorldRight + FVector(0, 0, -FallDownHeight), FColor::Green, TEXT("FallDownHeight %d"), LinkIndex);
 
 					Link.Right.Z -= FallDownHeight;
 				}
@@ -98,14 +108,13 @@ namespace NavigationHelper
 
 			if (Link.LeftProjectHeight > 0)
 			{
-				const FVector WorldLeft = LocalToWorld.TransformPosition(Link.Left);
-				const float FallDownHeight = RawGeometryFall(Actor, WorldLeft, Link.LeftProjectHeight);
+				const FVector WorldLeft = OwnerData.LinkToWorld.TransformPosition(Link.Left);
+				const float FallDownHeight = RawGeometryFall(OwnerData.Actor, WorldLeft, Link.LeftProjectHeight);
 
 				if (FallDownHeight > 0.f)
 				{
 					// @todo maybe it's a good idea to clear ModifiedLink.LeftProjectHeight here
-					UE_VLOG_SEGMENT(Actor, LogNavigation, Log, WorldLeft, WorldLeft + FVector(0, 0, -FallDownHeight)
-						, FColor::Green, TEXT("LeftProjectHeight %d"), LinkIndex);
+					UE_VLOG_SEGMENT(OwnerData.Actor, LogNavigation, Log, WorldLeft, WorldLeft + FVector(0, 0, -FallDownHeight), FColor::Green, TEXT("LeftProjectHeight %d"), LinkIndex);
 
 					Link.Left.Z -= FallDownHeight;
 				}
@@ -115,12 +124,11 @@ namespace NavigationHelper
 		CompositeModifier->Add(SimpleLink);
 	}
 
-	void DefaultNavLinkSegmentProcessorImpl(FCompositeNavModifier* OUT CompositeModifier, const AActor* Actor, const TArray<FNavigationSegmentLink>& IN NavLinks)
+	void DefaultNavLinkSegmentProcessorImpl(FCompositeNavModifier* OUT CompositeModifier, const FNavLinkOwnerData& OwnerData, const TArray<FNavigationSegmentLink>& IN NavLinks)
 	{
-		const FTransform LocalToWorld = Actor->ActorToWorld();
-		FSimpleLinkNavModifier SimpleLink(NavLinks, LocalToWorld);
+		FSimpleLinkNavModifier SimpleLink(NavLinks, OwnerData.LinkToWorld);
 
-			// adjust links if needed
+		// adjust links if needed
 		for (int32 LinkIndex = 0; LinkIndex < SimpleLink.SegmentLinks.Num(); ++LinkIndex)
 		{
 			FNavigationSegmentLink& Link = SimpleLink.SegmentLinks[LinkIndex];
@@ -134,25 +142,23 @@ namespace NavigationHelper
 
 			if (Link.MaxFallDownLength > 0)
 			{
-				const FVector WorldRightStart = LocalToWorld.TransformPosition(Link.RightStart);
-				const FVector WorldRightEnd = LocalToWorld.TransformPosition(Link.RightEnd);
+				const FVector WorldRightStart = OwnerData.LinkToWorld.TransformPosition(Link.RightStart);
+				const FVector WorldRightEnd = OwnerData.LinkToWorld.TransformPosition(Link.RightEnd);
 
-				const float FallDownHeightStart = RawGeometryFall(Actor, WorldRightStart, Link.MaxFallDownLength);
-				const float FallDownHeightEnd = RawGeometryFall(Actor, WorldRightEnd, Link.MaxFallDownLength);
+				const float FallDownHeightStart = RawGeometryFall(OwnerData.Actor, WorldRightStart, Link.MaxFallDownLength);
+				const float FallDownHeightEnd = RawGeometryFall(OwnerData.Actor, WorldRightEnd, Link.MaxFallDownLength);
 
 				if (FallDownHeightStart > 0.f)
 				{
 					// @todo maybe it's a good idea to clear ModifiedLink.MaxFallDownLength here
-					UE_VLOG_SEGMENT(Actor, LogNavigation, Log, WorldRightStart, WorldRightStart + FVector(0, 0, -FallDownHeightStart)
-						, FColor::Green, TEXT("FallDownHeightStart %d"), LinkIndex);
+					UE_VLOG_SEGMENT(OwnerData.Actor, LogNavigation, Log, WorldRightStart, WorldRightStart + FVector(0, 0, -FallDownHeightStart), FColor::Green, TEXT("FallDownHeightStart %d"), LinkIndex);
 
 					Link.RightStart.Z -= FallDownHeightStart;
 				}
 				if (FallDownHeightEnd > 0.f)
 				{
 					// @todo maybe it's a good idea to clear ModifiedLink.MaxFallDownLength here
-					UE_VLOG_SEGMENT(Actor, LogNavigation, Log, WorldRightEnd, WorldRightEnd + FVector(0, 0, -FallDownHeightEnd)
-						, FColor::Green, TEXT("FallDownHeightEnd %d"), LinkIndex);
+					UE_VLOG_SEGMENT(OwnerData.Actor, LogNavigation, Log, WorldRightEnd, WorldRightEnd + FVector(0, 0, -FallDownHeightEnd), FColor::Green, TEXT("FallDownHeightEnd %d"), LinkIndex);
 
 					Link.RightEnd.Z -= FallDownHeightEnd;
 				}
@@ -162,38 +168,54 @@ namespace NavigationHelper
 		CompositeModifier->Add(SimpleLink);
 	}
 
-	FNavLinkProcessorDelegate NavLinkProcessor = FNavLinkProcessorDelegate::CreateStatic(DefaultNavLinkProcessorImpl);
-	FNavLinkSegmentProcessorDelegate NavLinkSegmentProcessor = FNavLinkSegmentProcessorDelegate::CreateStatic(DefaultNavLinkSegmentProcessorImpl);
+	FNavLinkProcessorDataDelegate NavLinkProcessor = FNavLinkProcessorDataDelegate::CreateStatic(DefaultNavLinkProcessorImpl);
+	FNavLinkSegmentProcessorDataDelegate NavLinkSegmentProcessor = FNavLinkSegmentProcessorDataDelegate::CreateStatic(DefaultNavLinkSegmentProcessorImpl);
 
 	void ProcessNavLinkAndAppend(FCompositeNavModifier* OUT CompositeModifier, const AActor* Actor, const TArray<FNavigationLink>& IN NavLinks)
+	{
+		if (Actor)
+		{
+			ProcessNavLinkAndAppend(CompositeModifier, FNavLinkOwnerData(*Actor), NavLinks);
+		}
+	}
+
+	void ProcessNavLinkAndAppend(FCompositeNavModifier* OUT CompositeModifier, const FNavLinkOwnerData& OwnerData, const TArray<FNavigationLink>& IN NavLinks)
 	{
 		SCOPE_CYCLE_COUNTER(STAT_Navigation_AdjustingNavLinks);
 
 		if (NavLinks.Num())
 		{
 			check(NavLinkProcessor.IsBound());
-			NavLinkProcessor.Execute(CompositeModifier, Actor, NavLinks);
+			NavLinkProcessor.Execute(CompositeModifier, OwnerData, NavLinks);
 		}
 	}
 
 	void ProcessNavLinkSegmentAndAppend(FCompositeNavModifier* OUT CompositeModifier, const AActor* Actor, const TArray<FNavigationSegmentLink>& IN NavLinks)
+	{
+		if (Actor)
+		{
+			ProcessNavLinkSegmentAndAppend(CompositeModifier, FNavLinkOwnerData(*Actor), NavLinks);
+		}
+	}
+
+	void ProcessNavLinkSegmentAndAppend(FCompositeNavModifier* OUT CompositeModifier, const FNavLinkOwnerData& OwnerData, const TArray<FNavigationSegmentLink>& IN NavLinks)
 	{
 		SCOPE_CYCLE_COUNTER(STAT_Navigation_AdjustingNavLinks);
 
 		if (NavLinks.Num())
 		{
 			check(NavLinkSegmentProcessor.IsBound());
-			NavLinkSegmentProcessor.Execute(CompositeModifier, Actor, NavLinks);
+			NavLinkSegmentProcessor.Execute(CompositeModifier, OwnerData, NavLinks);
 		}
 	}
 
-	void SetNavLinkProcessorDelegate(const FNavLinkProcessorDelegate& NewDelegate)
+	void SetNavLinkProcessorDelegate(const FNavLinkProcessorDataDelegate& NewDelegate)
 	{
 		check(NewDelegate.IsBound());
 		NavLinkProcessor = NewDelegate;
 	}
 
-	void SetNavLinkSegmentProcessorDelegate(const FNavLinkSegmentProcessorDelegate& NewDelegate)
+	void SetNavLinkSegmentProcessorDelegate(const FNavLinkSegmentProcessorDataDelegate& NewDelegate)
 	{
 		check(NewDelegate.IsBound());
 		NavLinkSegmentProcessor = NewDelegate;
@@ -208,4 +230,33 @@ namespace NavigationHelper
 			// AND has full colliding capabilities 
 			&& BodySetup.DefaultInstance.GetCollisionEnabled() == ECollisionEnabled::QueryAndPhysics;
 	}
+
+	//////////////////////////////////////////////////////////////////////////
+	// DEPRECATED FUNCTIONS
+
+	void SetNavLinkProcessorDelegate(const FNavLinkProcessorDelegate& NewDelegate)
+	{
+		// don't set delegate, custom navlink processor will not be registered with deprecated function
+	}
+
+	void SetNavLinkSegmentProcessorDelegate(const FNavLinkSegmentProcessorDelegate& NewDelegate)
+	{
+		// don't set delegate, custom navlink processor will not be registered with deprecated function
+	}
+
+	void DefaultNavLinkProcessorImpl(FCompositeNavModifier* OUT CompositeModifier, const AActor* Actor, const TArray<FNavigationLink>& IN NavLinks)
+	{
+		if (Actor)
+		{
+			DefaultNavLinkProcessorImpl(CompositeModifier, FNavLinkOwnerData(*Actor), NavLinks);
+		}
+	}
+
+	void DefaultNavLinkSegmentProcessorImpl(FCompositeNavModifier* OUT CompositeModifier, const AActor* Actor, const TArray<FNavigationSegmentLink>& IN NavLinks)
+	{
+		if (Actor)
+		{
+			DefaultNavLinkSegmentProcessorImpl(CompositeModifier, FNavLinkOwnerData(*Actor), NavLinks);
+		}
+	}
 }
diff --git a/Engine/Source/Runtime/Engine/Private/AI/Navigation/RecastNavMeshGenerator.cpp b/Engine/Source/Runtime/Engine/Private/AI/Navigation/RecastNavMeshGenerator.cpp
--- a/Engine/Source/Runtime/Engine/Private/AI/Navigation/RecastNavMeshGenerator.cpp
+++ b/Engine/Source/Runtime/Engine/Private/AI/Navigation/RecastNavMeshGenerator.cpp
@@ -972,7 +972,7 @@ FORCEINLINE_DEBUGGABLE void ExportComponent(UActorComponent* Component, FRecastG
 		{
 			if (!bHasData)
 			{
-				ExportRigidBodySetup(*BodySetup, GeomExport.VertexBuffer, GeomExport.IndexBuffer, GeomExport.Data->Bounds, PrimComp->ComponentToWorld);
+				ExportRigidBodySetup(*BodySetup, GeomExport.VertexBuffer, GeomExport.IndexBuffer, GeomExport.Data->Bounds, PrimComp->GetComponentTransform());
 				bHasData = true;
 			}
 
diff --git a/Engine/Source/Runtime/Engine/Private/Actor.cpp b/Engine/Source/Runtime/Engine/Private/Actor.cpp
--- a/Engine/Source/Runtime/Engine/Private/Actor.cpp
+++ b/Engine/Source/Runtime/Engine/Private/Actor.cpp
@@ -1121,7 +1121,7 @@ FBox AActor::CalculateComponentsBoundingBoxInLocalSpace( bool bNonColliding ) co
 			// Only use collidable components to find collision bounding box.
 			if( PrimComp->IsRegistered() && ( bNonColliding || PrimComp->IsCollisionEnabled() ) )
 			{
-				const FTransform ComponentToActor = PrimComp->ComponentToWorld * WorldToActor;
+				const FTransform ComponentToActor = PrimComp->GetComponentTransform() * WorldToActor;
 				FBoxSphereBounds ActorSpaceComponentBounds = PrimComp->CalcBounds( ComponentToActor );
 
 				Box += ActorSpaceComponentBounds.GetBox();
@@ -1576,7 +1576,7 @@ void AActor::SnapRootComponentTo(AActor* InParentActor, FName InSocketName/* = N
 		USceneComponent* ParentDefaultAttachComponent = InParentActor->GetDefaultAttachComponent();
 		if (ParentDefaultAttachComponent)
 		{
-			RootComponent->SnapTo(ParentDefaultAttachComponent, InSocketName);
+			RootComponent->AttachToComponent(ParentDefaultAttachComponent, FAttachmentTransformRules::SnapToTargetNotIncludingScale, InSocketName);
 		}
 	}
 }
@@ -1614,14 +1614,13 @@ void AActor::DetachSceneComponentsFromParent(USceneComponent* InParentComponent,
 
 void AActor::DetachAllSceneComponents(USceneComponent* InParentComponent, const FDetachmentTransformRules& DetachmentRules)
 {
-	if (InParentComponent != NULL)
+	if (InParentComponent)
 	{
 		TInlineComponentArray<USceneComponent*> Components;
 		GetComponents(Components);
 
-		for (int32 Index = 0; Index < Components.Num(); ++Index)
+		for (USceneComponent* SceneComp : Components)
 		{
-			USceneComponent* SceneComp = Components[Index];
 			if (SceneComp->GetAttachParent() == InParentComponent)
 			{
 				SceneComp->DetachFromComponent(DetachmentRules);
@@ -1925,6 +1924,26 @@ void AActor::RouteEndPlay(const EEndPlayReason::Type EndPlayReason)
 		{
 			EndPlay(EndPlayReason);
 		}
+
+		// Behaviors specific to an actor being unloaded due to a streaming level removal
+		if (EndPlayReason == EEndPlayReason::RemovedFromWorld)
+		{
+			ClearComponentOverlaps();
+
+			bActorInitialized = false;
+			if (World)
+			{
+				World->RemoveNetworkActor(this);
+			}
+		}
+
+		// Clear any ticking lifespan timers
+		if (TimerHandle_LifeSpanExpired.IsValid())
+		{
+			SetLifeSpan(0.f);
+		}
+
+		UNavigationSystem::OnActorUnregistered(this);
 	}
 
 	UninitializeComponents();
@@ -1951,23 +1970,6 @@ void AActor::EndPlay(const EEndPlayReason::Type EndPlayReason)
 			}
 		}
 	}
-
-	// Behaviors specific to an actor being unloaded due to a streaming level removal
-	if (EndPlayReason == EEndPlayReason::RemovedFromWorld)
-	{
-		ClearComponentOverlaps();
-
-		bActorInitialized = false;
-		GetWorld()->RemoveNetworkActor(this);
-	}
-
-	// Clear any ticking lifespan timers
-	if (TimerHandle_LifeSpanExpired.IsValid())
-	{
-		SetLifeSpan(0.f);
-	}
-
-	UNavigationSystem::OnActorUnregistered(this);
 }
 
 FVector AActor::GetPlacementExtent() const
@@ -2838,7 +2840,7 @@ void AActor::FinishSpawning(const FTransform& UserTransform, bool bIsDefaultTran
 	{
 		bHasFinishedSpawning = true;
 
-		FTransform FinalRootComponentTransform = (RootComponent ? RootComponent->ComponentToWorld : UserTransform);
+		FTransform FinalRootComponentTransform = (RootComponent ? RootComponent->GetComponentTransform() : UserTransform);
 
 		// see if we need to adjust the transform (i.e. in deferred cases where the caller passes in a different transform here 
 		// than was passed in during the original SpawnActor call)
@@ -2857,7 +2859,7 @@ void AActor::FinishSpawning(const FTransform& UserTransform, bool bIsDefaultTran
 					// caller passed a different transform!
 					// undo the original spawn transform to get back to the template transform, so we can recompute a good
 					// final transform that takes into account the template's transform
-					FTransform const TemplateTransform = RootComponent->ComponentToWorld * OriginalSpawnTransform->Inverse();
+					FTransform const TemplateTransform = RootComponent->GetComponentTransform() * OriginalSpawnTransform->Inverse();
 					FinalRootComponentTransform = TemplateTransform * UserTransform;
 				}
 			}
diff --git a/Engine/Source/Runtime/Engine/Private/ActorConstruction.cpp b/Engine/Source/Runtime/Engine/Private/ActorConstruction.cpp
--- a/Engine/Source/Runtime/Engine/Private/ActorConstruction.cpp
+++ b/Engine/Source/Runtime/Engine/Private/ActorConstruction.cpp
@@ -286,7 +286,7 @@ void AActor::RerunConstructionScripts()
 				USceneComponent* AttachParent = ActorTransactionAnnotation->RootComponentData.AttachedParentInfo.AttachParent.Get();
 				AttachParentComponent = (AttachParent ? AttachParent : FindObjectFast<USceneComponent>(Parent, ActorTransactionAnnotation->RootComponentData.AttachedParentInfo.AttachParentName));
 				SocketName = ActorTransactionAnnotation->RootComponentData.AttachedParentInfo.SocketName;
-				DetachRootComponentFromParent();
+				DetachFromActor(FDetachmentTransformRules::KeepWorldTransform);
 			}
 
 			for (const FActorRootComponentReconstructionData::FAttachedActorInfo& CachedAttachInfo : ActorTransactionAnnotation->RootComponentData.AttachedToInfo)
@@ -301,7 +301,7 @@ void AActor::RerunConstructionScripts()
 					Info.RelativeTransform = CachedAttachInfo.RelativeTransform;
 					AttachedActorInfos.Add(Info);
 
-					AttachedActor->DetachRootComponentFromParent();
+					AttachedActor->DetachFromActor(FDetachmentTransformRules::KeepWorldTransform);
 				}
 			}
 
@@ -365,7 +365,7 @@ void AActor::RerunConstructionScripts()
 				// Update component transform and remember it so it can be reapplied to any new root component which exists after construction.
 				// (Component transform may be stale if we are here following an Undo)
 				RootComponent->UpdateComponentToWorld();
-				OldTransform = RootComponent->ComponentToWorld;
+				OldTransform = RootComponent->GetComponentTransform();
 			}
 		}
 
diff --git a/Engine/Source/Runtime/Engine/Private/ActorEditor.cpp b/Engine/Source/Runtime/Engine/Private/ActorEditor.cpp
--- a/Engine/Source/Runtime/Engine/Private/ActorEditor.cpp
+++ b/Engine/Source/Runtime/Engine/Private/ActorEditor.cpp
@@ -231,7 +231,7 @@ void AActor::DebugShowOneComponentHierarchy( USceneComponent* SceneComp, int32&
 	FString PosString;
 	if( bShowPosition )
 	{
-		FVector Posn = SceneComp->ComponentToWorld.GetLocation();
+		FVector Posn = SceneComp->GetComponentTransform().GetLocation();
 		//PosString = FString::Printf( TEXT("{R:%f,%f,%f- W:%f,%f,%f}"), SceneComp->RelativeLocation.X, SceneComp->RelativeLocation.Y, SceneComp->RelativeLocation.Z, Posn.X, Posn.Y, Posn.Z );
 		PosString = FString::Printf( TEXT("{R:%f- W:%f}"), SceneComp->RelativeLocation.Z, Posn.Z );
 	}
@@ -252,7 +252,7 @@ void AActor::DebugShowOneComponentHierarchy( USceneComponent* SceneComp, int32&
 	{
 		if( bShowPosition )
 		{
-			FVector Posn = SceneComp->ComponentToWorld.GetLocation();
+			FVector Posn = SceneComp->GetComponentTransform().GetLocation();
 			//PosString = FString::Printf( TEXT("{R:%f,%f,%f- W:%f,%f,%f}"), SceneComp->RelativeLocation.X, SceneComp->RelativeLocation.Y, SceneComp->RelativeLocation.Z, Posn.X, Posn.Y, Posn.Z );
 			PosString = FString::Printf( TEXT("{R:%f- W:%f}"), SceneComp->RelativeLocation.Z, Posn.Z );
 		}
@@ -287,7 +287,7 @@ AActor::FActorTransactionAnnotation::FActorTransactionAnnotation(const AActor* A
 	if (bCacheRootComponentData && ActorRootComponent && ActorRootComponent->IsCreatedByConstructionScript())
 	{
 		bRootComponentDataCached = true;
-		RootComponentData.Transform = ActorRootComponent->ComponentToWorld;
+		RootComponentData.Transform = ActorRootComponent->GetComponentTransform();
 		RootComponentData.Transform.SetTranslation(ActorRootComponent->GetComponentLocation()); // take into account any custom location
 
 		if (ActorRootComponent->GetAttachParent())
diff --git a/Engine/Source/Runtime/Engine/Private/Animation/AnimInstanceProxy.cpp b/Engine/Source/Runtime/Engine/Private/Animation/AnimInstanceProxy.cpp
--- a/Engine/Source/Runtime/Engine/Private/Animation/AnimInstanceProxy.cpp
+++ b/Engine/Source/Runtime/Engine/Private/Animation/AnimInstanceProxy.cpp
@@ -226,7 +226,7 @@ void FAnimInstanceProxy::PreUpdate(UAnimInstance* InAnimInstance, float DeltaSec
 		LODLevel = SkelMeshComp->PredictedLODLevel;
 
 		// Cache these transforms, so nodes don't have to pull it off the gamethread manually.
-		SkelMeshCompLocalToWorld = SkelMeshComp->ComponentToWorld;
+		SkelMeshCompLocalToWorld = SkelMeshComp->GetComponentTransform();
 		if (const AActor* Owner = SkelMeshComp->GetOwner())
 		{
 			SkelMeshCompOwnerTransform = Owner->GetTransform();
