Copying //UE4/Dev-Framework to //UE4/Dev-Main (Source: //UE4/Dev-Framework @ 3431384)

diff --git a/Engine/Source/Runtime/Engine/Private/AssetManager.cpp b/Engine/Source/Runtime/Engine/Private/AssetManager.cpp
--- a/Engine/Source/Runtime/Engine/Private/AssetManager.cpp
+++ b/Engine/Source/Runtime/Engine/Private/AssetManager.cpp
@@ -6,10 +6,12 @@
 #include "AssetData.h"
 #include "ARFilter.h"
 #include "Engine/Engine.h"
+#include "Engine/BlueprintGeneratedClass.h"
 #include "UObject/ConstructorHelpers.h"
 #include "Misc/FileHelper.h"
 #include "Misc/ScopedSlowTask.h"
 #include "Misc/Paths.h"
+#include "AssetRegistryState.h"
 
 #if WITH_EDITOR
 #include "Editor.h"
@@ -96,18 +98,11 @@ const FPrimaryAssetType UAssetManager::PrimaryAssetLabelType = FName(TEXT("Prima
 UAssetManager::UAssetManager()
 {
 	bIsGlobalAsyncScanEnvironment = false;
-	bShouldKeepHardRefs = !GIsEditor || IsRunningCommandlet();
-	
-#if WITH_EDITOR
-	// In editor builds guess the type/name, in cooked builds do not as it should be saved
-	bShouldGuessTypeAndName = true;
-#else
 	bShouldGuessTypeAndName = false;
-#endif
-
-	bShouldUseSynchronousLoad = IsRunningCommandlet();
+	bShouldUseSynchronousLoad = false;
 	bIsBulkScanning = false;
 	bIsManagementDatabaseCurrent = false;
+	bUpdateManagementDatabaseAfterScan = false;
 }
 
 void UAssetManager::PostInitProperties()
@@ -116,7 +111,7 @@ void UAssetManager::PostInitProperties()
 
 	if (!HasAnyFlags(RF_ClassDefaultObject))
 	{
-
+		const UAssetManagerSettings& Settings = GetSettings();
 #if WITH_EDITOR
 		bIsGlobalAsyncScanEnvironment = GIsEditor && !IsRunningCommandlet();
 
@@ -128,11 +123,20 @@ void UAssetManager::PostInitProperties()
 			AssetRegistry.OnFilesLoaded().AddUObject(this, &UAssetManager::OnAssetRegistryFilesLoaded);
 			AssetRegistry.OnInMemoryAssetCreated().AddUObject(this, &UAssetManager::OnInMemoryAssetCreated);
 			AssetRegistry.OnInMemoryAssetDeleted().AddUObject(this, &UAssetManager::OnInMemoryAssetDeleted);
+			AssetRegistry.OnAssetRenamed().AddUObject(this, &UAssetManager::OnAssetRenamed);
 		}
 
 		FEditorDelegates::PreBeginPIE.AddUObject(this, &UAssetManager::PreBeginPIE);
 		FEditorDelegates::EndPIE.AddUObject(this, &UAssetManager::EndPIE);
+
+		// In editor builds guess the type/name if allowed
+		bShouldGuessTypeAndName = GetSettings().bShouldGuessTypeAndNameInEditor;
+#else 
+		// Never guess type in cooked builds
+		bShouldGuessTypeAndName = false;
 #endif
+		bShouldUseSynchronousLoad = IsRunningCommandlet();
+
 		LoadRedirectorMaps();
 	}
 }
@@ -281,7 +285,12 @@ int32 UAssetManager::ScanPathsForPrimaryAssets(FPrimaryAssetType PrimaryAssetTyp
 		}
 		else
 		{
-			ARFilter.ClassNames.Add(UBlueprint::StaticClass()->GetFName());
+			TArray<UClass*> BlueprintCoreDerivedClasses;
+			GetDerivedClasses(UBlueprintCore::StaticClass(), BlueprintCoreDerivedClasses);
+			for (UClass* BPCoreClass : BlueprintCoreDerivedClasses)
+			{
+				ARFilter.ClassNames.Add(BPCoreClass->GetFName());
+			}
 
 			// Make sure this works, if it does remove post load check
 			ClassNames.Add(BaseClass->GetFName());
@@ -289,14 +298,14 @@ int32 UAssetManager::ScanPathsForPrimaryAssets(FPrimaryAssetType PrimaryAssetTyp
 
 			/* DerivedClassNames are short names, GeneratedClass is a long name
 			const FName GeneratedClassName = FName(TEXT("GeneratedClass"));
-			for (FName DerivedClassName : DerivedClassNames)
+			for (const FName& DerivedClassName : DerivedClassNames)
 			{
 			ARFilter.TagsAndValues.Add(GeneratedClassName, DerivedClassName.ToString());
 
 			// Add any old names to the list in case things haven't been resaved
 			TArray<FName> OldNames = FLinkerLoad::FindPreviousNamesForClass(BaseClass->GetPathName(), false);
 
-			for (FName OldClassName : OldNames)
+			for (const FName& OldClassName : OldNames)
 			{
 			ARFilter.TagsAndValues.Add(GeneratedClassName, OldClassName.ToString());
 			}
@@ -348,7 +357,7 @@ int32 UAssetManager::ScanPathsForPrimaryAssets(FPrimaryAssetType PrimaryAssetTyp
 				ValidNames = FLinkerLoad::FindPreviousNamesForClass(BaseClass->GetPathName(), false);
 #endif
 				ValidNames.Add(FName(*ClassName));		
-				for (FName ValidName : ValidNames)
+				for (const FName& ValidName : ValidNames)
 				{
 					if (DerivedClassNames.Contains(ValidName))
 					{
@@ -417,11 +426,8 @@ void UAssetManager::UpdateCachedAssetData(const FPrimaryAssetId& PrimaryAssetId,
 
 		if (OldData && OldData->AssetPtr.ToStringReference() != NewStringReference)
 		{
-			if (bAllowDuplicates)
-			{
-				UE_LOG(LogAssetManager, Warning, TEXT("Found Duplicate PrimaryAssetID %s during creation, this must be resolved before saving. Path %s is replacing path %s"), *PrimaryAssetId.ToString(), *OldData->AssetPtr.ToStringReference().ToString(), *NewStringReference.ToString());
-			}
-			else
+			UE_LOG(LogAssetManager, Warning, TEXT("Found Duplicate PrimaryAssetID %s, this must be resolved before saving. Path %s is replacing path %s"), *PrimaryAssetId.ToString(), *OldData->AssetPtr.ToStringReference().ToString(), *NewStringReference.ToString());
+			if (!bAllowDuplicates)
 			{
 				ensureMsgf(!OldData, TEXT("Found Duplicate PrimaryAssetID %s! Path %s is replacing path %s"), *PrimaryAssetId.ToString(), *OldData->AssetPtr.ToStringReference().ToString(), *NewStringReference.ToString());
 			}
@@ -467,6 +473,9 @@ void UAssetManager::UpdateCachedAssetData(const FPrimaryAssetId& PrimaryAssetId,
 			CachedAssetBundles.Remove(PrimaryAssetId);
 		}
 		
+		// Mark these as editor only if our type is editor only
+		FStringAssetReferenceSerializationScope SerializationScope(NAME_None, NAME_None, TypeData.Info.bIsEditorOnly ? EStringAssetReferenceCollectType::EditorOnlyCollect : EStringAssetReferenceCollectType::AlwaysCollect);
+
 		FAssetBundleData BundleData;
 		if (BundleData.SetFromAssetData(NewAssetData))
 		{
@@ -677,7 +686,7 @@ bool UAssetManager::GetPrimaryAssetDataList(FPrimaryAssetType PrimaryAssetType,
 	{
 		const FPrimaryAssetTypeData& TypeData = FoundType->Get();
 
-		for (TPair<FName, FPrimaryAssetData> Pair : TypeData.AssetMap)
+		for (const TPair<FName, FPrimaryAssetData>& Pair : TypeData.AssetMap)
 		{
 			const FAssetData* CachedAssetData = Registry.GetCachedAssetDataForObjectPath(Pair.Value.AssetDataPath);
 
@@ -713,7 +722,7 @@ bool UAssetManager::GetPrimaryAssetObjectList(FPrimaryAssetType PrimaryAssetType
 	{
 		const FPrimaryAssetTypeData& TypeData = FoundType->Get();
 
-		for (TPair<FName, FPrimaryAssetData> Pair : TypeData.AssetMap)
+		for (const TPair<FName, FPrimaryAssetData>& Pair : TypeData.AssetMap)
 		{
 			UObject* FoundObject = Pair.Value.AssetPtr.Get();
 
@@ -747,9 +756,12 @@ bool UAssetManager::GetPrimaryAssetPathList(FPrimaryAssetType PrimaryAssetType,
 	{
 		const FPrimaryAssetTypeData& TypeData = FoundType->Get();
 
-		for (TPair<FName, FPrimaryAssetData> Pair : TypeData.AssetMap)
+		for (const TPair<FName, FPrimaryAssetData>& Pair : TypeData.AssetMap)
 		{
-			AssetPathList.AddUnique(Pair.Value.AssetPtr.ToStringReference());
+			if (!Pair.Value.AssetPtr.IsNull())
+			{
+				AssetPathList.AddUnique(Pair.Value.AssetPtr.ToStringReference());
+			}
 		}
 	}
 
@@ -855,7 +867,7 @@ bool UAssetManager::GetPrimaryAssetIdList(FPrimaryAssetType PrimaryAssetType, TA
 	{
 		const FPrimaryAssetTypeData& TypeData = FoundType->Get();
 
-		for (TPair<FName, FPrimaryAssetData> Pair : TypeData.AssetMap)
+		for (const TPair<FName, FPrimaryAssetData>& Pair : TypeData.AssetMap)
 		{
 			PrimaryAssetIdList.Add(FPrimaryAssetId(PrimaryAssetType, Pair.Key));
 		}
@@ -880,7 +892,7 @@ bool UAssetManager::GetPrimaryAssetTypeInfo(FPrimaryAssetType PrimaryAssetType,
 
 void UAssetManager::GetPrimaryAssetTypeInfoList(TArray<FPrimaryAssetTypeInfo>& AssetTypeInfoList) const
 {
-	for (const TPair<FName, TSharedRef<FPrimaryAssetTypeData>> TypePair : AssetTypeMap)
+	for (const TPair<FName, TSharedRef<FPrimaryAssetTypeData>>& TypePair : AssetTypeMap)
 	{
 		const FPrimaryAssetTypeData& TypeData = TypePair.Value.Get();
 
@@ -911,13 +923,13 @@ TSharedPtr<FStreamableHandle> UAssetManager::ChangeBundleStateForPrimaryAssets(c
 			{
 				NewBundleState = CurrentBundleState;
 
-				for (FName RemoveBundle : RemoveBundles)
+				for (const FName& RemoveBundle : RemoveBundles)
 				{
 					NewBundleState.Remove(RemoveBundle);
 				}
 			}
 
-			for (FName AddBundle : AddBundles)
+			for (const FName& AddBundle : AddBundles)
 			{
 				NewBundleState.AddUnique(AddBundle);
 			}
@@ -951,7 +963,7 @@ TSharedPtr<FStreamableHandle> UAssetManager::ChangeBundleStateForPrimaryAssets(c
 				PathsToLoad.Add(AssetPath);
 			}
 			
-			for (FName BundleName : NewBundleState)
+			for (const FName& BundleName : NewBundleState)
 			{
 				FAssetBundleEntry Entry = GetAssetBundleEntry(PrimaryAssetId, BundleName);
 
@@ -959,7 +971,10 @@ TSharedPtr<FStreamableHandle> UAssetManager::ChangeBundleStateForPrimaryAssets(c
 				{
 					PathsToLoad.Append(Entry.BundleAssets);
 				}
-				UE_LOG(LogAssetManager, Verbose, TEXT("ChangeBundleStateForPrimaryAssets: No assets for bundle %s::%s"), *PrimaryAssetId.ToString(), *BundleName.ToString());
+				else
+				{
+					UE_LOG(LogAssetManager, Verbose, TEXT("ChangeBundleStateForPrimaryAssets: No assets for bundle %s::%s"), *PrimaryAssetId.ToString(), *BundleName.ToString());
+				}
 			}
 
 			TSharedPtr<FStreamableHandle> NewHandle;
@@ -1058,6 +1073,80 @@ TSharedPtr<FStreamableHandle> UAssetManager::ChangeBundleStateForPrimaryAssets(c
 	return ReturnHandle;
 }
 
+
+TSharedPtr<FStreamableHandle> UAssetManager::PreloadPrimaryAssets(const TArray<FPrimaryAssetId>& AssetsToLoad, const TArray<FName>& LoadBundles, bool bLoadRecursive, FStreamableDelegate DelegateToCall, TAsyncLoadPriority Priority)
+{
+	TSet<FStringAssetReference> PathsToLoad;
+	FString DebugName;
+	TSharedPtr<FStreamableHandle> ReturnHandle;
+
+	for (const FPrimaryAssetId& PrimaryAssetId : AssetsToLoad)
+	{
+		FPrimaryAssetData* NameData = GetNameData(PrimaryAssetId);
+
+		if (NameData)
+		{
+			// Gather asset refs
+			const FStringAssetReference& AssetPath = NameData->AssetPtr.ToStringReference();
+
+			if (!AssetPath.IsNull())
+			{
+				// Dynamic types can have no base asset path
+				PathsToLoad.Add(AssetPath);
+			}
+
+			// Construct a temporary bundle data with the bundles specified
+			FAssetBundleData TempBundleData;
+			for (const FName& BundleName : LoadBundles)
+			{
+				FAssetBundleEntry Entry = GetAssetBundleEntry(PrimaryAssetId, BundleName);
+
+				if (Entry.IsValid())
+				{
+					TempBundleData.Bundles.Add(Entry);
+				}
+			}
+
+			if (bLoadRecursive)
+			{
+				RecursivelyExpandBundleData(TempBundleData);
+			}
+
+			for (const FAssetBundleEntry& Entry : TempBundleData.Bundles)
+			{
+				PathsToLoad.Append(Entry.BundleAssets);
+			}
+
+			if (DebugName.IsEmpty())
+			{
+				DebugName += TEXT("Preloading ");
+			}
+			else
+			{
+				DebugName += TEXT(", ");
+			}
+			DebugName += PrimaryAssetId.ToString();
+		}
+	}
+
+	if (bShouldUseSynchronousLoad)
+	{
+		ReturnHandle = StreamableManager.RequestSyncLoad(PathsToLoad.Array(), false, DebugName);
+		FStreamableHandle::ExecuteDelegate(DelegateToCall);
+	}
+	else
+	{
+		ReturnHandle = StreamableManager.RequestAsyncLoad(PathsToLoad.Array(), DelegateToCall, Priority, false, DebugName);
+	}
+
+	if (!ensureMsgf(ReturnHandle.IsValid(), TEXT("Requested preload of Primary Asset with no referenced assets!")))
+	{
+		return nullptr;
+	}
+
+	return ReturnHandle;
+}
+
 void UAssetManager::OnAssetStateChangeCompleted(FPrimaryAssetId PrimaryAssetId, TSharedPtr<FStreamableHandle> BoundHandle, FStreamableDelegate WrappedDelegate)
 {
 	FPrimaryAssetData* NameData = GetNameData(PrimaryAssetId);
@@ -1102,16 +1191,16 @@ TSharedPtr<FStreamableHandle> UAssetManager::LoadPrimaryAssetsWithType(FPrimaryA
 	return LoadPrimaryAssets(Assets, LoadBundles, DelegateToCall, Priority);
 }
 
-TSharedPtr<FStreamableHandle> UAssetManager::GetPrimaryAssetHandle(const FPrimaryAssetId& PrimaryAssetId, bool bForceCurrent, TArray<FName>* Bundles)
+TSharedPtr<FStreamableHandle> UAssetManager::GetPrimaryAssetHandle(const FPrimaryAssetId& PrimaryAssetId, bool bForceCurrent, TArray<FName>* Bundles) const
 {
-	FPrimaryAssetData* NameData = GetNameData(PrimaryAssetId);
+	const FPrimaryAssetData* NameData = GetNameData(PrimaryAssetId);
 
 	if (!NameData)
 	{
 		return nullptr;
 	}
 
-	FPrimaryAssetLoadState& LoadState = (bForceCurrent || !NameData->PendingState.IsValid()) ? NameData->CurrentState : NameData->PendingState;
+	const FPrimaryAssetLoadState& LoadState = (bForceCurrent || !NameData->PendingState.IsValid()) ? NameData->CurrentState : NameData->PendingState;
 
 	if (Bundles)
 	{
@@ -1120,11 +1209,11 @@ TSharedPtr<FStreamableHandle> UAssetManager::GetPrimaryAssetHandle(const FPrimar
 	return LoadState.Handle;
 }
 
-bool UAssetManager::GetPrimaryAssetsWithBundleState(TArray<FPrimaryAssetId>& PrimaryAssetList, const TArray<FName>& ValidTypes, const TArray<FName>& RequiredBundles, const TArray<FName>& ExcludedBundles, bool bForceCurrent)
+bool UAssetManager::GetPrimaryAssetsWithBundleState(TArray<FPrimaryAssetId>& PrimaryAssetList, const TArray<FName>& ValidTypes, const TArray<FName>& RequiredBundles, const TArray<FName>& ExcludedBundles, bool bForceCurrent) const
 {
 	bool bFoundAny = false;
 
-	for (TPair<FName, TSharedRef<FPrimaryAssetTypeData>> TypePair : AssetTypeMap)
+	for (const TPair<FName, TSharedRef<FPrimaryAssetTypeData>>& TypePair : AssetTypeMap)
 	{
 		if (ValidTypes.Num() > 0 && !ValidTypes.Contains(TypePair.Key))
 		{
@@ -1132,13 +1221,13 @@ bool UAssetManager::GetPrimaryAssetsWithBundleState(TArray<FPrimaryAssetId>& Pri
 			continue;
 		}
 
-		FPrimaryAssetTypeData& TypeData = TypePair.Value.Get();
+		const FPrimaryAssetTypeData& TypeData = TypePair.Value.Get();
 
-		for (TPair<FName, FPrimaryAssetData> NamePair : TypeData.AssetMap)
+		for (const TPair<FName, FPrimaryAssetData>& NamePair : TypeData.AssetMap)
 		{
-			FPrimaryAssetData& NameData = NamePair.Value;
+			const FPrimaryAssetData& NameData = NamePair.Value;
 
-			FPrimaryAssetLoadState& LoadState = (bForceCurrent || !NameData.PendingState.IsValid()) ? NameData.CurrentState : NameData.PendingState;
+			const FPrimaryAssetLoadState& LoadState = (bForceCurrent || !NameData.PendingState.IsValid()) ? NameData.CurrentState : NameData.PendingState;
 
 			if (!LoadState.IsValid())
 			{
@@ -1149,7 +1238,7 @@ bool UAssetManager::GetPrimaryAssetsWithBundleState(TArray<FPrimaryAssetId>& Pri
 			bool bFailedTest = false;
 
 			// Check bundle requirements
-			for (FName RequiredName : RequiredBundles)
+			for (const FName& RequiredName : RequiredBundles)
 			{
 				if (!LoadState.BundleNames.Contains(RequiredName))
 				{
@@ -1158,7 +1247,7 @@ bool UAssetManager::GetPrimaryAssetsWithBundleState(TArray<FPrimaryAssetId>& Pri
 				}
 			}
 
-			for (FName ExcludedName : ExcludedBundles)
+			for (const FName& ExcludedName : ExcludedBundles)
 			{
 				if (LoadState.BundleNames.Contains(ExcludedName))
 				{
@@ -1178,19 +1267,19 @@ bool UAssetManager::GetPrimaryAssetsWithBundleState(TArray<FPrimaryAssetId>& Pri
 	return bFoundAny;
 }
 
-void UAssetManager::GetPrimaryAssetBundleStateMap(TMap<FPrimaryAssetId, TArray<FName>>& BundleStateMap, bool bForceCurrent)
+void UAssetManager::GetPrimaryAssetBundleStateMap(TMap<FPrimaryAssetId, TArray<FName>>& BundleStateMap, bool bForceCurrent) const
 {
 	BundleStateMap.Reset();
 
-	for (TPair<FName, TSharedRef<FPrimaryAssetTypeData>> TypePair : AssetTypeMap)
+	for (const TPair<FName, TSharedRef<FPrimaryAssetTypeData>>& TypePair : AssetTypeMap)
 	{
-		FPrimaryAssetTypeData& TypeData = TypePair.Value.Get();
+		const FPrimaryAssetTypeData& TypeData = TypePair.Value.Get();
 
-		for (TPair<FName, FPrimaryAssetData> NamePair : TypeData.AssetMap)
+		for (const TPair<FName, FPrimaryAssetData>& NamePair : TypeData.AssetMap)
 		{
-			FPrimaryAssetData& NameData = NamePair.Value;
+			const FPrimaryAssetData& NameData = NamePair.Value;
 
-			FPrimaryAssetLoadState& LoadState = (bForceCurrent || !NameData.PendingState.IsValid()) ? NameData.CurrentState : NameData.PendingState;
+			const FPrimaryAssetLoadState& LoadState = (bForceCurrent || !NameData.PendingState.IsValid()) ? NameData.CurrentState : NameData.PendingState;
 
 			if (!LoadState.IsValid())
 			{
@@ -1289,7 +1378,6 @@ const FPrimaryAssetData* UAssetManager::GetNameData(const FPrimaryAssetId& Prima
 	const TSharedRef<FPrimaryAssetTypeData>* FoundType = AssetTypeMap.Find(PrimaryAssetId.PrimaryAssetType);
 
 	// Try redirected name
-
 	if (FoundType)
 	{
 		const FPrimaryAssetData* FoundName = (*FoundType)->AssetMap.Find(PrimaryAssetId.PrimaryAssetName);
@@ -1320,19 +1408,19 @@ void UAssetManager::RebuildObjectReferenceList()
 	ObjectReferenceList.Reset();
 
 	// Iterate primary asset map
-	for (TPair<FName, TSharedRef<FPrimaryAssetTypeData>> TypePair : AssetTypeMap)
+	for (TPair<FName, TSharedRef<FPrimaryAssetTypeData>>& TypePair : AssetTypeMap)
 	{
 		FPrimaryAssetTypeData& TypeData = TypePair.Value.Get();
 
 		// Add base class in case it's a blueprint
-		if (bShouldKeepHardRefs && !TypeData.Info.bIsDynamicAsset)
+		if (!TypeData.Info.bIsDynamicAsset)
 		{
 			ObjectReferenceList.AddUnique(TypeData.Info.AssetBaseClassLoaded);
 		}
 
 		TypeData.Info.NumberOfAssets = TypeData.AssetMap.Num();
 
-		for (TPair<FName, FPrimaryAssetData> NamePair : TypeData.AssetMap)
+		for (TPair<FName, FPrimaryAssetData>& NamePair : TypeData.AssetMap)
 		{
 			FPrimaryAssetData& NameData = NamePair.Value;
 			
@@ -1504,9 +1592,13 @@ void UAssetManager::ExtractStringAssetReferences(const UStruct* Struct, const vo
 	}
 }
 
-
 bool UAssetManager::GetAssetDataForPath(const FStringAssetReference& ObjectPath, FAssetData& AssetData) const
 {
+	if (ObjectPath.IsNull())
+	{
+		return false;
+	}
+
 	IAssetRegistry& AssetRegistry = GetAssetRegistry();
 
 	FString AssetPath = ObjectPath.ToString();
@@ -1564,7 +1656,7 @@ void UAssetManager::GetAssetDataForPathInternal(IAssetRegistry& AssetRegistry, c
 	bool bIsClass = AssetPath.EndsWith(TEXT("_C"), ESearchCase::CaseSensitive) && !AssetPath.Contains(TEXT("_C."), ESearchCase::CaseSensitive);
 
 	// If we're a class, first look for the asset data without the trailing _C
-	// We do this first because in cooked builds you have to search the asset registery for the Blueprint, not the class itself
+	// We do this first because in cooked builds you have to search the asset registry for the Blueprint, not the class itself
 	if (bIsClass)
 	{
 		// We need to strip the class suffix because the asset registry has it listed by blueprint name
@@ -1599,8 +1691,7 @@ bool UAssetManager::WriteCustomReport(FString FileName, TArray<FString>& FileLin
 		{
 			for (int32 Index = 0; Index < FileLines.Num(); ++Index)
 			{
-				FString LogEntry = FString::Printf(TEXT("%s"),
-					*FileLines[Index]) + LINE_TERMINATOR;
+				FString LogEntry = FString::Printf(TEXT("%s"), *FileLines[Index]) + LINE_TERMINATOR;
 				LogFile->Serialize(TCHAR_TO_ANSI(*LogEntry), LogEntry.Len());
 			}
 
@@ -1680,17 +1771,17 @@ void UAssetManager::DumpLoadedAssetState()
 
 		FPrimaryAssetTypeData& TypeData = Manager.AssetTypeMap.Find(TypeInfo.PrimaryAssetType)->Get();
 
-		for (TPair<FName, FPrimaryAssetData> NamePair : TypeData.AssetMap)
+		for (const TPair<FName, FPrimaryAssetData>& NamePair : TypeData.AssetMap)
 		{
-			FPrimaryAssetData& NameData = NamePair.Value;
+			const FPrimaryAssetData& NameData = NamePair.Value;
 
 			if (NameData.PendingState.IsValid() || NameData.CurrentState.IsValid())
 			{
-				FPrimaryAssetLoadState& LoadState = (!NameData.PendingState.IsValid()) ? NameData.CurrentState : NameData.PendingState;
+				const FPrimaryAssetLoadState& LoadState = (!NameData.PendingState.IsValid()) ? NameData.CurrentState : NameData.PendingState;
 
 				FString BundleString;
 
-				for (FName BundleName : LoadState.BundleNames)
+				for (const FName& BundleName : LoadState.BundleNames)
 				{
 					if (!BundleString.IsEmpty())
 					{
@@ -1717,6 +1808,35 @@ void UAssetManager::DumpLoadedAssetState()
 	}
 }
 
+static FAutoConsoleCommand CVarDumpAssetRegistryInfo(
+	TEXT("AssetManager.DumpAssetRegistryInfo"),
+	TEXT("Dumps extended info about asset registry to log"),
+	FConsoleCommandDelegate::CreateStatic(UAssetManager::DumpAssetRegistryInfo),
+	ECVF_Cheat);
+
+void UAssetManager::DumpAssetRegistryInfo()
+{
+	UE_LOG(LogAssetManager, Log, TEXT("=========== Asset Registry Summary ==========="));
+	UE_LOG(LogAssetManager, Log, TEXT("Current Registry Memory:"));
+
+	UAssetManager& Manager = Get();
+
+	// Output sizes
+	Manager.GetAssetRegistry().GetAllocatedSize(true);
+
+#if WITH_EDITOR
+	UE_LOG(LogAssetManager, Log, TEXT("Estimated Cooked Registry Memory:"));
+
+	FAssetRegistryState State;
+	FAssetRegistrySerializationOptions SaveOptions;
+
+	Manager.GetAssetRegistry().InitializeSerializationOptions(SaveOptions);
+	Manager.GetAssetRegistry().InitializeTemporaryAssetRegistryState(State, SaveOptions);
+
+	State.GetAllocatedSize(true);
+#endif
+}
+
 static FAutoConsoleCommand CVarDumpReferencersForPackage(
 	TEXT("AssetManager.DumpReferencersForPackage"),
 	TEXT("Generates a graph viz and log file of all references to a specified package"),
@@ -1737,7 +1857,7 @@ void UAssetManager::DumpReferencersForPackage(const TArray< FString >& PackageNa
 
 	ReportLines.Add(TEXT("digraph { "));
 
-	for (FString PackageString : PackageNames)
+	for (const FString& PackageString : PackageNames)
 	{
 		TArray<FAssetIdentifier> FoundReferencers;
 
@@ -1767,13 +1887,13 @@ void UAssetManager::ScanPrimaryAssetTypesFromConfig()
 
 	for (FPrimaryAssetTypeInfo TypeInfo : Settings.PrimaryAssetTypesToScan)
 	{
-		// Fill out runtime data on a copy
-		if (!TypeInfo.FillRuntimeData())
+		if (TypeInfo.bIsEditorOnly && !GIsEditor)
 		{
 			continue;
 		}
 
-		if (TypeInfo.bIsEditorOnly && !GIsEditor)
+		// Fill out runtime data on a copy, specifically not making a reference in this case
+		if (!TypeInfo.FillRuntimeData())
 		{
 			continue;
 		}
@@ -1782,6 +1902,7 @@ void UAssetManager::ScanPrimaryAssetTypesFromConfig()
 
 		SetPrimaryAssetTypeRules(TypeInfo.PrimaryAssetType, TypeInfo.Rules);
 	}
+
 	StopBulkScanning();
 
 	// Read primary asset rule overrides
@@ -1802,6 +1923,7 @@ void UAssetManager::PostInitialAssetScan()
 #if WITH_EDITOR
 	if (bUpdateManagementDatabaseAfterScan)
 	{
+		bUpdateManagementDatabaseAfterScan = false;
 		UpdateManagementDatabase(true);
 	}
 #endif
@@ -1827,7 +1949,7 @@ bool UAssetManager::GetManagedPackageList(FPrimaryAssetId PrimaryAssetId, TArray
 	return bFoundAny;
 }
 
-bool UAssetManager::GetPackageManagerList(FName PackageName, bool bRecurseToParents, TArray<FPrimaryAssetId>& ManagerList) const
+bool UAssetManager::GetPackageManagers(FName PackageName, bool bRecurseToParents, TSet<FPrimaryAssetId>& ManagerSet) const
 {
 	IAssetRegistry& AssetRegistry = GetAssetRegistry();
 
@@ -1842,7 +1964,7 @@ bool UAssetManager::GetPackageManagerList(FName PackageName, bool bRecurseToPare
 		if (PrimaryAssetId.IsValid())
 		{
 			bFoundAny = true;
-			ManagerList.AddUnique(PrimaryAssetId);
+			ManagerSet.Add(PrimaryAssetId);
 
 			if (bRecurseToParents)
 			{
@@ -1926,9 +2048,9 @@ void UAssetManager::FinishInitialLoading()
 	// See if we have pending scans, if so defer result
 	bool bWaitingOnDeferredScan = false;
 
-	for (TPair<FName, TSharedRef<FPrimaryAssetTypeData>> TypePair : AssetTypeMap)
+	for (const TPair<FName, TSharedRef<FPrimaryAssetTypeData>>& TypePair : AssetTypeMap)
 	{
-		FPrimaryAssetTypeData& TypeData = TypePair.Value.Get();
+		const FPrimaryAssetTypeData& TypeData = TypePair.Value.Get();
 
 		if (TypeData.DeferredAssetScanPaths.Num())
 		{
@@ -2028,7 +2150,7 @@ void UAssetManager::OnAssetRegistryFilesLoaded()
 {
 	StartBulkScanning();
 
-	for (TPair<FName, TSharedRef<FPrimaryAssetTypeData>> TypePair : AssetTypeMap)
+	for (TPair<FName, TSharedRef<FPrimaryAssetTypeData>>& TypePair : AssetTypeMap)
 	{
 		FPrimaryAssetTypeData& TypeData = TypePair.Value.Get();
 
@@ -2079,13 +2201,13 @@ void UAssetManager::UpdateManagementDatabase(bool bForceRefresh)
 	// List of references to not recurse on, priority doesn't matter
 	TMultiMap<FAssetIdentifier, FAssetIdentifier> NoReferenceManagementMap;
 
-	for (TPair<FName, TSharedRef<FPrimaryAssetTypeData>> TypePair : AssetTypeMap)
+	for (const TPair<FName, TSharedRef<FPrimaryAssetTypeData>>& TypePair : AssetTypeMap)
 	{
-		FPrimaryAssetTypeData& TypeData = TypePair.Value.Get();
+		const FPrimaryAssetTypeData& TypeData = TypePair.Value.Get();
 
-		for (TPair<FName, FPrimaryAssetData> NamePair : TypeData.AssetMap)
+		for (const TPair<FName, FPrimaryAssetData>& NamePair : TypeData.AssetMap)
 		{
-			FPrimaryAssetData& NameData = NamePair.Value;
+			const FPrimaryAssetData& NameData = NamePair.Value;
 			FPrimaryAssetId PrimaryAssetId(TypePair.Key, NamePair.Key);
 
 			FPrimaryAssetRules Rules = GetPrimaryAssetRules(PrimaryAssetId);
@@ -2107,9 +2229,9 @@ void UAssetManager::UpdateManagementDatabase(bool bForceRefresh)
 			// Add bundle references to manual reference list
 			if (BundleMap)
 			{
-				for (TPair<FName, FAssetBundleEntry>& BundlePair : *BundleMap)
+				for (const TPair<FName, FAssetBundleEntry>& BundlePair : *BundleMap)
 				{
-					for (FStringAssetReference& BundleAssetRef : BundlePair.Value.BundleAssets)
+					for (const FStringAssetReference& BundleAssetRef : BundlePair.Value.BundleAssets)
 					{
 						FString PackageName = BundleAssetRef.GetLongPackageName();
 
@@ -2118,7 +2240,7 @@ void UAssetManager::UpdateManagementDatabase(bool bForceRefresh)
 				}
 			}
 
-			for (FName AssetPackage : AssetPackagesReferenced)
+			for (const FName& AssetPackage : AssetPackagesReferenced)
 			{
 				TMultiMap<FAssetIdentifier, FAssetIdentifier>& ManagerMap = Rules.bApplyRecursively ? PriorityManagementMap.FindOrAdd(Rules.Priority) : NoReferenceManagementMap;
 
@@ -2161,15 +2283,14 @@ void UAssetManager::UpdateManagementDatabase(bool bForceRefresh)
 
 	TMultiMap<FAssetIdentifier, FAssetIdentifier> PrimaryAssetIdManagementMap;
 
-
 	// Update management parent list, which is PrimaryAssetId -> PrimaryAssetId
-	for (TPair<FName, TSharedRef<FPrimaryAssetTypeData>> TypePair : AssetTypeMap)
+	for (const TPair<FName, TSharedRef<FPrimaryAssetTypeData>>& TypePair : AssetTypeMap)
 	{
-		FPrimaryAssetTypeData& TypeData = TypePair.Value.Get();
+		const FPrimaryAssetTypeData& TypeData = TypePair.Value.Get();
 
-		for (TPair<FName, FPrimaryAssetData> NamePair : TypeData.AssetMap)
+		for (const TPair<FName, FPrimaryAssetData>& NamePair : TypeData.AssetMap)
 		{
-			FPrimaryAssetData& NameData = NamePair.Value;
+			const FPrimaryAssetData& NameData = NamePair.Value;
 			FPrimaryAssetId PrimaryAssetId(TypePair.Key, NamePair.Key);
 			const FStringAssetReference& AssetRef = NameData.AssetPtr.ToStringReference();
 
@@ -2177,9 +2298,9 @@ void UAssetManager::UpdateManagementDatabase(bool bForceRefresh)
 			{
 				FName PackageName = FName(*AssetRef.GetLongPackageName());
 
-				TArray<FPrimaryAssetId> Managers;
+				TSet<FPrimaryAssetId> Managers;
 
-				if (GetPackageManagerList(PackageName, false, Managers) && Managers.Num() > 1)
+				if (GetPackageManagers(PackageName, false, Managers) && Managers.Num() > 1)
 				{
 					// Find all managers that aren't this specific asset
 					for (const FPrimaryAssetId& Manager : Managers)
@@ -2208,7 +2329,9 @@ void UAssetManager::UpdateManagementDatabase(bool bForceRefresh)
 
 void UAssetManager::ApplyPrimaryAssetLabels()
 {
-	// Load all of them off disk
+	// Load all of them off disk. Turn off string asset reference tracking to avoid them getting cooked
+	FStringAssetReferenceSerializationScope SerializationScope(NAME_None, NAME_None, EStringAssetReferenceCollectType::NeverCollect);
+
 	TSharedPtr<FStreamableHandle> Handle = LoadPrimaryAssetsWithType(PrimaryAssetLabelType);
 
 	if (Handle.IsValid())
@@ -2219,7 +2342,7 @@ void UAssetManager::ApplyPrimaryAssetLabels()
 	// PostLoad in PrimaryAssetLabel sets PrimaryAssetRules overrides
 }
 
-void UAssetManager::ModifyCook(TArray<FName>& PackageNames)
+void UAssetManager::ModifyCook(TArray<FName>& PackagesToCook, TArray<FName>& PackagesToNeverCook)
 {
 	// Make sure management database is set up
 	UpdateManagementDatabase();
@@ -2232,11 +2355,7 @@ void UAssetManager::ModifyCook(TArray<FName>& PackageNames)
 	// Get package names in the libraries that we care about for cooking. Only get ones that are needed in production
 	for (const FPrimaryAssetTypeInfo& TypeInfo : TypeList)
 	{
-		if (TypeInfo.bIsEditorOnly)
-		{
-			continue;
-		}
-
+		// Cook these types
 		TArray<FAssetData> AssetDataList;
 		GetPrimaryAssetDataList(TypeInfo.PrimaryAssetType, AssetDataList);
 
@@ -2244,9 +2363,15 @@ void UAssetManager::ModifyCook(TArray<FName>& PackageNames)
 		{
 			EPrimaryAssetCookRule CookRule = GetPackageCookRule(AssetData.PackageName);
 
-			if (CookRule == EPrimaryAssetCookRule::AlwaysCook)
+			if (CookRule == EPrimaryAssetCookRule::AlwaysCook && !TypeInfo.bIsEditorOnly)
+			{
+				// If this is always cook and not editor only, cook it
+				PackagesToCook.AddUnique(AssetData.PackageName);
+			}
+			else if (!VerifyCanCookPackage(AssetData.PackageName, false))
 			{
-				PackageNames.AddUnique(AssetData.PackageName);
+				// If this package cannot be cooked, add to exclusion list
+				PackagesToNeverCook.AddUnique(AssetData.PackageName);
 			}
 		}
 	}
@@ -2256,8 +2381,8 @@ EPrimaryAssetCookRule UAssetManager::GetPackageCookRule(FName PackageName) const
 {
 	FPrimaryAssetRules BestRules;
 	FPrimaryAssetId BestId;
-	TArray<FPrimaryAssetId> Managers;
-	GetPackageManagerList(PackageName, true, Managers);
+	TSet<FPrimaryAssetId> Managers;
+	GetPackageManagers(PackageName, true, Managers);
 
 	for (const FPrimaryAssetId& PrimaryAssetId : Managers)
 	{
@@ -2324,9 +2449,9 @@ bool UAssetManager::GetPackageChunkIds(FName PackageName, const ITargetPlatform*
 
 	// Add all chunk ids from the asset rules of managers. By default priority will not override other chunks
 	bool bFoundAny = false;
-	TArray<FPrimaryAssetId> Managers;
+	TSet<FPrimaryAssetId> Managers;
 
-	GetPackageManagerList(PackageName, true, Managers);
+	GetPackageManagers(PackageName, true, Managers);
 
 	for (const FPrimaryAssetId& PrimaryAssetId : Managers)
 	{
@@ -2365,15 +2490,14 @@ void UAssetManager::PreBeginPIE(bool bStartSimulate)
 void UAssetManager::EndPIE(bool bStartSimulate)
 {
 	// Reset asset load state
-	for (TPair<FName, TSharedRef<FPrimaryAssetTypeData>> TypePair : AssetTypeMap)
+	for (const TPair<FName, TSharedRef<FPrimaryAssetTypeData>>& TypePair : AssetTypeMap)
 	{
-		FPrimaryAssetTypeData& TypeData = TypePair.Value.Get();
+		const FPrimaryAssetTypeData& TypeData = TypePair.Value.Get();
 
-		for (TPair<FName, FPrimaryAssetData> NamePair : TypeData.AssetMap)
+		for (const TPair<FName, FPrimaryAssetData>& NamePair : TypeData.AssetMap)
 		{
-			FPrimaryAssetData& NameData = NamePair.Value;
-
-			FPrimaryAssetLoadState& LoadState = (!NameData.PendingState.IsValid()) ? NameData.CurrentState : NameData.PendingState;
+			const FPrimaryAssetData& NameData = NamePair.Value;
+			const FPrimaryAssetLoadState& LoadState = (!NameData.PendingState.IsValid()) ? NameData.CurrentState : NameData.PendingState;
 
 			if (!LoadState.IsValid())
 			{
@@ -2403,7 +2527,7 @@ void UAssetManager::RefreshPrimaryAssetDirectory()
 {
 	StartBulkScanning();
 
-	for (TPair<FName, TSharedRef<FPrimaryAssetTypeData>> TypePair : AssetTypeMap)
+	for (TPair<FName, TSharedRef<FPrimaryAssetTypeData>>& TypePair : AssetTypeMap)
 	{
 		FPrimaryAssetTypeData& TypeData = TypePair.Value.Get();
 
@@ -2424,8 +2548,8 @@ void UAssetManager::RefreshPrimaryAssetDirectory()
 
 void UAssetManager::OnInMemoryAssetCreated(UObject *Object)
 {
-	// Ignore PIE changes
-	if (GIsPlayInEditorWorld || !Object)
+	// Ignore PIE and CDO changes
+	if (GIsPlayInEditorWorld || !Object || Object->HasAnyFlags(RF_ClassDefaultObject))
 	{
 		return;
 	}
@@ -2494,6 +2618,12 @@ void UAssetManager::OnAssetRenamed(const FAssetData& NewData, const FString& Old
 
 	FPrimaryAssetId OldPrimaryAssetId = GetPrimaryAssetIdForPath(OldPath);
 
+	// This may be a blueprint, try with _C
+	if (!OldPrimaryAssetId.IsValid())
+	{
+		OldPrimaryAssetId = GetPrimaryAssetIdForPath(OldPath + TEXT("_C"));
+	}
+
 	RemovePrimaryAssetId(OldPrimaryAssetId);
 
 	// This will always be in memory
@@ -2520,22 +2650,54 @@ void UAssetManager::RemovePrimaryAssetId(const FPrimaryAssetId& PrimaryAssetId)
 
 void UAssetManager::RefreshAssetData(UObject* ChangedObject)
 {
+	// If this is a BP CDO, call on class instead
+	if (ChangedObject->HasAnyFlags(RF_ClassDefaultObject))
+	{
+		UBlueprintGeneratedClass* AssetClass = Cast<UBlueprintGeneratedClass>(ChangedObject->GetClass());
+		if (AssetClass)
+		{
+			RefreshAssetData(AssetClass);
+		}
+		return;
+	}
+
 	// Only update things it knows about
-	FPrimaryAssetId PrimaryAssetId = ChangedObject->GetPrimaryAssetId();
+	IAssetRegistry& AssetRegistry = GetAssetRegistry();
 	FStringAssetReference ChangedObjectPath(ChangedObject);
-
-	if (PrimaryAssetId.IsValid() && GetPrimaryAssetPath(PrimaryAssetId) == ChangedObjectPath)
+	FPrimaryAssetId PrimaryAssetId = ChangedObject->GetPrimaryAssetId();
+	FPrimaryAssetId OldPrimaryAssetId = GetPrimaryAssetIdForPath(ChangedObjectPath);
+	
+	// This may be a blueprint, try with _C
+	if (!OldPrimaryAssetId.IsValid())
 	{
-		IAssetRegistry& AssetRegistry = GetAssetRegistry();
+		OldPrimaryAssetId = GetPrimaryAssetIdForPath(ChangedObjectPath.ToString() + TEXT("_C"));
+	}
 
-		// This will load it out of the in memory object
+	if (PrimaryAssetId.IsValid() && OldPrimaryAssetId == PrimaryAssetId)
+	{
+		// Same AssetId, this will update cache out of the in memory object
 		FAssetData NewData;
 		GetAssetDataForPathInternal(AssetRegistry, ChangedObjectPath.ToString(), NewData);
 
 		if (ensure(NewData.IsValid()))
 		{
 			UpdateCachedAssetData(PrimaryAssetId, NewData, false);
-		}	
+		}
+	}
+	else
+	{
+		// AssetId changed
+		if (OldPrimaryAssetId.IsValid())
+		{
+			// Remove old id if it was registered
+			RemovePrimaryAssetId(OldPrimaryAssetId);
+		}
+
+		if (PrimaryAssetId.IsValid())
+		{
+			// This will add new id
+			OnInMemoryAssetCreated(ChangedObject);
+		}
 	}
 }
 
diff --git a/Engine/Source/Runtime/Engine/Private/Blueprint.cpp b/Engine/Source/Runtime/Engine/Private/Blueprint.cpp
--- a/Engine/Source/Runtime/Engine/Private/Blueprint.cpp
+++ b/Engine/Source/Runtime/Engine/Private/Blueprint.cpp
@@ -582,9 +582,12 @@ UClass* UBlueprint::RegenerateClass(UClass* ClassToRegenerate, UObject* Previous
 			UBlueprint::ForceLoadMembers(GeneratedClassResolved->ClassDefaultObject);
 		}
 		UBlueprint::ForceLoadMembers(this);
+		
+		FBlueprintEditorUtils::PreloadConstructionScript( this );
+
+		FBlueprintEditorUtils::LinkExternalDependencies( this );
 
 		FBlueprintEditorUtils::RefreshVariables(this);
-		FBlueprintEditorUtils::PreloadConstructionScript( this );
 		
 		// Preload Overridden Components
 		if (InheritableComponentHandler)
@@ -592,8 +595,6 @@ UClass* UBlueprint::RegenerateClass(UClass* ClassToRegenerate, UObject* Previous
 			InheritableComponentHandler->PreloadAll();
 		}
 
-		FBlueprintEditorUtils::LinkExternalDependencies( this );
-
 		FBlueprintCompilationManager::NotifyBlueprintLoaded( this ); 
 		
 		// clear this now that we're not in a re-entrrant context - bHasBeenRegenerated will guard against 'real' 
@@ -1176,7 +1177,7 @@ void UBlueprint::BeginCacheForCookedPlatformData(const ITargetPlatform *TargetPl
 		// If nativization is enabled and this Blueprint class will NOT be nativized, we need to determine if any of its parent Blueprints will be nativized and flag it for the runtime code.
 		// Note: Currently, this flag is set on Actor-based Blueprint classes only. If it's ever needed for non-Actor-based Blueprint classes at runtime, then this needs to be updated to match.
 		const IBlueprintNativeCodeGenCore* NativeCodeGenCore = IBlueprintNativeCodeGenCore::Get();
-		if (GeneratedClass != nullptr && NativeCodeGenCore != nullptr && FParse::Param(FCommandLine::Get(), TEXT("NativizeAssets")))
+		if (GeneratedClass != nullptr && NativeCodeGenCore != nullptr)
 		{
 			ensure(TargetPlatform);
 			const FCompilerNativizationOptions& NativizationOptions = NativeCodeGenCore->GetNativizationOptionsForPlatform(TargetPlatform);
diff --git a/Engine/Source/Runtime/Engine/Private/BlueprintGeneratedClass.cpp b/Engine/Source/Runtime/Engine/Private/BlueprintGeneratedClass.cpp
--- a/Engine/Source/Runtime/Engine/Private/BlueprintGeneratedClass.cpp
+++ b/Engine/Source/Runtime/Engine/Private/BlueprintGeneratedClass.cpp
@@ -24,6 +24,7 @@
 #include "Kismet2/BlueprintEditorUtils.h"
 #include "Kismet2/KismetEditorUtilities.h"
 #include "BlueprintCompilationManager.h"
+#include "Engine/LevelScriptBlueprint.h"
 #endif //WITH_EDITOR
 
 DEFINE_STAT(STAT_PersistentUberGraphFrameMemory);
@@ -306,6 +307,17 @@ void UBlueprintGeneratedClass::ConditionalRecompileClass(TArray<UObject*>* ObjLo
 	}
 }
 
+void UBlueprintGeneratedClass::FlushCompilationQueueForLevel()
+{
+	if(GBlueprintUseCompilationManager)
+	{
+		if(Cast<ULevelScriptBlueprint>(ClassGeneratedBy))
+		{
+			FBlueprintCompilationManager::FlushCompilationQueue();
+		}
+	}
+}
+
 UObject* UBlueprintGeneratedClass::GetArchetypeForCDO() const
 {
 	if (OverridenArchetypeForCDO)
@@ -1326,6 +1338,11 @@ bool UBlueprintGeneratedClass::NeedsLoadForClient() const
 	return Super::NeedsLoadForClient();
 }
 
+bool UBlueprintGeneratedClass::NeedsLoadForEditorGame() const
+{
+	return true;
+}
+
 bool UBlueprintGeneratedClass::CanBeClusterRoot() const
 {
 	// Clustering level BPs doesn't work yet
diff --git a/Engine/Source/Runtime/Engine/Private/Collision/PhysXCollision.cpp b/Engine/Source/Runtime/Engine/Private/Collision/PhysXCollision.cpp
--- a/Engine/Source/Runtime/Engine/Private/Collision/PhysXCollision.cpp
+++ b/Engine/Source/Runtime/Engine/Private/Collision/PhysXCollision.cpp
@@ -628,8 +628,8 @@ PxQueryHitType::Enum FPxQueryFilterCallback::preFilter(const PxFilterData& filte
 	}
 
 	// Check if the shape is the right complexity for the trace 
-	PxFilterData ShapeFilter = shape->getQueryFilterData();
-	PxFilterData ShapeSimFilter = shape->getSimulationFilterData();	//This is a bit of a hack. We do this because word2 has our component ID
+	const PxFilterData ShapeFilter = shape->getQueryFilterData();
+
 #define ENABLE_PREFILTER_LOGGING 0
 #if ENABLE_PREFILTER_LOGGING
 	static bool bLoggingEnabled=false;
@@ -675,12 +675,19 @@ PxQueryHitType::Enum FPxQueryFilterCallback::preFilter(const PxFilterData& filte
 	// If not already rejected, check ignore actor and component list.
 	if (Result != PxQueryHitType::eNONE)
 	{
-		// See if we are ignoring the actor this shape belongs to (word0 of shape filterdata is actorID) or the component (word2 of shape sim filter data is componentID)
-		if (IgnoreActors.Contains(ShapeFilter.word0) || IgnoreComponents.Contains(ShapeSimFilter.word2))
+		// See if we are ignoring the actor this shape belongs to (word0 of shape filterdata is actorID)
+		if (IgnoreActors.Contains(ShapeFilter.word0))
 		{
 			//UE_LOG(LogTemp, Log, TEXT("Ignoring Actor: %d"), ShapeFilter.word0);
 			Result = PxQueryHitType::eNONE;
 		}
+
+		// We usually don't have ignore components so we try to avoid the virtual getSimulationFilterData() call below. 'word2' of shape sim filter data is componentID.
+		if (IgnoreComponents.Num() > 0 && IgnoreComponents.Contains(shape->getSimulationFilterData().word2))
+		{
+			//UE_LOG(LogTemp, Log, TEXT("Ignoring Component: %d"), shape->getSimulationFilterData().word2);
+			Result = PxQueryHitType::eNONE;
+		}
 	}
 
 #if ENABLE_PREFILTER_LOGGING
diff --git a/Engine/Source/Runtime/Engine/Private/ComponentInstanceDataCache.cpp b/Engine/Source/Runtime/Engine/Private/ComponentInstanceDataCache.cpp
--- a/Engine/Source/Runtime/Engine/Private/ComponentInstanceDataCache.cpp
+++ b/Engine/Source/Runtime/Engine/Private/ComponentInstanceDataCache.cpp
@@ -3,11 +3,150 @@
 #include "ComponentInstanceDataCache.h"
 #include "Serialization/ObjectWriter.h"
 #include "Serialization/ObjectReader.h"
+#include "Serialization/DuplicatedObject.h"
+#include "UObject/Package.h"
+#include "UObject/UObjectAnnotation.h"
+#include "UObject/UObjectGlobals.h"
+#include "UObject/UObjectHash.h"
 #include "Engine/EngineTypes.h"
 #include "Components/ActorComponent.h"
 #include "Components/SceneComponent.h"
 #include "GameFramework/Actor.h"
 
+class FComponentPropertyWriter : public FObjectWriter
+{
+public:
+	FComponentPropertyWriter(const UActorComponent* InComponent, TArray<uint8>& InBytes, TArray<UObject*>& InInstancedObjects)
+		: FObjectWriter(InBytes)
+		, Component(InComponent)
+		, InstancedObjects(InInstancedObjects)
+	{
+		// Include properties that would normally skip tagged serialization (e.g. bulk serialization of array properties).
+		ArPortFlags |= PPF_ForceTaggedSerialization;
+
+		if (Component)
+		{
+			UClass* ComponentClass = Component->GetClass();
+
+			Component->GetUCSModifiedProperties(PropertiesToSkip);
+
+			if (AActor* ComponentOwner = Component->GetOwner())
+			{
+				// If this is the owning Actor's root scene component, don't include relative transform properties. This is handled elsewhere.
+				if (Component == ComponentOwner->GetRootComponent())
+				{
+					PropertiesToSkip.Add(ComponentClass->FindPropertyByName(GET_MEMBER_NAME_CHECKED(USceneComponent, RelativeLocation)));
+					PropertiesToSkip.Add(ComponentClass->FindPropertyByName(GET_MEMBER_NAME_CHECKED(USceneComponent, RelativeRotation)));
+					PropertiesToSkip.Add(ComponentClass->FindPropertyByName(GET_MEMBER_NAME_CHECKED(USceneComponent, RelativeScale3D)));
+				}
+			}
+
+			ComponentClass->SerializeTaggedProperties(*this, (uint8*)Component, ComponentClass, (uint8*)Component->GetArchetype());
+		}
+	}
+
+	virtual ~FComponentPropertyWriter()
+	{
+		DuplicatedObjectAnnotation.RemoveAllAnnotations();
+	}
+
+	virtual bool ShouldSkipProperty(const UProperty* InProperty) const override
+	{
+		return (	InProperty->HasAnyPropertyFlags(CPF_Transient)
+			|| !InProperty->HasAnyPropertyFlags(CPF_Edit | CPF_Interp)
+			|| PropertiesToSkip.Contains(InProperty));
+	}
+
+
+	UObject* GetDuplicatedObject(UObject* Object)
+	{
+		UObject* Result = Object;
+		if (IsValid(Object))
+		{
+			// Check for an existing duplicate of the object.
+			FDuplicatedObject DupObjectInfo = DuplicatedObjectAnnotation.GetAnnotation( Object );
+			if( !DupObjectInfo.IsDefault() )
+			{
+				Result = DupObjectInfo.DuplicatedObject;
+			}
+			else if (Object->GetOuter() == Component)
+			{
+				Result = DuplicateObject(Object, GetTransientPackage());
+				InstancedObjects.Add(Result);
+			}
+			else
+			{
+				check(Object->IsIn(Component));
+
+				// Check to see if the object's outer is being duplicated.
+				UObject* DupOuter = GetDuplicatedObject(Object->GetOuter());
+				if (DupOuter != nullptr)
+				{
+					// First check if the duplicated outer already has an allocated duplicate of this object
+					Result = static_cast<UObject*>(FindObjectWithOuter(DupOuter, Object->GetClass(), Object->GetFName()));
+
+					if (Result == nullptr)
+					{
+						// The object's outer is being duplicated, create a duplicate of this object.
+						Result = DuplicateObject(Object, DupOuter);
+					}
+
+					DuplicatedObjectAnnotation.AddAnnotation( Object, FDuplicatedObject( Result ) );
+				}
+			}
+		}
+
+		return Result;
+	}
+
+	virtual FArchive& operator<<(UObject*& Object) override
+	{
+		UObject* SerializedObject = Object;
+		if (Object && Object->IsIn(Component))
+		{
+			SerializedObject = GetDuplicatedObject(Object);
+		}
+
+		// store the pointer to this object
+		Serialize(&SerializedObject, sizeof(UObject*));
+
+		return *this;
+	}
+
+private:
+
+	const UActorComponent* Component;
+
+	TSet<const UProperty*> PropertiesToSkip;
+	
+	TArray<UObject*>& InstancedObjects;
+
+	FUObjectAnnotationSparse<FDuplicatedObject,false> DuplicatedObjectAnnotation;
+};
+
+class FComponentPropertyReader : public FObjectReader
+{
+public:
+	FComponentPropertyReader(UActorComponent* InComponent, TArray<uint8>& InBytes)
+		: FObjectReader(InBytes)
+	{
+		// Include properties that would normally skip tagged serialization (e.g. bulk serialization of array properties).
+		ArPortFlags |= PPF_ForceTaggedSerialization;
+
+		InComponent->GetUCSModifiedProperties(PropertiesToSkip);
+
+		UClass* Class = InComponent->GetClass();
+		Class->SerializeTaggedProperties(*this, (uint8*)InComponent, Class, nullptr);
+	}
+
+	virtual bool ShouldSkipProperty(const UProperty* InProperty) const override
+	{
+		return PropertiesToSkip.Contains(InProperty);
+	}
+
+	TSet<const UProperty*> PropertiesToSkip;
+};
+
 FActorComponentInstanceData::FActorComponentInstanceData()
 	: SourceComponentTemplate(nullptr)
 	, SourceComponentTypeSerializedIndex(-1)
@@ -55,60 +194,20 @@ FActorComponentInstanceData::FActorComponentInstanceData(const UActorComponent*
 
 	if (SourceComponent->IsEditableWhenInherited())
 	{
-		class FComponentPropertyWriter : public FObjectWriter
-		{
-		public:
-			FComponentPropertyWriter(const UActorComponent* Component, TArray<uint8>& InBytes)
-				: FObjectWriter(InBytes)
-			{
-				// Include properties that would normally skip tagged serialization (e.g. bulk serialization of array properties).
-				ArPortFlags |= PPF_ForceTaggedSerialization;
-
-				if (Component)
-				{
-					UClass* ComponentClass = Component->GetClass();
-
-					Component->GetUCSModifiedProperties(PropertiesToSkip);
-
-					if (AActor* ComponentOwner = Component->GetOwner())
-					{
-						// If this is the owning Actor's root scene component, don't include relative transform properties. This is handled elsewhere.
-						if (Component == ComponentOwner->GetRootComponent())
-						{
-							PropertiesToSkip.Add(ComponentClass->FindPropertyByName(GET_MEMBER_NAME_CHECKED(USceneComponent, RelativeLocation)));
-							PropertiesToSkip.Add(ComponentClass->FindPropertyByName(GET_MEMBER_NAME_CHECKED(USceneComponent, RelativeRotation)));
-							PropertiesToSkip.Add(ComponentClass->FindPropertyByName(GET_MEMBER_NAME_CHECKED(USceneComponent, RelativeScale3D)));
-						}
-					}
-
-					ComponentClass->SerializeTaggedProperties(*this, (uint8*)Component, ComponentClass, (uint8*)Component->GetArchetype());
-				}
-			}
-
-			virtual bool ShouldSkipProperty(const UProperty* InProperty) const override
-			{
-				return (	InProperty->HasAnyPropertyFlags(CPF_Transient | CPF_ContainsInstancedReference | CPF_InstancedReference)
-						|| !InProperty->HasAnyPropertyFlags(CPF_Edit | CPF_Interp)
-						|| PropertiesToSkip.Contains(InProperty));
-			}
-
-		private:
-			TSet<const UProperty*> PropertiesToSkip;
-
-		};
-
-		FComponentPropertyWriter ComponentPropertyWriter(SourceComponent, SavedProperties);
+		FComponentPropertyWriter ComponentPropertyWriter(SourceComponent, SavedProperties, InstancedObjects);
 
 		// Cache off the length of an array that will come from SerializeTaggedProperties that had no properties saved in to it.
 		auto GetSizeOfEmptyArchive = [](const UActorComponent* DummyComponent) -> int32
 		{
 			TArray<uint8> NoWrittenPropertyReference;
-			FComponentPropertyWriter NullWriter(nullptr, NoWrittenPropertyReference);
+			TArray<UObject*> NoInstances;
+			FComponentPropertyWriter NullWriter(nullptr, NoWrittenPropertyReference, NoInstances);
 			UClass* ComponentClass = DummyComponent->GetClass();
 			
 			// By serializing the component with itself as its defaults we guarantee that no properties will be written out
 			ComponentClass->SerializeTaggedProperties(NullWriter, (uint8*)DummyComponent, ComponentClass, (uint8*)DummyComponent);
 
+			check(NoInstances.Num() == 0);
 			return NoWrittenPropertyReference.Num();
 		};
 
@@ -167,29 +266,12 @@ void FActorComponentInstanceData::ApplyToComponent(UActorComponent* Component, c
 	{
 		Component->DetermineUCSModifiedProperties();
 
-		class FComponentPropertyReader : public FObjectReader
+		for (UObject* InstancedObject : InstancedObjects)
 		{
-		public:
-			FComponentPropertyReader(UActorComponent* InComponent, TArray<uint8>& InBytes)
-				: FObjectReader(InBytes)
-			{
-				// Include properties that would normally skip tagged serialization (e.g. bulk serialization of array properties).
-				ArPortFlags |= PPF_ForceTaggedSerialization;
-
-				InComponent->GetUCSModifiedProperties(PropertiesToSkip);
-
-				UClass* Class = InComponent->GetClass();
-				Class->SerializeTaggedProperties(*this, (uint8*)InComponent, Class, nullptr);
-			}
-
-			virtual bool ShouldSkipProperty(const UProperty* InProperty) const override
-			{
-				return PropertiesToSkip.Contains(InProperty);
-			}
-
-			TSet<const UProperty*> PropertiesToSkip;
+			InstancedObject->Rename(nullptr, Component);
+		}
 
-		} ComponentPropertyReader(Component, SavedProperties);
+		FComponentPropertyReader ComponentPropertyReader(Component, SavedProperties);
 
 		if (Component->IsRegistered())
 		{
@@ -201,6 +283,7 @@ void FActorComponentInstanceData::ApplyToComponent(UActorComponent* Component, c
 void FActorComponentInstanceData::AddReferencedObjects(FReferenceCollector& Collector)
 {
 	Collector.AddReferencedObject(SourceComponentTemplate);
+	Collector.AddReferencedObjects(InstancedObjects);
 }
 
 FComponentInstanceDataCache::FComponentInstanceDataCache(const AActor* Actor)
diff --git a/Engine/Source/Runtime/Engine/Private/Components/ActorComponent.cpp b/Engine/Source/Runtime/Engine/Private/Components/ActorComponent.cpp
--- a/Engine/Source/Runtime/Engine/Private/Components/ActorComponent.cpp
+++ b/Engine/Source/Runtime/Engine/Private/Components/ActorComponent.cpp
@@ -1684,7 +1684,7 @@ void UActorComponent::DetermineUCSModifiedProperties()
 
 			virtual bool ShouldSkipProperty(const UProperty* InProperty) const override
 			{
-				return (    InProperty->HasAnyPropertyFlags(CPF_Transient | CPF_ContainsInstancedReference | CPF_InstancedReference)
+				return (    InProperty->HasAnyPropertyFlags(CPF_Transient)
 						|| !InProperty->HasAnyPropertyFlags(CPF_Edit | CPF_Interp));
 			}
 		} PropertySkipper;
diff --git a/Engine/Source/Runtime/Engine/Private/Components/AudioComponent.cpp b/Engine/Source/Runtime/Engine/Private/Components/AudioComponent.cpp
--- a/Engine/Source/Runtime/Engine/Private/Components/AudioComponent.cpp
+++ b/Engine/Source/Runtime/Engine/Private/Components/AudioComponent.cpp
@@ -193,7 +193,7 @@ void UAudioComponent::OnUpdateTransform(EUpdateTransformFlags UpdateTransformFla
 			DECLARE_CYCLE_STAT(TEXT("FAudioThreadTask.UpdateTransform"), STAT_AudioUpdateTransform, STATGROUP_AudioThreadCommands);
 
 			const uint64 MyAudioComponentID = AudioComponentID;
-			const FTransform& MyTransform = ComponentToWorld;
+			const FTransform& MyTransform = GetComponentTransform();
 
 			FAudioThread::RunCommandOnAudioThread([AudioDevice, MyAudioComponentID, MyTransform]()
 			{
@@ -237,7 +237,7 @@ void UAudioComponent::PlayInternal(const float StartTime, const float FadeInDura
 
 			float MaxDistance = 0.0f;
 			float FocusFactor = 0.0f;
-			FVector Location = ComponentToWorld.GetLocation();
+			FVector Location = GetComponentTransform().GetLocation();
 
 			AudioDevice->GetMaxDistanceAndFocusFactor(Sound, World, Location, AttenuationSettingsToApply, MaxDistance, FocusFactor);
 
@@ -288,7 +288,7 @@ void UAudioComponent::PlayInternal(const float StartTime, const float FadeInDura
 
 			if (NewActiveSound.bLocationDefined)
 			{
-				NewActiveSound.Transform = ComponentToWorld;
+				NewActiveSound.Transform = GetComponentTransform();
 			}
 
 			NewActiveSound.bAllowSpatialization = bAllowSpatialization;
diff --git a/Engine/Source/Runtime/Engine/Private/Components/BillboardComponent.cpp b/Engine/Source/Runtime/Engine/Private/Components/BillboardComponent.cpp
--- a/Engine/Source/Runtime/Engine/Private/Components/BillboardComponent.cpp
+++ b/Engine/Source/Runtime/Engine/Private/Components/BillboardComponent.cpp
@@ -53,7 +53,7 @@ public:
 		bWillEverBeLit = false;
 
 		// Calculate the scale factor for the sprite.
-		float Scale = InComponent->ComponentToWorld.GetMaximumAxisScale();
+		float Scale = InComponent->GetComponentTransform().GetMaximumAxisScale();
 
 		if(InComponent->Sprite)
 		{
@@ -323,7 +323,7 @@ bool UBillboardComponent::ComponentIsTouchingSelectionBox(const FBox& InSelBBox,
 
 	if (!bConsiderOnlyBSP && ShowFlags.BillboardSprites && Sprite != nullptr && Actor != nullptr)
 	{
-		const float Scale = ComponentToWorld.GetMaximumAxisScale();
+		const float Scale = GetComponentTransform().GetMaximumAxisScale();
 
 		// Construct a box representing the sprite
 		const FBox SpriteBox(
@@ -349,7 +349,7 @@ bool UBillboardComponent::ComponentIsTouchingSelectionFrustum(const FConvexVolum
 
 	if (!bConsiderOnlyBSP && ShowFlags.BillboardSprites && Sprite != nullptr && Actor != nullptr)
 	{
-		const float Scale = ComponentToWorld.GetMaximumAxisScale();
+		const float Scale = GetComponentTransform().GetMaximumAxisScale();
 		const float MaxExtent = FMath::Max(Sprite->GetSizeX(), Sprite->GetSizeY());
 		const FVector Extent = Scale * MaxExtent * FVector(0.5f, 0.5f, 0.0f);
 
diff --git a/Engine/Source/Runtime/Engine/Private/Components/BoxComponent.cpp b/Engine/Source/Runtime/Engine/Private/Components/BoxComponent.cpp
--- a/Engine/Source/Runtime/Engine/Private/Components/BoxComponent.cpp
+++ b/Engine/Source/Runtime/Engine/Private/Components/BoxComponent.cpp
@@ -31,7 +31,7 @@ void UBoxComponent::SetBoxExtent(FVector NewBoxExtent, bool bUpdateOverlaps)
 	if (bPhysicsStateCreated)
 	{
 		// Update physics engine collision shapes
-		BodyInstance.UpdateBodyScale(ComponentToWorld.GetScale3D(), true);
+		BodyInstance.UpdateBodyScale(GetComponentTransform().GetScale3D(), true);
 
 		if ( bUpdateOverlaps && IsCollisionEnabled() && GetOwner() )
 		{
diff --git a/Engine/Source/Runtime/Engine/Private/Components/BrushComponent.cpp b/Engine/Source/Runtime/Engine/Private/Components/BrushComponent.cpp
--- a/Engine/Source/Runtime/Engine/Private/Components/BrushComponent.cpp
+++ b/Engine/Source/Runtime/Engine/Private/Components/BrushComponent.cpp
@@ -628,7 +628,7 @@ bool UBrushComponent::ComponentIsTouchingSelectionBox(const FBox& InSelBBox, con
 				{
 					for (const auto& Vertex : Poly.Vertices)
 					{
-						const FVector Location = ComponentToWorld.TransformPosition(Vertex);
+						const FVector Location = GetComponentTransform().TransformPosition(Vertex);
 						const bool bLocationIntersected = FMath::PointBoxIntersection(Location, InSelBBox);
 
 						// If the selection box doesn't have to encompass the entire component and any poly vertex intersects with the selection
@@ -649,10 +649,10 @@ bool UBrushComponent::ComponentIsTouchingSelectionBox(const FBox& InSelBBox, con
 					const int32 NumVerts = Poly.Vertices.Num();
 					if (NumVerts > 0)
 					{
-						FVector StartVert = ComponentToWorld.TransformPosition(Poly.Vertices[NumVerts - 1]);
+						FVector StartVert = GetComponentTransform().TransformPosition(Poly.Vertices[NumVerts - 1]);
 						for (int32 Index = 0; Index < NumVerts; ++Index)
 						{
-							const FVector EndVert = ComponentToWorld.TransformPosition(Poly.Vertices[Index]);
+							const FVector EndVert = GetComponentTransform().TransformPosition(Poly.Vertices[Index]);
 
 							if (FMath::LineBoxIntersection(InSelBBox, StartVert, EndVert, EndVert - StartVert))
 							{
@@ -675,7 +675,7 @@ bool UBrushComponent::ComponentIsTouchingSelectionBox(const FBox& InSelBBox, con
 				// The component must be entirely within the bounding box...
 				for (const auto& Vertex : Poly.Vertices)
 				{
-					const FVector Location = ComponentToWorld.TransformPosition(Vertex);
+					const FVector Location = GetComponentTransform().TransformPosition(Vertex);
 					const bool bLocationIntersected = FMath::PointBoxIntersection(Location, InSelBBox);
 
 					// If the selection box has to encompass the entire component and a poly vertex didn't intersect with the selection
@@ -711,7 +711,7 @@ bool UBrushComponent::ComponentIsTouchingSelectionFrustum(const FConvexVolume& I
 		{
 			for (const auto& Vertex : Poly.Vertices)
 			{
-				const FVector Location = ComponentToWorld.TransformPosition(Vertex);
+				const FVector Location = GetComponentTransform().TransformPosition(Vertex);
 				const bool bIntersect = InFrustum.IntersectSphere(Location, 0.0f);
 
 				if (bIntersect && !bMustEncompassEntireComponent)
diff --git a/Engine/Source/Runtime/Engine/Private/Components/CapsuleComponent.cpp b/Engine/Source/Runtime/Engine/Private/Components/CapsuleComponent.cpp
--- a/Engine/Source/Runtime/Engine/Private/Components/CapsuleComponent.cpp
+++ b/Engine/Source/Runtime/Engine/Private/Components/CapsuleComponent.cpp
@@ -97,9 +97,9 @@ FBoxSphereBounds UCapsuleComponent::CalcBounds(const FTransform& LocalToWorld) c
 
 void UCapsuleComponent::CalcBoundingCylinder(float& CylinderRadius, float& CylinderHalfHeight) const 
 {
-	const float Scale = ComponentToWorld.GetMaximumAxisScale();
+	const float Scale = GetComponentTransform().GetMaximumAxisScale();
 	const float CapsuleEndCapCenter = FMath::Max(CapsuleHalfHeight - CapsuleRadius, 0.f);
-	const FVector ZAxis = ComponentToWorld.TransformVectorNoScale(FVector(0.f, 0.f, CapsuleEndCapCenter * Scale));
+	const FVector ZAxis = GetComponentTransform().TransformVectorNoScale(FVector(0.f, 0.f, CapsuleEndCapCenter * Scale));
 	
 	const float ScaledRadius = CapsuleRadius * Scale;
 	
@@ -165,7 +165,7 @@ void UCapsuleComponent::SetCapsuleSize(float NewRadius, float NewHalfHeight, boo
 	if (bPhysicsStateCreated)
 	{
 		// Update physics engine collision shapes
-		BodyInstance.UpdateBodyScale(ComponentToWorld.GetScale3D(), true);
+		BodyInstance.UpdateBodyScale(GetComponentTransform().GetScale3D(), true);
 
 		if ( bUpdateOverlaps && IsCollisionEnabled() && GetOwner() )
 		{
diff --git a/Engine/Source/Runtime/Engine/Private/Components/CharacterMovementComponent.cpp b/Engine/Source/Runtime/Engine/Private/Components/CharacterMovementComponent.cpp
--- a/Engine/Source/Runtime/Engine/Private/Components/CharacterMovementComponent.cpp
+++ b/Engine/Source/Runtime/Engine/Private/Components/CharacterMovementComponent.cpp
@@ -7408,12 +7408,7 @@ void UCharacterMovementComponent::ForcePositionUpdate(float DeltaTime)
 
 FNetworkPredictionData_Client* UCharacterMovementComponent::GetPredictionData_Client() const
 {
-	// Should only be called on client or listen server (for remote clients) in network games
-	check(CharacterOwner != NULL);
-	checkSlow(CharacterOwner->Role < ROLE_Authority || (CharacterOwner->GetRemoteRole() == ROLE_AutonomousProxy && GetNetMode() == NM_ListenServer));
-	checkSlow(GetNetMode() == NM_Client || GetNetMode() == NM_ListenServer);
-
-	if (!ClientPredictionData)
+	if (ClientPredictionData == nullptr)
 	{
 		UCharacterMovementComponent* MutableThis = const_cast<UCharacterMovementComponent*>(this);
 		MutableThis->ClientPredictionData = new FNetworkPredictionData_Client_Character(*this);
@@ -7424,12 +7419,7 @@ FNetworkPredictionData_Client* UCharacterMovementComponent::GetPredictionData_Cl
 
 FNetworkPredictionData_Server* UCharacterMovementComponent::GetPredictionData_Server() const
 {
-	// Should only be called on server in network games
-	check(CharacterOwner != NULL);
-	check(CharacterOwner->Role == ROLE_Authority);
-	checkSlow(GetNetMode() < NM_Client);
-
-	if (!ServerPredictionData)
+	if (ServerPredictionData == nullptr)
 	{
 		UCharacterMovementComponent* MutableThis = const_cast<UCharacterMovementComponent*>(this);
 		MutableThis->ServerPredictionData = new FNetworkPredictionData_Server_Character(*this);
@@ -7441,13 +7431,35 @@ FNetworkPredictionData_Server* UCharacterMovementComponent::GetPredictionData_Se
 
 FNetworkPredictionData_Client_Character* UCharacterMovementComponent::GetPredictionData_Client_Character() const
 {
-	return static_cast<class FNetworkPredictionData_Client_Character*>(GetPredictionData_Client());
+	// Should only be called on client or listen server (for remote clients) in network games
+	checkSlow(CharacterOwner != NULL);
+	checkSlow(CharacterOwner->Role < ROLE_Authority || (CharacterOwner->GetRemoteRole() == ROLE_AutonomousProxy && GetNetMode() == NM_ListenServer));
+	checkSlow(GetNetMode() == NM_Client || GetNetMode() == NM_ListenServer);
+
+	if (ClientPredictionData == nullptr)
+	{
+		UCharacterMovementComponent* MutableThis = const_cast<UCharacterMovementComponent*>(this);
+		MutableThis->ClientPredictionData = static_cast<class FNetworkPredictionData_Client_Character*>(GetPredictionData_Client());
+	}
+
+	return ClientPredictionData;
 }
 
 
 FNetworkPredictionData_Server_Character* UCharacterMovementComponent::GetPredictionData_Server_Character() const
 {
-	return static_cast<class FNetworkPredictionData_Server_Character*>(GetPredictionData_Server());
+	// Should only be called on server in network games
+	checkSlow(CharacterOwner != NULL);
+	checkSlow(CharacterOwner->Role == ROLE_Authority);
+	checkSlow(GetNetMode() < NM_Client);
+
+	if (ServerPredictionData == nullptr)
+	{
+		UCharacterMovementComponent* MutableThis = const_cast<UCharacterMovementComponent*>(this);
+		MutableThis->ServerPredictionData = static_cast<class FNetworkPredictionData_Server_Character*>(GetPredictionData_Server());
+	}
+
+	return ServerPredictionData;
 }
 
 bool UCharacterMovementComponent::HasPredictionData_Client() const
diff --git a/Engine/Source/Runtime/Engine/Private/Components/ChildActorComponent.cpp b/Engine/Source/Runtime/Engine/Private/Components/ChildActorComponent.cpp
--- a/Engine/Source/Runtime/Engine/Private/Components/ChildActorComponent.cpp
+++ b/Engine/Source/Runtime/Engine/Private/Components/ChildActorComponent.cpp
@@ -120,6 +120,30 @@ void UChildActorComponent::Serialize(FArchive& Ar)
 }
 
 #if WITH_EDITOR
+void UChildActorComponent::PostEditImport()
+{
+	Super::PostEditImport();
+
+	if (IsTemplate())
+	{
+		TArray<UObject*> Children;
+		GetObjectsWithOuter(this, Children, false);
+
+		for (UObject* Child : Children)
+		{
+			if (Child->GetClass() == ChildActorClass)
+			{
+				ChildActorTemplate = CastChecked<AActor>(Child);
+				break;
+			}
+		}
+	}
+	else
+	{
+		ChildActorTemplate = CastChecked<UChildActorComponent>(GetArchetype())->ChildActorTemplate;
+	}
+}
+
 void UChildActorComponent::PostEditChangeProperty(FPropertyChangedEvent& PropertyChangedEvent)
 {
 	if (PropertyChangedEvent.Property && PropertyChangedEvent.Property->GetFName() == GET_MEMBER_NAME_CHECKED(UChildActorComponent, ChildActorClass))
@@ -354,7 +378,7 @@ void UChildActorComponent::ApplyComponentInstanceData(FChildActorComponentInstan
 		USceneComponent* ChildActorRoot = ChildActor->GetRootComponent();
 		if (ChildActorRoot)
 		{
-			for (const auto& AttachInfo : ChildActorInstanceData->AttachedActors)
+			for (const FChildActorComponentInstanceData::FAttachedActorInfo& AttachInfo : ChildActorInstanceData->AttachedActors)
 			{
 				AActor* AttachedActor = AttachInfo.Actor.Get();
 				if (AttachedActor)
@@ -362,7 +386,7 @@ void UChildActorComponent::ApplyComponentInstanceData(FChildActorComponentInstan
 					USceneComponent* AttachedRootComponent = AttachedActor->GetRootComponent();
 					if (AttachedRootComponent)
 					{
-						AttachedActor->DetachRootComponentFromParent();
+						AttachedActor->DetachFromActor(FDetachmentTransformRules::KeepWorldTransform);
 						AttachedRootComponent->AttachToComponent(ChildActorRoot, FAttachmentTransformRules::KeepWorldTransform, AttachInfo.SocketName);
 						AttachedRootComponent->SetRelativeTransform(AttachInfo.RelativeTransform);
 						AttachedRootComponent->UpdateComponentToWorld();
@@ -487,7 +511,10 @@ void UChildActorComponent::CreateChildActor()
 				Params.bAllowDuringConstructionScript = true;
 				Params.OverrideLevel = (MyOwner ? MyOwner->GetLevel() : nullptr);
 				Params.Name = ChildActorName;
-				Params.Template = ChildActorTemplate;
+				if (ChildActorTemplate && ChildActorTemplate->GetClass() == ChildActorClass)
+				{
+					Params.Template = ChildActorTemplate;
+				}
 				Params.ObjectFlags |= (RF_TextExportTransient | RF_NonPIEDuplicateTransient);
 				if (!HasAllFlags(RF_Transactional))
 				{
@@ -510,7 +537,7 @@ void UChildActorComponent::CreateChildActor()
 
 					// Parts that we deferred from SpawnActor
 					const FComponentInstanceDataCache* ComponentInstanceData = (CachedInstanceData ? CachedInstanceData->ComponentInstanceData : nullptr);
-					ChildActor->FinishSpawning(ComponentToWorld, false, ComponentInstanceData);
+					ChildActor->FinishSpawning(GetComponentTransform(), false, ComponentInstanceData);
 
 					ChildActor->AttachToComponent(this, FAttachmentTransformRules::SnapToTargetNotIncludingScale);
 
diff --git a/Engine/Source/Runtime/Engine/Private/Components/DestructibleComponent.cpp b/Engine/Source/Runtime/Engine/Private/Components/DestructibleComponent.cpp
--- a/Engine/Source/Runtime/Engine/Private/Components/DestructibleComponent.cpp
+++ b/Engine/Source/Runtime/Engine/Private/Components/DestructibleComponent.cpp
@@ -143,7 +143,7 @@ void UDestructibleComponent::OnUpdateTransform(EUpdateTransformFlags UpdateTrans
 		return;
 	}
 
-	const FTransform& CurrentLocalToWorld = ComponentToWorld;
+	const FTransform& CurrentLocalToWorld = GetComponentTransform();
 
 #if !(UE_BUILD_SHIPPING)
 	if(CurrentLocalToWorld.ContainsNaN())
@@ -253,9 +253,9 @@ void UDestructibleComponent::OnCreatePhysicsState()
 	// Get the default actor descriptor NvParameterized data from the asset
 	NvParameterized::Interface* ActorParams = TheDestructibleMesh->GetDestructibleActorDesc(PhysMat);
 
-	// Create PhysX transforms from ComponentToWorld
-	const PxTransform GlobalPose(U2PVector(ComponentToWorld.GetTranslation()), U2PQuat(ComponentToWorld.GetRotation()));
-	const PxVec3 Scale = U2PVector(ComponentToWorld.GetScale3D());
+	// Create PhysX transforms from GetComponentTransform()
+	const PxTransform GlobalPose(U2PVector(GetComponentTransform().GetTranslation()), U2PQuat(GetComponentTransform().GetRotation()));
+	const PxVec3 Scale = U2PVector(GetComponentTransform().GetScale3D());
 
 	// Set the transform in the actor descriptor
 	verify( NvParameterized::setParamTransform(*ActorParams,"globalPose",GlobalPose) );
@@ -985,7 +985,7 @@ void UDestructibleComponent::UpdateDestructibleChunkTM(const TArray<const PxRigi
 		}
 		else
 		{
-			//if we haven't fractured it must mean that we're simulating a destructible and so we should update our ComponentToWorld based on the single rigid body
+			//if we haven't fractured it must mean that we're simulating a destructible and so we should update our GetComponentTransform() based on the single rigid body
 			DestructibleComponent->SyncComponentToRBPhysics();
 		}
 
@@ -998,7 +998,7 @@ void UDestructibleComponent::UpdateDestructibleChunkTM(const TArray<const PxRigi
 
 void UDestructibleComponent::SetChunksWorldTM(const TArray<FUpdateChunksInfo>& UpdateInfos)
 {
-	const FQuat InvRotation = ComponentToWorld.GetRotation().Inverse();
+	const FQuat InvRotation = GetComponentTransform().GetRotation().Inverse();
 
 	for (const FUpdateChunksInfo& UpdateInfo : UpdateInfos)
 	{
@@ -1008,7 +1008,7 @@ void UDestructibleComponent::SetChunksWorldTM(const TArray<FUpdateChunksInfo>& U
 		const FQuat WorldRotation		= UpdateInfo.WorldTM.GetRotation();
 
 		const FQuat BoneRotation = InvRotation*WorldRotation;
-		const FVector BoneTranslation = InvRotation.RotateVector(WorldTranslation - ComponentToWorld.GetTranslation()) / ComponentToWorld.GetScale3D();
+		const FVector BoneTranslation = InvRotation.RotateVector(WorldTranslation - GetComponentTransform().GetTranslation()) / GetComponentTransform().GetScale3D();
 
 		GetEditableComponentSpaceTransforms()[BoneIndex] = FTransform(BoneRotation, BoneTranslation);
 	}
@@ -1035,14 +1035,14 @@ void UDestructibleComponent::SetChunkWorldRT( int32 ChunkIndex, const FQuat& Wor
 	MarkRenderDynamicDataDirty();
 
 #if 0
-	// Scale is already applied to the ComponentToWorld transform, and is carried into the bones _locally_.
+	// Scale is already applied to the GetComponentTransform() transform, and is carried into the bones _locally_.
 	// So there is no need to set scale in the bone local transforms
-	const FTransform WorldRT(WorldRotation, WorldTranslation, ComponentToWorld.GetScale3D());
-	SpaceBases(BoneIndex) = WorldRT*ComponentToWorld.Inverse();
+	const FTransform WorldRT(WorldRotation, WorldTranslation, GetComponentTransform().GetScale3D());
+	SpaceBases(BoneIndex) = WorldRT*GetComponentTransform().Inverse();
 #elif 1
 	// More optimal form of the above
-	const FQuat BoneRotation = ComponentToWorld.GetRotation().Inverse()*WorldRotation;
-	const FVector BoneTranslation = ComponentToWorld.GetRotation().Inverse().RotateVector(WorldTranslation - ComponentToWorld.GetTranslation())/ComponentToWorld.GetScale3D();
+	const FQuat BoneRotation = GetComponentTransform().GetRotation().Inverse()*WorldRotation;
+	const FVector BoneTranslation = GetComponentTransform().GetRotation().Inverse().RotateVector(WorldTranslation - GetComponentTransform().GetTranslation())/GetComponentTransform().GetScale3D();
 	GetEditableComponentSpaceTransforms()[BoneIndex] = FTransform(BoneRotation, BoneTranslation);
 #endif
 }
@@ -1155,7 +1155,7 @@ bool UDestructibleComponent::DoCustomNavigableGeometryExport(FNavigableGeometryE
 #if WITH_EDITORONLY_DATA
 		if (DestructibleMesh && DestructibleMesh->SourceStaticMesh)
 		{
-			GeomExport.ExportRigidBodySetup(*DestructibleMesh->SourceStaticMesh->BodySetup, ComponentToWorld);
+			GeomExport.ExportRigidBodySetup(*DestructibleMesh->SourceStaticMesh->BodySetup, GetComponentTransform());
 			bExportFromBodySetup = false;
 		}
 #endif	// WITH_EDITORONLY_DATA
@@ -1165,7 +1165,7 @@ bool UDestructibleComponent::DoCustomNavigableGeometryExport(FNavigableGeometryE
 
 	apex::DestructibleActor* DestrActor = const_cast<apex::DestructibleActor*>(ApexDestructibleActor);
 
-	const FTransform ComponentToWorldNoScale(ComponentToWorld.GetRotation(), ComponentToWorld.GetTranslation(), FVector(1.f));
+	const FTransform ComponentToWorldNoScale(GetComponentTransform().GetRotation(), GetComponentTransform().GetTranslation(), FVector(1.f));
 	TArray<PxShape*> Shapes;
 	Shapes.AddUninitialized(8);
 	PxRigidDynamic** PActorBuffer = NULL;
diff --git a/Engine/Source/Runtime/Engine/Private/Components/LightComponent.cpp b/Engine/Source/Runtime/Engine/Private/Components/LightComponent.cpp
--- a/Engine/Source/Runtime/Engine/Private/Components/LightComponent.cpp
+++ b/Engine/Source/Runtime/Engine/Private/Components/LightComponent.cpp
@@ -904,7 +904,7 @@ void ULightComponent::SetIESTexture(UTextureLightProfile* NewValue)
 // GetDirection
 FVector ULightComponent::GetDirection() const 
 { 
-	return ComponentToWorld.GetUnitAxis( EAxis::X );
+	return GetComponentTransform().GetUnitAxis( EAxis::X );
 }
 
 void ULightComponent::UpdateColorAndBrightness()
@@ -977,7 +977,7 @@ class FPrecomputedLightInstanceData : public FSceneComponentInstanceData
 public:
 	FPrecomputedLightInstanceData(const ULightComponent* SourceComponent)
 		: FSceneComponentInstanceData(SourceComponent)
-		, Transform(SourceComponent->ComponentToWorld)
+		, Transform(SourceComponent->GetComponentTransform())
 		, LightGuid(SourceComponent->LightGuid)
 		, PreviewShadowMapChannel(SourceComponent->PreviewShadowMapChannel)
 	{}
@@ -1003,7 +1003,7 @@ void ULightComponent::ApplyComponentInstanceData(FPrecomputedLightInstanceData*
 {
 	check(LightMapData);
 
-	if (!LightMapData->Transform.Equals(ComponentToWorld))
+	if (!LightMapData->Transform.Equals(GetComponentTransform()))
 	{
 		return;
 	}
diff --git a/Engine/Source/Runtime/Engine/Private/Components/ModelComponent.cpp b/Engine/Source/Runtime/Engine/Private/Components/ModelComponent.cpp
--- a/Engine/Source/Runtime/Engine/Private/Components/ModelComponent.cpp
+++ b/Engine/Source/Runtime/Engine/Private/Components/ModelComponent.cpp
@@ -449,7 +449,7 @@ void UModelComponent::GetStreamingTextureInfo(FStreamingTextureLevelContext& Lev
 					const FBspNode& Node = Model->Nodes[SurfaceNodes[NodeIndex]];
 					for(int32 VertexIndex = 0;VertexIndex < Node.NumVertices;VertexIndex++)
 					{
-						const FVector WorldVertex = ComponentToWorld.TransformPosition(Model->Points[Model->Verts[Node.iVertPool + VertexIndex].pVertex]);
+						const FVector WorldVertex = GetComponentTransform().TransformPosition(Model->Points[Model->Verts[Node.iVertPool + VertexIndex].pVertex]);
 						SurfaceVertices.Add(WorldVertex);
 					}
 				}
diff --git a/Engine/Source/Runtime/Engine/Private/Components/PointLightComponent.cpp b/Engine/Source/Runtime/Engine/Private/Components/PointLightComponent.cpp
--- a/Engine/Source/Runtime/Engine/Private/Components/PointLightComponent.cpp
+++ b/Engine/Source/Runtime/Engine/Private/Components/PointLightComponent.cpp
@@ -172,7 +172,7 @@ void UPointLightComponent::SetSourceLength(float NewValue)
 
 bool UPointLightComponent::AffectsBounds(const FBoxSphereBounds& InBounds) const
 {
-	if((InBounds.Origin - ComponentToWorld.GetLocation()).SizeSquared() > FMath::Square(AttenuationRadius + InBounds.SphereRadius))
+	if((InBounds.Origin - GetComponentTransform().GetLocation()).SizeSquared() > FMath::Square(AttenuationRadius + InBounds.SphereRadius))
 	{
 		return false;
 	}
@@ -199,7 +199,7 @@ void UPointLightComponent::SendRenderTransform_Concurrent()
 //
 FVector4 UPointLightComponent::GetLightPosition() const
 {
-	return FVector4(ComponentToWorld.GetLocation(),1);
+	return FVector4(GetComponentTransform().GetLocation(),1);
 }
 
 /**
@@ -232,7 +232,7 @@ FBox UPointLightComponent::GetBoundingBox() const
 
 FSphere UPointLightComponent::GetBoundingSphere() const
 {
-	return FSphere(ComponentToWorld.GetLocation(), AttenuationRadius);
+	return FSphere(GetComponentTransform().GetLocation(), AttenuationRadius);
 }
 
 void UPointLightComponent::Serialize(FArchive& Ar)
diff --git a/Engine/Source/Runtime/Engine/Private/Components/PoseableMeshComponent.cpp b/Engine/Source/Runtime/Engine/Private/Components/PoseableMeshComponent.cpp
--- a/Engine/Source/Runtime/Engine/Private/Components/PoseableMeshComponent.cpp
+++ b/Engine/Source/Runtime/Engine/Private/Components/PoseableMeshComponent.cpp
@@ -206,7 +206,7 @@ FTransform GetBoneTransformByNameHelper(FName BoneName, EBoneSpaces::Type BoneSp
 	}
 	else
 	{
-		return CSPose.GetComponentSpaceTransform(BoneIndex) * Component->ComponentToWorld;
+		return CSPose.GetComponentSpaceTransform(BoneIndex) * Component->GetComponentTransform();
 	}
 }
 
diff --git a/Engine/Source/Runtime/Engine/Private/Components/PostProcessComponent.cpp b/Engine/Source/Runtime/Engine/Private/Components/PostProcessComponent.cpp
--- a/Engine/Source/Runtime/Engine/Private/Components/PostProcessComponent.cpp
+++ b/Engine/Source/Runtime/Engine/Private/Components/PostProcessComponent.cpp
@@ -21,7 +21,7 @@ bool UPostProcessComponent::EncompassesPoint(FVector Point, float SphereRadius/*
 		FVector ClosestPoint;
 		float Distance = ParentShape->GetDistanceToCollision(Point, ClosestPoint);
 #else
-		FBoxSphereBounds Bounds = ParentShape->CalcBounds(ParentShape->ComponentToWorld);
+		FBoxSphereBounds Bounds = ParentShape->CalcBounds(ParentShape->GetComponentTransform());
 		float Distance = 0;
 		if (ParentShape->IsA<USphereComponent>())
 		{
diff --git a/Engine/Source/Runtime/Engine/Private/Components/PrimitiveComponent.cpp b/Engine/Source/Runtime/Engine/Private/Components/PrimitiveComponent.cpp
--- a/Engine/Source/Runtime/Engine/Private/Components/PrimitiveComponent.cpp
+++ b/Engine/Source/Runtime/Engine/Private/Components/PrimitiveComponent.cpp
@@ -614,7 +614,7 @@ void UPrimitiveComponent::OnCreatePhysicsState()
 		if(BodySetup)
 		{
 			// Create new BodyInstance at given location.
-			FTransform BodyTransform = ComponentToWorld;
+			FTransform BodyTransform = GetComponentTransform();
 
 			// Here we make sure we don't have zero scale. This still results in a body being made and placed in
 			// world (very small) but is consistent with a body scaled to zero.
@@ -687,8 +687,8 @@ void UPrimitiveComponent::OnUpdateTransform(EUpdateTransformFlags UpdateTransfor
 
 void UPrimitiveComponent::SendPhysicsTransform(ETeleportType Teleport)
 {
-	BodyInstance.SetBodyTransform(ComponentToWorld, Teleport);
-	BodyInstance.UpdateBodyScale(ComponentToWorld.GetScale3D());
+	BodyInstance.SetBodyTransform(GetComponentTransform(), Teleport);
+	BodyInstance.UpdateBodyScale(GetComponentTransform().GetScale3D());
 }
 
 void UPrimitiveComponent::OnDestroyPhysicsState()
@@ -728,8 +728,8 @@ void UPrimitiveComponent::SendRenderDebugPhysics(FPrimitiveSceneProxy* OverrideS
 					FPrimitiveSceneProxy::FDebugMassData& RootMassData = DebugMassData[0];
 					const FTransform MassToWorld = BI->GetMassSpaceToWorldSpace();
 
-					RootMassData.LocalCenterOfMass = ComponentToWorld.InverseTransformPosition(MassToWorld.GetLocation());
-					RootMassData.LocalTensorOrientation = MassToWorld.GetRotation() * ComponentToWorld.GetRotation().Inverse();
+					RootMassData.LocalCenterOfMass = GetComponentTransform().InverseTransformPosition(MassToWorld.GetLocation());
+					RootMassData.LocalTensorOrientation = MassToWorld.GetRotation() * GetComponentTransform().GetRotation().Inverse();
 					RootMassData.MassSpaceInertiaTensor = BI->GetBodyInertiaTensor();
 					RootMassData.BoneIndex = INDEX_NONE;
 				}
@@ -747,7 +747,7 @@ void UPrimitiveComponent::SendRenderDebugPhysics(FPrimitiveSceneProxy* OverrideS
 
 FMatrix UPrimitiveComponent::GetRenderMatrix() const
 {
-	return ComponentToWorld.ToMatrixWithScale();
+	return GetComponentTransform().ToMatrixWithScale();
 }
 
 void UPrimitiveComponent::Serialize(FArchive& Ar)
@@ -1723,7 +1723,7 @@ bool UPrimitiveComponent::MoveComponentImpl( const FVector& Delta, const FQuat&
 	const FVector TraceStart = GetComponentLocation();
 	const FVector TraceEnd = TraceStart + Delta;
 	float DeltaSizeSq = (TraceEnd - TraceStart).SizeSquared();				// Recalc here to account for precision loss of float addition
-	const FQuat InitialRotationQuat = ComponentToWorld.GetRotation();
+	const FQuat InitialRotationQuat = GetComponentTransform().GetRotation();
 
 	// ComponentSweepMulti does nothing if moving < KINDA_SMALL_NUMBER in distance, so it's important to not try to sweep distances smaller than that. 
 	const float MinMovementDistSq = (bSweep ? FMath::Square(4.f*KINDA_SMALL_NUMBER) : 0.f);
diff --git a/Engine/Source/Runtime/Engine/Private/Components/ReflectionCaptureComponent.cpp b/Engine/Source/Runtime/Engine/Private/Components/ReflectionCaptureComponent.cpp
--- a/Engine/Source/Runtime/Engine/Private/Components/ReflectionCaptureComponent.cpp
+++ b/Engine/Source/Runtime/Engine/Private/Components/ReflectionCaptureComponent.cpp
@@ -1310,7 +1310,7 @@ void UReflectionCaptureComponent::UpdatePreviewShape()
 {
 	if (CaptureOffsetComponent)
 	{
-		CaptureOffsetComponent->RelativeLocation = CaptureOffset / ComponentToWorld.GetScale3D();
+		CaptureOffsetComponent->RelativeLocation = CaptureOffset / GetComponentTransform().GetScale3D();
 	}
 }
 
@@ -1699,7 +1699,7 @@ void UBoxReflectionCaptureComponent::UpdatePreviewShape()
 {
 	if (PreviewCaptureBox)
 	{
-		PreviewCaptureBox->InitBoxExtent(((ComponentToWorld.GetScale3D() - FVector(BoxTransitionDistance)) / ComponentToWorld.GetScale3D()).ComponentMax(FVector::ZeroVector));
+		PreviewCaptureBox->InitBoxExtent(((GetComponentTransform().GetScale3D() - FVector(BoxTransitionDistance)) / GetComponentTransform().GetScale3D()).ComponentMax(FVector::ZeroVector));
 	}
 
 	Super::UpdatePreviewShape();
@@ -1707,7 +1707,7 @@ void UBoxReflectionCaptureComponent::UpdatePreviewShape()
 
 float UBoxReflectionCaptureComponent::GetInfluenceBoundingRadius() const
 {
-	return (ComponentToWorld.GetScale3D() + FVector(BoxTransitionDistance)).Size();
+	return (GetComponentTransform().GetScale3D() + FVector(BoxTransitionDistance)).Size();
 }
 
 #if WITH_EDITOR
@@ -1741,7 +1741,7 @@ void UPlaneReflectionCaptureComponent::UpdatePreviewShape()
 
 float UPlaneReflectionCaptureComponent::GetInfluenceBoundingRadius() const
 {
-	return FVector2D(ComponentToWorld.GetScale3D().Y, ComponentToWorld.GetScale3D().Z).Size() * InfluenceRadiusScale;
+	return FVector2D(GetComponentTransform().GetScale3D().Y, GetComponentTransform().GetScale3D().Z).Size() * InfluenceRadiusScale;
 }
 
 FReflectionCaptureProxy::FReflectionCaptureProxy(const UReflectionCaptureComponent* InComponent)
@@ -1779,7 +1779,7 @@ FReflectionCaptureProxy::FReflectionCaptureProxy(const UReflectionCaptureCompone
 	Component = InComponent;
 	SM4FullHDRCubemap = InComponent->SM4FullHDRCubemapTexture;
 	EncodedHDRCubemap = InComponent->EncodedHDRCubemapTexture;
-	SetTransform(InComponent->ComponentToWorld.ToMatrixWithScale());
+	SetTransform(InComponent->GetComponentTransform().ToMatrixWithScale());
 	InfluenceRadius = InComponent->GetInfluenceBoundingRadius();
 	Brightness = InComponent->Brightness;
 	Guid = GetTypeHash( Component->GetPathName() );
diff --git a/Engine/Source/Runtime/Engine/Private/Components/SceneComponent.cpp b/Engine/Source/Runtime/Engine/Private/Components/SceneComponent.cpp
--- a/Engine/Source/Runtime/Engine/Private/Components/SceneComponent.cpp
+++ b/Engine/Source/Runtime/Engine/Private/Components/SceneComponent.cpp
@@ -1101,7 +1101,7 @@ void USceneComponent::UpdateBounds()
 	else
 	{
 		// Calculate new bounds
-		Bounds = CalcBounds(ComponentToWorld);
+		Bounds = CalcBounds(GetComponentTransform());
 	}
 
 
@@ -1216,26 +1216,26 @@ void USceneComponent::AddLocalTransform(const FTransform& DeltaTransform, bool b
 
 void USceneComponent::AddWorldOffset(FVector DeltaLocation, bool bSweep, FHitResult* OutSweepHitResult, ETeleportType Teleport)
 {
-	const FVector NewWorldLocation = DeltaLocation + ComponentToWorld.GetTranslation();
+	const FVector NewWorldLocation = DeltaLocation + GetComponentTransform().GetTranslation();
 	SetWorldLocation(NewWorldLocation, bSweep, OutSweepHitResult, Teleport);
 }
 
 void USceneComponent::AddWorldRotation(FRotator DeltaRotation, bool bSweep, FHitResult* OutSweepHitResult, ETeleportType Teleport)
 {
-	const FQuat NewWorldRotation = DeltaRotation.Quaternion() * ComponentToWorld.GetRotation();
+	const FQuat NewWorldRotation = DeltaRotation.Quaternion() * GetComponentTransform().GetRotation();
 	SetWorldRotation(NewWorldRotation, bSweep, OutSweepHitResult, Teleport);
 }
 
 void USceneComponent::AddWorldRotation(const FQuat& DeltaRotation, bool bSweep, FHitResult* OutSweepHitResult, ETeleportType Teleport)
 {
-	const FQuat NewWorldRotation = DeltaRotation * ComponentToWorld.GetRotation();
+	const FQuat NewWorldRotation = DeltaRotation * GetComponentTransform().GetRotation();
 	SetWorldRotation(NewWorldRotation, bSweep, OutSweepHitResult, Teleport);
 }
 
 void USceneComponent::AddWorldTransform(const FTransform& DeltaTransform, bool bSweep, FHitResult* OutSweepHitResult, ETeleportType Teleport)
 {
-	const FQuat NewWorldRotation = DeltaTransform.GetRotation() * ComponentToWorld.GetRotation();
-	const FVector NewWorldLocation = FTransform::AddTranslations(DeltaTransform, ComponentToWorld);
+	const FQuat NewWorldRotation = DeltaTransform.GetRotation() * GetComponentTransform().GetRotation();
+	const FVector NewWorldLocation = FTransform::AddTranslations(DeltaTransform, ComponentToWorld); // ComponentToWorld is sure to be accurate due to GetComponentTransform() on previous line
 	SetWorldTransform(FTransform(NewWorldRotation, NewWorldLocation, FVector(1,1,1)),bSweep, OutSweepHitResult, Teleport);
 }
 
@@ -1479,17 +1479,17 @@ FTransform USceneComponent::K2_GetComponentToWorld() const
 
 FVector USceneComponent::GetForwardVector() const
 {
-	return ComponentToWorld.GetUnitAxis( EAxis::X );
+	return GetComponentTransform().GetUnitAxis( EAxis::X );
 }
 
 FVector USceneComponent::GetRightVector() const
 {
-	return ComponentToWorld.GetUnitAxis( EAxis::Y );
+	return GetComponentTransform().GetUnitAxis( EAxis::Y );
 }
 
 FVector USceneComponent::GetUpVector() const
 {
-	return ComponentToWorld.GetUnitAxis( EAxis::Z );
+	return GetComponentTransform().GetUnitAxis( EAxis::Z );
 }
 
 FVector USceneComponent::K2_GetComponentLocation() const
@@ -1773,7 +1773,7 @@ bool USceneComponent::AttachToComponent(USceneComponent* Parent, const FAttachme
 							UpdateComponentToWorldWithParent(Parent, SocketName, EUpdateTransformFlags::None, RelativeRotationCache.RotatorToQuat(RelativeRotation));
 							if (AttachmentRules.LocationRule == EAttachmentRule::KeepRelative)
 							{
-								RelativeLocation = ComponentToWorld.GetLocation(); // or GetComponentLocation(), but worried about custom location...
+								RelativeLocation = GetComponentLocation();
 							}
 							if (AttachmentRules.RotationRule == EAttachmentRule::KeepRelative)
 							{
@@ -1829,7 +1829,7 @@ bool USceneComponent::AttachToComponent(USceneComponent* Parent, const FAttachme
 			return false;
 		}
 #endif
-		FTransform RelativeTM = ComponentToWorld.GetRelativeTransform(SocketTransform);
+		FTransform RelativeTM = GetComponentTransform().GetRelativeTransform(SocketTransform);
 #if ENABLE_NAN_DIAGNOSTIC
 		if (RelativeTM.ContainsNaN())
 		{
@@ -1846,7 +1846,7 @@ bool USceneComponent::AttachToComponent(USceneComponent* Parent, const FAttachme
 		case EAttachmentRule::KeepWorld:
 			if (bAbsoluteLocation)
 			{
-				RelativeLocation = ComponentToWorld.GetTranslation();
+				RelativeLocation = GetComponentTransform().GetTranslation();
 			}
 			else
 			{
@@ -1886,7 +1886,7 @@ bool USceneComponent::AttachToComponent(USceneComponent* Parent, const FAttachme
 		case EAttachmentRule::KeepWorld:
 			if (bAbsoluteScale)
 			{
-				RelativeScale3D = ComponentToWorld.GetScale3D();
+				RelativeScale3D = GetComponentTransform().GetScale3D();
 			}
 			else
 			{
@@ -2006,7 +2006,7 @@ void USceneComponent::DetachFromComponent(const FDetachmentTransformRules& Detac
 		case EDetachmentRule::KeepRelative:
 			break;
 		case EDetachmentRule::KeepWorld:
-			RelativeLocation = ComponentToWorld.GetTranslation(); // or GetComponentLocation, but worried about custom location...
+			RelativeLocation = GetComponentTransform().GetTranslation(); // or GetComponentLocation, but worried about custom location...
 			break;
 		}
 
@@ -2154,7 +2154,7 @@ void USceneComponent::UpdateChildTransforms(EUpdateTransformFlags UpdateTransfor
 	SCOPE_CYCLE_COUNTER(STAT_UpdateChildTransforms);
 
 #if ENABLE_NAN_DIAGNOSTIC
-	if (!ComponentToWorld.IsValid())
+	if (!GetComponentTransform().IsValid())
 	{
 		logOrEnsureNanError(TEXT("USceneComponent::UpdateChildTransforms found NaN/INF in ComponentToWorld: %s"), *ComponentToWorld.ToString());
 	}
@@ -2222,7 +2222,7 @@ FTransform USceneComponent::GetSocketTransform(FName SocketName, ERelativeTransf
 	{
 		case RTS_Actor:
 		{
-			return ComponentToWorld.GetRelativeTransform( GetOwner()->GetTransform() );
+			return GetComponentTransform().GetRelativeTransform( GetOwner()->GetTransform() );
 			break;
 		}
 		case RTS_Component:
@@ -2232,7 +2232,7 @@ FTransform USceneComponent::GetSocketTransform(FName SocketName, ERelativeTransf
 		}
 		default:
 		{
-			return ComponentToWorld;
+			return GetComponentTransform();
 		}
 	}
 }
@@ -2693,7 +2693,7 @@ bool USceneComponent::MoveComponentImpl(const FVector& Delta, const FQuat& NewRo
 	if( Delta.IsZero() )
 	{
 		// Skip if no vector or rotation.
-		if (NewRotation.Equals(ComponentToWorld.GetRotation(), SCENECOMPONENT_QUAT_TOLERANCE))
+		if (NewRotation.Equals(GetComponentTransform().GetRotation(), SCENECOMPONENT_QUAT_TOLERANCE))
 		{
 			return true;
 		}
@@ -3124,7 +3124,7 @@ FScopedPreventAttachedComponentMove::~FScopedPreventAttachedComponentMove()
 		{
 			// Need to keep RelativeLocation/Rotation/Scale in sync. ComponentToWorld() will remain correct because child isn't moving.
 			const FTransform ParentToWorld = Owner->GetAttachParent()->GetSocketTransform(Owner->GetAttachSocketName());
-			const FTransform ChildRelativeTM = Owner->ComponentToWorld.GetRelativeTransform(ParentToWorld);
+			const FTransform ChildRelativeTM = Owner->GetComponentTransform().GetRelativeTransform(ParentToWorld);
 
 			if (!bSavedAbsoluteLocation)
 			{
diff --git a/Engine/Source/Runtime/Engine/Private/Components/SkeletalMeshComponent.cpp b/Engine/Source/Runtime/Engine/Private/Components/SkeletalMeshComponent.cpp
--- a/Engine/Source/Runtime/Engine/Private/Components/SkeletalMeshComponent.cpp
+++ b/Engine/Source/Runtime/Engine/Private/Components/SkeletalMeshComponent.cpp
@@ -2214,7 +2214,7 @@ void USkeletalMeshComponent::DebugDrawBones(UCanvas* Canvas, bool bSimpleBones)
 		{
 			int32 BoneIndex = RequiredBones[Index];
 			int32 ParentIndex = SkeletalMesh->RefSkeleton.GetParentIndex(BoneIndex);
-			FTransform BoneTM = (GetComponentSpaceTransforms()[BoneIndex] * ComponentToWorld);
+			FTransform BoneTM = (GetComponentSpaceTransforms()[BoneIndex] * GetComponentTransform());
 			FVector Start, End;
 			FLinearColor LineColor;
 
@@ -2222,12 +2222,12 @@ void USkeletalMeshComponent::DebugDrawBones(UCanvas* Canvas, bool bSimpleBones)
 
 			if (ParentIndex >=0)
 			{
-				Start = (GetComponentSpaceTransforms()[ParentIndex] * ComponentToWorld).GetLocation();
+				Start = (GetComponentSpaceTransforms()[ParentIndex] * GetComponentTransform()).GetLocation();
 				LineColor = FLinearColor::White;
 			}
 			else
 			{
-				Start = ComponentToWorld.GetLocation();
+				Start = GetComponentTransform().GetLocation();
 				LineColor = FLinearColor::Red;
 			}
 
@@ -2529,30 +2529,27 @@ bool USkeletalMeshComponent::PoseTickedThisFrame() const
 FTransform USkeletalMeshComponent::ConvertLocalRootMotionToWorld(const FTransform& InTransform)
 {
 	// Make sure component to world is up to date
-	if (!bWorldToComponentUpdated)
-	{
-		UpdateComponentToWorld();
-	}
+	ConditionalUpdateComponentToWorld();
 
 #if !(UE_BUILD_SHIPPING)
-	if (ComponentToWorld.ContainsNaN())
+	if (GetComponentTransform().ContainsNaN())
 	{
-		logOrEnsureNanError(TEXT("SkeletalMeshComponent: ComponentToWorld contains NaN!"));
-		ComponentToWorld = FTransform::Identity;
+		logOrEnsureNanError(TEXT("SkeletalMeshComponent: GetComponentTransform() contains NaN!"));
+		SetComponentToWorld(FTransform::Identity);
 	}
 #endif
 
 	//Calculate new actor transform after applying root motion to this component
 	const FTransform ActorToWorld = GetOwner()->GetTransform();
 
-	const FTransform ComponentToActor = ActorToWorld.GetRelativeTransform(ComponentToWorld);
-	const FTransform NewComponentToWorld = InTransform * ComponentToWorld;
+	const FTransform ComponentToActor = ActorToWorld.GetRelativeTransform(GetComponentTransform());
+	const FTransform NewComponentToWorld = InTransform * GetComponentTransform();
 	const FTransform NewActorTransform = ComponentToActor * NewComponentToWorld;
 
 	const FVector DeltaWorldTranslation = NewActorTransform.GetTranslation() - ActorToWorld.GetTranslation();
 
-	const FQuat NewWorldRotation = ComponentToWorld.GetRotation() * InTransform.GetRotation();
-	const FQuat DeltaWorldRotation = NewWorldRotation * ComponentToWorld.GetRotation().Inverse();
+	const FQuat NewWorldRotation = GetComponentTransform().GetRotation() * InTransform.GetRotation();
+	const FQuat DeltaWorldRotation = NewWorldRotation * GetComponentTransform().GetRotation().Inverse();
 	
 	const FTransform DeltaWorldTransform(DeltaWorldRotation, DeltaWorldTranslation);
 
@@ -2627,7 +2624,7 @@ bool USkeletalMeshComponent::ComponentIsTouchingSelectionBox(const FBox& InSelBB
 		{
 			for (const auto& Vertex : Section.SoftVertices)
 			{
-				const FVector Location = ComponentToWorld.TransformPosition(Vertex.Position);
+				const FVector Location = GetComponentTransform().TransformPosition(Vertex.Position);
 				const bool bLocationIntersected = FMath::PointBoxIntersection(Location, InSelBBox);
 
 				// If the selection box doesn't have to encompass the entire component and a skeletal mesh vertex has intersected with
@@ -2671,7 +2668,7 @@ bool USkeletalMeshComponent::ComponentIsTouchingSelectionFrustum(const FConvexVo
 		{
 			for (const auto& Vertex : Section.SoftVertices)
 			{
-				const FVector Location = ComponentToWorld.TransformPosition(Vertex.Position);
+				const FVector Location = GetComponentTransform().TransformPosition(Vertex.Position);
 				const bool bLocationIntersected = InFrustum.IntersectSphere(Location, 0.0f);
 
 				// If the selection box doesn't have to encompass the entire component and a skeletal mesh vertex has intersected with
@@ -2960,7 +2957,7 @@ void USkeletalMeshComponent::UnbindClothFromMasterPoseComponent(bool bRestoreSim
 bool USkeletalMeshComponent::DoCustomNavigableGeometryExport(FNavigableGeometryExport& GeomExport) const
 {
 	UPhysicsAsset* PhysicsAsset = GetPhysicsAsset();
-	if (PhysicsAsset && ComponentToWorld.GetScale3D().IsUniform())
+	if (PhysicsAsset && GetComponentTransform().GetScale3D().IsUniform())
 	{
 		const int32 MaxBodies = PhysicsAsset->SkeletalBodySetups.Num();
 		for (int32 Idx = 0; Idx < MaxBodies; Idx++)
@@ -2970,7 +2967,7 @@ bool USkeletalMeshComponent::DoCustomNavigableGeometryExport(FNavigableGeometryE
 
 			if (BoneIndex != INDEX_NONE)
 			{
-				FTransform WorldBoneTransform = GetBoneTransform(BoneIndex, ComponentToWorld);
+				FTransform WorldBoneTransform = GetBoneTransform(BoneIndex, GetComponentTransform());
 				if (FMath::Abs(WorldBoneTransform.GetDeterminant()) > (float)KINDA_SMALL_NUMBER)
 				{
 					GeomExport.ExportRigidBodySetup(*BS, WorldBoneTransform);
diff --git a/Engine/Source/Runtime/Engine/Private/Components/SkinnedMeshComponent.cpp b/Engine/Source/Runtime/Engine/Private/Components/SkinnedMeshComponent.cpp
--- a/Engine/Source/Runtime/Engine/Private/Components/SkinnedMeshComponent.cpp
+++ b/Engine/Source/Runtime/Engine/Private/Components/SkinnedMeshComponent.cpp
@@ -760,7 +760,7 @@ bool USkinnedMeshComponent::GetMaterialStreamingData(int32 MaterialIndex, FPrimi
 
 void USkinnedMeshComponent::GetStreamingTextureInfo(FStreamingTextureLevelContext& LevelContext, TArray<FStreamingTexturePrimitiveInfo>& OutStreamingTextures) const
 {
-	GetStreamingTextureInfoInner(LevelContext, nullptr, ComponentToWorld.GetMaximumAxisScale(), OutStreamingTextures);
+	GetStreamingTextureInfoInner(LevelContext, nullptr, GetComponentTransform().GetMaximumAxisScale(), OutStreamingTextures);
 }
 
 bool USkinnedMeshComponent::ShouldUpdateBoneVisibility() const
@@ -915,7 +915,7 @@ FMatrix USkinnedMeshComponent::GetBoneMatrix(int32 BoneIdx) const
 	if ( !IsRegistered() )
 	{
 		// if not registered, we don't have SpaceBases yet. 
-		// also ComponentToWorld isn't set yet (They're set from relativetranslation, relativerotation, relativescale)
+		// also GetComponentTransform() isn't set yet (They're set from relativetranslation, relativerotation, relativescale)
 		return FMatrix::Identity;
 	}
 
@@ -931,7 +931,7 @@ FMatrix USkinnedMeshComponent::GetBoneMatrix(int32 BoneIdx) const
 			if(	ParentBoneIndex != INDEX_NONE && 
 				ParentBoneIndex < MasterPoseComponentInst->GetNumComponentSpaceTransforms())
 			{
-				return MasterPoseComponentInst->GetComponentSpaceTransforms()[ParentBoneIndex].ToMatrixWithScale() * ComponentToWorld.ToMatrixWithScale();
+				return MasterPoseComponentInst->GetComponentSpaceTransforms()[ParentBoneIndex].ToMatrixWithScale() * GetComponentTransform().ToMatrixWithScale();
 			}
 			else
 			{
@@ -949,7 +949,7 @@ FMatrix USkinnedMeshComponent::GetBoneMatrix(int32 BoneIdx) const
 	{
 		if(GetNumComponentSpaceTransforms() && BoneIdx < GetNumComponentSpaceTransforms() )
 		{
-			return GetComponentSpaceTransforms()[BoneIdx].ToMatrixWithScale() * ComponentToWorld.ToMatrixWithScale();
+			return GetComponentSpaceTransforms()[BoneIdx].ToMatrixWithScale() * GetComponentTransform().ToMatrixWithScale();
 		}
 		else
 		{
@@ -964,11 +964,11 @@ FTransform USkinnedMeshComponent::GetBoneTransform(int32 BoneIdx) const
 	if (!IsRegistered())
 	{
 		// if not registered, we don't have SpaceBases yet. 
-		// also ComponentToWorld isn't set yet (They're set from relativelocation, relativerotation, relativescale)
+		// also GetComponentTransform() isn't set yet (They're set from relativelocation, relativerotation, relativescale)
 		return FTransform::Identity;
 	}
 
-	return GetBoneTransform(BoneIdx, ComponentToWorld);
+	return GetBoneTransform(BoneIdx, GetComponentTransform());
 }
 
 FTransform USkinnedMeshComponent::GetBoneTransform(int32 BoneIdx, const FTransform& LocalToWorld) const
@@ -1389,7 +1389,7 @@ void USkinnedMeshComponent::UpdateMasterBoneMap()
 
 FTransform USkinnedMeshComponent::GetSocketTransform(FName InSocketName, ERelativeTransformSpace TransformSpace) const
 {
-	FTransform OutSocketTransform = ComponentToWorld;
+	FTransform OutSocketTransform = GetComponentTransform();
 
 	if (InSocketName != NAME_None)
 	{
@@ -1427,7 +1427,7 @@ FTransform USkinnedMeshComponent::GetSocketTransform(FName InSocketName, ERelati
 					{
 						return OutSocketTransform.GetRelativeTransform(GetBoneTransform(ParentIndex));
 					}
-					return OutSocketTransform.GetRelativeTransform(ComponentToWorld);
+					return OutSocketTransform.GetRelativeTransform(GetComponentTransform());
 				}
 			}
 		}
@@ -1445,7 +1445,7 @@ FTransform USkinnedMeshComponent::GetSocketTransform(FName InSocketName, ERelati
 		}
 		case RTS_Component:
 		{
-			return OutSocketTransform.GetRelativeTransform( ComponentToWorld );
+			return OutSocketTransform.GetRelativeTransform( GetComponentTransform() );
 		}
 	}
 
@@ -1710,7 +1710,7 @@ FName USkinnedMeshComponent::FindClosestBone(FVector TestLocation, FVector* Bone
 		}
 
 		// transform the TestLocation into mesh local space so we don't have to transform the (mesh local) bone locations
-		TestLocation = ComponentToWorld.InverseTransformPosition(TestLocation);
+		TestLocation = GetComponentTransform().InverseTransformPosition(TestLocation);
 		
 		float IgnoreScaleSquared = FMath::Square(IgnoreScale);
 		float BestDistSquared = BIG_NUMBER;
@@ -1749,7 +1749,7 @@ FName USkinnedMeshComponent::FindClosestBone(FVector TestLocation, FVector* Bone
 			// transform the bone location into world space
 			if (BoneLocation != NULL)
 			{
-				*BoneLocation = (GetComponentSpaceTransforms()[BestIndex] * ComponentToWorld).GetLocation();
+				*BoneLocation = (GetComponentSpaceTransforms()[BestIndex] * GetComponentTransform()).GetLocation();
 			}
 			return SkeletalMesh->RefSkeleton.GetBoneName(BestIndex);
 		}
diff --git a/Engine/Source/Runtime/Engine/Private/Components/SphereComponent.cpp b/Engine/Source/Runtime/Engine/Private/Components/SphereComponent.cpp
--- a/Engine/Source/Runtime/Engine/Private/Components/SphereComponent.cpp
+++ b/Engine/Source/Runtime/Engine/Private/Components/SphereComponent.cpp
@@ -26,7 +26,7 @@ FBoxSphereBounds USphereComponent::CalcBounds(const FTransform& LocalToWorld) co
 
 void USphereComponent::CalcBoundingCylinder(float& CylinderRadius, float& CylinderHalfHeight) const
 {
-	CylinderRadius = SphereRadius * ComponentToWorld.GetMaximumAxisScale();
+	CylinderRadius = SphereRadius * GetComponentTransform().GetMaximumAxisScale();
 	CylinderHalfHeight = CylinderRadius;
 }
 
@@ -87,7 +87,7 @@ void USphereComponent::SetSphereRadius( float InSphereRadius, bool bUpdateOverla
 	if (bPhysicsStateCreated)
 	{
 		// Update physics engine collision shapes
-		BodyInstance.UpdateBodyScale(ComponentToWorld.GetScale3D(), true);
+		BodyInstance.UpdateBodyScale(GetComponentTransform().GetScale3D(), true);
 
 		if ( bUpdateOverlaps && IsCollisionEnabled() && GetOwner() )
 		{
