Copying //UE4/Dev-Framework to //UE4/Dev-Main (Source: //UE4/Dev-Framework @ 3431384)

diff --git a/Engine/Source/Runtime/Engine/Private/Components/SplineComponent.cpp b/Engine/Source/Runtime/Engine/Private/Components/SplineComponent.cpp
--- a/Engine/Source/Runtime/Engine/Private/Components/SplineComponent.cpp
+++ b/Engine/Source/Runtime/Engine/Private/Components/SplineComponent.cpp
@@ -185,7 +185,7 @@ void FSplineCurves::UpdateSpline(bool bClosedLoop, bool bStationaryEndpoints, in
 
 void USplineComponent::UpdateSpline()
 {
-	SplineCurves.UpdateSpline(bClosedLoop, bStationaryEndpoints, ReparamStepsPerSegment, bLoopPositionOverride, LoopPosition, ComponentToWorld.GetScale3D());
+	SplineCurves.UpdateSpline(bClosedLoop, bStationaryEndpoints, ReparamStepsPerSegment, bLoopPositionOverride, LoopPosition, GetComponentTransform().GetScale3D());
 
 #if !UE_BUILD_SHIPPING
 	if (bDrawDebug)
@@ -264,7 +264,7 @@ float FSplineCurves::GetSegmentLength(const int32 Index, const float Param, bool
 
 float USplineComponent::GetSegmentLength(const int32 Index, const float Param) const
 {
-	return SplineCurves.GetSegmentLength(Index, Param, bClosedLoop, ComponentToWorld.GetScale3D());
+	return SplineCurves.GetSegmentLength(Index, Param, bClosedLoop, GetComponentTransform().GetScale3D());
 }
 
 float USplineComponent::GetSegmentParamFromLength(const int32 Index, const float Length, const float SegmentLength) const
@@ -322,7 +322,7 @@ FVector USplineComponent::GetLocationAtSplineInputKey(float InKey, ESplineCoordi
 
 	if (CoordinateSpace == ESplineCoordinateSpace::World)
 	{
-		Location = ComponentToWorld.TransformPosition(Location);
+		Location = GetComponentTransform().TransformPosition(Location);
 	}
 
 	return Location;
@@ -335,7 +335,7 @@ FVector USplineComponent::GetTangentAtSplineInputKey(float InKey, ESplineCoordin
 
 	if (CoordinateSpace == ESplineCoordinateSpace::World)
 	{
-		Tangent = ComponentToWorld.TransformVector(Tangent);
+		Tangent = GetComponentTransform().TransformVector(Tangent);
 	}
 
 	return Tangent;
@@ -348,7 +348,7 @@ FVector USplineComponent::GetDirectionAtSplineInputKey(float InKey, ESplineCoord
 
 	if (CoordinateSpace == ESplineCoordinateSpace::World)
 	{
-		Direction = ComponentToWorld.TransformVectorNoScale(Direction);
+		Direction = GetComponentTransform().TransformVectorNoScale(Direction);
 	}
 
 	return Direction;
@@ -373,7 +373,7 @@ FQuat USplineComponent::GetQuaternionAtSplineInputKey(float InKey, ESplineCoordi
 
 	if (CoordinateSpace == ESplineCoordinateSpace::World)
 	{
-		Rot = ComponentToWorld.GetRotation() * Rot;
+		Rot = GetComponentTransform().GetRotation() * Rot;
 	}
 
 	return Rot;
@@ -387,7 +387,7 @@ FVector USplineComponent::GetUpVectorAtSplineInputKey(float InKey, ESplineCoordi
 
 	if (CoordinateSpace == ESplineCoordinateSpace::World)
 	{
-		UpVector = ComponentToWorld.TransformVectorNoScale(UpVector);
+		UpVector = GetComponentTransform().TransformVectorNoScale(UpVector);
 	}
 
 	return UpVector;
@@ -401,7 +401,7 @@ FVector USplineComponent::GetRightVectorAtSplineInputKey(float InKey, ESplineCoo
 
 	if (CoordinateSpace == ESplineCoordinateSpace::World)
 	{
-		RightVector = ComponentToWorld.TransformVectorNoScale(RightVector);
+		RightVector = GetComponentTransform().TransformVectorNoScale(RightVector);
 	}
 
 	return RightVector;
@@ -418,7 +418,7 @@ FTransform USplineComponent::GetTransformAtSplineInputKey(float InKey, ESplineCo
 
 	if (CoordinateSpace == ESplineCoordinateSpace::World)
 	{
-		Transform = Transform * ComponentToWorld;
+		Transform = Transform * GetComponentTransform();
 	}
 
 	return Transform;
@@ -589,7 +589,7 @@ void USplineComponent::AddPoints(const TArray<FSplinePoint>& InSplinePoints, boo
 void USplineComponent::AddSplinePoint(const FVector& Position, ESplineCoordinateSpace::Type CoordinateSpace, bool bUpdateSpline)
 {
 	const FVector TransformedPosition = (CoordinateSpace == ESplineCoordinateSpace::World) ?
-		ComponentToWorld.InverseTransformPosition(Position) : Position;
+		GetComponentTransform().InverseTransformPosition(Position) : Position;
 
 	// Add the spline point at the end of the array, adding 1.0 to the current last input key.
 	// This continues the former behavior in which spline points had to be separated by an interval of 1.0.
@@ -613,7 +613,7 @@ void USplineComponent::AddSplinePoint(const FVector& Position, ESplineCoordinate
 void USplineComponent::AddSplinePointAtIndex(const FVector& Position, int32 Index, ESplineCoordinateSpace::Type CoordinateSpace, bool bUpdateSpline)
 {
 	const FVector TransformedPosition = (CoordinateSpace == ESplineCoordinateSpace::World) ?
-		ComponentToWorld.InverseTransformPosition(Position) : Position;
+		GetComponentTransform().InverseTransformPosition(Position) : Position;
 
 	int32 NumPoints = SplineCurves.Position.Points.Num();
 
@@ -689,7 +689,7 @@ void USplineComponent::SetSplinePoints(const TArray<FVector>& Points, ESplineCoo
 	for (const auto& Point : Points)
 	{
 		const FVector TransformedPoint = (CoordinateSpace == ESplineCoordinateSpace::World) ?
-			ComponentToWorld.InverseTransformPosition(Point) : Point;
+			GetComponentTransform().InverseTransformPosition(Point) : Point;
 
 		SplineCurves.Position.Points.Emplace(InputKey, TransformedPoint, FVector::ZeroVector, FVector::ZeroVector, CIM_CurveAuto);
 		SplineCurves.Rotation.Points.Emplace(InputKey, FQuat::Identity, FQuat::Identity, FQuat::Identity, CIM_CurveAuto);
@@ -714,7 +714,7 @@ void USplineComponent::SetLocationAtSplinePoint(int32 PointIndex, const FVector&
 	if ((PointIndex >= 0) && (PointIndex < NumPoints))
 	{
 		const FVector TransformedLocation = (CoordinateSpace == ESplineCoordinateSpace::World) ?
-			ComponentToWorld.InverseTransformPosition(InLocation) : InLocation;
+			GetComponentTransform().InverseTransformPosition(InLocation) : InLocation;
 
 		SplineCurves.Position.Points[PointIndex].OutVal = TransformedLocation;
 
@@ -733,7 +733,7 @@ void USplineComponent::SetTangentAtSplinePoint(int32 PointIndex, const FVector&
 	if ((PointIndex >= 0) && (PointIndex < NumPoints))
 	{
 		const FVector TransformedTangent = (CoordinateSpace == ESplineCoordinateSpace::World) ?
-			ComponentToWorld.InverseTransformVector(InTangent) : InTangent;
+			GetComponentTransform().InverseTransformVector(InTangent) : InTangent;
 
 		SplineCurves.Position.Points[PointIndex].LeaveTangent = TransformedTangent;
 		SplineCurves.Position.Points[PointIndex].ArriveTangent = TransformedTangent;
@@ -754,9 +754,9 @@ void USplineComponent::SetTangentsAtSplinePoint(int32 PointIndex, const FVector&
 	if ((PointIndex >= 0) && (PointIndex < NumPoints))
 	{
 		const FVector TransformedArriveTangent = (CoordinateSpace == ESplineCoordinateSpace::World) ?
-			ComponentToWorld.InverseTransformVector(InArriveTangent) : InArriveTangent;
+			GetComponentTransform().InverseTransformVector(InArriveTangent) : InArriveTangent;
 		const FVector TransformedLeaveTangent = (CoordinateSpace == ESplineCoordinateSpace::World) ?
-			ComponentToWorld.InverseTransformVector(InLeaveTangent) : InLeaveTangent;
+			GetComponentTransform().InverseTransformVector(InLeaveTangent) : InLeaveTangent;
 
 		SplineCurves.Position.Points[PointIndex].ArriveTangent = TransformedArriveTangent;
 		SplineCurves.Position.Points[PointIndex].LeaveTangent = TransformedLeaveTangent;
@@ -777,7 +777,7 @@ void USplineComponent::SetUpVectorAtSplinePoint(int32 PointIndex, const FVector&
 	if ((PointIndex >= 0) && (PointIndex < NumPoints))
 	{
 		const FVector TransformedUpVector = (CoordinateSpace == ESplineCoordinateSpace::World) ?
-			ComponentToWorld.InverseTransformVector(InUpVector.GetSafeNormal()) : InUpVector.GetSafeNormal();
+			GetComponentTransform().InverseTransformVector(InUpVector.GetSafeNormal()) : InUpVector.GetSafeNormal();
 
 		FQuat Quat = FQuat::FindBetween(DefaultUpVector, TransformedUpVector);
 		SplineCurves.Rotation.Points[PointIndex].OutVal = Quat;
@@ -827,7 +827,7 @@ FVector USplineComponent::GetLocationAtSplinePoint(int32 PointIndex, ESplineCoor
 	const int32 NumPoints = SplineCurves.Position.Points.Num();
 	const int32 ClampedIndex = (bClosedLoop && PointIndex >= NumPoints) ? 0 : FMath::Clamp(PointIndex, 0, NumPoints - 1);
 	const FVector Location = SplineCurves.Position.Points[ClampedIndex].OutVal;
-	return (CoordinateSpace == ESplineCoordinateSpace::World) ? ComponentToWorld.TransformPosition(Location) : Location;
+	return (CoordinateSpace == ESplineCoordinateSpace::World) ? GetComponentTransform().TransformPosition(Location) : Location;
 }
 
 
@@ -836,7 +836,7 @@ FVector USplineComponent::GetDirectionAtSplinePoint(int32 PointIndex, ESplineCoo
 	const int32 NumPoints = SplineCurves.Position.Points.Num();
 	const int32 ClampedIndex = (bClosedLoop && PointIndex >= NumPoints) ? 0 : FMath::Clamp(PointIndex, 0, NumPoints - 1);
 	const FVector Direction = SplineCurves.Position.Points[ClampedIndex].LeaveTangent.GetSafeNormal();
-	return (CoordinateSpace == ESplineCoordinateSpace::World) ? ComponentToWorld.TransformVector(Direction) : Direction;
+	return (CoordinateSpace == ESplineCoordinateSpace::World) ? GetComponentTransform().TransformVector(Direction) : Direction;
 }
 
 
@@ -845,7 +845,7 @@ FVector USplineComponent::GetTangentAtSplinePoint(int32 PointIndex, ESplineCoord
 	const int32 NumPoints = SplineCurves.Position.Points.Num();
 	const int32 ClampedIndex = (bClosedLoop && PointIndex >= NumPoints) ? 0 : FMath::Clamp(PointIndex, 0, NumPoints - 1);
 	const FVector Direction = SplineCurves.Position.Points[ClampedIndex].LeaveTangent;
-	return (CoordinateSpace == ESplineCoordinateSpace::World) ? ComponentToWorld.TransformVector(Direction) : Direction;
+	return (CoordinateSpace == ESplineCoordinateSpace::World) ? GetComponentTransform().TransformVector(Direction) : Direction;
 }
 
 
@@ -854,7 +854,7 @@ FVector USplineComponent::GetArriveTangentAtSplinePoint(int32 PointIndex, ESplin
 	const int32 NumPoints = SplineCurves.Position.Points.Num();
 	const int32 ClampedIndex = (bClosedLoop && PointIndex >= NumPoints) ? 0 : FMath::Clamp(PointIndex, 0, NumPoints - 1);
 	const FVector Direction = SplineCurves.Position.Points[ClampedIndex].ArriveTangent;
-	return (CoordinateSpace == ESplineCoordinateSpace::World) ? ComponentToWorld.TransformVector(Direction) : Direction;
+	return (CoordinateSpace == ESplineCoordinateSpace::World) ? GetComponentTransform().TransformVector(Direction) : Direction;
 }
 
 
@@ -863,7 +863,7 @@ FVector USplineComponent::GetLeaveTangentAtSplinePoint(int32 PointIndex, ESpline
 	const int32 NumPoints = SplineCurves.Position.Points.Num();
 	const int32 ClampedIndex = (bClosedLoop && PointIndex >= NumPoints) ? 0 : FMath::Clamp(PointIndex, 0, NumPoints - 1);
 	const FVector Direction = SplineCurves.Position.Points[ClampedIndex].LeaveTangent;
-	return (CoordinateSpace == ESplineCoordinateSpace::World) ? ComponentToWorld.TransformVector(Direction) : Direction;
+	return (CoordinateSpace == ESplineCoordinateSpace::World) ? GetComponentTransform().TransformVector(Direction) : Direction;
 }
 
 
@@ -969,7 +969,7 @@ void USplineComponent::SetDefaultUpVector(const FVector& UpVector, ESplineCoordi
 {
 	if (CoordinateSpace == ESplineCoordinateSpace::World)
 	{
-		DefaultUpVector = ComponentToWorld.InverseTransformVector(UpVector);
+		DefaultUpVector = GetComponentTransform().InverseTransformVector(UpVector);
 	}
 	else
 	{
@@ -984,7 +984,7 @@ FVector USplineComponent::GetDefaultUpVector(ESplineCoordinateSpace::Type Coordi
 {
 	if (CoordinateSpace == ESplineCoordinateSpace::World)
 	{
-		return ComponentToWorld.TransformVector(DefaultUpVector);
+		return GetComponentTransform().TransformVector(DefaultUpVector);
 	}
 	else
 	{
@@ -1289,7 +1289,7 @@ FVector USplineComponent::GetScaleAtTime(float Time, bool bUseConstantVelocity)
 
 float USplineComponent::FindInputKeyClosestToWorldLocation(const FVector& WorldLocation) const
 {
-	const FVector LocalLocation = ComponentToWorld.InverseTransformPosition(WorldLocation);
+	const FVector LocalLocation = GetComponentTransform().InverseTransformPosition(WorldLocation);
 	float Dummy;
 	return SplineCurves.Position.InaccurateFindNearest(LocalLocation, Dummy);
 }
diff --git a/Engine/Source/Runtime/Engine/Private/Components/SplineMeshComponent.cpp b/Engine/Source/Runtime/Engine/Private/Components/SplineMeshComponent.cpp
--- a/Engine/Source/Runtime/Engine/Private/Components/SplineMeshComponent.cpp
+++ b/Engine/Source/Runtime/Engine/Private/Components/SplineMeshComponent.cpp
@@ -1054,7 +1054,7 @@ bool USplineMeshComponent::DoCustomNavigableGeometryExport(FNavigableGeometryExp
 				}
 				GeomExport.ExportCustomMesh(VertexBuffer.GetData(), VertexBuffer.Num(),
 					NavCollision->ConvexCollision.IndexBuffer.GetData(), NavCollision->ConvexCollision.IndexBuffer.Num(),
-					ComponentToWorld);
+					GetComponentTransform());
 
 				VertexBuffer.Reset();
 				for (int32 i = 0; i < NavCollision->TriMeshCollision.VertexBuffer.Num(); ++i)
@@ -1065,7 +1065,7 @@ bool USplineMeshComponent::DoCustomNavigableGeometryExport(FNavigableGeometryExp
 				}
 				GeomExport.ExportCustomMesh(VertexBuffer.GetData(), VertexBuffer.Num(),
 					NavCollision->TriMeshCollision.IndexBuffer.GetData(), NavCollision->TriMeshCollision.IndexBuffer.Num(),
-					ComponentToWorld);
+					GetComponentTransform());
 
 				return false;
 			}
@@ -1282,7 +1282,7 @@ float USplineMeshComponent::GetTextureStreamingTransformScale() const
 
 		// We do this by looking at the ratio between current bounds (including deformation) and undeformed (straight from staticmesh)
 		const float MinExtent = 1.0f;
-		FBoxSphereBounds UndeformedBounds = GetStaticMesh()->GetBounds().TransformBy(ComponentToWorld);
+		FBoxSphereBounds UndeformedBounds = GetStaticMesh()->GetBounds().TransformBy(GetComponentTransform());
 		if (UndeformedBounds.BoxExtent.X >= MinExtent)
 		{
 			SplineDeformFactor = FMath::Max(SplineDeformFactor, Bounds.BoxExtent.X / UndeformedBounds.BoxExtent.X);
diff --git a/Engine/Source/Runtime/Engine/Private/Components/SpotLightComponent.cpp b/Engine/Source/Runtime/Engine/Private/Components/SpotLightComponent.cpp
--- a/Engine/Source/Runtime/Engine/Private/Components/SpotLightComponent.cpp
+++ b/Engine/Source/Runtime/Engine/Private/Components/SpotLightComponent.cpp
@@ -217,7 +217,7 @@ FSphere USpotLightComponent::GetBoundingSphere() const
 
 	// Use the law of cosines to find the distance to the furthest edge of the spotlight cone from a position that is halfway down the spotlight direction
 	const float BoundsRadius = FMath::Sqrt(1.25f * AttenuationRadius * AttenuationRadius - AttenuationRadius * AttenuationRadius * CosOuterCone);
-	return FSphere(ComponentToWorld.GetLocation() + .5f * GetDirection() * AttenuationRadius, BoundsRadius);
+	return FSphere(GetComponentTransform().GetLocation() + .5f * GetDirection() * AttenuationRadius, BoundsRadius);
 }
 
 bool USpotLightComponent::AffectsBounds(const FBoxSphereBounds& InBounds) const
diff --git a/Engine/Source/Runtime/Engine/Private/Components/StaticMeshComponent.cpp b/Engine/Source/Runtime/Engine/Private/Components/StaticMeshComponent.cpp
--- a/Engine/Source/Runtime/Engine/Private/Components/StaticMeshComponent.cpp
+++ b/Engine/Source/Runtime/Engine/Private/Components/StaticMeshComponent.cpp
@@ -730,7 +730,7 @@ bool UStaticMeshComponent::BuildTextureStreamingData(ETextureStreamingBuildType
 				MaterialStreamingRelativeBoxes.Empty(NumMaterials);
 				for (int32 MaterialIndex = 0; MaterialIndex < NumMaterials; ++MaterialIndex)
 				{
-					MaterialStreamingRelativeBoxes.Add(PackRelativeBox(Bounds.GetBox(), GetStaticMesh()->GetMaterialBox(MaterialIndex, ComponentToWorld)));
+					MaterialStreamingRelativeBoxes.Add(PackRelativeBox(Bounds.GetBox(), GetStaticMesh()->GetMaterialBox(MaterialIndex, GetComponentTransform())));
 				}
 
 				// Update since proxy has a copy of the material bounds.
@@ -821,7 +821,7 @@ bool UStaticMeshComponent::BuildTextureStreamingData(ETextureStreamingBuildType
 
 float UStaticMeshComponent::GetTextureStreamingTransformScale() const
 {
-	return ComponentToWorld.GetMaximumAxisScale();
+	return GetComponentTransform().GetMaximumAxisScale();
 }
 
 void UStaticMeshComponent::GetStreamingTextureInfo(FStreamingTextureLevelContext& LevelContext, TArray<FStreamingTexturePrimitiveInfo>& OutStreamingTextures) const
@@ -966,7 +966,7 @@ FTransform UStaticMeshComponent::GetSocketTransform(FName InSocketName, ERelativ
 					}
 					case RTS_Component:
 					{
-						return SocketWorldTransform.GetRelativeTransform(ComponentToWorld);
+						return SocketWorldTransform.GetRelativeTransform(GetComponentTransform());
 					}
 				}
 			}
@@ -1673,7 +1673,6 @@ void UStaticMeshComponent::PostLoad()
 	// Remap the materials array if the static mesh materials may have been remapped to remove zero triangle sections.
 	if (GetStaticMesh() && GetLinkerUE4Version() < VER_UE4_REMOVE_ZERO_TRIANGLE_SECTIONS && OverrideMaterials.Num())
 	{
-		GetStaticMesh()->ConditionalPostLoad();
 		if (GetStaticMesh()->HasValidRenderData()
 			&& GetStaticMesh()->RenderData->MaterialIndexToImportIndex.Num())
 		{
@@ -2100,7 +2099,7 @@ void UStaticMeshComponent::GetUsedMaterials(TArray<UMaterialInterface*>& OutMate
 int32 UStaticMeshComponent::GetBlueprintCreatedComponentIndex() const
 {
 	int32 ComponentIndex = 0;
-	for(const auto& Component : GetOwner()->BlueprintCreatedComponents)
+	for (const UActorComponent* Component : GetOwner()->BlueprintCreatedComponents)
 	{
 		if(Component == this)
 		{
@@ -2121,7 +2120,7 @@ FActorComponentInstanceData* UStaticMeshComponent::GetComponentInstanceData() co
 
 	// Fill in info
 	const_cast<UStaticMeshComponent*>(this)->ConditionalUpdateComponentToWorld(); // sadness
-	StaticMeshInstanceData->CachedStaticLighting.Transform = ComponentToWorld;
+	StaticMeshInstanceData->CachedStaticLighting.Transform = GetComponentTransform();
 
 	for (const FStaticMeshComponentLODInfo& LODDataEntry : LODData)
 	{
@@ -2170,7 +2169,7 @@ void UStaticMeshComponent::ApplyComponentInstanceData(FStaticMeshComponentInstan
 	if (HasStaticLighting() && NumLODLightMaps > 0)
 	{
 		// See if data matches current state
-		if (StaticMeshInstanceData->CachedStaticLighting.Transform.Equals(ComponentToWorld, 1.e-3f))
+		if (StaticMeshInstanceData->CachedStaticLighting.Transform.Equals(GetComponentTransform(), 1.e-3f))
 		{
 			SetLODDataCount(NumLODLightMaps, NumLODLightMaps);
 
@@ -2183,7 +2182,7 @@ void UStaticMeshComponent::ApplyComponentInstanceData(FStaticMeshComponentInstan
 		{
 			UE_LOG(LogStaticMesh, Warning, TEXT("Cached component instance data transform did not match!  Discarding cached lighting data which will cause lighting to be unbuilt.\n%s\nCurrent: %s Cached: %s"), 
 				*GetPathName(),
-				*ComponentToWorld.ToString(),
+				*GetComponentTransform().ToString(),
 				*StaticMeshInstanceData->CachedStaticLighting.Transform.ToString());
 		}
 	}
@@ -2215,15 +2214,15 @@ bool UStaticMeshComponent::DoCustomNavigableGeometryExport(FNavigableGeometryExp
 		
 		if (NavCollision->bHasConvexGeometry)
 		{
-			const FVector Scale3D = ComponentToWorld.GetScale3D();
+			const FVector Scale3D = GetComponentTransform().GetScale3D();
 			// if any of scales is 0 there's no point in exporting it
 			if (!Scale3D.IsZero())
 			{
 				GeomExport.ExportCustomMesh(NavCollision->ConvexCollision.VertexBuffer.GetData(), NavCollision->ConvexCollision.VertexBuffer.Num(),
-					NavCollision->ConvexCollision.IndexBuffer.GetData(), NavCollision->ConvexCollision.IndexBuffer.Num(), ComponentToWorld);
+					NavCollision->ConvexCollision.IndexBuffer.GetData(), NavCollision->ConvexCollision.IndexBuffer.Num(), GetComponentTransform());
 
 				GeomExport.ExportCustomMesh(NavCollision->TriMeshCollision.VertexBuffer.GetData(), NavCollision->TriMeshCollision.VertexBuffer.Num(),
-					NavCollision->TriMeshCollision.IndexBuffer.GetData(), NavCollision->TriMeshCollision.IndexBuffer.Num(), ComponentToWorld);
+					NavCollision->TriMeshCollision.IndexBuffer.GetData(), NavCollision->TriMeshCollision.IndexBuffer.Num(), GetComponentTransform());
 			}
 
 			// regardless of above we don't want "regular" collision export for this mesh instance
@@ -2288,7 +2287,7 @@ void UStaticMeshComponent::GetNavigationData(FNavigationRelevantData& Data) cons
 
 		if (bExportAsObstacle)
 		{
-			NavCollision->GetNavigationModifier(Data.Modifiers, ComponentToWorld);
+			NavCollision->GetNavigationModifier(Data.Modifiers, GetComponentTransform());
 		}
 	}
 }
@@ -2318,7 +2317,7 @@ bool UStaticMeshComponent::ComponentIsTouchingSelectionBox(const FBox& InSelBBox
 					{
 						int32 VertexIndex = Indices[FirstIndex + i];
 						FVector LocalPosition = LODModel.PositionVertexBuffer.VertexPosition(VertexIndex);
-						Vertex.Emplace(ComponentToWorld.TransformPosition(LocalPosition));
+						Vertex.Emplace(GetComponentTransform().TransformPosition(LocalPosition));
 					}
 
 					// Check if the triangle is colliding with the bounding box.
@@ -2362,7 +2361,7 @@ bool UStaticMeshComponent::ComponentIsTouchingSelectionFrustum(const FConvexVolu
 			for (uint32 VertexIndex = 0; VertexIndex < NumVertices; ++VertexIndex)
 			{
 				const FVector& LocalPosition = LODModel.PositionVertexBuffer.VertexPosition(VertexIndex);
-				const FVector WorldPosition = ComponentToWorld.TransformPosition(LocalPosition);
+				const FVector WorldPosition = GetComponentTransform().TransformPosition(LocalPosition);
 				bool bLocationIntersected = InFrustum.IntersectSphere(WorldPosition, 0.0f);
 				if (bLocationIntersected && !bMustEncompassEntireComponent)
 				{
diff --git a/Engine/Source/Runtime/Engine/Private/Components/TextRenderComponent.cpp b/Engine/Source/Runtime/Engine/Private/Components/TextRenderComponent.cpp
--- a/Engine/Source/Runtime/Engine/Private/Components/TextRenderComponent.cpp
+++ b/Engine/Source/Runtime/Engine/Private/Components/TextRenderComponent.cpp
@@ -1230,7 +1230,7 @@ FMatrix UTextRenderComponent::GetRenderMatrix() const
 		const float VerticalAlignmentOffset = -ComputeVerticalAlignmentOffset(SizeY, VerticalAlignment, FirstLineHeight);
 		VerticalTransform = VerticalTransform.ConcatTranslation(FVector(0.f, 0.f, VerticalAlignmentOffset));
 
-		return VerticalTransform *  ComponentToWorld.ToMatrixWithScale();
+		return VerticalTransform *  GetComponentTransform().ToMatrixWithScale();
 
 	}
 	return UPrimitiveComponent::GetRenderMatrix();
@@ -1329,7 +1329,7 @@ FVector UTextRenderComponent::GetTextLocalSize() const
 
 FVector UTextRenderComponent::GetTextWorldSize() const
 {
-	const FBoxSphereBounds TextBounds = CalcBounds(ComponentToWorld);
+	const FBoxSphereBounds TextBounds = CalcBounds(GetComponentTransform());
 	return TextBounds.GetBox().GetSize();
 }
 
diff --git a/Engine/Source/Runtime/Engine/Private/EdGraph/EdGraphNode.cpp b/Engine/Source/Runtime/Engine/Private/EdGraph/EdGraphNode.cpp
--- a/Engine/Source/Runtime/Engine/Private/EdGraph/EdGraphNode.cpp
+++ b/Engine/Source/Runtime/Engine/Private/EdGraph/EdGraphNode.cpp
@@ -71,16 +71,15 @@ UEdGraphNode::UEdGraphNode(const FObjectInitializer& ObjectInitializer)
 	, AdvancedPinDisplay(ENodeAdvancedPins::NoPins)
 	, EnabledState(ENodeEnabledState::Enabled)
 	, bUserSetEnabledState(false)
-	, bIsNodeEnabled_DEPRECATED(true)
 	, bAllowSplitPins_DEPRECATED(false)
-{
-
+	, bIsNodeEnabled_DEPRECATED(true)
 #if WITH_EDITORONLY_DATA
-	bCommentBubblePinned = false;
-	bCommentBubbleVisible = false;
-	bCommentBubbleMakeVisible = false;
-	bCanResizeNode = false;
+	, bCanResizeNode(false)
 #endif // WITH_EDITORONLY_DATA
+	, bCommentBubblePinned(false)
+	, bCommentBubbleVisible(false)
+	, bCommentBubbleMakeVisible(false)
+{
 }
 
 #if WITH_EDITOR
@@ -163,7 +162,12 @@ UEdGraphPin* UEdGraphNode::CreatePin(EEdGraphPinDirection Dir, const FEdGraphPin
 
 UEdGraphPin* UEdGraphNode::CreatePin(EEdGraphPinDirection Dir, const FString& PinCategory, const FString& PinSubCategory, UObject* PinSubCategoryObject, bool bIsArray, bool bIsReference, const FString& PinName, bool bIsConst /*= false*/, int32 Index /*= INDEX_NONE*/, bool bIsSet /*= false*/, bool bIsMap /*= false*/, const FEdGraphTerminalType& ValueTerminalType /*= FEdGraphTerminalType()*/)
 {
-	FEdGraphPinType PinType(PinCategory, PinSubCategory, PinSubCategoryObject, bIsArray, bIsReference, bIsSet, bIsMap, ValueTerminalType);
+	return CreatePin(Dir, PinCategory, PinSubCategory, PinSubCategoryObject, PinName, FEdGraphPinType::ToPinContainerType(bIsArray, bIsSet, bIsMap), bIsReference, bIsConst, Index, ValueTerminalType);
+}
+
+UEdGraphPin* UEdGraphNode::CreatePin(EEdGraphPinDirection Dir, const FString& PinCategory, const FString& PinSubCategory, UObject* PinSubCategoryObject, const FString& PinName, EPinContainerType PinContainerType /* EPinContainerType::None */, bool bIsReference /* = false */, bool bIsConst /*= false*/, int32 Index /*= INDEX_NONE*/, const FEdGraphTerminalType& ValueTerminalType /*= FEdGraphTerminalType()*/)
+{
+	FEdGraphPinType PinType(PinCategory, PinSubCategory, PinSubCategoryObject, PinContainerType, bIsReference, ValueTerminalType);
 	PinType.bIsConst = bIsConst;
 
 	return CreatePin(Dir, PinType, PinName, Index);
diff --git a/Engine/Source/Runtime/Engine/Private/EdGraph/EdGraphPin.cpp b/Engine/Source/Runtime/Engine/Private/EdGraph/EdGraphPin.cpp
--- a/Engine/Source/Runtime/Engine/Private/EdGraph/EdGraphPin.cpp
+++ b/Engine/Source/Runtime/Engine/Private/EdGraph/EdGraphPin.cpp
@@ -2,6 +2,7 @@
 
 #include "EdGraph/EdGraphPin.h"
 #include "UObject/BlueprintsObjectVersion.h"
+#include "UObject/FrameworkObjectVersion.h"
 #include "UObject/UnrealType.h"
 #include "UObject/TextProperty.h"
 #include "EdGraph/EdGraph.h"
@@ -82,21 +83,46 @@ bool FEdGraphPinType::Serialize(FArchive& Ar)
 		}
 	}
 
-	Ar.UsingCustomVersion(FBlueprintsObjectVersion::GUID);
-	if (Ar.CustomVer(FBlueprintsObjectVersion::GUID) >= FBlueprintsObjectVersion::AdvancedContainerSupport)
+	Ar.UsingCustomVersion(FFrameworkObjectVersion::GUID);
+	if (Ar.CustomVer(FFrameworkObjectVersion::GUID) >= FFrameworkObjectVersion::EdGraphPinContainerType)
 	{
-		Ar << bIsMap;
-		if (bIsMap)
+		Ar << ContainerType;
+		if (IsMap())
 		{
 			Ar << PinValueType;
 		}
-		Ar << bIsSet;
 	}
+	else
+	{
+		bool bIsMap = false;
+		bool bIsSet = false;
+		bool bIsArray = false;
+
+		Ar.UsingCustomVersion(FBlueprintsObjectVersion::GUID);
+		if (Ar.CustomVer(FBlueprintsObjectVersion::GUID) >= FBlueprintsObjectVersion::AdvancedContainerSupport)
+		{
+			Ar << bIsMap;
+			if (bIsMap)
+			{
+				Ar << PinValueType;
+			}
+			Ar << bIsSet;
+		}
+
+		Ar << bIsArray;
+
+		if (Ar.IsLoading())
+		{
+			ContainerType = ToPinContainerType(bIsArray, bIsSet, bIsMap);
+		}
+	}
+
+	bool bIsReferenceBool = bIsReference;
+	bool bIsWeakPointerBool = bIsWeakPointer;
+
+	Ar << bIsReferenceBool;
+	Ar << bIsWeakPointerBool;
 
-	Ar << bIsArray;
-	Ar << bIsReference;
-	Ar << bIsWeakPointer;
-	
 	if (Ar.UE4Ver() >= VER_UE4_MEMBERREFERENCE_IN_PINTYPE)
 	{
 		Ar << PinSubCategoryMemberReference;
@@ -118,13 +144,18 @@ bool FEdGraphPinType::Serialize(FArchive& Ar)
 		}
 	}
 
+	bool bIsConstBool = false;
+
 	if (Ar.UE4Ver() >= VER_UE4_SERIALIZE_PINTYPE_CONST)
 	{
-		Ar << bIsConst;
+		Ar << bIsConstBool;
 	}
-	else if (Ar.IsLoading())
+
+	if (Ar.IsLoading())
 	{
-		bIsConst = false;
+		bIsReference = bIsReferenceBool;
+		bIsWeakPointer = bIsWeakPointerBool;
+		bIsConst = bIsConstBool;
 	}
 
 	return true;
@@ -145,9 +176,7 @@ FEdGraphPinType FEdGraphPinType::GetTerminalTypeForContainer( const FEdGraphPinT
 	ensure(ContainerType.IsContainer());
 	
 	FEdGraphPinType TerminalType = ContainerType;
-	TerminalType.bIsArray = false;
-	TerminalType.bIsMap = false;
-	TerminalType.bIsSet = false;
+	TerminalType.ContainerType = EPinContainerType::None;
 	return TerminalType;
 }
 
@@ -352,7 +381,6 @@ void UEdGraphPin::BreakAllPinLinks()
 	}
 }
 
-
 void UEdGraphPin::CopyPersistentDataFromOldPin(const UEdGraphPin& SourcePin)
 {
 	// The name matches already, doesn't get copied here
@@ -361,7 +389,7 @@ void UEdGraphPin::CopyPersistentDataFromOldPin(const UEdGraphPin& SourcePin)
 	PinId = SourcePin.PinId;
 
 	// Only move the default value if it was modified; inherit the new default value otherwise
-	if (SourcePin.DefaultValue != SourcePin.AutogeneratedDefaultValue || SourcePin.DefaultObject != nullptr || SourcePin.DefaultTextValue.ToString() != SourcePin.AutogeneratedDefaultValue)
+	if (!SourcePin.DoesDefaultValueMatchAutogenerated())
 	{
 		DefaultObject = SourcePin.DefaultObject;
 		DefaultValue = SourcePin.DefaultValue;
@@ -473,6 +501,12 @@ FString UEdGraphPin::GetDefaultAsString() const
 	}
 }
 
+bool UEdGraphPin::DoesDefaultValueMatchAutogenerated() const
+{
+	// Ignoring case on purpose to match default behavior
+	return GetDefaultAsString().Equals(AutogeneratedDefaultValue, ESearchCase::IgnoreCase);
+}
+
 #if WITH_EDITORONLY_DATA
 FText UEdGraphPin::GetDisplayName() const
 {
diff --git a/Engine/Source/Runtime/Engine/Private/GameFramework/SpringArmComponent.cpp b/Engine/Source/Runtime/Engine/Private/GameFramework/SpringArmComponent.cpp
--- a/Engine/Source/Runtime/Engine/Private/GameFramework/SpringArmComponent.cpp
+++ b/Engine/Source/Runtime/Engine/Private/GameFramework/SpringArmComponent.cpp
@@ -48,6 +48,18 @@ void USpringArmComponent::UpdateDesiredArmLocation(bool bDoTrace, bool bDoLocati
 {
 	FRotator DesiredRot = GetComponentRotation();
 
+	if (bUsePawnControlRotation)
+	{
+		if (APawn* OwningPawn = Cast<APawn>(GetOwner()))
+		{
+			const FRotator PawnViewRotation = OwningPawn->GetViewRotation();
+			if (DesiredRot != PawnViewRotation)
+			{
+				DesiredRot = PawnViewRotation;
+			}
+		}
+	}
+
 	// If inheriting rotation, check options for which components to inherit
 	if(!bAbsoluteRotation)
 	{
@@ -173,7 +185,7 @@ void USpringArmComponent::UpdateDesiredArmLocation(bool bDoTrace, bool bDoLocati
 	// Form a transform for new world transform for camera
 	FTransform WorldCamTM(DesiredRot, ResultLoc);
 	// Convert to relative to component
-	FTransform RelCamTM = WorldCamTM.GetRelativeTransform(ComponentToWorld);
+	FTransform RelCamTM = WorldCamTM.GetRelativeTransform(GetComponentTransform());
 
 	// Update socket location/rotation
 	RelativeSocketLocation = RelCamTM.GetLocation();
@@ -213,19 +225,6 @@ void USpringArmComponent::PostLoad()
 void USpringArmComponent::TickComponent(float DeltaTime, enum ELevelTick TickType, FActorComponentTickFunction* ThisTickFunction)
 {
 	Super::TickComponent(DeltaTime, TickType, ThisTickFunction);
-
-	if (bUsePawnControlRotation)
-	{
-		if (APawn* OwningPawn = Cast<APawn>(GetOwner()))
-		{
-			const FRotator PawnViewRotation = OwningPawn->GetViewRotation();
-			if (PawnViewRotation != GetComponentRotation())
-			{
-				SetWorldRotation(PawnViewRotation);
-			}
-		}
-	}
-
 	UpdateDesiredArmLocation(bDoCollisionTest, bEnableCameraLag, bEnableCameraRotationLag, DeltaTime);
 }
 
@@ -237,14 +236,14 @@ FTransform USpringArmComponent::GetSocketTransform(FName InSocketName, ERelative
 	{
 		case RTS_World:
 		{
-			return RelativeTransform * ComponentToWorld;
+			return RelativeTransform * GetComponentTransform();
 			break;
 		}
 		case RTS_Actor:
 		{
 			if( const AActor* Actor = GetOwner() )
 			{
-				FTransform SocketTransform = RelativeTransform * ComponentToWorld;
+				FTransform SocketTransform = RelativeTransform * GetComponentTransform();
 				return SocketTransform.GetRelativeTransform(Actor->GetTransform());
 			}
 			break;
diff --git a/Engine/Source/Runtime/Engine/Private/HierarchicalInstancedStaticMesh.cpp b/Engine/Source/Runtime/Engine/Private/HierarchicalInstancedStaticMesh.cpp
--- a/Engine/Source/Runtime/Engine/Private/HierarchicalInstancedStaticMesh.cpp
+++ b/Engine/Source/Runtime/Engine/Private/HierarchicalInstancedStaticMesh.cpp
@@ -851,7 +851,7 @@ public:
 			//int32 NumPerNode = (1 + ClusterTree[0].LastInstance - ClusterTree[0].FirstInstance) / NumNodes;
 			//UE_LOG(LogTemp, Display, TEXT("Occlusion level %d   %d inst / node"), NumNodes, NumPerNode);
 			OcclusionBounds.Reserve(NumNodes);
-			FMatrix XForm = InComponent->ComponentToWorld.ToMatrixWithScale();
+			FMatrix XForm = InComponent->GetComponentTransform().ToMatrixWithScale();
 			for (int32 Index = FirstOcclusionNode; Index <= LastOcclusionNode; Index++)
 			{
 				OcclusionBounds.Add(FBoxSphereBounds(FBox(ClusterTree[Index].BoundMin, ClusterTree[Index].BoundMax).TransformBy(XForm)));
@@ -1974,7 +1974,7 @@ bool UHierarchicalInstancedStaticMeshComponent::UpdateInstanceTransform(int32 In
 
 	int32 RenderIndex = InstanceReorderTable[InstanceIndex];
 	const FMatrix OldTransform = PerInstanceSMData[InstanceIndex].Transform;
-	const FTransform NewLocalTransform = bWorldSpace ? NewInstanceTransform.GetRelativeTransform(ComponentToWorld) : NewInstanceTransform;
+	const FTransform NewLocalTransform = bWorldSpace ? NewInstanceTransform.GetRelativeTransform(GetComponentTransform()) : NewInstanceTransform;
 	const FVector NewLocalLocation = NewLocalTransform.GetTranslation();
 
 	// if we are only updating rotation/scale we update the instance directly in the cluster tree
@@ -2682,7 +2682,7 @@ static void GatherInstanceTransformsInArea(const UHierarchicalInstancedStaticMes
 	if (ClusterTree.Num())
 	{
 		const FClusterNode& ChildNode = ClusterTree[Child];
-		const FBox WorldNodeBox = FBox(ChildNode.BoundMin, ChildNode.BoundMax).TransformBy(Component.ComponentToWorld);
+		const FBox WorldNodeBox = FBox(ChildNode.BoundMin, ChildNode.BoundMax).TransformBy(Component.GetComponentTransform());
 	
 		if (AreaBox.Intersect(WorldNodeBox))
 		{
@@ -2712,7 +2712,7 @@ static void GatherInstanceTransformsInArea(const UHierarchicalInstancedStaticMes
 					
 					if (!InstanceToComponent.GetScale3D().IsZero())
 					{
-						InstanceData.Add(InstanceToComponent*Component.ComponentToWorld);
+						InstanceData.Add(InstanceToComponent*Component.GetComponentTransform());
 					}
 				}
 			}
@@ -2823,7 +2823,7 @@ void UHierarchicalInstancedStaticMeshComponent::PartialNavigationUpdate(int32 In
 		if (NavSys && NavSys->GetObjectsNavOctreeId(this))
 		{
 			FTransform InstanceTransform(PerInstanceSMData[InstanceIdx].Transform);
-			FBox InstanceBox = GetStaticMesh()->GetBounds().TransformBy(InstanceTransform*ComponentToWorld).GetBox(); // in world space
+			FBox InstanceBox = GetStaticMesh()->GetBounds().TransformBy(InstanceTransform*GetComponentTransform()).GetBox(); // in world space
 			AccumulatedNavigationDirtyArea+= InstanceBox;
 		}
 	}
@@ -2840,7 +2840,7 @@ void UHierarchicalInstancedStaticMeshComponent::FlushAccumulatedNavigationUpdate
 		// Check if this component is registered in navigation system
 		if (ClusterTree.Num() && NavSys && NavSys->GetObjectsNavOctreeId(this))
 		{
-			FBox NewBounds = FBox(ClusterTree[0].BoundMin, ClusterTree[0].BoundMax).TransformBy(ComponentToWorld);
+			FBox NewBounds = FBox(ClusterTree[0].BoundMin, ClusterTree[0].BoundMax).TransformBy(GetComponentTransform());
 			NavSys->UpdateNavOctreeElementBounds(this, NewBounds, AccumulatedNavigationDirtyArea);
 		}
 			
@@ -2853,7 +2853,7 @@ static void GatherInstancesOverlappingArea(const UHierarchicalInstancedStaticMes
 {
 	const TArray<FClusterNode>& ClusterTree = *Component.ClusterTreePtr;
 	const FClusterNode& ChildNode = ClusterTree[Child];
-	const FBox WorldNodeBox = FBox(ChildNode.BoundMin, ChildNode.BoundMax).TransformBy(Component.ComponentToWorld);
+	const FBox WorldNodeBox = FBox(ChildNode.BoundMin, ChildNode.BoundMax).TransformBy(Component.GetComponentTransform());
 
 	if (AreaBox.Intersect(WorldNodeBox))
 	{
@@ -2897,11 +2897,11 @@ TArray<int32> UHierarchicalInstancedStaticMeshComponent::GetInstancesOverlapping
 		FBox WorldSpaceAABB(Sphere.Center - FVector(Sphere.W), Sphere.Center + FVector(Sphere.W));
 		if (bSphereInWorldSpace)
 		{
-			Sphere = Sphere.TransformBy(ComponentToWorld.Inverse());
+			Sphere = Sphere.TransformBy(GetComponentTransform().Inverse());
 		}
 		else
 		{
-			WorldSpaceAABB = WorldSpaceAABB.TransformBy(ComponentToWorld);
+			WorldSpaceAABB = WorldSpaceAABB.TransformBy(GetComponentTransform());
 		}
 
 		const float StaticMeshBoundsRadius = GetStaticMesh()->GetBounds().SphereRadius;
@@ -2930,11 +2930,11 @@ TArray<int32> UHierarchicalInstancedStaticMeshComponent::GetInstancesOverlapping
 		FBox LocalSpaceSpaceBox(InBox);
 		if (bBoxInWorldSpace)
 		{
-			LocalSpaceSpaceBox = LocalSpaceSpaceBox.TransformBy(ComponentToWorld.Inverse());
+			LocalSpaceSpaceBox = LocalSpaceSpaceBox.TransformBy(GetComponentTransform().Inverse());
 		}
 		else
 		{
-			WorldSpaceBox = WorldSpaceBox.TransformBy(ComponentToWorld);
+			WorldSpaceBox = WorldSpaceBox.TransformBy(GetComponentTransform());
 		}
 
 		const FBox StaticMeshBox = GetStaticMesh()->GetBounds().GetBox();
diff --git a/Engine/Source/Runtime/Engine/Private/InstancedStaticMesh.cpp b/Engine/Source/Runtime/Engine/Private/InstancedStaticMesh.cpp
--- a/Engine/Source/Runtime/Engine/Private/InstancedStaticMesh.cpp
+++ b/Engine/Source/Runtime/Engine/Private/InstancedStaticMesh.cpp
@@ -831,7 +831,7 @@ FActorComponentInstanceData* UInstancedStaticMeshComponent::GetComponentInstance
 		InstanceData = StaticMeshInstanceData = new FInstancedStaticMeshComponentInstanceData(*this);	
 
 		// Fill in info (copied from UStaticMeshComponent::GetComponentInstanceData)
-		StaticMeshInstanceData->CachedStaticLighting.Transform = ComponentToWorld;
+		StaticMeshInstanceData->CachedStaticLighting.Transform = GetComponentTransform();
 
 		for (const FStaticMeshComponentLODInfo& LODDataEntry : LODData)
 		{
@@ -864,7 +864,7 @@ void UInstancedStaticMeshComponent::ApplyComponentInstanceData(FInstancedStaticM
 
 	// Check for any instance having moved as that would invalidate static lighting
 	if (PerInstanceSMData.Num() == InstancedMeshData->PerInstanceSMData.Num() &&
-		InstancedMeshData->CachedStaticLighting.Transform.Equals(ComponentToWorld))
+		InstancedMeshData->CachedStaticLighting.Transform.Equals(GetComponentTransform()))
 	{
 		bMatch = true;
 
@@ -949,7 +949,7 @@ void UInstancedStaticMeshComponent::InitInstanceBody(int32 InstanceIdx, FBodyIns
 	check(BodySetup);
 
 	// Get transform of the instance
-	FTransform InstanceTransform = FTransform(PerInstanceSMData[InstanceIdx].Transform) * ComponentToWorld;
+	FTransform InstanceTransform = FTransform(PerInstanceSMData[InstanceIdx].Transform) * GetComponentTransform();
 	
 	InstanceBodyInstance->CopyBodyInstancePropertiesFrom(&BodyInstance);
 	InstanceBodyInstance->InstanceBodyIndex = InstanceIdx; // Set body index 
@@ -985,7 +985,7 @@ void UInstancedStaticMeshComponent::CreateAllInstanceBodies()
 	    Transforms.Reserve(NumBodies);
 	    for (int32 i = 0; i < NumBodies; ++i)
 	    {
-			const FTransform InstanceTM = FTransform(PerInstanceSMData[i].Transform) * ComponentToWorld;
+			const FTransform InstanceTM = FTransform(PerInstanceSMData[i].Transform) * GetComponentTransform();
 			if (InstanceTM.GetScale3D().IsNearlyZero())
 			{
 				InstanceBodies[i] = nullptr;
@@ -1401,7 +1401,7 @@ int32 UInstancedStaticMeshComponent::AddInstance(const FTransform& InstanceTrans
 int32 UInstancedStaticMeshComponent::AddInstanceWorldSpace(const FTransform& WorldTransform)
  {
 	// Transform from world space to local space
-	FTransform RelativeTM = WorldTransform.GetRelativeTransform(ComponentToWorld);
+	FTransform RelativeTM = WorldTransform.GetRelativeTransform(GetComponentTransform());
 	return AddInstance(RelativeTM);
 }
 
@@ -1453,7 +1453,7 @@ bool UInstancedStaticMeshComponent::GetInstanceTransform(int32 InstanceIndex, FT
 	OutInstanceTransform = FTransform(InstanceData.Transform);
 	if (bWorldSpace)
 	{
-		OutInstanceTransform = OutInstanceTransform * ComponentToWorld;
+		OutInstanceTransform = OutInstanceTransform * GetComponentTransform();
 	}
 
 	return true;
@@ -1472,7 +1472,7 @@ void UInstancedStaticMeshComponent::OnUpdateTransform(EUpdateTransformFlags Upda
 		for (int32 i = 0; i < PerInstanceSMData.Num(); i++)
 		{
 			const FTransform InstanceTransform(PerInstanceSMData[i].Transform);
-			UpdateInstanceTransform(i, InstanceTransform * ComponentToWorld, /* bWorldSpace= */true, /* bMarkRenderStateDirty= */false, bTeleport);
+			UpdateInstanceTransform(i, InstanceTransform * GetComponentTransform(), /* bWorldSpace= */true, /* bMarkRenderStateDirty= */false, bTeleport);
 		}
 	}
 }
@@ -1495,13 +1495,13 @@ bool UInstancedStaticMeshComponent::UpdateInstanceTransform(int32 InstanceIndex,
 	// Should find some way around this for performance.
     
 	// Render data uses local transform of the instance
-	FTransform LocalTransform = bWorldSpace ? NewInstanceTransform.GetRelativeTransform(ComponentToWorld) : NewInstanceTransform;
+	FTransform LocalTransform = bWorldSpace ? NewInstanceTransform.GetRelativeTransform(GetComponentTransform()) : NewInstanceTransform;
 	InstanceData.Transform = LocalTransform.ToMatrixWithScale();
 
 	if (bPhysicsStateCreated)
 	{
 		// Physics uses world transform of the instance
-		FTransform WorldTransform = bWorldSpace ? NewInstanceTransform : (LocalTransform * ComponentToWorld);
+		FTransform WorldTransform = bWorldSpace ? NewInstanceTransform : (LocalTransform * GetComponentTransform());
 		FBodyInstance*& InstanceBodyInstance = InstanceBodies[InstanceIndex];
 #if WITH_PHYSX
 		if (NewInstanceTransform.GetScale3D().IsNearlyZero())
@@ -1552,7 +1552,7 @@ TArray<int32> UInstancedStaticMeshComponent::GetInstancesOverlappingSphere(const
 	FSphere Sphere(Center, Radius);
 	if (bSphereInWorldSpace)
 	{
-		Sphere = Sphere.TransformBy(ComponentToWorld.Inverse());
+		Sphere = Sphere.TransformBy(GetComponentTransform().Inverse());
 	}
 
 	float StaticMeshBoundsRadius = GetStaticMesh()->GetBounds().SphereRadius;
@@ -1578,7 +1578,7 @@ TArray<int32> UInstancedStaticMeshComponent::GetInstancesOverlappingBox(const FB
 	FBox Box(InBox);
 	if (bBoxInWorldSpace)
 	{
-		Box = Box.TransformBy(ComponentToWorld.Inverse());
+		Box = Box.TransformBy(GetComponentTransform().Inverse());
 	}
 
 	FVector StaticMeshBoundsExtent = GetStaticMesh()->GetBounds().BoxExtent;
@@ -1782,7 +1782,7 @@ void UInstancedStaticMeshComponent::GetNavigationPerInstanceTransforms(const FBo
 		const FTransform InstanceToComponent(InstancedData.Transform);
 		if (!InstanceToComponent.GetScale3D().IsZero())
 		{
-			InstanceData.Add(InstanceToComponent*ComponentToWorld);
+			InstanceData.Add(InstanceToComponent*GetComponentTransform());
 		}
 	}
 }
diff --git a/Engine/Source/Runtime/Engine/Private/InstancedStaticMesh.h b/Engine/Source/Runtime/Engine/Private/InstancedStaticMesh.h
--- a/Engine/Source/Runtime/Engine/Private/InstancedStaticMesh.h
+++ b/Engine/Source/Runtime/Engine/Private/InstancedStaticMesh.h
@@ -688,7 +688,7 @@ public:
 		: FStaticMeshStaticLightingMesh(InPrimitive, LODIndex, InRelevantLights)
 	{
 		// override the local to world to combine the per instance transform with the component's standard transform
-		SetLocalToWorld(InPrimitive->PerInstanceSMData[InstanceIndex].Transform * InPrimitive->ComponentToWorld.ToMatrixWithScale());
+		SetLocalToWorld(InPrimitive->PerInstanceSMData[InstanceIndex].Transform * InPrimitive->GetComponentTransform().ToMatrixWithScale());
 	}
 };
 
diff --git a/Engine/Source/Runtime/Engine/Private/KismetTextLibrary.cpp b/Engine/Source/Runtime/Engine/Private/KismetTextLibrary.cpp
--- a/Engine/Source/Runtime/Engine/Private/KismetTextLibrary.cpp
+++ b/Engine/Source/Runtime/Engine/Private/KismetTextLibrary.cpp
@@ -6,18 +6,10 @@
 
 #define LOCTEXT_NAMESPACE "Kismet"
 
-
-/* UKismetTextLibrary structors
- *****************************************************************************/
-
 UKismetTextLibrary::UKismetTextLibrary(const FObjectInitializer& ObjectInitializer)
 	: Super(ObjectInitializer)
 { }
 
-
-/* UKismetTextLibrary static functions
- *****************************************************************************/
-
 FText UKismetTextLibrary::Conv_VectorToText(FVector InVec)
 {
 	FFormatNamedArguments Args;
@@ -28,7 +20,6 @@ FText UKismetTextLibrary::Conv_VectorToText(FVector InVec)
 	return FText::Format(NSLOCTEXT("Core", "Vector3", "X={X} Y={Y} Z={Z}"), Args);
 }
 
-
 FText UKismetTextLibrary::Conv_Vector2dToText(FVector2D InVec)
 {
 	FFormatNamedArguments Args;
@@ -38,7 +29,6 @@ FText UKismetTextLibrary::Conv_Vector2dToText(FVector2D InVec)
 	return FText::Format(NSLOCTEXT("Core", "Vector2", "X={X} Y={Y}"), Args);
 }
 
-
 FText UKismetTextLibrary::Conv_RotatorToText(FRotator InRot)
 {
 	FFormatNamedArguments Args;
@@ -49,18 +39,29 @@ FText UKismetTextLibrary::Conv_RotatorToText(FRotator InRot)
 	return FText::Format(NSLOCTEXT("Core", "Rotator", "P={P} Y={Y} R={R}"), Args);
 }
 
-
 FText UKismetTextLibrary::Conv_TransformToText(const FTransform& InTrans)
 {
-	return FText::FromString(InTrans.ToString());
-}
+	const FVector T(InTrans.GetTranslation());
+	const FRotator R(InTrans.Rotator());
+	const FVector S(InTrans.GetScale3D());
+
+	FFormatNamedArguments Args;
+	Args.Add(TEXT("T"), Conv_VectorToText(T));
+	Args.Add(TEXT("R"), Conv_RotatorToText(R));
+	Args.Add(TEXT("S"), Conv_VectorToText(S));
 
+	return FText::Format(NSLOCTEXT("Core", "Transform", "Translation: {T} Rotation: {R} Scale: {S}"), Args);
+}
 
 FText UKismetTextLibrary::Conv_ObjectToText(class UObject* InObj)
 {
-	return FText::FromString((InObj != NULL) ? InObj->GetName() : FString(TEXT("None")));
-}
+	if (InObj)
+	{
+		return FText::AsCultureInvariant(InObj->GetName());
+	}
 
+	return GNone;
+}
 
 FText UKismetTextLibrary::Conv_ColorToText(FLinearColor InColor)
 {
@@ -73,121 +74,101 @@ FText UKismetTextLibrary::Conv_ColorToText(FLinearColor InColor)
 	return FText::Format(NSLOCTEXT("Core", "LinearColor", "R={R} G={G} B={B} A={A}"), Args);
 }
 
-
 FString UKismetTextLibrary::Conv_TextToString(const FText& InText)
 {
 	return InText.ToString();
 }
 
-
 FText UKismetTextLibrary::Conv_StringToText(const FString& InString)
 {	
-	return FText::FromString(InString);
+	return FText::AsCultureInvariant(InString);
 }
 
-
 FText UKismetTextLibrary::Conv_NameToText(FName InName)
 {
-	return FText::FromName(InName);
+	return FText::AsCultureInvariant(InName.ToString());
 }
 
-
 bool UKismetTextLibrary::TextIsEmpty(const FText& InText)
 {
 	return InText.IsEmpty();
 }
 
-
 bool UKismetTextLibrary::TextIsTransient(const FText& InText)
 {
 	return InText.IsTransient();
 }
 
-
 bool UKismetTextLibrary::TextIsCultureInvariant(const FText& InText)
 {
 	return InText.IsCultureInvariant();
 }
 
-
 FText UKismetTextLibrary::TextToLower(const FText& InText)
 {
 	return InText.ToLower();
 }
 
-
 FText UKismetTextLibrary::TextToUpper(const FText& InText)
 {
 	return InText.ToUpper();
 }
 
-
 FText UKismetTextLibrary::TextTrimPreceding(const FText& InText)
 {
 	return FText::TrimPreceding(InText);
 }
 
-
 FText UKismetTextLibrary::TextTrimTrailing(const FText& InText)
 {
 	return FText::TrimTrailing(InText);
 }
 
-
 FText UKismetTextLibrary::TextTrimPrecedingAndTrailing(const FText& InText)
 {
 	return FText::TrimPrecedingAndTrailing(InText);
 }
 
-
 FText UKismetTextLibrary::GetEmptyText()
 {
 	return FText::GetEmpty();
 }
 
-
 bool UKismetTextLibrary::FindTextInLocalizationTable(const FString& Namespace, const FString& Key, FText& OutText)
 {
 	return FText::FindText(Namespace, Key, OutText);
 }
 
-
 bool UKismetTextLibrary::EqualEqual_IgnoreCase_TextText(const FText& A, const FText& B)
 {
 	return A.EqualToCaseIgnored( B );
 }
 
-
 bool UKismetTextLibrary::EqualEqual_TextText(const FText& A, const FText& B)
 {
 	return A.EqualTo( B );
 }
 
-
 bool UKismetTextLibrary::NotEqual_IgnoreCase_TextText(const FText& A, const FText& B)
 {
 	return !A.EqualToCaseIgnored( B );
 }
 
-
 bool UKismetTextLibrary::NotEqual_TextText(const FText& A, const FText& B)
 {
 	return !A.EqualTo( B );
 }
 
-
 FText UKismetTextLibrary::Conv_BoolToText(bool InBool)
 {
 	return InBool ? LOCTEXT("True", "true") : LOCTEXT("False", "false");
 }
 
-
 FText UKismetTextLibrary::Conv_ByteToText(uint8 Value)
 {
 	return FText::AsNumber(Value, &FNumberFormattingOptions::DefaultNoGrouping());
 }
 
-
 FText UKismetTextLibrary::Conv_IntToText(int32 Value, bool bUseGrouping/* = true*/, int32 MinimumIntegralDigits/* = 1*/, int32 MaximumIntegralDigits/* = 324*/)
 {
 	// Only update the values that need to be changed from the default FNumberFormattingOptions, 
@@ -200,7 +181,6 @@ FText UKismetTextLibrary::Conv_IntToText(int32 Value, bool bUseGrouping/* = true
 	return FText::AsNumber(Value, &NumberFormatOptions);
 }
 
-
 FText UKismetTextLibrary::Conv_FloatToText(float Value, TEnumAsByte<ERoundingMode> RoundingMode, bool bUseGrouping/* = true*/, int32 MinimumIntegralDigits/* = 1*/, int32 MaximumIntegralDigits/* = 324*/, int32 MinimumFractionalDigits/* = 0*/, int32 MaximumFractionalDigits/* = 3*/)
 {
 	FNumberFormattingOptions NumberFormatOptions;
@@ -258,67 +238,56 @@ FText UKismetTextLibrary::AsPercent_Float(float Value, TEnumAsByte<ERoundingMode
 	return FText::AsPercent(Value, &NumberFormatOptions);
 }
 
-
 FText UKismetTextLibrary::AsDate_DateTime(const FDateTime& InDateTime)
 {
 	return FText::AsDate(InDateTime, EDateTimeStyle::Default, FText::GetInvariantTimeZone());
 }
 
-
 FText UKismetTextLibrary::AsTimeZoneDate_DateTime(const FDateTime& InDateTime, const FString& InTimeZone)
 {
 	return FText::AsDate(InDateTime, EDateTimeStyle::Default, InTimeZone);
 }
 
-
 FText UKismetTextLibrary::AsDateTime_DateTime(const FDateTime& InDateTime)
 {
 	return FText::AsDateTime(InDateTime, EDateTimeStyle::Default, EDateTimeStyle::Default, FText::GetInvariantTimeZone());
 }
 
-
 FText UKismetTextLibrary::AsTimeZoneDateTime_DateTime(const FDateTime& InDateTime, const FString& InTimeZone)
 {
 	return FText::AsDateTime(InDateTime, EDateTimeStyle::Default, EDateTimeStyle::Default, InTimeZone);
 }
 
-
 FText UKismetTextLibrary::AsTime_DateTime(const FDateTime& InDateTime)
 {
 	return FText::AsTime(InDateTime, EDateTimeStyle::Default, FText::GetInvariantTimeZone());
 }
 
-
 FText UKismetTextLibrary::AsTimeZoneTime_DateTime(const FDateTime& InDateTime, const FString& InTimeZone)
 {
 	return FText::AsTime(InDateTime, EDateTimeStyle::Default, InTimeZone);
 }
 
-
 FText UKismetTextLibrary::AsTimespan_Timespan(const FTimespan& InTimespan)
 {
 	return FText::AsTimespan(InTimespan);
 }
 
-
 FText UKismetTextLibrary::Format(FText InPattern, TArray<FFormatArgumentData> InArgs)
 {
 	return FTextFormatter::Format(MoveTemp(InPattern), MoveTemp(InArgs), false, false);
 }
 
-
 bool UKismetTextLibrary::TextIsFromStringTable(const FText& Text)
 {
 	return Text.IsFromStringTable();
 }
 
-
 FText UKismetTextLibrary::TextFromStringTable(const FName TableId, const FString& Key)
 {
 	return FText::FromStringTable(TableId, Key);
 }
 
-
 bool UKismetTextLibrary::StringTableIdAndKeyFromText(FText Text, FName& OutTableId, FString& OutKey)
 {
 	return FStringTableRegistry::Get().FindTableIdAndKey(Text, OutTableId, OutKey);
diff --git a/Engine/Source/Runtime/Engine/Private/LevelActor.cpp b/Engine/Source/Runtime/Engine/Private/LevelActor.cpp
--- a/Engine/Source/Runtime/Engine/Private/LevelActor.cpp
+++ b/Engine/Source/Runtime/Engine/Private/LevelActor.cpp
@@ -605,7 +605,7 @@ bool UWorld::DestroyActor( AActor* ThisActor, bool bNetForce, bool bShouldModify
 			OldParentActor->Modify();
 		}
 
-		ThisActor->DetachRootComponentFromParent();
+		ThisActor->DetachFromActor(FDetachmentTransformRules::KeepWorldTransform);
 
 #if WITH_EDITOR
 		if( GIsEditor )
diff --git a/Engine/Source/Runtime/Engine/Private/LevelBounds.cpp b/Engine/Source/Runtime/Engine/Private/LevelBounds.cpp
--- a/Engine/Source/Runtime/Engine/Private/LevelBounds.cpp
+++ b/Engine/Source/Runtime/Engine/Private/LevelBounds.cpp
@@ -43,7 +43,7 @@ FBox ALevelBounds::GetComponentsBoundingBox(bool bNonColliding) const
 {
 	checkf(RootComponent != nullptr, TEXT("LevelBounds actor with null root component: %s"), *GetPathNameSafe(this));
 	FVector BoundsCenter = RootComponent->GetComponentLocation();
-	FVector BoundsExtent = RootComponent->ComponentToWorld.GetScale3D() * 0.5f;
+	FVector BoundsExtent = RootComponent->GetComponentTransform().GetScale3D() * 0.5f;
 	return FBox(BoundsCenter - BoundsExtent, 
 				BoundsCenter + BoundsExtent);
 }
diff --git a/Engine/Source/Runtime/Engine/Private/LevelStreaming.cpp b/Engine/Source/Runtime/Engine/Private/LevelStreaming.cpp
--- a/Engine/Source/Runtime/Engine/Private/LevelStreaming.cpp
+++ b/Engine/Source/Runtime/Engine/Private/LevelStreaming.cpp
@@ -227,7 +227,7 @@ void ULevelStreaming::PostLoad()
 			// Convert the FName reference to a TAssetPtr, then broadcast that we loaded a reference
 			// so this reference is gathered by the cooker without having to resave the package.
 			SetWorldAssetByPackageName(PackageName_DEPRECATED);
-			FCoreUObjectDelegates::StringAssetReferenceLoaded.ExecuteIfBound(WorldAsset.ToStringReference().ToString());
+			WorldAsset.GetUniqueID().PostLoadPath();
 		}
 		else
 		{
@@ -899,7 +899,7 @@ FBox ULevelStreaming::GetStreamingVolumeBounds()
 		ALevelStreamingVolume* StreamingVol = EditorStreamingVolumes[VolIdx];
 		if(StreamingVol && StreamingVol->GetBrushComponent())
 		{
-			Bounds += StreamingVol->GetBrushComponent()->BrushBodySetup->AggGeom.CalcAABB(StreamingVol->GetBrushComponent()->ComponentToWorld);
+			Bounds += StreamingVol->GetBrushComponent()->BrushBodySetup->AggGeom.CalcAABB(StreamingVol->GetBrushComponent()->GetComponentTransform());
 		}
 	}
 
diff --git a/Engine/Source/Runtime/Engine/Private/LevelTick.cpp b/Engine/Source/Runtime/Engine/Private/LevelTick.cpp
--- a/Engine/Source/Runtime/Engine/Private/LevelTick.cpp
+++ b/Engine/Source/Runtime/Engine/Private/LevelTick.cpp
@@ -1155,48 +1155,55 @@ void EndTickDrawEvent(TDrawEvent<FRHICommandList>* TickDrawEvent)
 }
 
 
-void FTickableGameObject::TickObjects(UWorld* World, int32 InTickType, bool bIsPaused, float DeltaSeconds)
+void FTickableGameObject::TickObjects(UWorld* World, const int32 InTickType, const bool bIsPaused, const float DeltaSeconds)
 {
-	check(!bIsTickingObjects);
-	bIsTickingObjects = true;
+	if (TickableObjects.Num() > 0)
+	{
+		check(!bIsTickingObjects);
+		bIsTickingObjects = true;
 
-	bool bNeedsCleanup = false;
-	ELevelTick TickType = (ELevelTick)InTickType;
+		bool bNeedsCleanup = false;
+		const ELevelTick TickType = (ELevelTick)InTickType;
 
-	for( int32 i=0; i < TickableObjects.Num(); ++i )
-	{
-		if (FTickableGameObject* TickableObject = TickableObjects[i])
+		for (int32 i = 0; i < TickableObjects.Num(); ++i)
 		{
-			const bool bTickIt = TickableObject->IsTickable() && (TickableObject->GetTickableGameObjectWorld() == World) &&
-				(
-					(TickType != LEVELTICK_TimeOnly && !bIsPaused) ||
-					(bIsPaused && TickableObject->IsTickableWhenPaused()) ||
-					(GIsEditor && (World == nullptr || !World->IsPlayInEditor()) && TickableObject->IsTickableInEditor())
-					);
-
-			if (bTickIt)
+			if (FTickableGameObject* TickableObject = TickableObjects[i])
 			{
-				STAT(FScopeCycleCounter Context(TickableObject->GetStatId());)
-				TickableObject->Tick(DeltaSeconds);
-
-				if (TickableObjects[i] == nullptr)
+				// If it is tickable and in this world
+				if (TickableObject->IsTickable() && (TickableObject->GetTickableGameObjectWorld() == World))
 				{
-					bNeedsCleanup = true;
+					// If this is a game world or the object is tickable in editor
+					const bool bIsGameWorld = InTickType == LEVELTICK_All || (World && World->IsGameWorld());
+
+					if (bIsGameWorld || (GIsEditor && TickableObject->IsTickableInEditor()))
+					{
+						// Then tick if we are not doing a time only (paused) update and we are not paused or the object is tickable when paused
+						if ((!bIsPaused && TickType != LEVELTICK_TimeOnly) || (bIsPaused && TickableObject->IsTickableWhenPaused()))
+						{
+							STAT(FScopeCycleCounter Context(TickableObject->GetStatId());)
+							TickableObject->Tick(DeltaSeconds);
+
+							if (TickableObjects[i] == nullptr)
+							{
+								bNeedsCleanup = true;
+							}
+						}
+					}
 				}
 			}
+			else
+			{
+				bNeedsCleanup = true;
+			}
 		}
-		else
+
+		if (bNeedsCleanup)
 		{
-			bNeedsCleanup = true;
+			TickableObjects.RemoveAll([](FTickableGameObject* Object) { return Object == nullptr; });
 		}
-	}
 
-	if (bNeedsCleanup)
-	{
-		TickableObjects.RemoveAll([](FTickableGameObject* Object) { return Object == nullptr; });
+		bIsTickingObjects = false;
 	}
-
-	bIsTickingObjects = false;
 }
 
 /**
diff --git a/Engine/Source/Runtime/Engine/Private/Materials/MaterialInterface.cpp b/Engine/Source/Runtime/Engine/Private/Materials/MaterialInterface.cpp
--- a/Engine/Source/Runtime/Engine/Private/Materials/MaterialInterface.cpp
+++ b/Engine/Source/Runtime/Engine/Private/Materials/MaterialInterface.cpp
@@ -14,6 +14,7 @@
 #include "Engine/Texture2D.h"
 #include "Engine/SubsurfaceProfile.h"
 #include "Engine/TextureStreamingTypes.h"
+#include "Algo/BinarySearch.h"
 #include "Interfaces/ITargetPlatform.h"
 #include "Components.h"
 
@@ -429,35 +430,6 @@ void UMaterialInterface::UpdateMaterialRenderProxy(FMaterialRenderProxy& Proxy)
 	}
 }
 
-namespace Algo
-{
-	template<typename T, typename PredType>
-	int32 BinarySearch(const TArray<T>& Container, PredType Pred)
-	{
-		int32 Min = 0; // Min is included
-		int32 Max = Container.Num(); // Max is excluded
-
-		while (Min != Max)
-		{
-			const int32 Curr = (Min + Max) / 2;
-			const int32 Comp = Pred(Container[Curr]);
-			if (Comp < 0) // Pred < Ele
-			{
-				Max = Curr;
-			}
-			else if (Comp > 0) // Pred > Ele
-			{
-				Min = Curr + 1;
-			}
-			else // Pred == Ele
-			{
-				return Curr;
-			}
-		}
-		return INDEX_NONE;
-	}
-}
-
 bool FMaterialTextureInfo::IsValid(bool bCheckTextureIndex) const
 { 
 #if WITH_EDITORONLY_DATA
@@ -507,13 +479,6 @@ extern 	TAutoConsoleVariable<int32> CVarStreamingUseMaterialData;
 
 bool UMaterialInterface::FindTextureStreamingDataIndexRange(FName TextureName, int32& LowerIndex, int32& HigherIndex) const
 {
-	struct FNameSearch
-	{
-		FName Name;
-		FNameSearch(FName InName) : Name(InName) {}
-		FORCEINLINE int32 operator()(const FMaterialTextureInfo& Rhs) const { return Name.Compare(Rhs.TextureName); }
-	};
-
 #if WITH_EDITORONLY_DATA
 	// Because of redirectors (when textures are renammed), the texture names might be invalid and we need to udpate the data at every load.
 	// Normally we would do that in the post load, but since the process needs to resolve the StringAssetReference, this is forbidden at that place.
@@ -526,17 +491,13 @@ bool UMaterialInterface::FindTextureStreamingDataIndexRange(FName TextureName, i
 		return false;
 	}
 
-	const int32 MatchingIndex = Algo::BinarySearch(TextureStreamingData, FNameSearch(TextureName));
+	const int32 MatchingIndex = Algo::BinarySearchBy(TextureStreamingData, TextureName, &FMaterialTextureInfo::TextureName);
 	if (MatchingIndex != INDEX_NONE)
 	{
 		// Find the range of entries for this texture. 
 		// This is possible because the same texture could be bound to several register and also be used with different sampling UV.
 		LowerIndex = MatchingIndex;
 		HigherIndex = MatchingIndex;
-		while (LowerIndex > 0 && TextureStreamingData[LowerIndex - 1].TextureName == TextureName)
-		{
-			--LowerIndex;
-		}
 		while (HigherIndex + 1 < TextureStreamingData.Num() && TextureStreamingData[HigherIndex + 1].TextureName == TextureName)
 		{
 			++HigherIndex;
diff --git a/Engine/Source/Runtime/Engine/Private/Particles/FXSystem.cpp b/Engine/Source/Runtime/Engine/Private/Particles/FXSystem.cpp
--- a/Engine/Source/Runtime/Engine/Private/Particles/FXSystem.cpp
+++ b/Engine/Source/Runtime/Engine/Private/Particles/FXSystem.cpp
@@ -217,7 +217,7 @@ void FFXSystem::AddVectorField( UVectorFieldComponent* VectorFieldComponent )
 				FAddVectorFieldCommand,
 				FFXSystem*, FXSystem, this,
 				FVectorFieldInstance*, Instance, Instance,
-				FMatrix, ComponentToWorld, VectorFieldComponent->ComponentToWorld.ToMatrixWithScale(),
+				FMatrix, ComponentToWorld, VectorFieldComponent->GetComponentTransform().ToMatrixWithScale(),
 			{
 				Instance->UpdateTransforms( ComponentToWorld );
 				Instance->Index = FXSystem->VectorFields.AddUninitialized().Index;
@@ -273,7 +273,7 @@ void FFXSystem::UpdateVectorField( UVectorFieldComponent* VectorFieldComponent )
 
 			FUpdateVectorFieldParams UpdateParams;
 			UpdateParams.Bounds = VectorFieldComponent->Bounds.GetBox();
-			UpdateParams.ComponentToWorld = VectorFieldComponent->ComponentToWorld.ToMatrixWithScale();
+			UpdateParams.ComponentToWorld = VectorFieldComponent->GetComponentTransform().ToMatrixWithScale();
 			UpdateParams.Intensity = VectorFieldComponent->Intensity;
 			UpdateParams.Tightness = VectorFieldComponent->Tightness;
 
diff --git a/Engine/Source/Runtime/Engine/Private/Particles/ParticleBeam2EmitterInstance.cpp b/Engine/Source/Runtime/Engine/Private/Particles/ParticleBeam2EmitterInstance.cpp
--- a/Engine/Source/Runtime/Engine/Private/Particles/ParticleBeam2EmitterInstance.cpp
+++ b/Engine/Source/Runtime/Engine/Private/Particles/ParticleBeam2EmitterInstance.cpp
@@ -658,7 +658,7 @@ void FParticleBeam2EmitterInstance::UpdateBoundingBox(float DeltaTime)
 		}
 
 		// Take scale into account as well
-		FVector Scale = Component->ComponentToWorld.GetScale3D();
+		FVector Scale = Component->GetComponentTransform().GetScale3D();
 
 		// Take each particle into account
 		for (int32 i=0; i<ActiveParticles; i++)
@@ -732,7 +732,7 @@ void FParticleBeam2EmitterInstance::ForceUpdateBoundingBox()
 		}
 
 		// Take scale into account as well
-		FVector Scale = Component->ComponentToWorld.GetScale3D();
+		FVector Scale = Component->GetComponentTransform().GetScale3D();
 
 		// Take each particle into account
 		for (int32 i=0; i<ActiveParticles; i++)
diff --git a/Engine/Source/Runtime/Engine/Private/Particles/ParticleBeamModules.cpp b/Engine/Source/Runtime/Engine/Private/Particles/ParticleBeamModules.cpp
--- a/Engine/Source/Runtime/Engine/Private/Particles/ParticleBeamModules.cpp
+++ b/Engine/Source/Runtime/Engine/Private/Particles/ParticleBeamModules.cpp
@@ -126,7 +126,7 @@ void UParticleModuleTypeDataBeam2::Spawn(FParticleEmitterInstance* Owner, int32
 	if (BeamInst->BeamModule_Source == NULL)
 	{
 		BeamData->SourcePoint	= Component->GetComponentLocation();
-		BeamData->SourceTangent = Component->ComponentToWorld.GetScaledAxis( EAxis::X );
+		BeamData->SourceTangent = Component->GetComponentTransform().GetScaledAxis( EAxis::X );
 		BeamData->SourceStrength = 1.0f;
 	}
 
@@ -136,7 +136,7 @@ void UParticleModuleTypeDataBeam2::Spawn(FParticleEmitterInstance* Owner, int32
 		// Set the particle target based on the distance
 		float	TotalDistance	= Distance.GetValue(Particle.RelativeTime, Component);
 		// Always use the X-axis of the component as the direction
-		FVector	Direction		= Component->ComponentToWorld.GetScaledAxis( EAxis::X );
+		FVector	Direction		= Component->GetComponentTransform().GetScaledAxis( EAxis::X );
 		Direction.Normalize();
 		// Calculate the final target point
 		BeamData->TargetPoint	= BeamData->SourcePoint + Direction * TotalDistance;
@@ -245,7 +245,7 @@ void UParticleModuleTypeDataBeam2::Update(FParticleEmitterInstance* Owner, int32
 		if (BeamInst->BeamModule_Source == NULL)
 		{
 			BeamData->SourcePoint	= Component->GetComponentLocation();
-			BeamData->SourceTangent = Component->ComponentToWorld.GetScaledAxis( EAxis::X );
+			BeamData->SourceTangent = Component->GetComponentTransform().GetScaledAxis( EAxis::X );
 		}
 
 		// If the method is set for distance, or there is no target, determine the target point
@@ -253,7 +253,7 @@ void UParticleModuleTypeDataBeam2::Update(FParticleEmitterInstance* Owner, int32
 		{
 			// Set the particle target based on the distance
 			float	TotalDistance	= Distance.GetValue(Particle.RelativeTime, Component);
-			FVector	Direction		= Component->ComponentToWorld.GetScaledAxis( EAxis::X );
+			FVector	Direction		= Component->GetComponentTransform().GetScaledAxis( EAxis::X );
 			Direction.Normalize();
 			BeamData->TargetPoint	= BeamData->SourcePoint + Direction * TotalDistance;
 			BeamData->TargetTangent = -Direction;
@@ -496,7 +496,7 @@ void UParticleModuleTypeDataBeam2::Update(FParticleEmitterInstance* Owner, int32
 #endif	//#if defined(_BEAM2_TYPEDATA_NORMAL_TANGENTS_)
 			if (SourceTangent.IsNearlyZero())
 			{
-				SourceTangent	= Component->ComponentToWorld.GetScaledAxis( EAxis::X );
+				SourceTangent	= Component->GetComponentTransform().GetScaledAxis( EAxis::X );
 			}
 			SourceTangent	*= BeamData->SourceStrength;
 
@@ -507,7 +507,7 @@ void UParticleModuleTypeDataBeam2::Update(FParticleEmitterInstance* Owner, int32
 #endif	//#if defined(_BEAM2_TYPEDATA_NORMAL_TANGENTS_)
 			if (TargetTangent.IsNearlyZero())
 			{
-				TargetTangent	= Component->ComponentToWorld.GetScaledAxis( EAxis::X );
+				TargetTangent	= Component->GetComponentTransform().GetScaledAxis( EAxis::X );
 			}
 			TargetTangent	*= BeamData->TargetStrength;
 
@@ -2010,7 +2010,7 @@ bool UParticleModuleBeamSource::ResolveSourceData(FParticleBeam2EmitterInstance*
 			else
 			{
 				// Use the value as a local space position.
-				BeamData->SourcePoint	= BeamInst->Component->ComponentToWorld.TransformPosition(
+				BeamData->SourcePoint	= BeamInst->Component->GetComponentTransform().TransformPosition(
 					Source.GetValue(BeamInst->EmitterTime, BeamInst->Component));
 			}
 		}
@@ -2024,7 +2024,7 @@ bool UParticleModuleBeamSource::ResolveSourceData(FParticleBeam2EmitterInstance*
 		{
 		case PEB2STTM_Direct:
 			// Use the emitter direction as the tangent
-			BeamData->SourceTangent	= BeamInst->Component->ComponentToWorld.GetScaledAxis( EAxis::X );
+			BeamData->SourceTangent	= BeamInst->Component->GetComponentTransform().GetScaledAxis( EAxis::X );
 			bSetSourceTangent		= true;
 			break;
 		case PEB2STTM_UserSet:
@@ -2050,7 +2050,7 @@ bool UParticleModuleBeamSource::ResolveSourceData(FParticleBeam2EmitterInstance*
 			break;
 		case PEB2STTM_Emitter:
 			// Use the emitter direction as the tangent
-			BeamData->SourceTangent	= BeamInst->Component->ComponentToWorld.GetScaledAxis( EAxis::X );
+			BeamData->SourceTangent	= BeamInst->Component->GetComponentTransform().GetScaledAxis( EAxis::X );
 			bSetSourceTangent		= true;
 			break;
 		}
@@ -2063,7 +2063,7 @@ bool UParticleModuleBeamSource::ResolveSourceData(FParticleBeam2EmitterInstance*
 			if (bSourceAbsolute == false)
 			{
 				// If not tagged as absolute, transform it to world space
-				BeamData->SourceTangent	= BeamInst->Component->ComponentToWorld.TransformVector(BeamData->SourceTangent);
+				BeamData->SourceTangent	= BeamInst->Component->GetComponentTransform().TransformVector(BeamData->SourceTangent);
 			}
 		}
 	}
@@ -2318,7 +2318,7 @@ bool UParticleModuleBeamTarget::ResolveTargetData(FParticleBeam2EmitterInstance*
 			{
 				Distance	= 0.001f;
 			}
-			FVector	Direction		= BeamInst->Component->ComponentToWorld.GetScaledAxis( EAxis::X );
+			FVector	Direction		= BeamInst->Component->GetComponentTransform().GetScaledAxis( EAxis::X );
 			Direction.Normalize();
 			BeamData->TargetPoint	= BeamData->SourcePoint + Direction * Distance;
 			bSetTarget				= true;
@@ -2471,7 +2471,7 @@ bool UParticleModuleBeamTarget::ResolveTargetData(FParticleBeam2EmitterInstance*
 			}
 			else
 			{
-				BeamData->TargetPoint	= BeamInst->Component->ComponentToWorld.TransformPosition(
+				BeamData->TargetPoint	= BeamInst->Component->GetComponentTransform().TransformPosition(
 					Target.GetValue(BeamInst->EmitterTime, BeamInst->Component));
 			}
 		}
@@ -2484,7 +2484,7 @@ bool UParticleModuleBeamTarget::ResolveTargetData(FParticleBeam2EmitterInstance*
 		switch (TargetTangentMethod)
 		{
 		case PEB2STTM_Direct:
-			BeamData->TargetTangent	= BeamInst->Component->ComponentToWorld.GetScaledAxis( EAxis::X );
+			BeamData->TargetTangent	= BeamInst->Component->GetComponentTransform().GetScaledAxis( EAxis::X );
 			bSetTargetTangent		= true;
 			break;
 		case PEB2STTM_UserSet:
@@ -2507,7 +2507,7 @@ bool UParticleModuleBeamTarget::ResolveTargetData(FParticleBeam2EmitterInstance*
 			bSetTargetTangent		= true;
 			break;
 		case PEB2STTM_Emitter:
-			BeamData->TargetTangent	= BeamInst->Component->ComponentToWorld.GetScaledAxis( EAxis::X );
+			BeamData->TargetTangent	= BeamInst->Component->GetComponentTransform().GetScaledAxis( EAxis::X );
 			bSetTargetTangent		= true;
 			break;
 		}
@@ -2518,7 +2518,7 @@ bool UParticleModuleBeamTarget::ResolveTargetData(FParticleBeam2EmitterInstance*
 			BeamData->TargetTangent	= TargetTangent.GetValue(Particle.RelativeTime, BeamInst->Component);
 			if (bTargetAbsolute == false)
 			{
-				BeamData->TargetTangent	= BeamInst->Component->ComponentToWorld.TransformVector(BeamData->TargetTangent);
+				BeamData->TargetTangent	= BeamInst->Component->GetComponentTransform().TransformVector(BeamData->TargetTangent);
 			}
 		}
 	}
diff --git a/Engine/Source/Runtime/Engine/Private/Particles/ParticleComponents.cpp b/Engine/Source/Runtime/Engine/Private/Particles/ParticleComponents.cpp
--- a/Engine/Source/Runtime/Engine/Private/Particles/ParticleComponents.cpp
+++ b/Engine/Source/Runtime/Engine/Private/Particles/ParticleComponents.cpp
@@ -3849,7 +3849,7 @@ FParticleDynamicData* UParticleSystemComponent::CreateDynamicData(ERHIFeatureLev
 
 	if (Template)
 	{
-		ParticleDynamicData->SystemPositionForMacroUVs = ComponentToWorld.TransformPosition(Template->MacroUVPosition);
+		ParticleDynamicData->SystemPositionForMacroUVs = GetComponentTransform().TransformPosition(Template->MacroUVPosition);
 		ParticleDynamicData->SystemRadiusForMacroUVs = Template->MacroUVRadius;
 	}
 
@@ -4166,7 +4166,7 @@ void UParticleSystemComponent::OrientZAxisTowardCamera()
 		DirToCamera.Normalize();
 
 		// Convert the camera direction to local space
-		DirToCamera = ComponentToWorld.InverseTransformVectorNoScale(DirToCamera);
+		DirToCamera = GetComponentTransform().InverseTransformVectorNoScale(DirToCamera);
 		
 		// Local Z axis
 		const FVector LocalZAxis = FVector(0,0,1);
@@ -4611,7 +4611,7 @@ void UParticleSystemComponent::TickComponent(float DeltaTime, enum ELevelTick Ti
 			SCOPE_CYCLE_COUNTER(STAT_UParticleSystemComponent_Marshall)
 		bAsyncDataCopyIsValid = true;
 		check(!bParallelRenderThreadUpdate);
-		AsyncComponentToWorld = ComponentToWorld;
+		AsyncComponentToWorld = GetComponentTransform();
 		AsyncInstanceParameters.Reset();
 		AsyncInstanceParameters.Append(InstanceParameters);
 			AsyncBounds = Bounds;
diff --git a/Engine/Source/Runtime/Engine/Private/Particles/ParticleEmitterInstances.cpp b/Engine/Source/Runtime/Engine/Private/Particles/ParticleEmitterInstances.cpp
--- a/Engine/Source/Runtime/Engine/Private/Particles/ParticleEmitterInstances.cpp
+++ b/Engine/Source/Runtime/Engine/Private/Particles/ParticleEmitterInstances.cpp
@@ -1132,7 +1132,7 @@ void FParticleEmitterInstance::UpdateBoundingBox(float DeltaTime)
 		bool bUpdateBox = ((Component->bWarmingUp == false) && (Component->Template != NULL) && (Component->Template->bUseFixedRelativeBoundingBox == false));
 
 		// Take component scale into account
-		FVector Scale = Component->ComponentToWorld.GetScale3D();
+		FVector Scale = Component->GetComponentTransform().GetScale3D();
 
 		UParticleLODLevel* LODLevel = GetCurrentLODLevelChecked();
 
@@ -1267,7 +1267,7 @@ void FParticleEmitterInstance::ForceUpdateBoundingBox()
 	if ( Component )
 	{
 		// Take component scale into account
-		FVector Scale = Component->ComponentToWorld.GetScale3D();
+		FVector Scale = Component->GetComponentTransform().GetScale3D();
 
 		UParticleLODLevel* LODLevel = GetCurrentLODLevelChecked();
 		UParticleLODLevel* HighestLODLevel = SpriteTemplate->LODLevels[0];
@@ -2542,7 +2542,7 @@ bool FParticleEmitterInstance::FillReplayData( FDynamicEmitterReplayDataBase& Ou
 	OutData.Scale = FVector(1.0f, 1.0f, 1.0f);
 	if (Component)
 	{
-		OutData.Scale = Component->ComponentToWorld.GetScale3D();
+		OutData.Scale = Component->GetComponentTransform().GetScale3D();
 	}
 
 	int32 ParticleMemSize = MaxActiveParticles * ParticleStride;
@@ -2707,7 +2707,7 @@ void FParticleEmitterInstance::GetScreenAlignmentAndScale(int32& OutScreenAlign,
 	OutScale = FVector(1.0f, 1.0f, 1.0f);
 	if (Component)
 	{
-		OutScale = Component->ComponentToWorld.GetScale3D();
+		OutScale = Component->GetComponentTransform().GetScale3D();
 	}
 }
 
@@ -3139,7 +3139,7 @@ void FParticleMeshEmitterInstance::UpdateBoundingBox(float DeltaTime)
 			(Component->Template != NULL) && (Component->Template->bUseFixedRelativeBoundingBox == false));
 
 		// Take scale into account
-		FVector Scale = Component->ComponentToWorld.GetScale3D();
+		FVector Scale = Component->GetComponentTransform().GetScale3D();
 
 		// Get the static mesh bounds
 		FBoxSphereBounds MeshBound;
@@ -3621,7 +3621,7 @@ bool FParticleMeshEmitterInstance::FillReplayData( FDynamicEmitterReplayDataBase
 		{
 			if (!bIgnoreComponentScale)
 			{
-				NewReplayData->Scale = Component->ComponentToWorld.GetScale3D();
+				NewReplayData->Scale = Component->GetComponentTransform().GetScale3D();
 			}
 		}
 	}
diff --git a/Engine/Source/Runtime/Engine/Private/Particles/ParticleGpuSimulation.cpp b/Engine/Source/Runtime/Engine/Private/Particles/ParticleGpuSimulation.cpp
--- a/Engine/Source/Runtime/Engine/Private/Particles/ParticleGpuSimulation.cpp
+++ b/Engine/Source/Runtime/Engine/Private/Particles/ParticleGpuSimulation.cpp
@@ -3154,7 +3154,7 @@ public:
 		UParticleSystem *Template = Component->Template;
 
 		const bool bLocalSpace = EmitterInfo.RequiredModule->bUseLocalSpace;
-		const FMatrix ComponentToWorldMatrix = Component->ComponentToWorld.ToMatrixWithScale();
+		const FMatrix ComponentToWorldMatrix = Component->GetComponentTransform().ToMatrixWithScale();
 		const FMatrix ComponentToWorld = (bLocalSpace || EmitterInfo.LocalVectorField.bIgnoreComponentTransform) ? FMatrix::Identity : ComponentToWorldMatrix;
 
 		const FRotationMatrix VectorFieldTransform(LocalVectorFieldRotation);
@@ -3177,7 +3177,7 @@ public:
 		DynamicData->bUseLocalSpace = EmitterInfo.RequiredModule->bUseLocalSpace;
 
 		// Account for LocalToWorld scaling
-		FVector ComponentScale = Component->ComponentToWorld.GetScale3D();
+		FVector ComponentScale = Component->GetComponentTransform().GetScale3D();
 		// Figure out if we need to replicate the X channel of size to Y.
 		const bool bSquare = (EmitterInfo.ScreenAlignment == PSA_Square)
 			|| (EmitterInfo.ScreenAlignment == PSA_FacingCameraPosition)
diff --git a/Engine/Source/Runtime/Engine/Private/Particles/ParticleModules.cpp b/Engine/Source/Runtime/Engine/Private/Particles/ParticleModules.cpp
--- a/Engine/Source/Runtime/Engine/Private/Particles/ParticleModules.cpp
+++ b/Engine/Source/Runtime/Engine/Private/Particles/ParticleModules.cpp
@@ -2543,7 +2543,7 @@ void UParticleModuleAccelerationConstant::Spawn(FParticleEmitterInstance* Owner,
 	check(LODLevel);
 	if (bAlwaysInWorldSpace && LODLevel->RequiredModule->bUseLocalSpace)
 	{
-		FVector LocalAcceleration = Owner->Component->ComponentToWorld.InverseTransformVector(Acceleration);
+		FVector LocalAcceleration = Owner->Component->GetComponentTransform().InverseTransformVector(Acceleration);
 		Particle.Velocity		+= LocalAcceleration * SpawnTime;
 		Particle.BaseVelocity	+= LocalAcceleration * SpawnTime;
 	}
@@ -2572,7 +2572,7 @@ void UParticleModuleAccelerationConstant::Update(FParticleEmitterInstance* Owner
 	FPlatformMisc::Prefetch(Owner->ParticleData, (Owner->ParticleIndices[0] * Owner->ParticleStride) + PLATFORM_CACHE_LINE_SIZE);
 	if (bAlwaysInWorldSpace && LODLevel->RequiredModule->bUseLocalSpace)
 	{
-		FTransform Mat = Owner->Component->ComponentToWorld;
+		FTransform Mat = Owner->Component->GetComponentTransform();
 		FVector LocalAcceleration = Mat.InverseTransformVector(Acceleration);
 		BEGIN_UPDATE_LOOP;
 		{
@@ -2866,14 +2866,14 @@ void UParticleModuleAcceleration::Spawn(FParticleEmitterInstance* Owner, int32 O
 	UsedAcceleration = Acceleration.GetValue(Owner->EmitterTime, Owner->Component);
 	if ((bApplyOwnerScale == true) && Owner && Owner->Component)
 	{
-		FVector Scale = Owner->Component->ComponentToWorld.GetScale3D();
+		FVector Scale = Owner->Component->GetComponentTransform().GetScale3D();
 		UsedAcceleration *= Scale;
 	}
 	UParticleLODLevel* LODLevel	= Owner->SpriteTemplate->GetCurrentLODLevel(Owner);
 	check(LODLevel);
 	if (bAlwaysInWorldSpace && LODLevel->RequiredModule->bUseLocalSpace)
 	{
-		FVector TempUsedAcceleration = Owner->Component->ComponentToWorld.InverseTransformVector(UsedAcceleration);
+		FVector TempUsedAcceleration = Owner->Component->GetComponentTransform().InverseTransformVector(UsedAcceleration);
 		Particle.Velocity		+= TempUsedAcceleration * SpawnTime;
 		Particle.BaseVelocity	+= TempUsedAcceleration * SpawnTime;
 	}
@@ -2901,7 +2901,7 @@ void UParticleModuleAcceleration::Update(FParticleEmitterInstance* Owner, int32
 	FPlatformMisc::Prefetch(Owner->ParticleData, (Owner->ParticleIndices[0] * Owner->ParticleStride) + PLATFORM_CACHE_LINE_SIZE);
 	if (bAlwaysInWorldSpace && LODLevel->RequiredModule->bUseLocalSpace)
 	{
-		FTransform Mat = Owner->Component->ComponentToWorld;
+		FTransform Mat = Owner->Component->GetComponentTransform();
 		BEGIN_UPDATE_LOOP;
 		{
 			FVector& UsedAcceleration = *((FVector*)(ParticleBase + CurrentOffset));																\
@@ -2991,7 +2991,7 @@ void UParticleModuleAccelerationOverLifetime::Update(FParticleEmitterInstance* O
 	check(LODLevel);
 	if (bAlwaysInWorldSpace && LODLevel->RequiredModule->bUseLocalSpace)
 	{
-		FTransform Mat = Owner->Component->ComponentToWorld;
+		FTransform Mat = Owner->Component->GetComponentTransform();
 		BEGIN_UPDATE_LOOP;
 			// Acceleration should always be in world space...
 			FVector Accel = AccelOverLife.GetValue(Particle.RelativeTime, Owner->Component);
@@ -3338,7 +3338,7 @@ void UParticleModuleLight::Render3DPreview(FParticleEmitterInstance* Owner, cons
 		int32 Offset = 0;
 		UParticleLODLevel* LODLevel	= Owner->SpriteTemplate->GetCurrentLODLevel(Owner);
 		const bool bLocalSpace = LODLevel->RequiredModule->bUseLocalSpace;
-		const FVector Scale = Owner->Component->ComponentToWorld.GetScale3D();
+		const FVector Scale = Owner->Component->GetComponentTransform().GetScale3D();
 		const FMatrix LocalToWorld = Owner->EmitterToSimulation * Owner->SimulationToWorld;
 		check(LODLevel);
 
@@ -3607,11 +3607,11 @@ void UParticleModuleKillBox::Update(FParticleEmitterInstance* Owner, int32 Offse
 
 		if (LODLevel->RequiredModule->bUseLocalSpace)
 		{
-			Position = Owner->Component->ComponentToWorld.TransformVector(Position);
+			Position = Owner->Component->GetComponentTransform().TransformVector(Position);
 		}
 		else if ((bAxisAlignedAndFixedSize == false) && (bAbsolute == false))
 		{
-			Position = Owner->Component->ComponentToWorld.InverseTransformPosition(Position) + Owner->Component->ComponentToWorld.GetLocation();
+			Position = Owner->Component->GetComponentTransform().InverseTransformPosition(Position) + Owner->Component->GetComponentTransform().GetLocation();
 		}
 
 		// Determine if the particle is inside the box
@@ -3656,7 +3656,7 @@ void UParticleModuleKillBox::Render3DPreview(FParticleEmitterInstance* Owner, co
 		{
 			for (int32 i = 0; i < 8; ++i)
 			{
-				KillboxVerts[i] = Owner->Component->ComponentToWorld.TransformPosition(KillboxVerts[i]);
+				KillboxVerts[i] = Owner->Component->GetComponentTransform().TransformPosition(KillboxVerts[i]);
 			}
 		}
 		else
@@ -3736,7 +3736,7 @@ void UParticleModuleKillHeight::Update(FParticleEmitterInstance* Owner, int32 Of
 	float CheckHeight = Height.GetValue(Owner->EmitterTime, Owner->Component);
 	if (bApplyPSysScale == true)
 	{
-		FVector OwnerScale = Owner->Component->ComponentToWorld.GetScale3D();
+		FVector OwnerScale = Owner->Component->GetComponentTransform().GetScale3D();
 		CheckHeight *= OwnerScale.Z;
 	}
 
@@ -3751,7 +3751,7 @@ void UParticleModuleKillHeight::Update(FParticleEmitterInstance* Owner, int32 Of
 
 		if (LODLevel->RequiredModule->bUseLocalSpace)
 		{
-			Position = Owner->Component->ComponentToWorld.TransformVector(Position);
+			Position = Owner->Component->GetComponentTransform().TransformVector(Position);
 		}
 
 		if ((bFloor == true) && (Position.Z < CheckHeight))
@@ -3780,7 +3780,7 @@ void UParticleModuleKillHeight::Render3DPreview(FParticleEmitterInstance* Owner,
 		float ScaleValue = 1.0f;
 		if (bApplyPSysScale == true)
 		{
-			FVector OwnerScale = Owner->Component->ComponentToWorld.GetScale3D();
+			FVector OwnerScale = Owner->Component->GetComponentTransform().GetScale3D();
 			ScaleValue = OwnerScale.Z;
 		}
 		CheckHeight *= ScaleValue;
@@ -4296,11 +4296,11 @@ void UParticleModuleAttractorParticle::Update(FParticleEmitterInstance* Owner, i
 		{
 			if (bSrcUseLocalSpace)
 			{
-				SrcLocation = Owner->Component->ComponentToWorld.TransformVector(SrcLocation);
+				SrcLocation = Owner->Component->GetComponentTransform().TransformVector(SrcLocation);
 			}
 			if (bUseLocalSpace)
 			{
-				ParticleLocation = Owner->Component->ComponentToWorld.TransformVector(Particle.Location);
+				ParticleLocation = Owner->Component->GetComponentTransform().TransformVector(Particle.Location);
 			}
 		}
 
@@ -4416,9 +4416,9 @@ void UParticleModuleAttractorPoint::Update(FParticleEmitterInstance* Owner, int3
 	if ( (LODLevel->RequiredModule->bUseLocalSpace == false) && ( bUseWorldSpacePosition == false ) )
 	{
 		// Transform the attractor into world space
-		AttractorPosition = Component->ComponentToWorld.TransformPosition(AttractorPosition);
+		AttractorPosition = Component->GetComponentTransform().TransformPosition(AttractorPosition);
 
-		Scale *= Component->ComponentToWorld.GetScale3D();
+		Scale *= Component->GetComponentTransform().GetScale3D();
 	}
 	float ScaleSize = Scale.Size();
 
diff --git a/Engine/Source/Runtime/Engine/Private/Particles/ParticleModules_Location.cpp b/Engine/Source/Runtime/Engine/Private/Particles/ParticleModules_Location.cpp
--- a/Engine/Source/Runtime/Engine/Private/Particles/ParticleModules_Location.cpp
+++ b/Engine/Source/Runtime/Engine/Private/Particles/ParticleModules_Location.cpp
@@ -232,7 +232,7 @@ void UParticleModuleLocationWorldOffset::SpawnEx(FParticleEmitterInstance* Owner
 	else
 	{
 		// We need to inverse transform the location so that the bUseLocalSpace transform uses the proper value
-		FMatrix InvMat = Owner->Component->ComponentToWorld.ToMatrixWithScale().InverseFast();
+		FMatrix InvMat = Owner->Component->GetComponentTransform().ToMatrixWithScale().InverseFast();
 		FVector StartLoc = StartLocation.GetValue(Owner->EmitterTime, Owner->Component, 0, InRandomStream);
 		Particle.Location += InvMat.TransformVector(StartLoc);
 	}
@@ -339,7 +339,7 @@ void UParticleModuleLocationDirect::Spawn(FParticleEmitterInstance* Owner, int32
 	else
 	{
 		FVector StartLoc	= Location.GetValue(Particle.RelativeTime, Owner->Component);
-		StartLoc = Owner->Component->ComponentToWorld.TransformPosition(StartLoc);
+		StartLoc = Owner->Component->GetComponentTransform().TransformPosition(StartLoc);
 		Particle.Location	= StartLoc;
 	}
 
@@ -509,12 +509,12 @@ void UParticleModuleLocationEmitter::Spawn(FParticleEmitterInstance* Owner, int3
 					else if ((bSourceIsInLocalSpace == true) && (bInLocalSpace == false))
 					{
 						// We need to transform it into world space
-						Particle.Location = LocationEmitterInst->Component->ComponentToWorld.TransformPosition(pkParticle->Location);
+						Particle.Location = LocationEmitterInst->Component->GetComponentTransform().TransformPosition(pkParticle->Location);
 					}
 					else //if ((bSourceIsInLocalSpace == false) && (bInLocalSpace == true))
 					{
 						// We need to transform it into local space
-						Particle.Location = LocationEmitterInst->Component->ComponentToWorld.InverseTransformPosition(pkParticle->Location);
+						Particle.Location = LocationEmitterInst->Component->GetComponentTransform().InverseTransformPosition(pkParticle->Location);
 					}
 				}
 				if (InheritSourceVelocity)
@@ -1652,7 +1652,7 @@ void UParticleModuleLocationBoneSocket::Spawn(FParticleEmitterInstance* Owner, i
 				PayloadData->Rotation = RotationQuat.Euler();
 				if (Owner->CurrentLODLevel->RequiredModule->bUseLocalSpace == true)
 				{
-					PayloadData->Rotation = Owner->Component->ComponentToWorld.InverseTransformVectorNoScale(PayloadData->Rotation);
+					PayloadData->Rotation = Owner->Component->GetComponentTransform().InverseTransformVectorNoScale(PayloadData->Rotation);
 				}
 			}
 		}
@@ -2154,7 +2154,7 @@ bool UParticleModuleLocationBoneSocket::GetParticleLocation(FModuleLocationBoneS
 
 	if (Owner->CurrentLODLevel->RequiredModule->bUseLocalSpace == true)
 	{
-		OutPosition = Owner->Component->ComponentToWorld.InverseTransformPosition(OutPosition);
+		OutPosition = Owner->Component->GetComponentTransform().InverseTransformPosition(OutPosition);
 	}
 
 	return true;
@@ -2445,7 +2445,7 @@ void UParticleModuleLocationSkelVertSurface::Spawn(FParticleEmitterInstance* Own
 				FVector Rot = SourceRotation.Euler();
 				if (Owner->CurrentLODLevel->RequiredModule->bUseLocalSpace == true)
 				{
-					Rot = Owner->Component->ComponentToWorld.InverseTransformVectorNoScale(Rot);
+					Rot = Owner->Component->GetComponentTransform().InverseTransformVectorNoScale(Rot);
 				}
 				PayloadData->Rotation = Rot;
 				PayloadData->InitRotation = Rot;
@@ -2809,7 +2809,7 @@ bool UParticleModuleLocationSkelVertSurface::GetParticleLocation(FParticleEmitte
 		if (SourceType == VERTSURFACESOURCE_Vert)
 		{
 			FVector VertPos = InSkelMeshComponent->GetSkinnedVertexPosition(InPrimaryVertexIndex);
-			OutPosition = InSkelMeshComponent->ComponentToWorld.TransformPosition(VertPos);
+			OutPosition = InSkelMeshComponent->GetComponentTransform().TransformPosition(VertPos);
 			OutRotation = FQuat::Identity;
 		}
 		else if (SourceType == VERTSURFACESOURCE_Surface)
@@ -2821,9 +2821,9 @@ bool UParticleModuleLocationSkelVertSurface::GetParticleLocation(FParticleEmitte
 			VertIndex[0] = LODModel.MultiSizeIndexContainer.GetIndexBuffer()->Get( InPrimaryVertexIndex );
 			VertIndex[1] = LODModel.MultiSizeIndexContainer.GetIndexBuffer()->Get( InPrimaryVertexIndex+1 );
 			VertIndex[2] = LODModel.MultiSizeIndexContainer.GetIndexBuffer()->Get( InPrimaryVertexIndex+2 );
-			Verts[0] = InSkelMeshComponent->ComponentToWorld.TransformPosition(InSkelMeshComponent->GetSkinnedVertexPosition(VertIndex[0]));
-			Verts[1] = InSkelMeshComponent->ComponentToWorld.TransformPosition(InSkelMeshComponent->GetSkinnedVertexPosition(VertIndex[1]));
-			Verts[2] = InSkelMeshComponent->ComponentToWorld.TransformPosition(InSkelMeshComponent->GetSkinnedVertexPosition(VertIndex[2]));
+			Verts[0] = InSkelMeshComponent->GetComponentTransform().TransformPosition(InSkelMeshComponent->GetSkinnedVertexPosition(VertIndex[0]));
+			Verts[1] = InSkelMeshComponent->GetComponentTransform().TransformPosition(InSkelMeshComponent->GetSkinnedVertexPosition(VertIndex[1]));
+			Verts[2] = InSkelMeshComponent->GetComponentTransform().TransformPosition(InSkelMeshComponent->GetSkinnedVertexPosition(VertIndex[2]));
 
 			FVector V0ToV2 = (Verts[2] - Verts[0]);
 			V0ToV2.Normalize();
@@ -2869,7 +2869,7 @@ bool UParticleModuleLocationSkelVertSurface::GetParticleLocation(FParticleEmitte
 
 	if (Owner->CurrentLODLevel->RequiredModule->bUseLocalSpace == true)
 	{
-		OutPosition = Owner->Component->ComponentToWorld.InverseTransformPosition(OutPosition);
+		OutPosition = Owner->Component->GetComponentTransform().InverseTransformPosition(OutPosition);
 	}
 
 	OutPosition += UniversalOffset;
diff --git a/Engine/Source/Runtime/Engine/Private/Particles/ParticleModules_Velocity.cpp b/Engine/Source/Runtime/Engine/Private/Particles/ParticleModules_Velocity.cpp
--- a/Engine/Source/Runtime/Engine/Private/Particles/ParticleModules_Velocity.cpp
+++ b/Engine/Source/Runtime/Engine/Private/Particles/ParticleModules_Velocity.cpp
@@ -86,7 +86,7 @@ void UParticleModuleVelocity::SpawnEx(FParticleEmitterInstance* Owner, int32 Off
 		FVector OwnerScale(1.0f);
 		if ((bApplyOwnerScale == true) && Owner->Component)
 		{
-			OwnerScale = Owner->Component->ComponentToWorld.GetScale3D();
+			OwnerScale = Owner->Component->GetComponentTransform().GetScale3D();
 		}
 
 		UParticleLODLevel* LODLevel	= Owner->SpriteTemplate->GetCurrentLODLevel(Owner);
@@ -195,7 +195,7 @@ void UParticleModuleVelocityInheritParent::Spawn(FParticleEmitterInstance* Owner
 	check(LODLevel);
 	if (LODLevel->RequiredModule->bUseLocalSpace)
 	{
-		Vel = Owner->Component->ComponentToWorld.InverseTransformVector(Owner->Component->PartSysVelocity);
+		Vel = Owner->Component->GetComponentTransform().InverseTransformVector(Owner->Component->PartSysVelocity);
 	}
 	else
 	{
@@ -255,7 +255,7 @@ void UParticleModuleVelocityOverLifetime::Spawn(FParticleEmitterInstance* Owner,
 		FVector OwnerScale(1.0f);
 		if ((bApplyOwnerScale == true) && Owner && Owner->Component)
 		{
-			OwnerScale = Owner->Component->ComponentToWorld.GetScale3D();
+			OwnerScale = Owner->Component->GetComponentTransform().GetScale3D();
 		}
 		FVector Vel = VelOverLife.GetValue(Particle.RelativeTime, Owner->Component) * OwnerScale;
 		Particle.Velocity		= Vel;
@@ -431,7 +431,7 @@ void UParticleModuleVelocityCone::SpawnEx(FParticleEmitterInstance* Owner, int32
 	FVector OwnerScale(1.0f);
 	if ((bApplyOwnerScale == true) && Owner->Component)
 	{
-		OwnerScale = Owner->Component->ComponentToWorld.GetScale3D();
+		OwnerScale = Owner->Component->GetComponentTransform().GetScale3D();
 	}
 	
 	// Spawn particles
@@ -471,11 +471,11 @@ void UParticleModuleVelocityCone::SpawnEx(FParticleEmitterInstance* Owner, int32
 		// Transform according to world and local space flags 
 		if (!LODLevel->RequiredModule->bUseLocalSpace && !bInWorldSpace)
 		{
-			SpawnDirection = Owner->Component->ComponentToWorld.TransformVector(SpawnDirection);
+			SpawnDirection = Owner->Component->GetComponentTransform().TransformVector(SpawnDirection);
 		}
 		else if (LODLevel->RequiredModule->bUseLocalSpace && bInWorldSpace)
 		{
-			SpawnDirection = Owner->Component->ComponentToWorld.InverseTransformVector(SpawnDirection);
+			SpawnDirection = Owner->Component->GetComponentTransform().InverseTransformVector(SpawnDirection);
 		}
 
 		// Set final velocity vector
@@ -548,13 +548,13 @@ void UParticleModuleVelocityCone::Render3DPreview(FParticleEmitterInstance* Owne
 		{
 			if (bApplyOwnerScale == true)
 			{
-				OwnerScale = Owner->Component->ComponentToWorld.GetScale3D();
+				OwnerScale = Owner->Component->GetComponentTransform().GetScale3D();
 			}
 
 			OwnerRotation = FQuatRotationMatrix(Actor->GetActorQuat());
 		}
-	  LocalToWorldOrigin = Owner->Component->ComponentToWorld.GetLocation();
-	  LocalToWorld = Owner->Component->ComponentToWorld.ToMatrixWithScale().RemoveTranslation();
+	  LocalToWorldOrigin = Owner->Component->GetComponentTransform().GetLocation();
+	  LocalToWorld = Owner->Component->GetComponentTransform().ToMatrixWithScale().RemoveTranslation();
 	  LocalToWorld.RemoveScaling();
 	}
 	FMatrix Transform;
diff --git a/Engine/Source/Runtime/Engine/Private/Particles/ParticleTrail2EmitterInstance.cpp b/Engine/Source/Runtime/Engine/Private/Particles/ParticleTrail2EmitterInstance.cpp
--- a/Engine/Source/Runtime/Engine/Private/Particles/ParticleTrail2EmitterInstance.cpp
+++ b/Engine/Source/Runtime/Engine/Private/Particles/ParticleTrail2EmitterInstance.cpp
@@ -186,7 +186,7 @@ void FParticleTrailsEmitterInstance_Base::UpdateBoundingBox(float DeltaTime)
 		ParticleBoundingBox.IsValid = true;
 
 		// Take scale into account
-		FVector Scale = Component->ComponentToWorld.GetScale3D();
+		FVector Scale = Component->GetComponentTransform().GetScale3D();
 
 		// As well as each particle
 		int32 LocalActiveParticles = ActiveParticles;
@@ -240,7 +240,7 @@ void FParticleTrailsEmitterInstance_Base::UpdateBoundingBox(float DeltaTime)
 		{
 			if (LODLevel->RequiredModule->bUseLocalSpace) 
 			{
-				ParticleBoundingBox = ParticleBoundingBox.TransformBy(Component->ComponentToWorld);
+				ParticleBoundingBox = ParticleBoundingBox.TransformBy(Component->GetComponentTransform());
 			}
 		}
 	}
@@ -272,7 +272,7 @@ void FParticleTrailsEmitterInstance_Base::ForceUpdateBoundingBox()
 		ParticleBoundingBox.IsValid = true;
 
 		// Take scale into account
-		FVector Scale = Component->ComponentToWorld.GetScale3D();
+		FVector Scale = Component->GetComponentTransform().GetScale3D();
 
 		// As well as each particle
 		int32 LocalActiveParticles = ActiveParticles;
@@ -299,7 +299,7 @@ void FParticleTrailsEmitterInstance_Base::ForceUpdateBoundingBox()
 		// Transform bounding box into world space if the emitter uses a local space coordinate system.
 		if (LODLevel->RequiredModule->bUseLocalSpace) 
 		{
-			ParticleBoundingBox = ParticleBoundingBox.TransformBy(Component->ComponentToWorld);
+			ParticleBoundingBox = ParticleBoundingBox.TransformBy(Component->GetComponentTransform());
 		}
 	}
 }
@@ -1294,7 +1294,7 @@ float FParticleRibbonEmitterInstance::Spawn(float DeltaTime)
 		FVector CurrentUp;
 		if (TrailTypeData->RenderAxis == Trails_SourceUp)
 		{
-			CurrentUp = Component->ComponentToWorld.GetScaledAxis( EAxis::Z );
+			CurrentUp = Component->GetComponentTransform().GetScaledAxis( EAxis::Z );
 		}
 		else
 		{
@@ -1768,7 +1768,7 @@ bool FParticleRibbonEmitterInstance::Spawn_Source(float DeltaTime)
 			float InvCount = 1.0f / MovementSpawnCount;
 			float Increment = DeltaTime / MovementSpawnCount;
 
-			FTransform SavedComponentToWorld = Component->ComponentToWorld;
+			FTransform SavedComponentToWorld = Component->GetComponentTransform();
 
 			// Spawn the given number of particles, interpolating between the current and last position/tangent
 			float CurrTimeStep = InvCount;
@@ -1805,7 +1805,7 @@ bool FParticleRibbonEmitterInstance::Spawn_Source(float DeltaTime)
 				float SpawnTime = DeltaTime - (SpawnIdx * Increment);
 				float TrueSpawnTime = Diff * TimeStep;
 
-				Component->ComponentToWorld = FTransform(CurrRotation, CurrPosition);
+				Component->SetComponentToWorld(FTransform(CurrRotation, CurrPosition));
 
 				// Standard spawn setup
 				PreSpawn(Particle, CurrPosition, FVector::ZeroVector);
@@ -1835,7 +1835,7 @@ bool FParticleRibbonEmitterInstance::Spawn_Source(float DeltaTime)
 				Particle->Size.Z = Particle->Size.Z;
 				Particle->BaseSize = Particle->Size;
 
-				Component->ComponentToWorld = SavedComponentToWorld;
+				Component->SetComponentToWorld(SavedComponentToWorld);
 
 				// Trail specific...
 				// Clear the next and previous - just to be safe
@@ -2326,7 +2326,7 @@ bool FParticleRibbonEmitterInstance::ResolveSourcePoint(int32 InTrailIdx,
 					}
 					OutTangentStrength = OutTangent.SizeSquared();
 					//@todo. Allow particle rotation to define up??
-					OutUp = SourceEmitter->Component->ComponentToWorld.GetScaledAxis(EAxis::Z);
+					OutUp = SourceEmitter->Component->GetComponentTransform().GetScaledAxis(EAxis::Z);
 
 					//@todo. Where to get rotation from????
 					OutRotation = FQuat(0,0,0,1);
@@ -2374,7 +2374,7 @@ bool FParticleRibbonEmitterInstance::ResolveSourcePoint(int32 InTrailIdx,
 				if (CurrentLODLevel && (CurrentLODLevel->RequiredModule->bUseLocalSpace == false))
 				{
 					// Transform it
-					SourceOffsetValue = Component->ComponentToWorld.TransformVector(SourceOffsetValue);
+					SourceOffsetValue = Component->GetComponentTransform().TransformVector(SourceOffsetValue);
 				}
 				OutPosition += SourceOffsetValue;
 			}
@@ -2382,7 +2382,7 @@ bool FParticleRibbonEmitterInstance::ResolveSourcePoint(int32 InTrailIdx,
 		OutRotation = Component->GetComponentQuat();
 		OutTangent = Component->PartSysVelocity;
 		OutTangentStrength = OutTangent.SizeSquared();
-		OutUp = Component->ComponentToWorld.GetScaledAxis(EAxis::Z);
+		OutUp = Component->GetComponentTransform().GetScaledAxis(EAxis::Z);
 
 		bSourceWasSet = true;
 	}
@@ -3008,7 +3008,7 @@ struct FAnimTrailParticleSpawnParams
 	
 	
 	//TODO - These params are for interpolated spawn particles which are currently disabled.
-	///** The ComponentToWorld transform of the particle component before spawning began. */
+	///** The GetComponentTransform() transform of the particle component before spawning began. */
 	//FTransform SavedComponentToWorld;
 	///** True if this particle is interpolated. False otherwise. */
 	//bool bInterpolated;
@@ -3037,7 +3037,7 @@ void FParticleAnimTrailEmitterInstance::SpawnParticle( int32& StartParticleIndex
 //		FVector InterpSourcePos = FMath::Lerp<FVector>(PreviousComponentTransform.GetLocation(), Params.SavedComponentToWorld.GetLocation(), InterpFactor);
 //		FQuat InterpSourceRot = FQuat::Slerp(PreviousComponentTransform.GetRotation(), Params.SavedComponentToWorld.GetRotation(), InterpFactor);
 //		FTransform InterpSourceTransform = FTransform(InterpSourceRot, InterpSourcePos);
-//		Component->ComponentToWorld = InterpSourceTransform;
+//		Component->GetComponentTransform() = InterpSourceTransform;
 //	}
 
 	//TODO - Multiple trails.
@@ -3418,8 +3418,8 @@ float FParticleAnimTrailEmitterInstance::Spawn(float DeltaTime)
 	TrailSpawnTimes[TrailIdx] = RunningTime;
 
 	//TODO - If I enable interpolated spawning then the component transform needs restoring and storing.
-	//Component->ComponentToWorld = SpawnParams.SavedComponentToWorld;
-	//PreviousComponentTransform = Component->ComponentToWorld;
+	//Component->GetComponentTransform() = SpawnParams.SavedComponentToWorld;
+	//PreviousComponentTransform = Component->GetComponentTransform();
 
 	if (bTagTrailAsDead == true)
 	{
@@ -3522,7 +3522,7 @@ void FParticleAnimTrailEmitterInstance::UpdateBoundingBox(float DeltaTime)
 		ParticleBoundingBox.IsValid = true;
 
 		// Take scale into account
-		FVector Scale = Component->ComponentToWorld.GetScale3D();
+		FVector Scale = Component->GetComponentTransform().GetScale3D();
 
 		// As well as each particle
 		int32 LocalActiveParticles = ActiveParticles;
@@ -3593,7 +3593,7 @@ void FParticleAnimTrailEmitterInstance::UpdateBoundingBox(float DeltaTime)
 		{
 			if (LODLevel->RequiredModule->bUseLocalSpace) 
 			{
-				ParticleBoundingBox = ParticleBoundingBox.TransformBy(Component->ComponentToWorld);
+				ParticleBoundingBox = ParticleBoundingBox.TransformBy(Component->GetComponentTransform());
 			}
 		}
 	}
@@ -3625,7 +3625,7 @@ void FParticleAnimTrailEmitterInstance::ForceUpdateBoundingBox()
 		ParticleBoundingBox.IsValid = true;
 
 		// Take scale into account
-		FVector Scale = Component->ComponentToWorld.GetScale3D();
+		FVector Scale = Component->GetComponentTransform().GetScale3D();
 
 		// As well as each particle
 		int32 LocalActiveParticles = ActiveParticles;
@@ -3677,7 +3677,7 @@ void FParticleAnimTrailEmitterInstance::ForceUpdateBoundingBox()
 		// Transform bounding box into world space if the emitter uses a local space coordinate system.
 		if (LODLevel->RequiredModule->bUseLocalSpace) 
 		{
-			ParticleBoundingBox = ParticleBoundingBox.TransformBy(Component->ComponentToWorld);
+			ParticleBoundingBox = ParticleBoundingBox.TransformBy(Component->GetComponentTransform());
 		}
 	}
 }
