Fix non-unity game compile

diff --git a/Engine/Source/Runtime/Engine/Private/AssetManager.cpp b/Engine/Source/Runtime/Engine/Private/AssetManager.cpp
--- a/Engine/Source/Runtime/Engine/Private/AssetManager.cpp
+++ b/Engine/Source/Runtime/Engine/Private/AssetManager.cpp
@@ -8,6 +8,7 @@
 #include "Engine/Engine.h"
 #include "Engine/BlueprintGeneratedClass.h"
 #include "UObject/ConstructorHelpers.h"
+#include "UObject/UObjectHash.h"
 #include "Misc/FileHelper.h"
 #include "Misc/ScopedSlowTask.h"
 #include "Misc/Paths.h"
