Minor improvements to comments for consistency and API documentation generation.

diff --git a/Engine/Plugins/Editor/AudioCapture/Source/AudioCapture/Private/RtAudio.cpp b/Engine/Plugins/Editor/AudioCapture/Source/AudioCapture/Private/RtAudio.cpp
--- a/Engine/Plugins/Editor/AudioCapture/Source/AudioCapture/Private/RtAudio.cpp
+++ b/Engine/Plugins/Editor/AudioCapture/Source/AudioCapture/Private/RtAudio.cpp
@@ -3170,7 +3170,7 @@ bool RtApiAsio::probeDeviceOpen(unsigned int device, StreamMode mode, unsigned i
 	if (result != ASE_OK) {
 		// Standard method failed. This can happen with strict/misbehaving drivers that return valid buffer size ranges
 		// but only accept the preferred buffer size as parameter for ASIOCreateBuffers. eg. Creatives ASIO driver
-		// in that case, let's be naïve and try that instead
+		// in that case, let's be naive and try that instead
 		*bufferSize = preferSize;
 		stream_.bufferSize = *bufferSize;
 		result = ASIOCreateBuffers(handle->bufferInfos, nChannels, stream_.bufferSize, &asioCallbacks);
diff --git a/Engine/Source/Runtime/Engine/Classes/Kismet/KismetMathLibrary.h b/Engine/Source/Runtime/Engine/Classes/Kismet/KismetMathLibrary.h
--- a/Engine/Source/Runtime/Engine/Classes/Kismet/KismetMathLibrary.h
+++ b/Engine/Source/Runtime/Engine/Classes/Kismet/KismetMathLibrary.h
@@ -578,7 +578,7 @@ class ENGINE_API UKismetMathLibrary : public UBlueprintFunctionLibrary
 	UFUNCTION(BlueprintPure, meta=(DisplayName = "Max (float)", CompactNodeTitle = "MAX", CommutativeAssociativeBinaryOperator = "true"), Category="Math|Float")
 	static float FMax(float A, float B);
 
-	/* Returns the value clamped between A and B (inclusive) */
+	/* Returns Value clamped between A and B (inclusive) */
 	UFUNCTION(BlueprintPure, meta=(DisplayName = "Clamp (float)", Min="0.0", Max="1.0"), Category="Math|Float")
 	static float FClamp(float Value, float Min, float Max);
 
@@ -723,7 +723,7 @@ class ENGINE_API UKismetMathLibrary : public UBlueprintFunctionLibrary
 	UFUNCTION(BlueprintPure, meta=(DisplayName = "vector / int", CompactNodeTitle = "/", Keywords = "/ divide division"), Category="Math|Vector")
 	static FVector Divide_VectorInt(FVector A, int32 B = 1);
 	
-	/* Vector divide by vector */
+	/* Element-wise Vector division (Result = {A.x/B.x, A.y/B.y, A.z/B.z}) */
 	UFUNCTION(BlueprintPure, meta=(DisplayName = "vector / vector", CompactNodeTitle = "/", Keywords = "/ divide division"), Category="Math|Vector")
 	static FVector Divide_VectorVector(FVector A, FVector B = FVector(1.f,1.f,1.f));
 
