Copying //UE4/Dev-Rendering to //UE4/Dev-Main (Source: //UE4/Dev-Rendering @ 3461187)

diff --git a/Engine/Source/Runtime/D3D12RHI/Private/D3D12RHI.cpp b/Engine/Source/Runtime/D3D12RHI/Private/D3D12RHI.cpp
--- a/Engine/Source/Runtime/D3D12RHI/Private/D3D12RHI.cpp
+++ b/Engine/Source/Runtime/D3D12RHI/Private/D3D12RHI.cpp
@@ -177,6 +177,9 @@ FD3D12DynamicRHI::FD3D12DynamicRHI(TArray<FD3D12Adapter*>& ChosenAdaptersIn) :
 	if (!GIsEditor)
 	{
 		GRHISupportsRHIThread = true;
+#if PLATFORM_XBOXONE
+		GRHISupportsRHIOnTaskThread = true;
+#endif
 	}
 	GRHISupportsParallelRHIExecute = D3D12_SUPPORTS_PARALLEL_RHI_EXECUTE;
 
diff --git a/Engine/Source/Runtime/D3D12RHI/Private/D3D12RHIPrivate.h b/Engine/Source/Runtime/D3D12RHI/Private/D3D12RHIPrivate.h
--- a/Engine/Source/Runtime/D3D12RHI/Private/D3D12RHIPrivate.h
+++ b/Engine/Source/Runtime/D3D12RHI/Private/D3D12RHIPrivate.h
@@ -360,7 +360,7 @@ public:
 			return false;
 		}
 
-		if (RHICmdList->Bypass() || !GRHIThread)
+		if (RHICmdList->Bypass() || !IsRunningRHIInSeparateThread())
 		{
 			return false;
 		}
diff --git a/Engine/Source/Runtime/D3D12RHI/Private/D3D12Texture.h b/Engine/Source/Runtime/D3D12RHI/Private/D3D12Texture.h
--- a/Engine/Source/Runtime/D3D12RHI/Private/D3D12Texture.h
+++ b/Engine/Source/Runtime/D3D12RHI/Private/D3D12Texture.h
@@ -270,7 +270,7 @@ public:
 			return false;
 		}
 
-		if (RHICmdList->Bypass() || !GRHIThread)
+		if (RHICmdList->Bypass() || !IsRunningRHIInSeparateThread())
 		{
 			return false;
 		}
diff --git a/Engine/Source/Runtime/D3D12RHI/Private/Windows/WindowsD3D12Device.cpp b/Engine/Source/Runtime/D3D12RHI/Private/Windows/WindowsD3D12Device.cpp
--- a/Engine/Source/Runtime/D3D12RHI/Private/Windows/WindowsD3D12Device.cpp
+++ b/Engine/Source/Runtime/D3D12RHI/Private/Windows/WindowsD3D12Device.cpp
@@ -33,6 +33,26 @@ static TAutoConsoleVariable<int32> CVarGraphicsAdapter(
 	TEXT("  1: Adpater #1, ..."),
 	ECVF_RenderThreadSafe);
 
+int D3D12RHI_PreferAdaperVendor()
+{
+	if (FParse::Param(FCommandLine::Get(), TEXT("preferAMD")))
+	{
+		return 0x1002;
+	}
+
+	if (FParse::Param(FCommandLine::Get(), TEXT("preferIntel")))
+	{
+		return 0x8086;
+	}
+
+	if (FParse::Param(FCommandLine::Get(), TEXT("preferNvidia")))
+	{
+		return 0x10DE;
+	}
+
+	return -1;
+}
+
 namespace D3D12RHI
 {
 
@@ -260,6 +280,7 @@ void FD3D12DynamicRHIModule::FindAdapter()
 	bool bIsAnyNVIDIA = false;
 	bool bRequestedWARP = D3D12RHI_ShouldCreateWithWarp();
 
+	int PreferredVendor = D3D12RHI_PreferAdaperVendor();
 	// Enumerate the DXGIFactory's adapters.
 	for (uint32 AdapterIndex = 0; DXGIFactory->EnumAdapters(AdapterIndex, TempAdapter.GetInitReference()) != DXGI_ERROR_NOT_FOUND; ++AdapterIndex)
 	{
@@ -309,18 +330,13 @@ void FD3D12DynamicRHIModule::FindAdapter()
 				// Requested WARP, reject all other adapters.
 				const bool bSkipRequestedWARP = bRequestedWARP && !bIsWARP;
 
-				// Add special check to support WARP and HMDs, which do not have associated outputs.
-				// This device has no outputs. Reject it, 
-				// http://msdn.microsoft.com/en-us/library/windows/desktop/bb205075%28v=vs.85%29.aspx#WARP_new_for_Win8
-				const bool bSkipHmdGraphicsAdapter = !OutputCount && !bIsWARP && !bUseHmdGraphicsAdapter;
-
 				// we don't allow the PerfHUD adapter
 				const bool bSkipPerfHUDAdapter = bIsPerfHUD && !bAllowPerfHUD;
 
 				// the user wants a specific adapter, not this one
 				const bool bSkipExplicitAdapter = CVarExplicitAdapterValue >= 0 && AdapterIndex != CVarExplicitAdapterValue;
 
-				const bool bSkipAdapter = bSkipRequestedWARP || bSkipHmdGraphicsAdapter || bSkipPerfHUDAdapter || bSkipExplicitAdapter;
+				const bool bSkipAdapter = bSkipRequestedWARP || bSkipPerfHUDAdapter || bSkipExplicitAdapter;
 
 				if (!bSkipAdapter)
 				{
@@ -328,11 +344,19 @@ void FD3D12DynamicRHIModule::FindAdapter()
 					{
 						FirstWithoutIntegratedAdapter = CurrentAdapter;
 					}
+					else if (PreferredVendor == AdapterDesc.VendorId && FirstWithoutIntegratedAdapter.IsValid())
+					{
+						FirstWithoutIntegratedAdapter = CurrentAdapter;
+					}
 
 					if (!FirstAdapter.IsValid())
 					{
 						FirstAdapter = CurrentAdapter;
 					}
+					else if (PreferredVendor == AdapterDesc.VendorId && FirstAdapter.IsValid())
+					{
+						FirstAdapter = CurrentAdapter;
+					}
 				}
 			}
 		}
diff --git a/Engine/Source/Runtime/D3D12RHI/Public/D3D12RHI.h b/Engine/Source/Runtime/D3D12RHI/Public/D3D12RHI.h
--- a/Engine/Source/Runtime/D3D12RHI/Public/D3D12RHI.h
+++ b/Engine/Source/Runtime/D3D12RHI/Public/D3D12RHI.h
@@ -52,7 +52,7 @@
 	// Xbox doesn't have DXGI but common code needs this defined for headers
 	#define DXGI_QUERY_VIDEO_MEMORY_INFO			int
 #elif PLATFORM_WINDOWS
-	#define ENABLE_RESIDENCY_MANAGEMENT				1
+	#define ENABLE_RESIDENCY_MANAGEMENT				0
 	#define ASYNC_DEFERRED_DELETION					1
 	#define PLATFORM_SUPPORTS_MGPU					1
 	#define PIPELINE_STATE_FILE_LOCATION			FPaths::GameSavedDir()
diff --git a/Engine/Source/Runtime/Engine/Classes/Components/SceneCaptureComponent2D.h b/Engine/Source/Runtime/Engine/Classes/Components/SceneCaptureComponent2D.h
--- a/Engine/Source/Runtime/Engine/Classes/Components/SceneCaptureComponent2D.h
+++ b/Engine/Source/Runtime/Engine/Classes/Components/SceneCaptureComponent2D.h
@@ -24,7 +24,7 @@ class ENGINE_API USceneCaptureComponent2D : public USceneCaptureComponent
 	TEnumAsByte<ECameraProjectionMode::Type> ProjectionType;
 
 	/** Camera field of view (in degrees). */
-	UPROPERTY(interp, Category=Projection, meta=(DisplayName = "Field of View", UIMin = "5.0", UIMax = "170", ClampMin = "0.001", ClampMax = "360.0"))
+	UPROPERTY(interp, EditAnywhere, BlueprintReadWrite, Category=Projection, meta=(DisplayName = "Field of View", UIMin = "5.0", UIMax = "170", ClampMin = "0.001", ClampMax = "360.0"))
 	float FOVAngle;
 
 	/** The desired width (in world units) of the orthographic view (ignored in Perspective mode) */
diff --git a/Engine/Source/Runtime/Engine/Classes/Components/SkinnedMeshComponent.h b/Engine/Source/Runtime/Engine/Classes/Components/SkinnedMeshComponent.h
--- a/Engine/Source/Runtime/Engine/Classes/Components/SkinnedMeshComponent.h
+++ b/Engine/Source/Runtime/Engine/Classes/Components/SkinnedMeshComponent.h
@@ -771,6 +771,8 @@ public:
 	DEPRECATED(4.13, "SetSpaceBaseDoubleBuffering is now renamed SetComponentSpaceTransformsDoubleBuffering")
 	void SetSpaceBaseDoubleBuffering(bool bInDoubleBufferedBlendSpaces) { SetComponentSpaceTransformsDoubleBuffering(bInDoubleBufferedBlendSpaces);  }
 
+	const FBoxSphereBounds& GetCachedLocalBounds() { return CachedLocalBounds; } 
+
 protected:
 
 	/** Flip the editable space base buffer */
diff --git a/Engine/Source/Runtime/Engine/Classes/Engine/Engine.h b/Engine/Source/Runtime/Engine/Classes/Engine/Engine.h
--- a/Engine/Source/Runtime/Engine/Classes/Engine/Engine.h
+++ b/Engine/Source/Runtime/Engine/Classes/Engine/Engine.h
@@ -781,6 +781,14 @@ public:
 	UPROPERTY(globalconfig)
 	FStringAssetReference DefaultBokehTextureName;
 
+	/** Texture used to bloom when using FFT, mimics characteristic bloom produced in a camera from a signle bright source */
+	UPROPERTY()
+	class UTexture2D* DefaultBloomKernelTexture;
+
+	/** @todo document */
+	UPROPERTY(globalconfig)
+	FStringAssetReference DefaultBloomKernelTextureName;
+
 	/** The material used to render wireframe meshes. */
 	UPROPERTY()
 	class UMaterial* WireframeMaterial;
diff --git a/Engine/Source/Runtime/Engine/Classes/Engine/Scene.h b/Engine/Source/Runtime/Engine/Classes/Engine/Scene.h
--- a/Engine/Source/Runtime/Engine/Classes/Engine/Scene.h
+++ b/Engine/Source/Runtime/Engine/Classes/Engine/Scene.h
@@ -716,9 +716,7 @@ struct FPostProcessSettings
 	UPROPERTY(interp, BlueprintReadWrite, Category="Lens|Bloom", AdvancedDisplay, meta=(editcondition = "bOverride_Bloom6Tint", DisplayName = "#6 Tint", HideAlphaChannel))
 	FLinearColor Bloom6Tint;
 
-	/**
-	* Texture that defines the convolution for bloom.
-	*/
+	/** Texture to replace default convolution bloom kernel */
 	UPROPERTY(EditAnywhere, BlueprintReadWrite, Category = "Lens|Bloom", meta = (editcondition = "bOverride_BloomConvolutionTexture", DisplayName = "Convolution Kernel"))
 	class UTexture2D* BloomConvolutionTexture;
 
@@ -1261,7 +1259,7 @@ struct FPostProcessSettings
 		Bloom6Size = 64.0f;
 		BloomConvolutionSize = 1.f;
 		BloomConvolutionCenterUV = FVector2D(0.5f, 0.5f);
-		BloomConvolutionPreFilter = FVector(0.f, -1.f, 1.f);
+		BloomConvolutionPreFilter = FVector(7.f, 15000.f, 15.f);
 		BloomConvolutionBufferScale = 0.133f;
 		BloomDirtMaskIntensity = 0.0f;
 		BloomDirtMaskTint = FLinearColor(0.5f, 0.5f, 0.5f);
diff --git a/Engine/Source/Runtime/Engine/Classes/Engine/StaticMesh.h b/Engine/Source/Runtime/Engine/Classes/Engine/StaticMesh.h
--- a/Engine/Source/Runtime/Engine/Classes/Engine/StaticMesh.h
+++ b/Engine/Source/Runtime/Engine/Classes/Engine/StaticMesh.h
@@ -519,9 +519,13 @@ class UStaticMesh : public UObject, public IInterface_CollisionDataProvider, pub
 	UPROPERTY(EditAnywhere, AdvancedDisplay, Category=Navigation)
 	uint32 bHasNavigationData:1;
 
-	/** If true, mesh will calculate data for fast uniform random sampling. This is approx 8 bytes per triangle so should not be enabled unless needed. */
+	/**	
+		Mesh supports uniformly distributed sampling in constant time.
+		Memory cost is 8 bytes per triangle.
+		Example usage is uniform spawning of particles.
+	*/
 	UPROPERTY(EditAnywhere, AdvancedDisplay, Category = StaticMesh)
-	uint32 bRequiresAreaWeightedSampling : 1;
+	uint32 bSupportUniformlyDistributedSampling : 1;
 
 	/** Bias multiplier for Light Propagation Volume lighting */
 	UPROPERTY(EditAnywhere, BlueprintReadOnly, Category=StaticMesh, meta=(UIMin = "0.0", UIMax = "3.0"))
diff --git a/Engine/Source/Runtime/Engine/Classes/Engine/Texture.h b/Engine/Source/Runtime/Engine/Classes/Engine/Texture.h
--- a/Engine/Source/Runtime/Engine/Classes/Engine/Texture.h
+++ b/Engine/Source/Runtime/Engine/Classes/Engine/Texture.h
@@ -342,8 +342,8 @@ struct FTexturePlatformData
 	void Cache(
 		class UTexture& InTexture,
 		const struct FTextureBuildSettings& InSettings,
-		uint32 InFlags
-		);
+		uint32 InFlags,
+		class ITextureCompressorModule* Compressor);
 	void FinishCache();
 	ENGINE_API bool TryInlineMipData();
 	bool AreDerivedMipsAvailable() const;
@@ -649,7 +649,7 @@ public:
 	 * @param bAsyncCache spawn a thread to cache the platform data 
 	 * @param bAllowAsyncBuild allow the building of the texture to happen on another thread !!!load BulkData and cache the Source mip data on the main thread before checking DDC!!!
 	 */
-	void CachePlatformData(bool bAsyncCache = false, bool bAllowAsyncBuild = false);
+	void CachePlatformData(bool bAsyncCache = false, bool bAllowAsyncBuild = false, class ITextureCompressorModule* Compressor = nullptr);
 
 	/**
 	 * Begins caching platform data in the background for the platform requested
diff --git a/Engine/Source/Runtime/Engine/Classes/Engine/TextureRenderTarget2D.h b/Engine/Source/Runtime/Engine/Classes/Engine/TextureRenderTarget2D.h
--- a/Engine/Source/Runtime/Engine/Classes/Engine/TextureRenderTarget2D.h
+++ b/Engine/Source/Runtime/Engine/Classes/Engine/TextureRenderTarget2D.h
@@ -15,6 +15,51 @@ struct FPropertyChangedEvent;
 extern ENGINE_API int32 GTextureRenderTarget2DMaxSizeX;
 extern ENGINE_API int32 GTextureRenderTarget2DMaxSizeY;
 
+/** Subset of EPixelFormat exposed to UTextureRenderTarget2D */
+UENUM()
+enum ETextureRenderTargetFormat
+{
+	/** R channel, 8 bit per channel fixed point, range [0, 1]. */
+	RTF_R8,
+	/** RG channels, 8 bit per channel fixed point, range [0, 1]. */
+	RTF_RG8,
+	/** RGBA channels, 8 bit per channel fixed point, range [0, 1]. */
+	RTF_RGBA8,
+	/** R channel, 16 bit per channel floating point, range [-65504, 65504] */
+	RTF_R16f,
+	/** RG channels, 16 bit per channel floating point, range [-65504, 65504] */
+	RTF_RG16f,
+	/** RGBA channels, 16 bit per channel floating point, range [-65504, 65504] */
+	RTF_RGBA16f,
+	/** R channel, 32 bit per channel floating point, range [-3.402823 x 10^38, 3.402823 x 10^38] */
+	RTF_R32f,
+	/** RG channels, 32 bit per channel floating point, range [-3.402823 x 10^38, 3.402823 x 10^38] */
+	RTF_RG32f,
+	/** RGBA channels, 32 bit per channel floating point, range [-3.402823 x 10^38, 3.402823 x 10^38] */
+	RTF_RGBA32f
+};
+
+inline EPixelFormat GetPixelFormatFromRenderTargetFormat(ETextureRenderTargetFormat RTFormat)
+{
+	switch (RTFormat)
+	{
+	case RTF_R8: return PF_G8; break;
+	case RTF_RG8: return PF_R8G8; break;
+	case RTF_RGBA8: return PF_B8G8R8A8; break;
+
+	case RTF_R16f: return PF_R16F; break;
+	case RTF_RG16f: return PF_G16R16F; break;
+	case RTF_RGBA16f: return PF_FloatRGBA; break;
+
+	case RTF_R32f: return PF_R32_FLOAT; break;
+	case RTF_RG32f: return PF_G32R32F; break;
+	case RTF_RGBA32f: return PF_A32B32G32R32F; break;
+	}
+
+	ensureMsgf(false, TEXT("Unhandled ETextureRenderTargetFormat entry %u"), (uint32)RTFormat);
+	return PF_Unknown;
+}
+
 /**
  * TextureRenderTarget2D
  *
@@ -22,7 +67,7 @@ extern ENGINE_API int32 GTextureRenderTarget2DMaxSizeY;
  * for rendering as well as rendered as a regular 2D texture resource.
  *
  */
-UCLASS(hidecategories=Object, hidecategories=Texture, MinimalAPI)
+UCLASS(hidecategories=Object, hidecategories=Texture, hidecategories=Compression, hidecategories=Adjustments, hidecategories=Compositing, MinimalAPI)
 class UTextureRenderTarget2D : public UTextureRenderTarget
 {
 	GENERATED_UCLASS_BODY()
@@ -52,8 +97,17 @@ class UTextureRenderTarget2D : public UTextureRenderTarget
 	uint32 bForceLinearGamma:1;
 
 	/** Whether to support storing HDR values, which requires more memory. */
+	UPROPERTY()
+	uint32 bHDR_DEPRECATED:1;
+
+	/** 
+	 * Format of the texture render target. 
+	 * Data written to the render target will be quantized to this format, which can limit the range and precision.
+	 * The largest format (RTF_RGBA32f) uses 16x more memory and bandwidth than the smallest (RTF_R8) and can greatly affect performance.  
+	 * Use the smallest format that has enough precision and range for what you are doing.
+	 */
 	UPROPERTY(EditAnywhere, BlueprintReadOnly, Category=TextureRenderTarget2D, AssetRegistrySearchable)
-	uint32 bHDR:1;
+	TEnumAsByte<enum ETextureRenderTargetFormat> RenderTargetFormat;
 
 	/** Whether to support GPU sharing of the underlying native texture resource. */
 	UPROPERTY(EditAnywhere, BlueprintReadOnly, Category = TextureRenderTarget2D, meta=(DisplayName = "Shared"), AssetRegistrySearchable, AdvancedDisplay)
@@ -63,7 +117,7 @@ class UTextureRenderTarget2D : public UTextureRenderTarget
 	UPROPERTY(EditAnywhere, BlueprintReadOnly, Category=TextureRenderTarget2D, AssetRegistrySearchable)
 	uint32 bAutoGenerateMips:1;
 
-	/** Normally the format is derived from bHDR, this allows code to set the format explicitly. */
+	/** Normally the format is derived from RenderTargetFormat, this allows code to set the format explicitly. */
 	UPROPERTY()
 	TEnumAsByte<enum EPixelFormat> OverrideFormat;
 
@@ -109,6 +163,7 @@ class UTextureRenderTarget2D : public UTextureRenderTarget
 #if WITH_EDITOR
 	ENGINE_API virtual void PostEditChangeProperty(FPropertyChangedEvent& PropertyChangedEvent) override;
 #endif // WITH_EDITOR
+	ENGINE_API virtual void Serialize(FArchive& Ar) override;
 	ENGINE_API virtual void PostLoad() override;
 	ENGINE_API virtual void GetResourceSizeEx(FResourceSizeEx& CumulativeResourceSize) override;
 	ENGINE_API virtual FString GetDesc() override;
@@ -124,7 +179,7 @@ class UTextureRenderTarget2D : public UTextureRenderTarget
 	{
 		if (OverrideFormat == PF_Unknown)
 		{
-			return bHDR ? PF_FloatRGBA : PF_B8G8R8A8;
+			return GetPixelFormatFromRenderTargetFormat(RenderTargetFormat);
 		}
 		else
 		{
diff --git a/Engine/Source/Runtime/Engine/Classes/Engine/World.h b/Engine/Source/Runtime/Engine/Classes/Engine/World.h
--- a/Engine/Source/Runtime/Engine/Classes/Engine/World.h
+++ b/Engine/Source/Runtime/Engine/Classes/Engine/World.h
@@ -2192,7 +2192,7 @@ public:
 	ULevel* GetActiveLightingScenario() const;
 
 	/** Propagates a change to the active lighting scenario. */
-	void PropagateLightingScenarioChange();
+	void PropagateLightingScenarioChange(bool bLevelWasMadeVisible);
 
 	/**
 	 * Associates the passed in level with the world. The work to make the level visible is spread across several frames and this
diff --git a/Engine/Source/Runtime/Engine/Classes/Materials/Material.h b/Engine/Source/Runtime/Engine/Classes/Materials/Material.h
--- a/Engine/Source/Runtime/Engine/Classes/Materials/Material.h
+++ b/Engine/Source/Runtime/Engine/Classes/Materials/Material.h
@@ -424,7 +424,7 @@ public:
 	FScalarMaterialInput PixelDepthOffset;
 
 	/** Indicates that the material should be rendered in the SeparateTranslucency Pass (not affected by DOF, requires bAllowSeparateTranslucency to be set in .ini). */
-	UPROPERTY(EditAnywhere, Category=Translucency, meta=(DisplayName = "Separate Translucency"), AdvancedDisplay)
+	UPROPERTY(EditAnywhere, Category=Translucency, meta=(DisplayName = "Render After DOF"), AdvancedDisplay)
 	uint32 bEnableSeparateTranslucency:1;
 
 	/** Indicates that the translucent material should not be affected by bloom or DOF. (Note: Depth testing is not available) */
diff --git a/Engine/Source/Runtime/Engine/Classes/Materials/MaterialExpressionPower.h b/Engine/Source/Runtime/Engine/Classes/Materials/MaterialExpressionPower.h
--- a/Engine/Source/Runtime/Engine/Classes/Materials/MaterialExpressionPower.h
+++ b/Engine/Source/Runtime/Engine/Classes/Materials/MaterialExpressionPower.h
@@ -28,6 +28,7 @@ class UMaterialExpressionPower : public UMaterialExpression
 #if WITH_EDITOR
 	virtual int32 Compile(class FMaterialCompiler* Compiler, int32 OutputIndex) override;
 	virtual void GetCaption(TArray<FString>& OutCaptions) const override;
+	virtual void GetExpressionToolTip(TArray<FString>& OutToolTip) override;
 #endif
 	//~ End UMaterialExpression Interface
 
diff --git a/Engine/Source/Runtime/Engine/Classes/Materials/MaterialInterface.h b/Engine/Source/Runtime/Engine/Classes/Materials/MaterialInterface.h
--- a/Engine/Source/Runtime/Engine/Classes/Materials/MaterialInterface.h
+++ b/Engine/Source/Runtime/Engine/Classes/Materials/MaterialInterface.h
@@ -53,9 +53,11 @@ struct ENGINE_API FMaterialRelevance
 	uint32 bOpaque : 1;
 	uint32 bMasked : 1;
 	uint32 bDistortion : 1;
-	uint32 bSeparateTranslucency : 1;
+	uint32 bSeparateTranslucency : 1; // Translucency After DOF
 	uint32 bMobileSeparateTranslucency : 1;
 	uint32 bNormalTranslucency : 1;
+	uint32 bUsesSceneColorCopy : 1;
+	uint32 bDisableOffscreenRendering : 1; // Blend Modulate
 	uint32 bDisableDepthTest : 1;
 	uint32 bOutputsVelocityInBasePass : 1;
 	uint32 bUsesGlobalDistanceField : 1;
diff --git a/Engine/Source/Runtime/Engine/Private/Components/PoseableMeshComponent.cpp b/Engine/Source/Runtime/Engine/Private/Components/PoseableMeshComponent.cpp
--- a/Engine/Source/Runtime/Engine/Private/Components/PoseableMeshComponent.cpp
+++ b/Engine/Source/Runtime/Engine/Private/Components/PoseableMeshComponent.cpp
@@ -66,6 +66,8 @@ void UPoseableMeshComponent::RefreshBoneTransforms(FActorComponentTickFunction*
 	FinalizeBoneTransform();
 
 	UpdateChildTransforms();
+	UpdateBounds();
+	MarkRenderTransformDirty();
 	MarkRenderDynamicDataDirty();
 }
 
diff --git a/Engine/Source/Runtime/Engine/Private/Components/SceneCaptureComponent.cpp b/Engine/Source/Runtime/Engine/Private/Components/SceneCaptureComponent.cpp
--- a/Engine/Source/Runtime/Engine/Private/Components/SceneCaptureComponent.cpp
+++ b/Engine/Source/Runtime/Engine/Private/Components/SceneCaptureComponent.cpp
@@ -391,8 +391,16 @@ USceneCaptureComponent2D::USceneCaptureComponent2D(const FObjectInitializer& Obj
 	CaptureStereoPass = EStereoscopicPass::eSSP_FULL;
 	CustomProjectionMatrix.SetIdentity();
 	ClipPlaneNormal = FVector(0, 0, 1);
-	// previous behavior was to capture 2d scene captures before cube scene captures.
-	CaptureSortPriority = 1;
+	
+	// Legacy initialization.
+	{
+		// previous behavior was to capture 2d scene captures before cube scene captures.
+		CaptureSortPriority = 1;
+
+		// previous behavior was not exposing MotionBlur and Temporal AA in scene capture 2d.
+		ShowFlags.TemporalAA = false;
+		ShowFlags.MotionBlur = false;
+	}
 }
 
 void USceneCaptureComponent2D::OnRegister()
@@ -522,9 +530,15 @@ void USceneCaptureComponent2D::Serialize(FArchive& Ar)
 {
 	Super::Serialize(Ar);
 
-	if(Ar.IsLoading())
+	if (Ar.IsLoading())
 	{
 		PostProcessSettings.OnAfterLoad();
+
+		if (Ar.CustomVer(FRenderingObjectVersion::GUID) < FRenderingObjectVersion::MotionBlurAndTAASupportInSceneCapture2d)
+		{
+			ShowFlags.TemporalAA = false;
+			ShowFlags.MotionBlur = false;
+		}
 	}
 }
 
diff --git a/Engine/Source/Runtime/Engine/Private/Components/SkinnedMeshComponent.cpp b/Engine/Source/Runtime/Engine/Private/Components/SkinnedMeshComponent.cpp
--- a/Engine/Source/Runtime/Engine/Private/Components/SkinnedMeshComponent.cpp
+++ b/Engine/Source/Runtime/Engine/Private/Components/SkinnedMeshComponent.cpp
@@ -760,7 +760,7 @@ bool USkinnedMeshComponent::GetMaterialStreamingData(int32 MaterialIndex, FPrimi
 
 void USkinnedMeshComponent::GetStreamingTextureInfo(FStreamingTextureLevelContext& LevelContext, TArray<FStreamingTexturePrimitiveInfo>& OutStreamingTextures) const
 {
-	GetStreamingTextureInfoInner(LevelContext, nullptr, GetComponentTransform().GetMaximumAxisScale(), OutStreamingTextures);
+	GetStreamingTextureInfoInner(LevelContext, nullptr, GetComponentTransform().GetMaximumAxisScale() * StreamingDistanceMultiplier, OutStreamingTextures);
 }
 
 bool USkinnedMeshComponent::ShouldUpdateBoneVisibility() const
diff --git a/Engine/Source/Runtime/Engine/Private/GPUSkinVertexFactory.cpp b/Engine/Source/Runtime/Engine/Private/GPUSkinVertexFactory.cpp
--- a/Engine/Source/Runtime/Engine/Private/GPUSkinVertexFactory.cpp
+++ b/Engine/Source/Runtime/Engine/Private/GPUSkinVertexFactory.cpp
@@ -164,7 +164,7 @@ static TAutoConsoleVariable<int32> CVarRHICmdDeferSkeletalLockAndFillToRHIThread
 
 static bool DeferSkeletalLockAndFillToRHIThread()
 {
-	return GRHIThread && CVarRHICmdDeferSkeletalLockAndFillToRHIThread.GetValueOnRenderThread() > 0;
+	return IsRunningRHIInSeparateThread() && CVarRHICmdDeferSkeletalLockAndFillToRHIThread.GetValueOnRenderThread() > 0;
 }
 
 struct FRHICommandUpdateBoneBuffer : public FRHICommand<FRHICommandUpdateBoneBuffer>
@@ -303,44 +303,6 @@ int32 FGPUBaseSkinVertexFactory::GetMaxGPUSkinBones()
 	return GCVarMaxGPUSkinBones;
 }
 
-/*-----------------------------------------------------------------------------
-	FBoneDataTexture
-	SizeX(32 * 1024) - Good size for UE3
------------------------------------------------------------------------------*/
-
-FBoneDataVertexBuffer::FBoneDataVertexBuffer()
-	: SizeX(80 * 1024)		// todo: we will replace this fixed size using FGlobalDynamicVertexBuffer
-{
-}
-
-float* FBoneDataVertexBuffer::LockData()
-{
-	checkSlow(IsInRenderingThread());
-	checkSlow(GetSizeX());
-	checkSlow(IsValidRef(BoneBuffer));
-
-	float* Data = (float*)RHILockVertexBuffer(BoneBuffer.VertexBufferRHI, 0, ComputeMemorySize(), RLM_WriteOnly);
-	checkSlow(Data);
-
-	return Data;
-}
-
-void FBoneDataVertexBuffer::UnlockData(uint32 SizeInBytes)
-{
-	checkSlow(IsValidRef(BoneBuffer));
-	RHIUnlockVertexBuffer(BoneBuffer.VertexBufferRHI);
-}
-
-uint32 FBoneDataVertexBuffer::GetSizeX() const
-{
-	return SizeX;
-}
-
-uint32 FBoneDataVertexBuffer::ComputeMemorySize()
-{
-	return SizeX * sizeof(FVector4);
-}
-
 /*-----------------------------------------------------------------------------
 TGPUSkinVertexFactory
 -----------------------------------------------------------------------------*/
diff --git a/Engine/Source/Runtime/Engine/Private/KismetRenderingLibrary.cpp b/Engine/Source/Runtime/Engine/Private/KismetRenderingLibrary.cpp
--- a/Engine/Source/Runtime/Engine/Private/KismetRenderingLibrary.cpp
+++ b/Engine/Source/Runtime/Engine/Private/KismetRenderingLibrary.cpp
@@ -44,7 +44,7 @@ void UKismetRenderingLibrary::ClearRenderTarget2D(UObject* WorldContextObject, U
 			[RenderTargetResource, ClearColor](FRHICommandList& RHICmdList)
 			{
 				SetRenderTarget(RHICmdList, RenderTargetResource->GetRenderTargetTexture(), FTextureRHIRef(), true);
-				DrawClearQuad(RHICmdList, GMaxRHIFeatureLevel, ClearColor);
+				DrawClearQuad(RHICmdList, ClearColor);
 			});
 	}
 }
diff --git a/Engine/Source/Runtime/Engine/Private/Level.cpp b/Engine/Source/Runtime/Engine/Private/Level.cpp
--- a/Engine/Source/Runtime/Engine/Private/Level.cpp
+++ b/Engine/Source/Runtime/Engine/Private/Level.cpp
@@ -1851,7 +1851,7 @@ void ULevel::SetLightingScenario(bool bNewIsLightingScenario)
 {
 	bIsLightingScenario = bNewIsLightingScenario;
 
-	OwningWorld->PropagateLightingScenarioChange();
+	OwningWorld->PropagateLightingScenarioChange(true);
 }
 
 #if WITH_EDITOR
diff --git a/Engine/Source/Runtime/Engine/Private/LightMap.cpp b/Engine/Source/Runtime/Engine/Private/LightMap.cpp
--- a/Engine/Source/Runtime/Engine/Private/LightMap.cpp
+++ b/Engine/Source/Runtime/Engine/Private/LightMap.cpp
@@ -380,7 +380,7 @@ struct FLightMapPendingTexture : public FTextureLayout
 	ULightMapTexture2D*				SkyOcclusionTexture;
 	ULightMapTexture2D*				AOMaterialMaskTexture;
 
-	TArray<FLightMapAllocation*>	Allocations;
+	TArray<TUniquePtr<FLightMapAllocation>> Allocations;
 	UObject*						Outer;
 	TWeakObjectPtr<UWorld>			OwningWorld;
 	/** Bounding volume for all mappings within this texture.							*/
@@ -427,7 +427,7 @@ struct FLightMapPendingTexture : public FTextureLayout
 	/**
 	 * Processes the textures and starts asynchronous compression tasks for all mip-levels.
 	 */
-	void StartEncoding(ULevel* Unused);
+	void StartEncoding(ULevel* Unused, ITextureCompressorModule* UnusedCompressor);
 
 	/**
 	 * Call this function after the IsFinishedEncoding function returns true
@@ -568,7 +568,7 @@ void FLightMapPendingTexture::PostEncode()
 
 	for (int32 AllocationIndex = 0; AllocationIndex < Allocations.Num(); AllocationIndex++)
 	{
-		FLightMapAllocation* Allocation = Allocations[AllocationIndex];
+		auto& Allocation = Allocations[AllocationIndex];
 
 		int32 PaddedSizeX = Allocation->TotalSizeX;
 		int32 PaddedSizeY = Allocation->TotalSizeY;
@@ -1142,7 +1142,7 @@ bool FLightMapPendingTexture::NeedsSkyOcclusionTexture() const
 
 	for (int32 AllocationIndex = 0; AllocationIndex < Allocations.Num(); AllocationIndex++)
 	{
-		FLightMapAllocation* Allocation = Allocations[AllocationIndex];
+		auto& Allocation = Allocations[AllocationIndex];
 
 		if (Allocation->bHasSkyShadowing)
 		{
@@ -1169,7 +1169,7 @@ bool FLightMapPendingTexture::NeedsAOMaterialMaskTexture() const
 	return false;
 }
 
-void FLightMapPendingTexture::StartEncoding(ULevel* Unused)
+void FLightMapPendingTexture::StartEncoding(ULevel* Unused, ITextureCompressorModule* UnusedCompressor)
 {
 	if (!bUObjectsCreated)
 	{
@@ -1222,7 +1222,7 @@ void FLightMapPendingTexture::StartEncoding(ULevel* Unused)
 		FIntRect TextureRect( MAX_int32, MAX_int32, MIN_int32, MIN_int32 );
 		for(int32 AllocationIndex = 0;AllocationIndex < Allocations.Num();AllocationIndex++)
 		{
-			FLightMapAllocation* Allocation = Allocations[AllocationIndex];
+			auto& Allocation = Allocations[AllocationIndex];
 			// Link the light-map to the texture.
 			Allocation->LightMap->SkyOcclusionTexture = Texture;
 			
@@ -1310,7 +1310,7 @@ void FLightMapPendingTexture::StartEncoding(ULevel* Unused)
 		FIntRect TextureRect( MAX_int32, MAX_int32, MIN_int32, MIN_int32 );
 		for(int32 AllocationIndex = 0;AllocationIndex < Allocations.Num();AllocationIndex++)
 		{
-			FLightMapAllocation* Allocation = Allocations[AllocationIndex];
+			auto& Allocation = Allocations[AllocationIndex];
 			// Link the light-map to the texture.
 			Allocation->LightMap->AOMaterialMaskTexture = Texture;
 			
@@ -1400,7 +1400,7 @@ void FLightMapPendingTexture::StartEncoding(ULevel* Unused)
 		
 		for(int32 AllocationIndex = 0;AllocationIndex < Allocations.Num();AllocationIndex++)
 		{
-			FLightMapAllocation* Allocation = Allocations[AllocationIndex];
+			auto& Allocation = Allocations[AllocationIndex];
 			// Link the light-map to the texture.
 			Allocation->LightMap->Textures[ CoefficientIndex / 2 ] = Texture;
 			for( int k = 0; k < 2; k++ )
@@ -1970,7 +1970,7 @@ void FLightMap2D::EncodeTextures( UWorld* InWorld, bool bLightingSuccessful, boo
 			// Give the texture ownership of the allocations
 			for (auto& Allocation : PendingGroup.Allocations)
 			{
-				Texture->Allocations.Add(Allocation.Release());
+				Texture->Allocations.Add(MoveTemp(Allocation));
 			}
 		}
 		PendingLightMaps.Empty();
@@ -1985,7 +1985,7 @@ void FLightMap2D::EncodeTextures( UWorld* InWorld, bool bLightingSuccessful, boo
 				// precreate the UObjects then give them to some threads to process
 				// need to precreate Uobjects 
 				Texture->CreateUObjects();
-				auto AsyncEncodeTask = new (AsyncEncodeTasks)FAsyncEncode<FLightMapPendingTexture>(Texture,NULL,Counter);
+				auto AsyncEncodeTask = new (AsyncEncodeTasks)FAsyncEncode<FLightMapPendingTexture>(Texture,nullptr,Counter,nullptr);
 				GLargeThreadPool->AddQueuedWork(AsyncEncodeTask);
 			}
 
@@ -2005,7 +2005,7 @@ void FLightMap2D::EncodeTextures( UWorld* InWorld, bool bLightingSuccessful, boo
 					GWarn->UpdateProgress(TextureIndex, PendingTextures.Num());
 				}
 				FLightMapPendingTexture* PendingTexture = PendingTextures[TextureIndex];
-				PendingTexture->StartEncoding(NULL);
+				PendingTexture->StartEncoding(nullptr,nullptr);
 			}
 		}
 
diff --git a/Engine/Source/Runtime/Engine/Private/MapBuildData.cpp b/Engine/Source/Runtime/Engine/Private/MapBuildData.cpp
--- a/Engine/Source/Runtime/Engine/Private/MapBuildData.cpp
+++ b/Engine/Source/Runtime/Engine/Private/MapBuildData.cpp
@@ -43,7 +43,7 @@ ULevel* UWorld::GetActiveLightingScenario() const
 	return NULL;
 }
 
-void UWorld::PropagateLightingScenarioChange()
+void UWorld::PropagateLightingScenarioChange(bool bLevelWasMadeVisible)
 {
 	for (FActorIterator It(this); It; ++It)
 	{
@@ -68,8 +68,13 @@ void UWorld::PropagateLightingScenarioChange()
 		}
 	}
 
-	//@todo - store reflection capture data in UMapBuildDataRegistry so it can work with multiple lighting scenarios without forcing a recapture
-	UpdateAllReflectionCaptures();
+	// Skipping the reflection capture update if made invisible, in most cases another lighting scenario level will be made visible shortly after, 
+	// Or we're unloading all levels and then it doesn't matter if lighting is updated.
+	if (bLevelWasMadeVisible)
+	{
+		//@todo - store reflection capture data in UMapBuildDataRegistry so it can work with multiple lighting scenarios without forcing a recapture
+		UpdateAllReflectionCaptures();
+	}
 }
 
 UMapBuildDataRegistry* CreateRegistryForLegacyMap(ULevel* Level)
diff --git a/Engine/Source/Runtime/Engine/Private/Materials/HLSLMaterialTranslator.h b/Engine/Source/Runtime/Engine/Private/Materials/HLSLMaterialTranslator.h
--- a/Engine/Source/Runtime/Engine/Private/Materials/HLSLMaterialTranslator.h
+++ b/Engine/Source/Runtime/Engine/Private/Materials/HLSLMaterialTranslator.h
@@ -514,9 +514,9 @@ public:
 
 			if (Domain == MD_Surface)
 			{
-				if (Material->GetBlendMode() == BLEND_Modulate && Material->IsSeparateTranslucencyEnabled())
+				if (Material->GetBlendMode() == BLEND_Modulate && Material->IsTranslucencyAfterDOFEnabled())
 				{
-					Errorf(TEXT("Separate translucency with BLEND_Modulate is not supported. Consider using BLEND_Translucent with black emissive"));
+					Errorf(TEXT("Translucency after DOF with BLEND_Modulate is not supported. Consider using BLEND_Translucent with black emissive"));
 				}
 			}
 
@@ -3113,7 +3113,7 @@ protected:
 	
 		if(MipValueMode == TMVM_None)
 		{
-			SampleCode += FString(TEXT("(%s,")) + SamplerStateCode + TEXT(",%s)");
+			SampleCode += TEXT("(%s,") + SamplerStateCode + TEXT(",%s)");
 		}
 		else if(MipValueMode == TMVM_MipLevel)
 		{
@@ -3126,11 +3126,11 @@ protected:
 				return INDEX_NONE;
 			}
 
-			SampleCode += TEXT("Level(%s,%sSampler,%s,%s)");
+			SampleCode += TEXT("Level(%s,") + SamplerStateCode + TEXT(",%s,%s)");
 		}
 		else if(MipValueMode == TMVM_MipBias)
 		{
-			SampleCode += TEXT("Bias(%s,%sSampler,%s,%s)");
+			SampleCode += TEXT("Bias(%s,") + SamplerStateCode + TEXT(",%s,%s)");
 		}
 		else if(MipValueMode == TMVM_Derivative)
 		{
@@ -3143,7 +3143,7 @@ protected:
 				return Errorf(TEXT("Missing DDY(UVs) parameter"));
 			}
 
-			SampleCode += TEXT("Grad(%s,%sSampler,%s,%s,%s)");
+			SampleCode += TEXT("Grad(%s,") + SamplerStateCode + TEXT(",%s,%s,%s)");
 
 			MipValue0Code = CoerceParameter(MipValue0Index, UVsType);
 			MipValue1Code = CoerceParameter(MipValue1Index, UVsType);
@@ -4718,6 +4718,25 @@ protected:
 			*GetParameterCode(Depth), FunctionValueIndex);
 	}
 
+	virtual int32 Sobol(int32 Cell, int32 Index, int32 Seed) override
+	{
+		return AddCodeChunk(MCT_Float2,
+			TEXT("floor(%s) + float2(DirectSobol(PixelSobol(uint2(%s), uint(%s) & 7u), uint(%s) >> 3u).xy ^ uint2(%s * 0x10000) & 0xffff) / 0x10000"),
+			*GetParameterCode(Cell),
+			*GetParameterCode(Cell),
+			*GetParameterCode(Index),
+			*GetParameterCode(Index),
+			*GetParameterCode(Seed));
+	}
+
+	virtual int32 TemporalSobol(int32 Index, int32 Seed) override
+	{
+		return AddCodeChunk(MCT_Float2,
+			TEXT("float2(DirectSobol(PixelSobol(uint2(Parameters.SvPosition.xy), View.FrameNumber), uint(%s)).xy ^ uint2(%s * 0x10000) & 0xffff) / 0x10000"),
+			*GetParameterCode(Index),
+			*GetParameterCode(Seed));
+	}
+
 	virtual int32 Noise(int32 Position, float Scale, int32 Quality, uint8 NoiseFunction, bool bTurbulence, int32 Levels, float OutputMin, float OutputMax, float LevelScale, int32 FilterWidth, bool bTiling, uint32 RepeatSize) override
 	{
 		// GradientTex3D uses 3D texturing, which is not available on ES2
diff --git a/Engine/Source/Runtime/Engine/Private/Materials/Material.cpp b/Engine/Source/Runtime/Engine/Private/Materials/Material.cpp
--- a/Engine/Source/Runtime/Engine/Private/Materials/Material.cpp
+++ b/Engine/Source/Runtime/Engine/Private/Materials/Material.cpp
@@ -136,22 +136,26 @@ int32 FMaterialResource::CompilePropertyAndSetMaterialProperty(EMaterialProperty
 	};
 	
 	EMaterialValueType AttributeType = FMaterialAttributeDefinitionMap::GetValueType(Property);
-	FMaterialUniformExpression* Expression = Compiler->GetParameterUniformExpression(Ret);
 
-	if (Expression && Expression->IsConstant())
+	if (Ret != INDEX_NONE)
 	{
-		// Where possible we want to preserve constant expressions allowing default value checks
-		EMaterialValueType ResultType = Compiler->GetParameterType(Ret);
-		EMaterialValueType ExactAttributeType = (AttributeType == MCT_Float) ? MCT_Float1 : AttributeType;
-		EMaterialValueType ExactResultType = (ResultType == MCT_Float) ? MCT_Float1 : ResultType;
+		FMaterialUniformExpression* Expression = Compiler->GetParameterUniformExpression(Ret);
 
-		if (ExactAttributeType == ExactResultType)
+		if (Expression && Expression->IsConstant())
 		{
-			return Ret;
-		}
-		else if (ResultType == MCT_Float || (ExactAttributeType == MCT_Float1 && ResultType & MCT_Float))
-		{
-			return Compiler->ComponentMask(Ret, true, ExactAttributeType >= MCT_Float2, ExactAttributeType >= MCT_Float3, ExactAttributeType >= MCT_Float4);
+			// Where possible we want to preserve constant expressions allowing default value checks
+			EMaterialValueType ResultType = Compiler->GetParameterType(Ret);
+			EMaterialValueType ExactAttributeType = (AttributeType == MCT_Float) ? MCT_Float1 : AttributeType;
+			EMaterialValueType ExactResultType = (ResultType == MCT_Float) ? MCT_Float1 : ResultType;
+
+			if (ExactAttributeType == ExactResultType)
+			{
+				return Ret;
+			}
+			else if (ResultType == MCT_Float || (ExactAttributeType == MCT_Float1 && ResultType & MCT_Float))
+			{
+				return Compiler->ComponentMask(Ret, true, ExactAttributeType >= MCT_Float2, ExactAttributeType >= MCT_Float3, ExactAttributeType >= MCT_Float4);
+			}
 		}
 	}
 
@@ -3558,6 +3562,14 @@ void UMaterial::BeginDestroy()
 {
 	Super::BeginDestroy();
 
+	for (int32 InstanceIndex = 0; InstanceIndex < 3; ++InstanceIndex)
+	{
+		if (DefaultMaterialInstances[InstanceIndex])
+		{
+			BeginReleaseResource(DefaultMaterialInstances[InstanceIndex]);
+		}
+	}
+
 	ReleaseFence.BeginFence();
 }
 
@@ -3695,20 +3707,6 @@ void UMaterial::UpdateMaterialShaders(TArray<FShaderType*>& ShaderTypesToFlush,
 	{
 		FMaterialUpdateContext UpdateContext(FMaterialUpdateContext::EOptions::Default, ShaderPlatform);
 
-		// Go through all material shader maps and flush the appropriate shaders
-		FMaterialShaderMap::FlushShaderTypes(ShaderTypesToFlush, ShaderPipelineTypesToFlush, VFTypesToFlush);
-
-		// There should be no references to the given material shader types at this point
-		// If there still are shaders of the given types, they may be reused when we call CacheResourceShaders instead of compiling new shaders
-		for (int32 ShaderTypeIndex = 0; ShaderTypeIndex < ShaderTypesToFlush.Num(); ShaderTypeIndex++)
-		{
-			FShaderType* CurrentType = ShaderTypesToFlush[ShaderTypeIndex];
-			if (CurrentType->GetMaterialShaderType() || CurrentType->GetMeshMaterialShaderType())
-			{
-				checkf(CurrentType->GetNumShaders() == 0, TEXT("Type %s, Shaders %u"), CurrentType->GetName(), CurrentType->GetNumShaders());
-			}
-		}
-
 		int32 NumMaterials = 0;
 
 		for( TObjectIterator<UMaterial> It; It; ++It )
diff --git a/Engine/Source/Runtime/Engine/Private/Materials/MaterialExpressions.cpp b/Engine/Source/Runtime/Engine/Private/Materials/MaterialExpressions.cpp
--- a/Engine/Source/Runtime/Engine/Private/Materials/MaterialExpressions.cpp
+++ b/Engine/Source/Runtime/Engine/Private/Materials/MaterialExpressions.cpp
@@ -149,6 +149,7 @@
 #include "Materials/MaterialExpressionSceneTexture.h"
 #include "Materials/MaterialExpressionScreenPosition.h"
 #include "Materials/MaterialExpressionSine.h"
+#include "Materials/MaterialExpressionSobol.h"
 #include "Materials/MaterialExpressionSpeedTree.h"
 #include "Materials/MaterialExpressionSphereMask.h"
 #include "Materials/MaterialExpressionSphericalParticleOpacity.h"
@@ -158,6 +159,7 @@
 #include "Materials/MaterialExpressionSubtract.h"
 #include "Materials/MaterialExpressionTangent.h"
 #include "Materials/MaterialExpressionTangentOutput.h"
+#include "Materials/MaterialExpressionTemporalSobol.h"
 #include "Materials/MaterialExpressionTextureBase.h"
 #include "Materials/MaterialExpressionTextureObject.h"
 #include "Materials/MaterialExpressionTextureProperty.h"
@@ -6916,6 +6918,11 @@ void UMaterialExpressionPower::GetCaption(TArray<FString>& OutCaptions) const
 
 	OutCaptions.Add(ret);
 }
+
+void UMaterialExpressionPower::GetExpressionToolTip(TArray<FString>& OutToolTip) 
+{
+	ConvertToMultilineToolTip(TEXT("Returns the Base value raised to the power of Exponent. Base value must be positive, values less than 0 will be clamped."), 40, OutToolTip);
+}
 #endif // WITH_EDITOR
 
 UMaterialExpressionLogarithm2::UMaterialExpressionLogarithm2(const FObjectInitializer& ObjectInitializer)
@@ -10585,6 +10592,91 @@ void UMaterialExpressionSphereMask::GetCaption(TArray<FString>& OutCaptions) con
 }
 #endif // WITH_EDITOR
 
+///////////////////////////////////////////////////////////////////////////////
+// UMaterialExpressionSobol
+///////////////////////////////////////////////////////////////////////////////
+UMaterialExpressionSobol::UMaterialExpressionSobol(const FObjectInitializer& ObjectInitializer)
+	: Super(ObjectInitializer)
+{
+	// Structure to hold one-time initialization
+	struct FConstructorStatics
+	{
+		FText NAME_Utility;
+		FConstructorStatics()
+			: NAME_Utility(LOCTEXT("Utility", "Utility"))
+		{
+		}
+	};
+	static FConstructorStatics ConstructorStatics;
+
+	ConstIndex = 0;
+	ConstSeed = FVector2D(0.f, 0.f);
+}
+
+#if WITH_EDITOR
+int32 UMaterialExpressionSobol::Compile(class FMaterialCompiler* Compiler, int32 OutputIndex)
+{
+	int32 CellInput = Cell.GetTracedInput().Expression ? Cell.Compile(Compiler) : Compiler->Constant2(0.f, 0.f);
+	int32 IndexInput = Index.GetTracedInput().Expression ? Index.Compile(Compiler) : Compiler->Constant(ConstIndex);
+	int32 SeedInput = Seed.GetTracedInput().Expression ? Seed.Compile(Compiler) : Compiler->Constant2(ConstSeed.X, ConstSeed.Y);
+	return Compiler->Sobol(CellInput, IndexInput, SeedInput);
+}
+
+void UMaterialExpressionSobol::GetCaption(TArray<FString>& OutCaptions) const
+{
+	FString Caption = TEXT("Sobol");
+
+	if (!Index.GetTracedInput().Expression)
+	{
+		Caption += FString::Printf(TEXT(" (%d)"), ConstIndex);;
+	}
+
+	OutCaptions.Add(Caption);
+}
+#endif // WITH_EDITOR
+
+///////////////////////////////////////////////////////////////////////////////
+// UMaterialExpressionTemporalSobol
+///////////////////////////////////////////////////////////////////////////////
+UMaterialExpressionTemporalSobol::UMaterialExpressionTemporalSobol(const FObjectInitializer& ObjectInitializer)
+	: Super(ObjectInitializer)
+{
+	// Structure to hold one-time initialization
+	struct FConstructorStatics
+	{
+		FText NAME_Utility;
+		FConstructorStatics()
+			: NAME_Utility(LOCTEXT("Utility", "Utility"))
+		{
+		}
+	};
+	static FConstructorStatics ConstructorStatics;
+
+	ConstIndex = 0;
+	ConstSeed = FVector2D(0.f, 0.f);
+}
+
+#if WITH_EDITOR
+int32 UMaterialExpressionTemporalSobol::Compile(class FMaterialCompiler* Compiler, int32 OutputIndex)
+{
+	int32 IndexInput = Index.GetTracedInput().Expression ? Index.Compile(Compiler) : Compiler->Constant(ConstIndex);
+	int32 SeedInput = Seed.GetTracedInput().Expression ? Seed.Compile(Compiler) : Compiler->Constant2(ConstSeed.X, ConstSeed.Y);
+	return Compiler->TemporalSobol(IndexInput, SeedInput);
+}
+
+void UMaterialExpressionTemporalSobol::GetCaption(TArray<FString>& OutCaptions) const
+{
+	FString Caption = TEXT("Temporal Sobol");
+
+	if (!Index.GetTracedInput().Expression)
+	{
+		Caption += FString::Printf(TEXT(" (%d)"), ConstIndex);;
+	}
+
+	OutCaptions.Add(Caption);
+}
+#endif // WITH_EDITOR
+
 ///////////////////////////////////////////////////////////////////////////////
 // UMaterialExpressionNoise
 ///////////////////////////////////////////////////////////////////////////////
@@ -10679,6 +10771,9 @@ int32 UMaterialExpressionNoise::Compile(class FMaterialCompiler* Compiler, int32
 
 void UMaterialExpressionNoise::GetCaption(TArray<FString>& OutCaptions) const
 {
+	const UEnum* NFEnum = FindObject<UEnum>(nullptr, TEXT("Engine.ENoiseFunction"));
+	check(NFEnum);
+	OutCaptions.Add(NFEnum->GetDisplayNameTextByValue(NoiseFunction).ToString());
 	OutCaptions.Add(TEXT("Noise"));
 }
 #endif // WITH_EDITOR
@@ -10752,6 +10847,9 @@ int32 UMaterialExpressionVectorNoise::Compile(class FMaterialCompiler* Compiler,
 
 void UMaterialExpressionVectorNoise::GetCaption(TArray<FString>& OutCaptions) const
 {
+	const UEnum* VNFEnum = FindObject<UEnum>(nullptr, TEXT("Engine.EVectorNoiseFunction"));
+	check(VNFEnum);
+	OutCaptions.Add(VNFEnum->GetDisplayNameTextByValue(NoiseFunction).ToString());
 	OutCaptions.Add(TEXT("Vector Noise"));
 }
 #endif // WITH_EDITOR
diff --git a/Engine/Source/Runtime/Engine/Private/Materials/MaterialInstance.cpp b/Engine/Source/Runtime/Engine/Private/Materials/MaterialInstance.cpp
--- a/Engine/Source/Runtime/Engine/Private/Materials/MaterialInstance.cpp
+++ b/Engine/Source/Runtime/Engine/Private/Materials/MaterialInstance.cpp
@@ -2114,6 +2114,17 @@ void UMaterialInstance::BeginDestroy()
 {
 	Super::BeginDestroy();
 
+	if (!HasAnyFlags(RF_ClassDefaultObject))
+	{
+		BeginReleaseResource(Resources[0]);
+
+		if(GIsEditor)
+		{
+			BeginReleaseResource(Resources[1]);
+			BeginReleaseResource(Resources[2]);
+		}
+	}
+
 	ReleaseFence.BeginFence();
 }
 
diff --git a/Engine/Source/Runtime/Engine/Private/Materials/MaterialInterface.cpp b/Engine/Source/Runtime/Engine/Private/Materials/MaterialInterface.cpp
--- a/Engine/Source/Runtime/Engine/Private/Materials/MaterialInterface.cpp
+++ b/Engine/Source/Runtime/Engine/Private/Materials/MaterialInterface.cpp
@@ -33,6 +33,8 @@ void FMaterialRelevance::SetPrimitiveViewRelevance(FPrimitiveViewRelevance& OutV
 	OutViewRelevance.bSeparateTranslucencyRelevance = bSeparateTranslucency;
 	OutViewRelevance.bMobileSeparateTranslucencyRelevance = bMobileSeparateTranslucency;
 	OutViewRelevance.bNormalTranslucencyRelevance = bNormalTranslucency;
+	OutViewRelevance.bUsesSceneColorCopy = bUsesSceneColorCopy;
+	OutViewRelevance.bDisableOffscreenRendering = bDisableOffscreenRendering;
 	OutViewRelevance.ShadingModelMaskRelevance = ShadingModelMask;
 	OutViewRelevance.bUsesGlobalDistanceField = bUsesGlobalDistanceField;
 	OutViewRelevance.bUsesWorldPositionOffset = bUsesWorldPositionOffset;
@@ -89,7 +91,8 @@ FMaterialRelevance UMaterialInterface::GetRelevance_Internal(const UMaterial* Ma
 	if(Material)
 	{
 		const FMaterialResource* MaterialResource = Material->GetMaterialResource(InFeatureLevel);
-		const bool bIsTranslucent = IsTranslucentBlendMode((EBlendMode)GetBlendMode());
+		const EBlendMode BlendMode = (EBlendMode)GetBlendMode();
+		const bool bIsTranslucent = IsTranslucentBlendMode(BlendMode);
 
 		EMaterialShadingModel ShadingModel = GetShadingModel();
 		EMaterialDomain Domain = (EMaterialDomain)MaterialResource->GetMaterialDomain();
@@ -113,7 +116,9 @@ FMaterialRelevance UMaterialInterface::GetRelevance_Internal(const UMaterial* Ma
 			MaterialRelevance.bSeparateTranslucency = bIsTranslucent && Material->bEnableSeparateTranslucency;
 			MaterialRelevance.bMobileSeparateTranslucency = bIsTranslucent && Material->bEnableMobileSeparateTranslucency;
 			MaterialRelevance.bNormalTranslucency = bIsTranslucent && !Material->bEnableSeparateTranslucency;
-			MaterialRelevance.bDisableDepthTest = bIsTranslucent && Material->bDisableDepthTest;		
+			MaterialRelevance.bDisableDepthTest = bIsTranslucent && Material->bDisableDepthTest;	
+			MaterialRelevance.bUsesSceneColorCopy = bIsTranslucent && MaterialResource->RequiresSceneColorCopy_GameThread();
+			MaterialRelevance.bDisableOffscreenRendering = BlendMode == BLEND_Modulate; // Blend Modulate must be rendered directly in the scene color.
 			MaterialRelevance.bOutputsVelocityInBasePass = Material->bOutputVelocityOnBasePass;	
 			MaterialRelevance.bUsesGlobalDistanceField = MaterialResource->UsesGlobalDistanceField_GameThread();
 			MaterialRelevance.bUsesWorldPositionOffset = MaterialResource->UsesWorldPositionOffset_GameThread();
diff --git a/Engine/Source/Runtime/Engine/Private/Materials/MaterialShader.cpp b/Engine/Source/Runtime/Engine/Private/Materials/MaterialShader.cpp
--- a/Engine/Source/Runtime/Engine/Private/Materials/MaterialShader.cpp
+++ b/Engine/Source/Runtime/Engine/Private/Materials/MaterialShader.cpp
@@ -852,11 +852,11 @@ void FMaterialShaderMap::FixupShaderTypes(EShaderPlatform Platform, const TMap<F
 	}
 }
 
-void FMaterialShaderMap::LoadFromDerivedDataCache(const FMaterial* Material, const FMaterialShaderMapId& ShaderMapId, EShaderPlatform Platform, TRefCountPtr<FMaterialShaderMap>& InOutShaderMap)
+void FMaterialShaderMap::LoadFromDerivedDataCache(const FMaterial* Material, const FMaterialShaderMapId& ShaderMapId, EShaderPlatform InPlatform, TRefCountPtr<FMaterialShaderMap>& InOutShaderMap)
 {
 	if (InOutShaderMap != NULL)
 	{
-		check(InOutShaderMap->Platform == Platform);
+		check(InOutShaderMap->GetShaderPlatform() == InPlatform);
 		// If the shader map was non-NULL then it was found in memory but is incomplete, attempt to load the missing entries from memory
 		InOutShaderMap->LoadMissingShadersFromMemory(Material);
 	}
@@ -869,12 +869,12 @@ void FMaterialShaderMap::LoadFromDerivedDataCache(const FMaterial* Material, con
 			COOK_STAT(auto Timer = MaterialShaderCookStats::UsageStats.TimeSyncWork());
 
 			TArray<uint8> CachedData;
-			const FString DataKey = GetMaterialShaderMapKeyString(ShaderMapId, Platform);
+			const FString DataKey = GetMaterialShaderMapKeyString(ShaderMapId, InPlatform);
 
 			if (GetDerivedDataCacheRef().GetSynchronous(*DataKey, CachedData))
 			{
 				COOK_STAT(Timer.AddHit(CachedData.Num()));
-				InOutShaderMap = new FMaterialShaderMap();
+				InOutShaderMap = new FMaterialShaderMap(InPlatform);
 				FMemoryReader Ar(CachedData, true);
 
 				// Deserialize from the cached data
@@ -884,7 +884,7 @@ void FMaterialShaderMap::LoadFromDerivedDataCache(const FMaterial* Material, con
 				checkSlow(InOutShaderMap->GetShaderMapId() == ShaderMapId);
 
 				// Register in the global map
-				InOutShaderMap->Register(Platform);
+				InOutShaderMap->Register(InPlatform);
 			}
 			else
 			{
@@ -904,7 +904,7 @@ void FMaterialShaderMap::SaveToDerivedDataCache()
 	FMemoryWriter Ar(SaveData, true);
 	Serialize(Ar);
 
-	GetDerivedDataCacheRef().Put(*GetMaterialShaderMapKeyString(ShaderMapId, Platform), SaveData);
+	GetDerivedDataCacheRef().Put(*GetMaterialShaderMapKeyString(ShaderMapId, GetShaderPlatform()), SaveData);
 	COOK_STAT(Timer.AddMiss(SaveData.Num()));
 }
 
@@ -1100,7 +1100,7 @@ void FMaterialShaderMap::LoadForRemoteRecompile(FArchive& Ar, EShaderPlatform Sh
 
 			if (bIsValid)
 			{
-				FMaterialShaderMap* ShaderMap = new FMaterialShaderMap;
+				FMaterialShaderMap* ShaderMap = new FMaterialShaderMap(ShaderPlatform);
 
 				// serialize the id and the material shader map
 				ShaderMap->Serialize(Ar, false);
@@ -1285,7 +1285,7 @@ void FMaterialShaderMap::Compile(
 					{
 						// Create a new mesh material shader map.
 						MeshShaderMapIndex = MeshShaderMaps.Num();
-						MeshShaderMap = new(MeshShaderMaps) FMeshMaterialShaderMap(VertexFactoryType);
+						MeshShaderMap = new(MeshShaderMaps) FMeshMaterialShaderMap(InPlatform, VertexFactoryType);
 					}
   
 					// Enqueue compilation all mesh material shaders for this material and vertex factory type combo.
@@ -1468,6 +1468,7 @@ FShader* FMaterialShaderMap::ProcessCompilationResultsForSingleJob(FShaderCompil
 bool FMaterialShaderMap::ProcessCompilationResults(const TArray<FShaderCommonCompileJob*>& InCompilationResults, int32& InOutJobIndex, float& TimeBudget, TMap<const FVertexFactoryType*, TArray<const FShaderPipelineType*> >& SharedPipelines)
 {
 	check(InOutJobIndex < InCompilationResults.Num());
+	check(!bCompilationFinalized);
 
 	double StartTime = FPlatformTime::Seconds();
 
@@ -1653,7 +1654,7 @@ bool FMaterialShaderMap::TryToAddToExistingCompilationTask(FMaterial* Material)
 bool FMaterialShaderMap::IsMaterialShaderComplete(const FMaterial* Material, const FMaterialShaderType* ShaderType, const FShaderPipelineType* Pipeline, bool bSilent)
 {
 	// If we should cache this material, it's incomplete if the pipeline is missing or the shader itself is missing
-	if (ShouldCacheMaterialShader(ShaderType, Platform, Material) &&
+	if (ShouldCacheMaterialShader(ShaderType, GetShaderPlatform(), Material) &&
 		((Pipeline && !HasShaderPipeline(Pipeline)) || (!Pipeline && !HasShader((FShaderType*)ShaderType))))
 	{
 		if (!bSilent)
@@ -1695,7 +1696,7 @@ bool FMaterialShaderMap::IsComplete(const FMaterial* Material, bool bSilent)
 		{
 			// Find the shaders for this vertex factory type.
 			const FMeshMaterialShaderMap* MeshShaderMap = GetMeshShaderMap(VertexFactoryType);
-			if (!FMeshMaterialShaderMap::IsComplete(MeshShaderMap,Platform,Material,VertexFactoryType,bSilent))
+			if (!FMeshMaterialShaderMap::IsComplete(MeshShaderMap,GetShaderPlatform(),Material,VertexFactoryType,bSilent))
 			{
 				if (!MeshShaderMap && !bSilent)
 				{
@@ -1730,7 +1731,7 @@ bool FMaterialShaderMap::IsComplete(const FMaterial* Material, bool bSilent)
 			for (int32 Index = 0; Index < StageTypes.Num(); ++Index)
 			{
 				auto* ShaderType = StageTypes[Index]->GetMaterialShaderType();
-				if (ShouldCacheMaterialShader(ShaderType, Platform, Material))
+				if (ShouldCacheMaterialShader(ShaderType, GetShaderPlatform(), Material))
 				{
 					++NumShouldCache;
 				}
@@ -1778,9 +1779,9 @@ void FMaterialShaderMap::LoadMissingShadersFromMemory(const FMaterial* Material)
 	for (TLinkedList<FShaderType*>::TIterator ShaderTypeIt(FShaderType::GetTypeList());ShaderTypeIt;ShaderTypeIt.Next())
 	{
 		FMaterialShaderType* ShaderType = ShaderTypeIt->GetMaterialShaderType();
-		if (ShaderType && ShouldCacheMaterialShader(ShaderType, Platform, Material) && !HasShader(ShaderType))
+		if (ShaderType && ShouldCacheMaterialShader(ShaderType, GetShaderPlatform(), Material) && !HasShader(ShaderType))
 		{
-			FShaderId ShaderId(MaterialShaderMapHash, nullptr, nullptr, ShaderType, FShaderTarget(ShaderType->GetFrequency(), Platform));
+			FShaderId ShaderId(MaterialShaderMapHash, nullptr, nullptr, ShaderType, FShaderTarget(ShaderType->GetFrequency(), GetShaderPlatform()));
 			FShader* FoundShader = ShaderType->FindShaderById(ShaderId);
 			if (FoundShader)
 			{
@@ -1801,7 +1802,7 @@ void FMaterialShaderMap::LoadMissingShadersFromMemory(const FMaterial* Material)
 			for (const FShaderType* Shader : Stages)
 			{
 				FMaterialShaderType* ShaderType = (FMaterialShaderType*)Shader->GetMaterialShaderType();
-				if (ShaderType && ShouldCacheMaterialShader(ShaderType, Platform, Material))
+				if (ShaderType && ShouldCacheMaterialShader(ShaderType, GetShaderPlatform(), Material))
 				{
 					++NumShaders;
 				}
@@ -1815,7 +1816,7 @@ void FMaterialShaderMap::LoadMissingShadersFromMemory(const FMaterial* Material)
 					FMaterialShaderType* ShaderType = (FMaterialShaderType*)Shader->GetMaterialShaderType();
 					if (!HasShader(ShaderType))
 					{
-						FShaderId ShaderId(MaterialShaderMapHash, PipelineType->ShouldOptimizeUnusedOutputs() ? PipelineType : nullptr, nullptr, ShaderType, FShaderTarget(ShaderType->GetFrequency(), Platform));
+						FShaderId ShaderId(MaterialShaderMapHash, PipelineType->ShouldOptimizeUnusedOutputs() ? PipelineType : nullptr, nullptr, ShaderType, FShaderTarget(ShaderType->GetFrequency(), GetShaderPlatform()));
 						FShader* FoundShader = ShaderType->FindShaderById(ShaderId);
 						if (FoundShader)
 						{
@@ -1846,7 +1847,7 @@ void FMaterialShaderMap::LoadMissingShadersFromMemory(const FMaterial* Material)
 
 			if (MeshShaderMap)
 			{
-				MeshShaderMap->LoadMissingShadersFromMemory(MaterialShaderMapHash, Material, Platform);
+				MeshShaderMap->LoadMissingShadersFromMemory(MaterialShaderMapHash, Material, GetShaderPlatform());
 			}
 		}
 	}
@@ -1876,7 +1877,7 @@ void FMaterialShaderMap::GetShaderPipelineList(TArray<FShaderPipeline*>& OutShad
 void FMaterialShaderMap::Register(EShaderPlatform InShaderPlatform)
 {
 	extern int32 GCreateShadersOnLoad;
-	if (GCreateShadersOnLoad && Platform == InShaderPlatform)
+	if (GCreateShadersOnLoad && GetShaderPlatform() == InShaderPlatform)
 	{
 		for (auto KeyValue : GetShaders())
 		{
@@ -1906,7 +1907,7 @@ void FMaterialShaderMap::Register(EShaderPlatform InShaderPlatform)
 		INC_DWORD_STAT_BY(STAT_Shaders_ShaderMapMemory, GetSizeBytes());
 	}
 
-	GIdToMaterialShaderMap[Platform].Add(ShaderMapId,this); 
+	GIdToMaterialShaderMap[GetShaderPlatform()].Add(ShaderMapId,this);
 	bRegistered = true;
 }
 
@@ -1926,7 +1927,7 @@ void FMaterialShaderMap::Release()
 			DEC_DWORD_STAT(STAT_Shaders_NumShaderMaps);
 			DEC_DWORD_STAT_BY(STAT_Shaders_ShaderMapMemory, GetSizeBytes());
 
-			GIdToMaterialShaderMap[Platform].Remove(ShaderMapId);
+			GIdToMaterialShaderMap[GetShaderPlatform()].Remove(ShaderMapId);
 			bRegistered = false;
 		}
 
@@ -1934,8 +1935,8 @@ void FMaterialShaderMap::Release()
 	}
 }
 
-FMaterialShaderMap::FMaterialShaderMap() :
-	Platform(SP_NumPlatforms),
+FMaterialShaderMap::FMaterialShaderMap(EShaderPlatform InPlatform) :
+	TShaderMap<FMaterialShaderType>(InPlatform),
 	CompilingId(1),
 	NumRefs(0),
 	bDeletedThroughDeferredCleanup(false),
@@ -2032,7 +2033,7 @@ void FMaterialShaderMap::Serialize(FArchive& Ar, bool bInlineShaderResources)
 	ShaderMapId.Serialize(Ar);
 
 	// serialize the platform enum as a uint8
-	int32 TempPlatform = (int32)Platform;
+	int32 TempPlatform = (int32)GetShaderPlatform();
 	Ar << TempPlatform;
 	Platform = (EShaderPlatform)TempPlatform;
 
@@ -2103,7 +2104,7 @@ void FMaterialShaderMap::Serialize(FArchive& Ar, bool bInlineShaderResources)
 
 			if (VertexFactoryType->IsUsedWithMaterials())
 			{
-				new(MeshShaderMaps) FMeshMaterialShaderMap(VertexFactoryType);
+				new(MeshShaderMaps) FMeshMaterialShaderMap(GetShaderPlatform(), VertexFactoryType);
 			}
 		}
 
diff --git a/Engine/Source/Runtime/Engine/Private/Materials/MaterialShared.cpp b/Engine/Source/Runtime/Engine/Private/Materials/MaterialShared.cpp
--- a/Engine/Source/Runtime/Engine/Private/Materials/MaterialShared.cpp
+++ b/Engine/Source/Runtime/Engine/Private/Materials/MaterialShared.cpp
@@ -35,6 +35,14 @@
 
 DEFINE_LOG_CATEGORY(LogMaterial);
 
+int32 GDeferUniformExpressionCaching = 1;
+FAutoConsoleVariableRef CVarDeferUniformExpressionCaching(
+	TEXT("r.DeferUniformExpressionCaching"),
+	GDeferUniformExpressionCaching,
+	TEXT("Whether to defer caching of uniform expressions until a rendering command needs them up to date.  Deferring updates is more efficient because multiple SetVectorParameterValue calls in a frame will only result in one update."),
+	ECVF_RenderThreadSafe
+	);
+
 FName MaterialQualityLevelNames[] = 
 {
 	FName(TEXT("Low")),
@@ -501,12 +509,7 @@ const TArray<TRefCountPtr<FMaterialUniformExpression> >& FMaterial::GetUniformSc
 
 bool FMaterial::RequiresSceneColorCopy_GameThread() const
 {
-	check(IsInGameThread());
-	if (GameThreadShaderMap)
-	{
-		return GameThreadShaderMap->RequiresSceneColorCopy();
-	}
-	return false;
+	return GameThreadShaderMap.GetReference() ? GameThreadShaderMap->RequiresSceneColorCopy() : false; 
 }
 
 bool FMaterial::RequiresSceneColorCopy_RenderThread() const
@@ -899,7 +902,7 @@ bool FMaterialResource::IsCrackFreeDisplacementEnabled() const
 	return Material->bEnableCrackFreeDisplacement;
 }
 
-bool FMaterialResource::IsSeparateTranslucencyEnabled() const 
+bool FMaterialResource::IsTranslucencyAfterDOFEnabled() const 
 { 
 	return Material->bEnableSeparateTranslucency && !IsUIMaterial() && !IsDeferredDecal();
 }
@@ -1976,36 +1979,17 @@ void FMaterialRenderProxy::CacheUniformExpressions()
 	check((bUsingNewLoader && GIsInitialLoad) || // The EDL at boot time maybe not load the default materials first; we need to intialize materials before the default materials are done
 		UMaterial::GetDefaultMaterial(MD_Surface));
 
+	DeferredUniformExpressionCacheRequests.Add(this);
+
 	UMaterialInterface::IterateOverActiveFeatureLevels([&](ERHIFeatureLevel::Type InFeatureLevel)
 	{
-		const FMaterial* MaterialNoFallback = GetMaterialNoFallback(InFeatureLevel);
-
-		if (MaterialNoFallback && MaterialNoFallback->GetRenderingThreadShaderMap())
-		{
-			const FMaterial* Material = bUsingNewLoader ? MaterialNoFallback : GetMaterial(InFeatureLevel);
-
-			// Do not cache uniform expressions for fallback materials. This step could
-			// be skipped where we don't allow for asynchronous shader compiling.
-			bool bIsFallbackMaterial = bUsingNewLoader ? false : (Material != MaterialNoFallback);
-
-			if (!bIsFallbackMaterial)
-			{
-				FMaterialRenderContext MaterialRenderContext(this , *Material, nullptr);
-				MaterialRenderContext.bShowSelection = GIsEditor;
-				EvaluateUniformExpressions(UniformExpressionCache[(int32)InFeatureLevel], MaterialRenderContext);
-			}
-			else
-			{
-				InvalidateUniformExpressionCache();
-				return;
-			}
-		}
-		else
-		{
-			InvalidateUniformExpressionCache();
-			return;
-		}
+		InvalidateUniformExpressionCache();
 	});
+
+	if (!GDeferUniformExpressionCaching)
+	{
+		FMaterialRenderProxy::UpdateDeferredCachedUniformExpressions();
+	}
 }
 
 void FMaterialRenderProxy::CacheUniformExpressions_GameThread()
@@ -2072,16 +2056,60 @@ FMaterialRenderProxy::~FMaterialRenderProxy()
 
 void FMaterialRenderProxy::InitDynamicRHI()
 {
-	FMaterialRenderProxy::MaterialRenderProxyMap.Add(this);
+	// MaterialRenderProxyMap is only used by shader compiling
+	if (!FPlatformProperties::RequiresCookedData())
+	{
+		FMaterialRenderProxy::MaterialRenderProxyMap.Add(this);
+	}
 }
 
 void FMaterialRenderProxy::ReleaseDynamicRHI()
 {
-	FMaterialRenderProxy::MaterialRenderProxyMap.Remove(this);
+	if (!FPlatformProperties::RequiresCookedData())
+	{
+		FMaterialRenderProxy::MaterialRenderProxyMap.Remove(this);
+	}
+
+	DeferredUniformExpressionCacheRequests.Remove(this);
+
 	InvalidateUniformExpressionCache();
 }
 
+void FMaterialRenderProxy::UpdateDeferredCachedUniformExpressions()
+{
+	check(IsInRenderingThread());
+
+	for (TSet<FMaterialRenderProxy*>::TConstIterator It(DeferredUniformExpressionCacheRequests); It; ++It)
+	{
+		FMaterialRenderProxy* MaterialProxy = *It;
+
+		UMaterialInterface::IterateOverActiveFeatureLevels([&](ERHIFeatureLevel::Type InFeatureLevel)
+		{
+			const FMaterial* MaterialNoFallback = MaterialProxy->GetMaterialNoFallback(InFeatureLevel);
+
+			if (MaterialNoFallback && MaterialNoFallback->GetRenderingThreadShaderMap())
+			{
+				const FMaterial* Material = MaterialProxy->GetMaterial(InFeatureLevel);
+
+				// Do not cache uniform expressions for fallback materials. This step could
+				// be skipped where we don't allow for asynchronous shader compiling.
+				bool bIsFallbackMaterial = (Material != MaterialNoFallback);
+
+				if (!bIsFallbackMaterial)
+				{
+					FMaterialRenderContext MaterialRenderContext(MaterialProxy, *Material, nullptr);
+					MaterialRenderContext.bShowSelection = GIsEditor;
+					MaterialProxy->EvaluateUniformExpressions(MaterialProxy->UniformExpressionCache[(int32)InFeatureLevel], MaterialRenderContext);
+				}
+			}
+		});
+	}
+
+	DeferredUniformExpressionCacheRequests.Reset();
+}
+
 TSet<FMaterialRenderProxy*> FMaterialRenderProxy::MaterialRenderProxyMap;
+TSet<FMaterialRenderProxy*> FMaterialRenderProxy::DeferredUniformExpressionCacheRequests;
 
 /*-----------------------------------------------------------------------------
 	FColoredMaterialRenderProxy
diff --git a/Engine/Source/Runtime/Engine/Private/Materials/MaterialUniformExpressions.h b/Engine/Source/Runtime/Engine/Private/Materials/MaterialUniformExpressions.h
--- a/Engine/Source/Runtime/Engine/Private/Materials/MaterialUniformExpressions.h
+++ b/Engine/Source/Runtime/Engine/Private/Materials/MaterialUniformExpressions.h
@@ -139,7 +139,7 @@ public:
 	{
 		OutValue.R = OutValue.G = OutValue.B = OutValue.A = 0;
 
-		if(!Context.MaterialRenderProxy->GetVectorValue(ParameterName, &OutValue, Context))
+		if(!Context.MaterialRenderProxy || !Context.MaterialRenderProxy->GetVectorValue(ParameterName, &OutValue, Context))
 		{
 			GetDefaultValue(OutValue);
 		}
@@ -212,7 +212,7 @@ public:
 	// inefficient compared to GetGameThreadNumberValue(), for editor purpose
 	virtual void GetNumberValue(const FMaterialRenderContext& Context,FLinearColor& OutValue) const
 	{
-		if(Context.MaterialRenderProxy->GetScalarValue(ParameterName, &OutValue.R, Context))
+		if(Context.MaterialRenderProxy && Context.MaterialRenderProxy->GetScalarValue(ParameterName, &OutValue.R, Context))
 		{
 			OutValue.G = OutValue.B = OutValue.A = OutValue.R;
 		}
@@ -319,7 +319,7 @@ public:
 		else
 		{
 			OutValue = NULL;
-			if(!Context.MaterialRenderProxy->GetTextureValue(ParameterName,&OutValue,Context))
+			if(!Context.MaterialRenderProxy || !Context.MaterialRenderProxy->GetTextureValue(ParameterName,&OutValue,Context))
 			{
 				OutValue = GetIndexedTexture(Material, TextureIndex);
 			}
diff --git a/Engine/Source/Runtime/Engine/Private/Materials/MeshMaterialShader.cpp b/Engine/Source/Runtime/Engine/Private/Materials/MeshMaterialShader.cpp
--- a/Engine/Source/Runtime/Engine/Private/Materials/MeshMaterialShader.cpp
+++ b/Engine/Source/Runtime/Engine/Private/Materials/MeshMaterialShader.cpp
@@ -156,7 +156,7 @@ uint32 FMeshMaterialShaderMap::BeginCompile(
 	const FMaterialShaderMapId& InShaderMapId, 
 	const FMaterial* Material,
 	FShaderCompilerEnvironment* MaterialEnvironment,
-	EShaderPlatform Platform,
+	EShaderPlatform InPlatform,
 	TArray<FShaderCommonCompileJob*>& NewJobs
 	)
 {
@@ -173,7 +173,7 @@ uint32 FMeshMaterialShaderMap::BeginCompile(
 	for (TLinkedList<FShaderType*>::TIterator ShaderTypeIt(FShaderType::GetTypeList());ShaderTypeIt;ShaderTypeIt.Next())
 	{
 		FMeshMaterialShaderType* ShaderType = ShaderTypeIt->GetMeshMaterialShaderType();
-		if (ShaderType && ShouldCacheMeshShader(ShaderType, Platform, Material, VertexFactoryType))
+		if (ShaderType && ShouldCacheMeshShader(ShaderType, InPlatform, Material, VertexFactoryType))
 		{
 			// Verify that the shader map Id contains inputs for any shaders that will be put into this shader map
 			check(InShaderMapId.ContainsVertexFactoryType(VertexFactoryType));
@@ -186,7 +186,7 @@ uint32 FMeshMaterialShaderMap::BeginCompile(
 				// Compile this mesh material shader for this material and vertex factory type.
 				auto* Job = ShaderType->BeginCompileShader(
 					ShaderMapId,
-					Platform,
+					InPlatform,
 					Material,
 					MaterialEnvironment,
 					VertexFactoryType,
@@ -211,7 +211,7 @@ uint32 FMeshMaterialShaderMap::BeginCompile(
 			for (auto* Shader : StageTypes)
 			{
 				const FMeshMaterialShaderType* ShaderType = Shader->GetMeshMaterialShaderType();
-				if (ShouldCacheMeshShader(ShaderType, Platform, Material, VertexFactoryType))
+				if (ShouldCacheMeshShader(ShaderType, InPlatform, Material, VertexFactoryType))
 				{
 					++NumShaderStagesToCompile;
 				}
@@ -241,7 +241,7 @@ uint32 FMeshMaterialShaderMap::BeginCompile(
 					}
 
 					// Make a pipeline job with all the stages
-					FMeshMaterialShaderType::BeginCompileShaderPipeline(ShaderMapId, Platform, Material, MaterialEnvironment, VertexFactoryType, Pipeline, ShaderStagesToCompile, NewJobs);
+					FMeshMaterialShaderType::BeginCompileShaderPipeline(ShaderMapId, InPlatform, Material, MaterialEnvironment, VertexFactoryType, Pipeline, ShaderStagesToCompile, NewJobs);
 				}
 				else
 				{
@@ -358,14 +358,14 @@ bool FMeshMaterialShaderMap::IsComplete(
 void FMeshMaterialShaderMap::LoadMissingShadersFromMemory(
 	const FSHAHash& MaterialShaderMapHash, 
 	const FMaterial* Material, 
-	EShaderPlatform Platform)
+	EShaderPlatform InPlatform)
 {
 	for (TLinkedList<FShaderType*>::TIterator ShaderTypeIt(FShaderType::GetTypeList());ShaderTypeIt;ShaderTypeIt.Next())
 	{
 		FMeshMaterialShaderType* ShaderType = ShaderTypeIt->GetMeshMaterialShaderType();
-		if (ShaderType && ShouldCacheMeshShader(ShaderType, Platform, Material, VertexFactoryType) && !HasShader((FShaderType*)ShaderType))
+		if (ShaderType && ShouldCacheMeshShader(ShaderType, InPlatform, Material, VertexFactoryType) && !HasShader((FShaderType*)ShaderType))
 		{
-			const FShaderId ShaderId(MaterialShaderMapHash, nullptr, VertexFactoryType, (FShaderType*)ShaderType, FShaderTarget(ShaderType->GetFrequency(), Platform));
+			const FShaderId ShaderId(MaterialShaderMapHash, nullptr, VertexFactoryType, (FShaderType*)ShaderType, FShaderTarget(ShaderType->GetFrequency(), InPlatform));
 			FShader* FoundShader = ((FShaderType*)ShaderType)->FindShaderById(ShaderId);
 
 			if (FoundShader)
@@ -387,7 +387,7 @@ void FMeshMaterialShaderMap::LoadMissingShadersFromMemory(
 			for (const FShaderType* Shader : Stages)
 			{
 				FMeshMaterialShaderType* ShaderType = (FMeshMaterialShaderType*)Shader->GetMeshMaterialShaderType();
-				if (ShaderType && ShouldCacheMeshShader(ShaderType, Platform, Material, VertexFactoryType))
+				if (ShaderType && ShouldCacheMeshShader(ShaderType, InPlatform, Material, VertexFactoryType))
 				{
 					++NumShaders;
 				}
@@ -405,7 +405,7 @@ void FMeshMaterialShaderMap::LoadMissingShadersFromMemory(
 					FMeshMaterialShaderType* ShaderType = (FMeshMaterialShaderType*)Shader->GetMeshMaterialShaderType();
 					if (!HasShader(ShaderType))
 					{
-						const FShaderId ShaderId(MaterialShaderMapHash, PipelineType->ShouldOptimizeUnusedOutputs() ? PipelineType : nullptr, VertexFactoryType, ShaderType, FShaderTarget(ShaderType->GetFrequency(), Platform));
+						const FShaderId ShaderId(MaterialShaderMapHash, PipelineType->ShouldOptimizeUnusedOutputs() ? PipelineType : nullptr, VertexFactoryType, ShaderType, FShaderTarget(ShaderType->GetFrequency(), InPlatform));
 						FShader* FoundShader = ShaderType->FindShaderById(ShaderId);
 						if (FoundShader)
 						{
diff --git a/Engine/Source/Runtime/Engine/Private/Particles/ParticleComponents.cpp b/Engine/Source/Runtime/Engine/Private/Particles/ParticleComponents.cpp
--- a/Engine/Source/Runtime/Engine/Private/Particles/ParticleComponents.cpp
+++ b/Engine/Source/Runtime/Engine/Private/Particles/ParticleComponents.cpp
@@ -1977,8 +1977,8 @@ UParticleSystem::UParticleSystem(const FObjectInitializer& ObjectInitializer)
 	EditorLODSetting = 0;
 #endif // WITH_EDITORONLY_DATA
 	FixedRelativeBoundingBox.Min = FVector(-1.0f, -1.0f, -1.0f);
-
 	FixedRelativeBoundingBox.Max = FVector(1.0f, 1.0f, 1.0f);
+	FixedRelativeBoundingBox.IsValid = true;
 
 	LODMethod = PARTICLESYSTEMLODMETHOD_Automatic;
 	LODDistanceCheckTime = 0.25f;
@@ -4050,6 +4050,24 @@ void UParticleSystemComponent::SetMaterial(int32 ElementIndex, UMaterialInterfac
 		}
 		EmitterMaterials[ElementIndex] = Material;
 		bIsViewRelevanceDirty = true;
+
+		for (int32 EmitterIndex = 0; EmitterIndex < EmitterInstances.Num(); ++EmitterIndex)
+		{
+			if (FParticleEmitterInstance* Inst = EmitterInstances[EmitterIndex])
+			{
+				if (!Inst->Tick_MaterialOverrides())
+				{
+					if (EmitterMaterials.IsValidIndex(EmitterIndex))
+					{
+						if (EmitterMaterials[EmitterIndex])
+						{
+							Inst->CurrentMaterial = EmitterMaterials[EmitterIndex];
+						}
+					}
+				}
+			}
+		}
+		MarkRenderDynamicDataDirty();
 	}
 }
 
diff --git a/Engine/Source/Runtime/Engine/Private/Particles/ParticleSystemRender.cpp b/Engine/Source/Runtime/Engine/Private/Particles/ParticleSystemRender.cpp
--- a/Engine/Source/Runtime/Engine/Private/Particles/ParticleSystemRender.cpp
+++ b/Engine/Source/Runtime/Engine/Private/Particles/ParticleSystemRender.cpp
@@ -1102,7 +1102,7 @@ void FDynamicSpriteEmitterData::GetDynamicMeshElementsEmitter(const FParticleSys
 									FParticleOrder* ParticleOrder = NULL;
 									if (bSort)
 									{
-										ParticleOrder = (FParticleOrder*)FMemStack::Get().Alloc(sizeof(FParticleOrder)* ParticleCount, ALIGNOF(FParticleOrder));
+										ParticleOrder = (FParticleOrder*)FMemStack::Get().Alloc(sizeof(FParticleOrder)* ParticleCount, alignof(FParticleOrder));
 										SortSpriteParticles(SourceData->SortMode, SourceData->bUseLocalSpace, SourceData->ActiveParticleCount, 
 											SourceData->DataContainer.ParticleData, SourceData->ParticleStride, SourceData->DataContainer.ParticleIndices,
 											View, Proxy->GetLocalToWorld(), ParticleOrder);
@@ -1739,10 +1739,14 @@ void FDynamicMeshEmitterData::GetDynamicMeshElementsEmitter(const FParticleSyste
 			{
 				for (int32 SectionIndex = 0; SectionIndex < LODModel.Sections.Num(); SectionIndex++)
 				{
-					FMaterialRenderProxy* MaterialProxy = MeshMaterials[SectionIndex]->GetRenderProxy(bSelected);
+					FMaterialRenderProxy* MaterialProxy = nullptr;
+					if (SectionIndex < MeshMaterials.Num() && MeshMaterials[SectionIndex])
+					{
+						MaterialProxy = MeshMaterials[SectionIndex]->GetRenderProxy(bSelected);
+					}
 					const FStaticMeshSection& Section = LODModel.Sections[SectionIndex];
 
-					if ((Section.NumTriangles == 0) || (MaterialProxy == NULL))
+					if ((Section.NumTriangles == 0) || (MaterialProxy == nullptr))
 					{
 						//@todo. This should never occur, but it does occasionally.
 						continue;
diff --git a/Engine/Source/Runtime/Engine/Private/PhysicsEngine/PhysAnim.cpp b/Engine/Source/Runtime/Engine/Private/PhysicsEngine/PhysAnim.cpp
--- a/Engine/Source/Runtime/Engine/Private/PhysicsEngine/PhysAnim.cpp
+++ b/Engine/Source/Runtime/Engine/Private/PhysicsEngine/PhysAnim.cpp
@@ -108,7 +108,7 @@ public:
 	}
 };
 
-typedef TArray<FAssetWorldBoneTM, TMemStackAllocator<ALIGNOF(FAssetWorldBoneTM)>> TAssetWorldBoneTMArray;
+typedef TArray<FAssetWorldBoneTM, TMemStackAllocator<alignof(FAssetWorldBoneTM)>> TAssetWorldBoneTMArray;
 // Use current pose to calculate world-space position of this bone without physics now.
 void UpdateWorldBoneTM(TAssetWorldBoneTMArray& WorldBoneTMs, const TArray<FTransform>& InBoneSpaceTransforms, int32 BoneIndex, USkeletalMeshComponent* SkelComp, const FTransform &LocalToWorldTM, const FVector& Scale3D)
 {
@@ -185,7 +185,7 @@ void USkeletalMeshComponent::PerformBlendPhysicsBones(const TArray<FBoneIndexTyp
 
 #define DEPERCATED_PHYSBLEND_UPDATES_PHYSX 0	//If you really need the old behavior of physics blending set this to 1. Note this is inefficient and doesn't lead to good results. Please use UPhysicalAnimationComponent
 #if DEPERCATED_PHYSBLEND_UPDATES_PHYSX
-	TArray<FBodyTMPair, TMemStackAllocator<ALIGNOF(FBodyTMPair)>> PendingBodyTMs;
+	TArray<FBodyTMPair, TMemStackAllocator<alignof(FBodyTMPair)>> PendingBodyTMs;
 #endif
 
 #if WITH_PHYSX
diff --git a/Engine/Source/Runtime/Engine/Private/PhysicsEngine/PhysDrawing.cpp b/Engine/Source/Runtime/Engine/Private/PhysicsEngine/PhysDrawing.cpp
--- a/Engine/Source/Runtime/Engine/Private/PhysicsEngine/PhysDrawing.cpp
+++ b/Engine/Source/Runtime/Engine/Private/PhysicsEngine/PhysDrawing.cpp
@@ -64,7 +64,7 @@ void FKSphereElem::DrawElemWire(class FPrimitiveDrawInterface* PDI, const FTrans
 
 void FKSphereElem::DrawElemSolid(class FPrimitiveDrawInterface* PDI, const FTransform& ElemTM, const FVector& Scale3D, const FMaterialRenderProxy* MaterialRenderProxy) const
 {
-	DrawSphere(PDI, ElemTM.GetLocation(), FVector(this->Radius * Scale3D.GetAbsMin()), DrawCollisionSides, DrawCollisionSides / 2, MaterialRenderProxy, SDPG_World);
+	DrawSphere(PDI, ElemTM.GetLocation(), FRotator::ZeroRotator, FVector(this->Radius * Scale3D.GetAbsMin()), DrawCollisionSides, DrawCollisionSides / 2, MaterialRenderProxy, SDPG_World);
 }
 
 void FKSphereElem::GetElemSolid(const FTransform& ElemTM, const FVector& Scale3D, const FMaterialRenderProxy* MaterialRenderProxy, int32 ViewIndex, FMeshElementCollector& Collector) const
@@ -965,7 +965,7 @@ void FConstraintInstance::DrawConstraintImp(const FPDIOrCollector& PDIOrCollecto
 				}
 				else
 				{
-					DrawSphere(PDI, Con1Pos, FVector(Length * 0.9f), DrawConeLimitSides, DrawConeLimitSides, LimitMaterialX->GetRenderProxy(false), Layer);
+					DrawSphere(PDI, Con1Pos, FRotator::ZeroRotator, FVector(Length * 0.9f), DrawConeLimitSides, DrawConeLimitSides, LimitMaterialX->GetRenderProxy(false), Layer);
 				}
 			}
 			else
diff --git a/Engine/Source/Runtime/Engine/Private/PostProcessVolume.cpp b/Engine/Source/Runtime/Engine/Private/PostProcessVolume.cpp
--- a/Engine/Source/Runtime/Engine/Private/PostProcessVolume.cpp
+++ b/Engine/Source/Runtime/Engine/Private/PostProcessVolume.cpp
@@ -141,18 +141,14 @@ bool APostProcessVolume::CanEditChange(const UProperty* InProperty) const
 				return (Settings.BloomMethod == EBloomMethod::BM_SOG);
 			}
 
-			if (PropertyName == GET_MEMBER_NAME_STRING_CHECKED(FPostProcessSettings, BloomConvolutionTexture))
-			{
-				return (Settings.BloomMethod == EBloomMethod::BM_FFT);
-			}
-
 			// Parameters that are only of use with the bloom texture based fft
-			if (PropertyName == GET_MEMBER_NAME_STRING_CHECKED(FPostProcessSettings, BloomConvolutionSize) ||
+			if (PropertyName == GET_MEMBER_NAME_STRING_CHECKED(FPostProcessSettings, BloomConvolutionTexture) ||
+				PropertyName == GET_MEMBER_NAME_STRING_CHECKED(FPostProcessSettings, BloomConvolutionSize) ||
 				PropertyName == GET_MEMBER_NAME_STRING_CHECKED(FPostProcessSettings, BloomConvolutionCenterUV) ||
 				PropertyName == GET_MEMBER_NAME_STRING_CHECKED(FPostProcessSettings, BloomConvolutionPreFilter) ||
 				PropertyName == GET_MEMBER_NAME_STRING_CHECKED(FPostProcessSettings, BloomConvolutionBufferScale))
 			{
-				return (Settings.BloomMethod == EBloomMethod::BM_FFT && Settings.bOverride_BloomConvolutionTexture);
+				return (Settings.BloomMethod == EBloomMethod::BM_FFT);
 			}
 
 		}
diff --git a/Engine/Source/Runtime/Engine/Private/PrimitiveDrawingUtils.cpp b/Engine/Source/Runtime/Engine/Private/PrimitiveDrawingUtils.cpp
--- a/Engine/Source/Runtime/Engine/Private/PrimitiveDrawingUtils.cpp
+++ b/Engine/Source/Runtime/Engine/Private/PrimitiveDrawingUtils.cpp
@@ -312,7 +312,7 @@ extern ENGINE_API void GetSphereMesh(const FVector& Center, const FVector& Radii
 	GetHalfSphereMesh(Center, Radii, NumSides, NumRings, 0, PI, MaterialRenderProxy, DepthPriority, bDisableBackfaceCulling, ViewIndex, Collector, bUseSelectionOutline, HitProxy);
 }
 
-void DrawSphere(FPrimitiveDrawInterface* PDI,const FVector& Center,const FVector& Radii,int32 NumSides,int32 NumRings,const FMaterialRenderProxy* MaterialRenderProxy,uint8 DepthPriority,bool bDisableBackfaceCulling)
+void DrawSphere(FPrimitiveDrawInterface* PDI,const FVector& Center,const FRotator& Orientation,const FVector& Radii,int32 NumSides,int32 NumRings,const FMaterialRenderProxy* MaterialRenderProxy,uint8 DepthPriority,bool bDisableBackfaceCulling)
 {
 	// Use a mesh builder to draw the sphere.
 	FDynamicMeshBuilder MeshBuilder;
@@ -392,7 +392,7 @@ void DrawSphere(FPrimitiveDrawInterface* PDI,const FVector& Center,const FVector
 		FMemory::Free(Verts);
 		FMemory::Free(ArcVerts);
 	}
-	MeshBuilder.Draw(PDI, FScaleMatrix( Radii ) * FTranslationMatrix( Center ), MaterialRenderProxy, DepthPriority,bDisableBackfaceCulling);
+	MeshBuilder.Draw(PDI, FScaleMatrix( Radii ) * FRotationMatrix(Orientation) * FTranslationMatrix( Center ), MaterialRenderProxy, DepthPriority,bDisableBackfaceCulling);
 }
 
 FVector CalcConeVert(float Angle1, float Angle2, float AzimuthAngle)
diff --git a/Engine/Source/Runtime/Engine/Private/SceneManagement.cpp b/Engine/Source/Runtime/Engine/Private/SceneManagement.cpp
--- a/Engine/Source/Runtime/Engine/Private/SceneManagement.cpp
+++ b/Engine/Source/Runtime/Engine/Private/SceneManagement.cpp
@@ -546,6 +546,8 @@ FViewUniformShaderParameters::FViewUniformShaderParameters()
 	PerlinNoise3DTexture = BlackVolume;
 	PerlinNoise3DTextureSampler = TStaticSamplerState<SF_Bilinear, AM_Wrap, AM_Wrap, AM_Wrap>::GetRHI();
 
+	SobolSamplingTexture = GWhiteTexture->TextureRHI;
+
 	GlobalDistanceFieldTexture0_UB = BlackVolume;
 	GlobalDistanceFieldSampler0_UB = TStaticSamplerState<SF_Bilinear, AM_Wrap, AM_Wrap, AM_Wrap>::GetRHI();
 	GlobalDistanceFieldTexture1_UB = BlackVolume;
diff --git a/Engine/Source/Runtime/Engine/Private/SceneUtils.cpp b/Engine/Source/Runtime/Engine/Private/SceneUtils.cpp
--- a/Engine/Source/Runtime/Engine/Private/SceneUtils.cpp
+++ b/Engine/Source/Runtime/Engine/Private/SceneUtils.cpp
@@ -63,7 +63,7 @@ template struct TDrawEvent<FRHIAsyncComputeCommandList>;
 
 void FDrawEventRHIExecute::Start(IRHIComputeContext& InRHICommandContext, FColor Color, const TCHAR* Fmt, ...)
 {
-	check(IsInParallelRenderingThread() || IsInRHIThread() || (!GRHIThread && IsInRenderingThread()));
+	check(IsInParallelRenderingThread() || IsInRHIThread() || (!IsRunningRHIInSeparateThread() && IsInRenderingThread()));
 	{
 		va_list ptr;
 		va_start(ptr, Fmt);
diff --git a/Engine/Source/Runtime/Engine/Private/SceneView.cpp b/Engine/Source/Runtime/Engine/Private/SceneView.cpp
--- a/Engine/Source/Runtime/Engine/Private/SceneView.cpp
+++ b/Engine/Source/Runtime/Engine/Private/SceneView.cpp
@@ -311,6 +311,16 @@ static TAutoConsoleVariable<float> CVarRoughnessMax(
 	ECVF_Cheat | ECVF_RenderThreadSafe
 	);
 
+static TAutoConsoleVariable<int32> CVarAllowTranslucencyAfterDOF(
+	TEXT("r.SeparateTranslucency"),
+	1,
+	TEXT("Allows to disable the separate translucency feature (all translucency is rendered in separate RT and composited\n")
+	TEXT("after DOF, if not specified otherwise in the material).\n")
+	TEXT(" 0: off (translucency is affected by depth of field)\n")
+	TEXT(" 1: on costs GPU performance and memory but keeps translucency unaffected by Depth of Field. (default)"),
+	ECVF_RenderThreadSafe);
+
+
 /** Global vertex color view mode setting when SHOW_VertexColors show flag is set */
 EVertexColorViewMode::Type GVertexColorViewMode = EVertexColorViewMode::Color;
 
@@ -2460,6 +2470,17 @@ const FSceneView& FSceneViewFamily::GetStereoEyeView(const EStereoscopicPass Eye
 	}
 }
 
+bool FSceneViewFamily::AllowTranslucencyAfterDOF() const
+{
+	return CVarAllowTranslucencyAfterDOF.GetValueOnRenderThread() != 0
+		&& GetFeatureLevel() >= ERHIFeatureLevel::SM4
+		// && EngineShowFlags.PostProcessing
+		&& !UseDebugViewPS()
+		&& EngineShowFlags.SeparateTranslucency;
+		// If not, translucency after DOF will be rendered in standard translucency.
+}
+
+
 FSceneViewFamilyContext::~FSceneViewFamilyContext()
 {
 	// Cleanup the views allocated for this view family.
diff --git a/Engine/Source/Runtime/Engine/Private/ShaderCompiler/ShaderCompiler.cpp b/Engine/Source/Runtime/Engine/Private/ShaderCompiler/ShaderCompiler.cpp
--- a/Engine/Source/Runtime/Engine/Private/ShaderCompiler/ShaderCompiler.cpp
+++ b/Engine/Source/Runtime/Engine/Private/ShaderCompiler/ShaderCompiler.cpp
@@ -109,8 +109,9 @@ static int32 GDumpShaderDebugInfo = 0;
 static FAutoConsoleVariableRef CVarDumpShaderDebugInfo(
 	TEXT("r.DumpShaderDebugInfo"),
 	GDumpShaderDebugInfo,
-	TEXT("When set to 1, will cause any shaders that are then compiled to dump debug info to GameName/Saved/ShaderDebugInfo\n")
+	TEXT("When set to 1, will cause any material shaders that are then compiled to dump debug info to GameName/Saved/ShaderDebugInfo\n")
 	TEXT("The debug info is platform dependent, but usually includes a preprocessed version of the shader source.\n")
+	TEXT("Global shaders automatically dump debug info if r.ShaderDevelopmentMode is enabled, this cvar is not necessary.\n")
 	TEXT("On iOS, if the PowerVR graphics SDK is installed to the default path, the PowerVR shader compiler will be called and errors will be reported during the cook.")
 	);
 
@@ -2523,7 +2524,7 @@ void GlobalBeginCompileShader(
 	FShaderCompilerInput& Input = NewJob->Input;
 	Input.Target = Target;
 	Input.ShaderFormat = LegacyShaderPlatformToShaderFormat(EShaderPlatform(Target.Platform));
-	Input.SourceFilename = SourceFilename;
+	Input.SourceFilename = *FindShaderRelativePath(SourceFilename);
 	Input.EntryPointName = FunctionName;
 	Input.bCompilingForShaderPipeline = false;
 	Input.bIncludeUsedOutputs = false;
@@ -2612,7 +2613,10 @@ void GlobalBeginCompileShader(
 		}
 	}
 	
-	if (GDumpShaderDebugInfo != 0)
+	static const auto CVarShaderDevelopmentMode = IConsoleManager::Get().FindConsoleVariable(TEXT("r.ShaderDevelopmentMode"));
+
+	// Setup the debug info path if requested, or if this is a global shader and shader development mode is enabled
+	if (GDumpShaderDebugInfo != 0 || (ShaderType->GetGlobalShaderType() != NULL && CVarShaderDevelopmentMode->GetInt() != 0))
 	{
 		Input.DumpDebugInfoPath = Input.DumpDebugInfoRootPath / Input.DebugGroupName;
 		
@@ -3847,7 +3851,7 @@ void VerifyGlobalShaders(EShaderPlatform Platform, bool bLoadedFromCacheFile)
 	check(!FPlatformProperties::IsServerOnly());
 	check(GGlobalShaderMap[Platform]);
 
-	UE_LOG(LogShaders, Log, TEXT("Verifying Global Shaders for %s"), *LegacyShaderPlatformToShaderFormat(Platform).ToString());
+	UE_LOG(LogMaterial, Log, TEXT("Verifying Global Shaders for %s"), *LegacyShaderPlatformToShaderFormat(Platform).ToString());
 
 	// Ensure that the global shader map contains all global shader types.
 	TShaderMap<FGlobalShaderType>* GlobalShaderMap = GetGlobalShaderMap(Platform);
@@ -4135,7 +4139,7 @@ void CompileGlobalShaderMap(EShaderPlatform Platform, bool bRefreshShaderMap)
 	{
 		if (!GGlobalShaderMap[Platform])
 		{
-			GGlobalShaderMap[Platform] = new TShaderMap<FGlobalShaderType>();
+			GGlobalShaderMap[Platform] = new TShaderMap<FGlobalShaderType>(Platform);
 		}
 		return;
 	}
@@ -4163,7 +4167,7 @@ void CompileGlobalShaderMap(EShaderPlatform Platform, bool bRefreshShaderMap)
 		SlowTask.EnterProgressFrame(20);
 		VerifyShaderSourceFiles();
 
-		GGlobalShaderMap[Platform] = new TShaderMap<FGlobalShaderType>();
+		GGlobalShaderMap[Platform] = new TShaderMap<FGlobalShaderType>(Platform);
 
 		bool bLoadedFromCacheFile = false;
 
diff --git a/Engine/Source/Runtime/Engine/Private/ShadowMap.cpp b/Engine/Source/Runtime/Engine/Private/ShadowMap.cpp
--- a/Engine/Source/Runtime/Engine/Private/ShadowMap.cpp
+++ b/Engine/Source/Runtime/Engine/Private/ShadowMap.cpp
@@ -14,6 +14,9 @@
 #include "GameFramework/WorldSettings.h"
 
 #if WITH_EDITOR
+
+	#include "TextureCompressorModule.h"
+
 	// NOTE: We're only counting the top-level mip-map for the following variables.
 	/** Total number of texels allocated for all shadowmap textures. */
 	ENGINE_API uint64 GNumShadowmapTotalTexels = 0;
@@ -178,7 +181,7 @@ extern bool GGroupComponentLightmaps;
 
 struct FShadowMapPendingTexture : FTextureLayout
 {
-	TArray<FShadowMapAllocation*> Allocations;
+	TArray<TUniquePtr<FShadowMapAllocation>> Allocations;
 
 	UObject* Outer;
 
@@ -218,7 +221,7 @@ struct FShadowMapPendingTexture : FTextureLayout
 	 *
 	 * @param	InWorld	World in which the textures exist
 	 */
-	void StartEncoding(ULevel* LightingScenario);
+	void StartEncoding(ULevel* LightingScenario, ITextureCompressorModule* Compressor);
 
 	/**
 	 * Create UObjects required in the encoding step, this is so we can multithread teh encode step
@@ -330,7 +333,7 @@ void FShadowMapPendingTexture::CreateUObjects()
 	bCreatedUObjects = true;
 }
 
-void FShadowMapPendingTexture::StartEncoding(ULevel* LightingScenario)
+void FShadowMapPendingTexture::StartEncoding(ULevel* LightingScenario, ITextureCompressorModule* Compressor)
 {
 	// Create the shadow-map texture.
 	CreateUObjects();
@@ -381,7 +384,7 @@ void FShadowMapPendingTexture::StartEncoding(ULevel* LightingScenario)
 	}
 
 	// Update the texture resource.
-	Texture->CachePlatformData(true, true);
+	Texture->CachePlatformData(true, true, Compressor);
 
 	bFinishedEncoding = true;
 }
@@ -754,6 +757,8 @@ void FShadowMap2D::EncodeTextures(UWorld* InWorld, ULevel* LightingScenario, boo
 		GWarn->BeginSlowTask( NSLOCTEXT("ShadowMap2D", "BeginEncodingShadowMapsTask", "Encoding shadow-maps"), false );
 		const int32 PackedLightAndShadowMapTextureSize = InWorld->GetWorldSettings()->PackedLightAndShadowMapTextureSize;
 
+		ITextureCompressorModule* TextureCompressorModule = &FModuleManager::LoadModuleChecked<ITextureCompressorModule>(TEXTURE_COMPRESSOR_MODULENAME);
+
 		// Reset the pending shadow-map size.
 		PendingShadowMapSize = 0;
 
@@ -819,7 +824,7 @@ void FShadowMap2D::EncodeTextures(UWorld* InWorld, ULevel* LightingScenario, boo
 			// Give the texture ownership of the allocations
 			for (auto& Allocation : PendingGroup.Allocations)
 			{
-				Texture->Allocations.Add(Allocation.Release());
+				Texture->Allocations.Add(MoveTemp(Allocation));
 			}
 		}
 		PendingShadowMaps.Empty();
@@ -833,7 +838,7 @@ void FShadowMap2D::EncodeTextures(UWorld* InWorld, ULevel* LightingScenario, boo
 			for (auto& PendingTexture : PendingTextures)
 			{
 				PendingTexture.CreateUObjects();
-				auto AsyncEncodeTask = new (AsyncEncodeTasks)FAsyncEncode<FShadowMapPendingTexture>(&PendingTexture, LightingScenario, Counter);
+				auto AsyncEncodeTask = new (AsyncEncodeTasks)FAsyncEncode<FShadowMapPendingTexture>(&PendingTexture, LightingScenario, Counter, TextureCompressorModule);
 				GThreadPool->AddQueuedWork(AsyncEncodeTask);
 			}
 
@@ -849,7 +854,7 @@ void FShadowMap2D::EncodeTextures(UWorld* InWorld, ULevel* LightingScenario, boo
 			for (int32 TextureIndex = 0; TextureIndex < PendingTextures.Num(); TextureIndex++)
 			{
 				FShadowMapPendingTexture& PendingTexture = PendingTextures[TextureIndex];
-				PendingTexture.StartEncoding(LightingScenario);
+				PendingTexture.StartEncoding(LightingScenario, TextureCompressorModule);
 			}
 		}
 
@@ -1174,7 +1179,7 @@ int32 FShadowMap2D::EncodeSingleTexture(ULevel* LightingScenario, FShadowMapPend
 		}
 	}
 
-	for (FShadowMapAllocation* Allocation : PendingTexture.Allocations)
+	for (auto& Allocation : PendingTexture.Allocations)
 	{
 		Allocation->PostEncode();
 	}
diff --git a/Engine/Source/Runtime/Engine/Private/SkeletalMesh.cpp b/Engine/Source/Runtime/Engine/Private/SkeletalMesh.cpp
--- a/Engine/Source/Runtime/Engine/Private/SkeletalMesh.cpp
+++ b/Engine/Source/Runtime/Engine/Private/SkeletalMesh.cpp
@@ -1198,34 +1198,35 @@ FArchive& operator<<(FArchive& Ar,FSkelMeshSection& S)
 
 void FMorphTargetVertexInfoBuffers::InitRHI()
 {
-	if (PerVertexInfoList.Num() > 0)
+	check(NumTotalWorkItems > 0);
+
 	{
 		FRHIResourceCreateInfo CreateInfo;
-		void* PerVertexInfoListVBData = nullptr;
-		PerVertexInfoVB = RHICreateAndLockVertexBuffer(PerVertexInfoList.GetAllocatedSize(), BUF_Static | BUF_ShaderResource, CreateInfo, PerVertexInfoListVBData);
-		FMemory::Memcpy(PerVertexInfoListVBData, PerVertexInfoList.GetData(), PerVertexInfoList.GetAllocatedSize());
-		RHIUnlockVertexBuffer(PerVertexInfoVB);
-		PerVertexInfoSRV = RHICreateShaderResourceView(PerVertexInfoVB, sizeof(uint32), PF_R32_UINT);
-
-		void* FlattenedDeltasVBData = nullptr;
-		FlattenedDeltasVB = RHICreateAndLockVertexBuffer(FlattenedDeltaList.GetAllocatedSize(), BUF_Static | BUF_ShaderResource, CreateInfo, FlattenedDeltasVBData);
-		FMemory::Memcpy(FlattenedDeltasVBData, FlattenedDeltaList.GetData(), FlattenedDeltaList.GetAllocatedSize());
-		RHIUnlockVertexBuffer(FlattenedDeltasVB);
-		FlattenedDeltasSRV = RHICreateShaderResourceView(FlattenedDeltasVB, sizeof(uint32), PF_R32_UINT);
-
-		NumInfluencedVerticesByMorphs = (uint32)PerVertexInfoList.Num();
-
-		PerVertexInfoList.Empty();
-		FlattenedDeltaList.Empty();
+		void* VertexIndicesVBData = nullptr;
+		VertexIndicesVB = RHICreateAndLockVertexBuffer(VertexIndices.GetAllocatedSize(), BUF_Static | BUF_ShaderResource, CreateInfo, VertexIndicesVBData);
+		FMemory::Memcpy(VertexIndicesVBData, VertexIndices.GetData(), VertexIndices.GetAllocatedSize());
+		RHIUnlockVertexBuffer(VertexIndicesVB);
+		VertexIndicesSRV = RHICreateShaderResourceView(VertexIndicesVB, 4, PF_R32_UINT);
+	}
+	{
+		FRHIResourceCreateInfo CreateInfo;
+		void* MorphDeltasVBData = nullptr;
+		MorphDeltasVB = RHICreateAndLockVertexBuffer(MorphDeltas.GetAllocatedSize(), BUF_Static | BUF_ShaderResource, CreateInfo, MorphDeltasVBData);
+		FMemory::Memcpy(MorphDeltasVBData, MorphDeltas.GetData(), MorphDeltas.GetAllocatedSize());
+		RHIUnlockVertexBuffer(MorphDeltasVB);
+		MorphDeltasSRV = RHICreateShaderResourceView(MorphDeltasVB, 2, PF_R16F);
 	}
+
+	VertexIndices.Empty();
+	MorphDeltas.Empty();
 }
 
 void FMorphTargetVertexInfoBuffers::ReleaseRHI()
 {
-	PerVertexInfoVB.SafeRelease();
-	PerVertexInfoSRV.SafeRelease();
-	FlattenedDeltasVB.SafeRelease();
-	FlattenedDeltasSRV.SafeRelease();
+	VertexIndicesVB.SafeRelease();
+	VertexIndicesSRV.SafeRelease();
+	MorphDeltasVB.SafeRelease();
+	MorphDeltasSRV.SafeRelease();
 }
 
 /*-----------------------------------------------------------------------------
@@ -1543,44 +1544,65 @@ void FStaticLODModel::InitResources(bool bNeedsVertexColors, int32 LODIndex, TAr
 
 	if (RHISupportsComputeShaders(GMaxRHIShaderPlatform) && InMorphTargets.Num() > 0)
 	{
-		// Populate the arrays to be filled in later in the render thread
-
-		// Auxiliary mapping for affected vertex indices by morph to its weights
-		TMap<uint32, TArray<FMorphTargetVertexInfoBuffers::FFlattenedDelta>> AuxDeltaList;
+		MorphTargetVertexInfoBuffers.VertexIndices.Empty();
+		MorphTargetVertexInfoBuffers.MorphDeltas.Empty();
+		MorphTargetVertexInfoBuffers.WorkItemsPerMorph.Empty();
+		MorphTargetVertexInfoBuffers.StartOffsetPerMorph.Empty();
+		MorphTargetVertexInfoBuffers.MaximumValuePerMorph.Empty();
+		MorphTargetVertexInfoBuffers.MinimumValuePerMorph.Empty();
+		MorphTargetVertexInfoBuffers.NumTotalWorkItems = 0;
 
-		int32 TotalNumDeltas = 0;
+		// Populate the arrays to be filled in later in the render thread
 		for (int32 AnimIdx = 0; AnimIdx < InMorphTargets.Num(); ++AnimIdx)
 		{
-			UMorphTarget* MorphTarget = InMorphTargets[AnimIdx];
+			uint32 StartOffset = MorphTargetVertexInfoBuffers.NumTotalWorkItems;
+			MorphTargetVertexInfoBuffers.StartOffsetPerMorph.Add(StartOffset);
 
+			float MaximumValues[4] = { -FLT_MAX, -FLT_MAX, -FLT_MAX, -FLT_MAX };
+			float MinimumValues[4] = { +FLT_MAX, +FLT_MAX, +FLT_MAX, +FLT_MAX };
+			UMorphTarget* MorphTarget = InMorphTargets[AnimIdx];
 			int32 NumSrcDeltas = 0;
-			FMorphTargetDelta* SrcDelta = MorphTarget->GetMorphTargetDelta(LODIndex, NumSrcDeltas);
-			for (int32 SrcDeltaIndex = 0; SrcDeltaIndex < NumSrcDeltas; ++SrcDeltaIndex, ++SrcDelta)
+			FMorphTargetDelta* MorphDeltas = MorphTarget->GetMorphTargetDelta(LODIndex, NumSrcDeltas);
+			for (int32 DeltaIndex = 0; DeltaIndex < NumSrcDeltas; DeltaIndex++)
 			{
-				TArray<FMorphTargetVertexInfoBuffers::FFlattenedDelta>& FlattenedDeltas = AuxDeltaList.FindOrAdd(SrcDelta->SourceIdx);
-				FMorphTargetVertexInfoBuffers::FFlattenedDelta* NewDelta = new(FlattenedDeltas) FMorphTargetVertexInfoBuffers::FFlattenedDelta;
-				NewDelta->PosDelta = SrcDelta->PositionDelta;
-				NewDelta->TangentDelta = SrcDelta->TangentZDelta;
-				NewDelta->WeightIndex = AnimIdx;
-				++TotalNumDeltas;
+				const auto& MorphDelta = MorphDeltas[DeltaIndex];
+				if (!(MorphDelta.PositionDelta.IsNearlyZero(0.0000001f) && MorphDelta.TangentZDelta.IsNearlyZero(0.0000001f)))
+				{				
+					MaximumValues[0] = FMath::Max(MaximumValues[0], MorphDelta.PositionDelta.X);
+					MaximumValues[1] = FMath::Max(MaximumValues[1], MorphDelta.PositionDelta.Y);
+					MaximumValues[2] = FMath::Max(MaximumValues[2], MorphDelta.PositionDelta.Z);
+					MaximumValues[3] = FMath::Max(MaximumValues[3], FMath::Max(MorphDelta.TangentZDelta.X, FMath::Max(MorphDelta.TangentZDelta.Y, MorphDelta.TangentZDelta.Z)));
+
+					MinimumValues[0] = FMath::Min(MinimumValues[0], MorphDelta.PositionDelta.X);
+					MinimumValues[1] = FMath::Min(MinimumValues[1], MorphDelta.PositionDelta.Y);
+					MinimumValues[2] = FMath::Min(MinimumValues[2], MorphDelta.PositionDelta.Z);
+					MinimumValues[3] = FMath::Min(MinimumValues[3], FMath::Min(MorphDelta.TangentZDelta.X, FMath::Min(MorphDelta.TangentZDelta.Y, MorphDelta.TangentZDelta.Z)));
+
+					MorphTargetVertexInfoBuffers.VertexIndices.Add(MorphDelta.SourceIdx);
+					MorphTargetVertexInfoBuffers.MorphDeltas.Emplace(MorphDelta.PositionDelta, MorphDelta.TangentZDelta);
+					MorphTargetVertexInfoBuffers.NumTotalWorkItems++;
+				}
 			}
-		}
 
-		MorphTargetVertexInfoBuffers.FlattenedDeltaList.Empty(TotalNumDeltas);
-		MorphTargetVertexInfoBuffers.PerVertexInfoList.AddUninitialized(AuxDeltaList.Num());
-		int32 StartDelta = 0;
-		FMorphTargetVertexInfoBuffers::FPerVertexInfo* NewPerVertexInfo = MorphTargetVertexInfoBuffers.PerVertexInfoList.GetData();
-		for (auto& Pair : AuxDeltaList)
-		{
-			NewPerVertexInfo->DestVertexIndex = Pair.Key;
-			NewPerVertexInfo->StartDelta = StartDelta;
-			NewPerVertexInfo->NumDeltas = Pair.Value.Num();
-			MorphTargetVertexInfoBuffers.FlattenedDeltaList.Append(Pair.Value);
-			StartDelta += Pair.Value.Num();
-			++NewPerVertexInfo;
+			uint32 MorphTargetSize =  MorphTargetVertexInfoBuffers.NumTotalWorkItems - StartOffset;
+			if (MorphTargetSize > 0)
+			{
+				ensureMsgf(MaximumValues[0] < +32752.0f && MaximumValues[1] < +32752.0f && MaximumValues[2] < +32752.0f && MaximumValues[3] < +32752.0f, TEXT("Huge MorphTarget Delta found in %s at index %i, might break down because we use half float storage"), *MorphTarget->GetName(), AnimIdx);
+				ensureMsgf(MinimumValues[0] > -32752.0f && MinimumValues[1] > -32752.0f && MinimumValues[2] > -32752.0f && MaximumValues[3] > -32752.0f, TEXT("Huge MorphTarget Delta found in %s at index %i, might break down because we use half float storage"), *MorphTarget->GetName(), AnimIdx);
+			}
+			else
+			{
+				UE_LOG(LogSkeletalMesh, Warning, TEXT("Empty MorphTarget Animation found in %s at index %i"), *MorphTarget->GetName(), AnimIdx);
+			}
+			MorphTargetVertexInfoBuffers.WorkItemsPerMorph.Add(MorphTargetSize);
+			MorphTargetVertexInfoBuffers.MaximumValuePerMorph.Add(FVector4(MaximumValues[0], MaximumValues[1], MaximumValues[2], MaximumValues[3]));
+			MorphTargetVertexInfoBuffers.MinimumValuePerMorph.Add(FVector4(MinimumValues[0], MinimumValues[1], MinimumValues[2], MinimumValues[3]));
 		}
 
-		if (AuxDeltaList.Num() > 0)
+		check(MorphTargetVertexInfoBuffers.WorkItemsPerMorph.Num() == MorphTargetVertexInfoBuffers.StartOffsetPerMorph.Num());
+		check(MorphTargetVertexInfoBuffers.WorkItemsPerMorph.Num() == MorphTargetVertexInfoBuffers.MaximumValuePerMorph.Num());
+		check(MorphTargetVertexInfoBuffers.WorkItemsPerMorph.Num() == MorphTargetVertexInfoBuffers.MinimumValuePerMorph.Num());
+		if (MorphTargetVertexInfoBuffers.NumTotalWorkItems > 0)
 		{
 			BeginInitResource(&MorphTargetVertexInfoBuffers);
 		}
@@ -2846,6 +2868,8 @@ void USkeletalMesh::PostEditChangeProperty(FPropertyChangedEvent& PropertyChange
 		FMultiComponentReregisterContext ReregisterContext(ComponentsToReregister);
 	}
 
+	UpdateUVChannelData(true);
+
 	Super::PostEditChangeProperty(PropertyChangedEvent);
 }
 
@@ -3305,6 +3329,11 @@ void USkeletalMesh::PostLoad()
 			}
 		}
 	}
+
+	if (GetLinkerCustomVersion(FRenderingObjectVersion::GUID) < FRenderingObjectVersion::FixedMeshUVDensity)
+	{
+		UpdateUVChannelData(true);
+	}
 #endif // WITH_EDITOR
 
 	// initialize rendering resources
