Copying //UE4/Dev-Rendering to //UE4/Dev-Main (Source: //UE4/Dev-Rendering @ 3461187)

diff --git a/Engine/Source/Runtime/Engine/Private/SkeletalRenderGPUSkin.cpp b/Engine/Source/Runtime/Engine/Private/SkeletalRenderGPUSkin.cpp
--- a/Engine/Source/Runtime/Engine/Private/SkeletalRenderGPUSkin.cpp
+++ b/Engine/Source/Runtime/Engine/Private/SkeletalRenderGPUSkin.cpp
@@ -10,7 +10,7 @@
 #include "SkeletalRender.h"
 #include "GPUSkinCache.h"
 #include "Animation/MorphTarget.h"
-
+#include "ClearQuad.h"
 #include "ShaderParameterUtils.h"
 
 
@@ -18,7 +18,6 @@ DEFINE_LOG_CATEGORY_STATIC(LogSkeletalGPUSkinMesh, Warning, All);
 
 // 0/1
 #define UPDATE_PER_BONE_DATA_ONLY_FOR_OBJECT_BEEN_VISIBLE 1
-#define MORPH_TARGETS_USE_BYTE_BUFFER_UAV	0
 
 DECLARE_CYCLE_STAT(TEXT("Morph Vertex Buffer Update"),STAT_MorphVertexBuffer_Update,STATGROUP_MorphTarget);
 DECLARE_CYCLE_STAT(TEXT("Morph Vertex Buffer Init"),STAT_MorphVertexBuffer_Init,STATGROUP_MorphTarget);
@@ -58,40 +57,41 @@ void FMorphVertexBuffer::InitDynamicRHI()
 	// Create the buffer rendering resource
 	uint32 Size = LodModel.NumVertices * sizeof(FMorphGPUSkinVertex);
 	FRHIResourceCreateInfo CreateInfo;
-	void* BufferData = nullptr;
 
 	bool bSupportsComputeShaders = RHISupportsComputeShaders(GMaxRHIShaderPlatform);
 	bUsesComputeShader = GUseGPUMorphTargets != 0 && bSupportsComputeShaders;
 
 #if PLATFORM_PS4
 	// PS4 requires non-static buffers in order to be updated on the GPU while the CPU is writing into it
-	EBufferUsageFlags Flags = BUF_Dynamic;
+	EBufferUsageFlags Flags = bUsesComputeShader ? (EBufferUsageFlags)(BUF_Dynamic | BUF_UnorderedAccess) : BUF_Dynamic;
 #else
 	EBufferUsageFlags Flags = bUsesComputeShader ? (EBufferUsageFlags)(BUF_Static | BUF_UnorderedAccess) : BUF_Dynamic;
-#if MORPH_TARGETS_USE_BYTE_BUFFER_UAV
-	Flags = (EBufferUsageFlags)(Flags | (bUsesComputeShader ? BUF_ByteAddressBuffer : 0));
-#endif
 #endif
 
 	// BUF_ShaderResource is needed for Morph support of the SkinCache
 	Flags = (EBufferUsageFlags)(Flags | BUF_ShaderResource);
 
-	VertexBufferRHI = RHICreateAndLockVertexBuffer(Size, Flags, CreateInfo, BufferData);
+	VertexBufferRHI = RHICreateVertexBuffer(Size, Flags, CreateInfo);
 	bool bUsesSkinCache = bSupportsComputeShaders && IsGPUSkinCacheAvailable() && GEnableGPUSkinCache;
 	if (bUsesSkinCache)
 	{
 		SRVValue = RHICreateShaderResourceView(VertexBufferRHI, 4, PF_R32_FLOAT);
 	}
 
-	// Lock the buffer.
-	FMorphGPUSkinVertex* Buffer = (FMorphGPUSkinVertex*)BufferData;
-	FMemory::Memzero(&Buffer[0], sizeof(FMorphGPUSkinVertex)*LodModel.NumVertices);
-	// Unlock the buffer.
-	RHIUnlockVertexBuffer(VertexBufferRHI);
-	
-	if (bUsesComputeShader)
+	if (!bUsesComputeShader)
+	{
+		// Lock the buffer.
+		void* BufferData = RHILockVertexBuffer(VertexBufferRHI, 0, sizeof(FMorphGPUSkinVertex)*LodModel.NumVertices, RLM_WriteOnly);
+		FMorphGPUSkinVertex* Buffer = (FMorphGPUSkinVertex*)BufferData;
+		FMemory::Memzero(&Buffer[0], sizeof(FMorphGPUSkinVertex)*LodModel.NumVertices);
+		// Unlock the buffer.
+		RHIUnlockVertexBuffer(VertexBufferRHI);
+		bNeedsInitialClear = false;
+	}
+	else
 	{
 		UAVValue = RHICreateUnorderedAccessView(VertexBufferRHI, PF_R32_UINT);
+		bNeedsInitialClear = true;
 	}
 
 	// hasn't been updated yet
@@ -195,21 +195,6 @@ void FSkeletalMeshObjectGPUSkin::InitMorphResources(bool bInUsePerBoneMotionBlur
 		SkelLOD.InitMorphResources(MeshLODInfo,bInUsePerBoneMotionBlur, FeatureLevel);
 	}
 	bMorphResourcesInitialized = true;
-
-	if (MorphTargetWeights.Num() > 0)
-	{
-		int32 NumWeights = MorphTargetWeights.Num();
-		FVertexBufferAndSRV* MorphWeightsVertexBufferPtr = &MorphWeightsVertexBuffer;
-		ENQUEUE_RENDER_COMMAND(InitMorphWeightsVertexBuffer)(
-			[NumWeights, MorphWeightsVertexBufferPtr](FRHICommandList& RHICmdList)
-			{
-				const uint32 Size = (uint32)NumWeights * sizeof(float);
-				FRHIResourceCreateInfo CreateInfo;
-				MorphWeightsVertexBufferPtr->VertexBufferRHI = RHICreateVertexBuffer(Size, BUF_Dynamic | BUF_ShaderResource, CreateInfo);
-				MorphWeightsVertexBufferPtr->VertexBufferSRV = RHICreateShaderResourceView(MorphWeightsVertexBufferPtr->VertexBufferRHI, sizeof(float), PF_R32_UINT);
-			}
-		);
-	}
 }
 
 void FSkeletalMeshObjectGPUSkin::ReleaseMorphResources()
@@ -221,18 +206,6 @@ void FSkeletalMeshObjectGPUSkin::ReleaseMorphResources()
 		SkelLOD.ReleaseMorphResources();
 	}
 
-	if (MorphWeightsVertexBuffer.VertexBufferRHI)
-	{
-		FVertexBufferAndSRV* MorphWeightsVertexBufferPtr = &MorphWeightsVertexBuffer;
-		ENQUEUE_RENDER_COMMAND(InitMorphWeightsVertexBuffer)(
-			[MorphWeightsVertexBufferPtr](FRHICommandList& RHICmdList)
-			{
-				MorphWeightsVertexBufferPtr->VertexBufferRHI.SafeRelease();
-				MorphWeightsVertexBufferPtr->VertexBufferSRV.SafeRelease();
-			}
-		);
-	}
-
 	bMorphResourcesInitialized = false;
 }
 
@@ -385,6 +358,15 @@ void FSkeletalMeshObjectGPUSkin::ProcessUpdatedDynamicData(FGPUSkinCache* GPUSki
 
 	bool bGPUSkinCacheEnabled = GEnableGPUSkinCache && (FeatureLevel >= ERHIFeatureLevel::SM5);
 
+	if (LOD.MorphVertexBuffer.bNeedsInitialClear && !(bMorph && bMorphNeedsUpdate))
+	{
+		if (IsValidRef(LOD.MorphVertexBuffer.GetUAV()))
+		{
+			ClearUAV(RHICmdList, LOD.MorphVertexBuffer.GetUAV(), LOD.MorphVertexBuffer.GetUAVSize(), 0);
+		}
+	}
+	LOD.MorphVertexBuffer.bNeedsInitialClear = false;
+
 	if(bMorph) 
 	{
 		bDataPresent = true;
@@ -393,17 +375,16 @@ void FSkeletalMeshObjectGPUSkin::ProcessUpdatedDynamicData(FGPUSkinCache* GPUSki
 		// only update if the morph data changed and there are weighted morph targets
 		if(bMorphNeedsUpdate)
 		{
-			if (GUseGPUMorphTargets)
+			if (GUseGPUMorphTargets && RHISupportsComputeShaders(GMaxRHIShaderPlatform))
 			{
-				SCOPE_CYCLE_COUNTER(STAT_MorphVertexBuffer_RhiLockAndCopy);
-				const uint32 MorphWeightsSize = DynamicData->MorphTargetWeights.GetAllocatedSize();
-				void* NewMorphWeights = RHICmdList.LockVertexBuffer(MorphWeightsVertexBuffer.VertexBufferRHI, 0, MorphWeightsSize, RLM_WriteOnly);
-				FMemory::Memcpy(NewMorphWeights, DynamicData->MorphTargetWeights.GetData(), MorphWeightsSize);
-				RHICmdList.UnlockVertexBuffer(MorphWeightsVertexBuffer.VertexBufferRHI);
+				// update the morph data for the lod (before SkinCache)
+				LOD.UpdateMorphVertexBufferGPU(RHICmdList, DynamicData->MorphTargetWeights, LODModel.MorphTargetVertexInfoBuffers);
 			}
-
-			// update the morph data for the lod (before SkinCache)
-			LOD.UpdateMorphVertexBuffer(RHICmdList, DynamicData->ActiveMorphTargets, DynamicData->MorphTargetWeights, MorphWeightsVertexBuffer.VertexBufferSRV, LODModel.MorphTargetVertexInfoBuffers.GetNumInfluencedVerticesByMorphs(), LODModel.MorphTargetVertexInfoBuffers.PerVertexInfoSRV, LODModel.MorphTargetVertexInfoBuffers.FlattenedDeltasSRV);
+			else
+			{
+				// update the morph data for the lod (before SkinCache)
+				LOD.UpdateMorphVertexBufferCPU(DynamicData->ActiveMorphTargets, DynamicData->MorphTargetWeights);
+			}		
 		}
 	}
 	else
@@ -500,208 +481,286 @@ void FSkeletalMeshObjectGPUSkin::ProcessUpdatedDynamicData(FGPUSkinCache* GPUSki
 	}
 }
 
-void FMorphTargetBaseShader::SetParameters(FRHICommandList& RHICmdList, int32 NumLODVertices, int32 NumMorphVertices, FShaderResourceViewRHIRef MorphWeightsSRV, FShaderResourceViewRHIRef PerVertexInfoListSRV, FShaderResourceViewRHIRef FlattenedDeltaListSRV, FUnorderedAccessViewRHIParamRef UAV)
+TArray<float> FSkeletalMeshObjectGPUSkin::FSkeletalMeshObjectLOD::MorphAccumulatedWeightArray;
+
+void FGPUMorphUpdateCS::SetParameters(FRHICommandList& RHICmdList, const FVector4& LocalScale, const FMorphTargetVertexInfoBuffers& MorphTargetVertexInfoBuffers, FMorphVertexBuffer& MorphVertexBuffer)
 {
 	FComputeShaderRHIRef CS = GetComputeShader();
 	RHICmdList.SetComputeShader(CS);
-	FVector4 MorphTargetCounts(NumLODVertices, NumMorphVertices, 0, 0);
-	SetShaderValue(RHICmdList, CS, MorphTargetCountsParameter, MorphTargetCounts);
 
-	SetSRVParameter(RHICmdList, CS, PerVertexInfoListParameter, PerVertexInfoListSRV);
+	SetUAVParameter(RHICmdList, CS, MorphVertexBufferParameter, MorphVertexBuffer.GetUAV());
 
-	SetSRVParameter(RHICmdList, CS, FlattenedDeltaListParameter, FlattenedDeltaListSRV);
+	SetShaderValue(RHICmdList, CS, PositionScaleParameter, LocalScale);
 
-	SetSRVParameter(RHICmdList, CS, AllWeightsPerMorphsParameter, MorphWeightsSRV);
+	SetSRVParameter(RHICmdList, CS, VertexIndicesParameter, MorphTargetVertexInfoBuffers.VertexIndicesSRV);
+	SetSRVParameter(RHICmdList, CS, MorphDeltasParameter, MorphTargetVertexInfoBuffers.MorphDeltasSRV);
+}
 
-	SetUAVParameter(RHICmdList, CS, MorphVertexBufferParameter, UAV);
+void FGPUMorphUpdateCS::SetOffsetAndSize(FRHICommandList& RHICmdList, uint32 Offset, uint32 Size, float Weight)
+{
+	FComputeShaderRHIRef CS = GetComputeShader();
+	uint32 OffsetAndSize[2] = { Offset, Offset + Size };
+	SetShaderValue(RHICmdList, CS, OffsetAndSizeParameter, OffsetAndSize);
+	SetShaderValue(RHICmdList, CS, MorphTargetWeightParameter, Weight);
 }
 
 
-class FGPUMorphUpdateCS : public FMorphTargetBaseShader
+void FGPUMorphUpdateCS::Dispatch(FRHICommandList& RHICmdList, uint32 Size)
 {
-public:
-	DECLARE_SHADER_TYPE(FGPUMorphUpdateCS, Global);
+	RHICmdList.DispatchComputeShader(1, (Size + 31) / 32, 1);	
+}
 
-	FGPUMorphUpdateCS() {}
+void FGPUMorphUpdateCS::EndAllDispatches(FRHICommandList& RHICmdList)
+{
+	FComputeShaderRHIRef CS = GetComputeShader();
+	SetUAVParameter(RHICmdList, CS, MorphVertexBufferParameter, nullptr);
+}
 
-	FGPUMorphUpdateCS(const ShaderMetaType::CompiledShaderInitializerType& Initializer)
-		: FMorphTargetBaseShader(Initializer)
-	{
-	}
+IMPLEMENT_SHADER_TYPE(, FGPUMorphUpdateCS, TEXT("MorphTargets"), TEXT("GPUMorphUpdateCS"), SF_Compute);
 
-	static void ModifyCompilationEnvironment(EShaderPlatform Platform, FShaderCompilerEnvironment& OutEnvironment)
-	{
-		OutEnvironment.SetDefine(TEXT("USE_BYTE_BUFFER"), MORPH_TARGETS_USE_BYTE_BUFFER_UAV);
-	}
+void FGPUMorphNormalizeCS::SetParameters(FRHICommandList& RHICmdList, uint32 NumVerticies, const FVector4& InvLocalScale, const float AccumulatedWeight, FMorphVertexBuffer& MorphVertexBuffer)
+{
+	FComputeShaderRHIRef CS = GetComputeShader();
+	RHICmdList.SetComputeShader(CS);
 
-	void Dispatch(FRHICommandList& RHICmdList, int32 NumLODVertices, int32 NumMorphVertices, FShaderResourceViewRHIRef MorphWeightsSRV, FShaderResourceViewRHIRef PerVertexInfoListSRV, FShaderResourceViewRHIRef FlattenedDeltaListSRV, FUnorderedAccessViewRHIParamRef UAV)
-	{
-		FComputeShaderRHIRef CS = GetComputeShader();
-		FMorphTargetBaseShader::SetParameters(RHICmdList, NumLODVertices, NumMorphVertices, MorphWeightsSRV, PerVertexInfoListSRV, FlattenedDeltaListSRV, UAV);
-		RHICmdList.DispatchComputeShader((NumMorphVertices + 63) / 64, 1, 1);
-		SetUAVParameter(RHICmdList, CS, MorphVertexBufferParameter, nullptr);
-	}
-};
+	SetUAVParameter(RHICmdList, CS, MorphVertexBufferParameter, MorphVertexBuffer.GetUAV());
 
-IMPLEMENT_SHADER_TYPE(, FGPUMorphUpdateCS, TEXT("MorphTargets"), TEXT("GPUMorphUpdateCS"), SF_Compute);
+	SetShaderValue(RHICmdList, CS, MorphTargetWeightParameter, AccumulatedWeight);
+	SetShaderValue(RHICmdList, CS, MorphWorkItemsParameter, NumVerticies);
+	SetShaderValue(RHICmdList, CS, PositionScaleParameter, InvLocalScale);
+}
+
+void FGPUMorphNormalizeCS::Dispatch(FRHICommandList& RHICmdList, uint32 NumVerticies, const FVector4& InvLocalScale, const float AccumulatedWeight, FMorphVertexBuffer& MorphVertexBuffer)
+{
+	FComputeShaderRHIRef CS = GetComputeShader();
+	FGPUMorphNormalizeCS::SetParameters(RHICmdList, NumVerticies, InvLocalScale, AccumulatedWeight, MorphVertexBuffer);
+	RHICmdList.DispatchComputeShader(1, (NumVerticies + 31) / 32, 1);
+	SetUAVParameter(RHICmdList, CS, MorphVertexBufferParameter, nullptr);
+}
 
+IMPLEMENT_SHADER_TYPE(, FGPUMorphNormalizeCS, TEXT("MorphTargets"), TEXT("GPUMorphNormalizeCS"), SF_Compute);
 
-TArray<FVector> FSkeletalMeshObjectGPUSkin::FSkeletalMeshObjectLOD::MorphDeltaTangentZAccumulationArray;
-TArray<float> FSkeletalMeshObjectGPUSkin::FSkeletalMeshObjectLOD::MorphAccumulatedWeightArray;
+static const float MorphTargetWeightCutoffThreshold = 0.00000001f;
 
-void FSkeletalMeshObjectGPUSkin::FSkeletalMeshObjectLOD::UpdateMorphVertexBuffer(FRHICommandListImmediate& RHICmdList, const TArray<FActiveMorphTarget>& ActiveMorphTargets, const TArray<float>& MorphTargetWeights, FShaderResourceViewRHIRef MorphWeightsSRV, int32 NumInfluencedVerticesByMorph, FShaderResourceViewRHIRef MorphPerVertexInfoSRV, FShaderResourceViewRHIRef MorphFlattenedSRV)
+static void CalculateMorphDeltaBounds(const TArray<float>& MorphTargetWeights, const FMorphTargetVertexInfoBuffers& MorphTargetVertexInfoBuffers, FVector4& MorphScale, FVector4& InvMorphScale, float &InvTotalAccumulatedWeight)
 {
-	SCOPE_CYCLE_COUNTER(STAT_MorphVertexBuffer_Update);
-
-	if( IsValidRef(MorphVertexBuffer.VertexBufferRHI) )
+	double TotalAccumulatedWeight = 0.0;
+	double MinAccumScale[4] = { 0, 0, 0, 0 };
+	double MaxAccumScale[4] = { 0, 0, 0, 0 };
+	double MaxScale[4] = { 0, 0, 0, 0 };
+	for (uint32 i = 0; i < MorphTargetVertexInfoBuffers.GetNumMorphs(); i++)
 	{
-		extern ENGINE_API bool DoRecomputeSkinTangentsOnGPU_RT();
-		bool bBlendTangentsOnCPU = !DoRecomputeSkinTangentsOnGPU_RT();
-
-/*		// TODO: Need to finish this to avoid artifacts if the SkinCache is not handing all objects
-		if(!bBlendTangentsOnCPU)
+		float AbsoluteMorphTargetWeight = FMath::Abs(MorphTargetWeights[i]);
+		if (AbsoluteMorphTargetWeight > MorphTargetWeightCutoffThreshold)
 		{
-			// It's possible that we reject the object from the SkinCache (e.g. all memory is used up), the we want to render the normal way.
-			// Unfortunately we don't know that at this point and code needs to be changed to get the info here.
-			if(!GGPUSkinCache.IsElementProcessed())
+			TotalAccumulatedWeight += AbsoluteMorphTargetWeight;
+			FVector4 MinMorphScale = MorphTargetVertexInfoBuffers.GetMinimumMorphScale(i);
+			FVector4 MaxMorphScale = MorphTargetVertexInfoBuffers.GetMaximumMorphScale(i);
+
+			for (uint32 j = 0; j < 4; j++)
 			{
-				bBlendTangentsOnCPU = true;
+				MinAccumScale[j] += MorphTargetWeights[i] * MinMorphScale[j];
+				MaxAccumScale[j] += MorphTargetWeights[i] * MaxMorphScale[j];
+
+				double AbsMorphScale = FMath::Max(FMath::Abs(MinMorphScale[j]), FMath::Abs(MaxMorphScale[j]));
+				double AbsAccumScale = FMath::Max(FMath::Abs(MinAccumScale[j]), FMath::Abs(MaxAccumScale[j]));
+				//the maximum accumulated and the maximum local value have to fit into out int24
+				MaxScale[j] = FMath::Max(MaxScale[j], FMath::Max(AbsMorphScale, AbsAccumScale));
 			}
 		}
-*/
+	}
+
+	const double ScaleToInt24 = 16777216.0;
+	MorphScale = FVector4
+	(
+		ScaleToInt24 / (double)((uint64)(MaxScale[0] + 1.0)),
+		ScaleToInt24 / (double)((uint64)(MaxScale[1] + 1.0)),
+		ScaleToInt24 / (double)((uint64)(MaxScale[2] + 1.0)),
+		ScaleToInt24 / (double)((uint64)(MaxScale[3] + 1.0))
+	);
+
+	InvMorphScale = FVector4
+	(
+		(double)((uint64)(MaxScale[0] + 1.0)) / ScaleToInt24,
+		(double)((uint64)(MaxScale[1] + 1.0)) / ScaleToInt24,
+		(double)((uint64)(MaxScale[2] + 1.0)) / ScaleToInt24,
+		(double)((uint64)(MaxScale[3] + 1.0)) / ScaleToInt24
+	);
+
+	InvTotalAccumulatedWeight = (float)(1.0 / TotalAccumulatedWeight);
+}
+
+void FSkeletalMeshObjectGPUSkin::FSkeletalMeshObjectLOD::UpdateMorphVertexBufferGPU(FRHICommandListImmediate& RHICmdList, const TArray<float>& MorphTargetWeights, const FMorphTargetVertexInfoBuffers& MorphTargetVertexInfoBuffers)
+{
+	SCOPE_CYCLE_COUNTER(STAT_MorphVertexBuffer_Update);
+
+	if (IsValidRef(MorphVertexBuffer.VertexBufferRHI))
+	{
 		// LOD of the skel mesh is used to find number of vertices in buffer
 		FStaticLODModel& LodModel = SkelMeshResource->LODModels[LODIndex];
 
 		MorphVertexBuffer.RecreateResourcesIfRequired(GUseGPUMorphTargets != 0);
 
-		if (GUseGPUMorphTargets && RHISupportsComputeShaders(GMaxRHIShaderPlatform))
-		{
-			SCOPED_GPU_STAT(RHICmdList, Stat_GPU_MorphTargets);
+		SCOPED_GPU_STAT(RHICmdList, Stat_GPU_MorphTargets);
 
-			SCOPED_DRAW_EVENTF(RHICmdList, MorphUpdate,
-				TEXT("MorphUpdate LodVertices=%d InflVerts=%d"),
-				LodModel.NumVertices,
-				NumInfluencedVerticesByMorph);
-			RHICmdList.TransitionResource(EResourceTransitionAccess::EWritable, EResourceTransitionPipeline::EGfxToCompute, MorphVertexBuffer.GetUAV());
+		SCOPED_DRAW_EVENTF(RHICmdList, MorphUpdate,
+			TEXT("MorphUpdate LodVertices=%d Threads=%d"),
+			LodModel.NumVertices,
+			MorphTargetVertexInfoBuffers.GetNumWorkItems());
+		RHICmdList.TransitionResource(EResourceTransitionAccess::EWritable, EResourceTransitionPipeline::EGfxToCompute, MorphVertexBuffer.GetUAV());
 
-			RHICmdList.ClearTinyUAV(MorphVertexBuffer.GetUAV(), {0, 0, 0, 0});
+		ClearUAV(RHICmdList, MorphVertexBuffer.GetUAV(), MorphVertexBuffer.GetUAVSize(), 0);
 
+		{		
+			SCOPE_CYCLE_COUNTER(STAT_MorphVertexBuffer_ApplyDelta);
+
+			FVector4 MorphScale; 
+			FVector4 InvMorphScale; 
+			float InvTotalAccumulatedWeight;
+			CalculateMorphDeltaBounds(MorphTargetWeights, MorphTargetVertexInfoBuffers, MorphScale, InvMorphScale, InvTotalAccumulatedWeight);
+			
+			TShaderMapRef<FGPUMorphUpdateCS> GPUMorphUpdateCS(GetGlobalShaderMap(ERHIFeatureLevel::SM5));
+			for (uint32 i = 0; i < MorphTargetVertexInfoBuffers.GetNumMorphs(); i++)
 			{
-				SCOPE_CYCLE_COUNTER(STAT_MorphVertexBuffer_ApplyDelta);
-				TShaderMapRef<FGPUMorphUpdateCS> GPUMorphUpdateCS(GetGlobalShaderMap(ERHIFeatureLevel::SM5));
-				GPUMorphUpdateCS->Dispatch(RHICmdList, LodModel.NumVertices, NumInfluencedVerticesByMorph, MorphWeightsSRV, MorphPerVertexInfoSRV, MorphFlattenedSRV, MorphVertexBuffer.GetUAV());
+				uint32 NumMorphDeltas = MorphTargetVertexInfoBuffers.GetNumWorkItems(i);
+				if (FMath::Abs(MorphTargetWeights[i]) > MorphTargetWeightCutoffThreshold && NumMorphDeltas > 0)
+				{
+					GPUMorphUpdateCS->SetParameters(RHICmdList, MorphScale, MorphTargetVertexInfoBuffers, MorphVertexBuffer);
+					GPUMorphUpdateCS->SetOffsetAndSize(RHICmdList, MorphTargetVertexInfoBuffers.GetStartOffset(i), NumMorphDeltas, MorphTargetWeights[i]);
+					GPUMorphUpdateCS->Dispatch(RHICmdList, NumMorphDeltas);
+					RHICmdList.TransitionResource(EResourceTransitionAccess::ERWNoBarrier, EResourceTransitionPipeline::EComputeToCompute, MorphVertexBuffer.GetUAV());
+				}
 			}
+			GPUMorphUpdateCS->EndAllDispatches(RHICmdList);
+			RHICmdList.TransitionResource(EResourceTransitionAccess::ERWBarrier, EResourceTransitionPipeline::EComputeToCompute, MorphVertexBuffer.GetUAV());
 
-			RHICmdList.TransitionResource(EResourceTransitionAccess::EReadable, EResourceTransitionPipeline::EComputeToGfx, MorphVertexBuffer.GetUAV());
-
-			// set update flag
-			MorphVertexBuffer.bHasBeenUpdated = true;
+			TShaderMapRef<FGPUMorphNormalizeCS> GPUMorphNormalizeCS(GetGlobalShaderMap(ERHIFeatureLevel::SM5));
+			GPUMorphNormalizeCS->Dispatch(RHICmdList, LodModel.NumVertices, InvMorphScale, InvTotalAccumulatedWeight, MorphVertexBuffer);
 		}
-		else
+
+		RHICmdList.TransitionResource(EResourceTransitionAccess::EReadable, EResourceTransitionPipeline::EComputeToGfx, MorphVertexBuffer.GetUAV());
+
+		// set update flag
+		MorphVertexBuffer.bHasBeenUpdated = true;
+	}
+}
+
+void FSkeletalMeshObjectGPUSkin::FSkeletalMeshObjectLOD::UpdateMorphVertexBufferCPU(const TArray<FActiveMorphTarget>& ActiveMorphTargets, const TArray<float>& MorphTargetWeights)
+{
+	SCOPE_CYCLE_COUNTER(STAT_MorphVertexBuffer_Update);
+
+	if (IsValidRef(MorphVertexBuffer.VertexBufferRHI))
+	{
+		extern ENGINE_API bool DoRecomputeSkinTangentsOnGPU_RT();
+		bool bBlendTangentsOnCPU = !DoRecomputeSkinTangentsOnGPU_RT();
+
+		// LOD of the skel mesh is used to find number of vertices in buffer
+		FStaticLODModel& LodModel = SkelMeshResource->LODModels[LODIndex];
+
+		MorphVertexBuffer.RecreateResourcesIfRequired(GUseGPUMorphTargets != 0);
+
+		uint32 Size = LodModel.NumVertices * sizeof(FMorphGPUSkinVertex);
+
+		FMorphGPUSkinVertex* Buffer = nullptr;
 		{
-			uint32 Size = LodModel.NumVertices * sizeof(FMorphGPUSkinVertex);
+			SCOPE_CYCLE_COUNTER(STAT_MorphVertexBuffer_Alloc);
+			Buffer = (FMorphGPUSkinVertex*)FMemory::Malloc(Size);
+		}
 
-			FMorphGPUSkinVertex* Buffer = nullptr;
-			{
-				SCOPE_CYCLE_COUNTER(STAT_MorphVertexBuffer_Alloc);
-				Buffer = (FMorphGPUSkinVertex*)FMemory::Malloc(Size);
-			}
+		{
+			SCOPE_CYCLE_COUNTER(STAT_MorphVertexBuffer_Init);
 
+			if (bBlendTangentsOnCPU)
 			{
-				SCOPE_CYCLE_COUNTER(STAT_MorphVertexBuffer_Init);
-
-				if (bBlendTangentsOnCPU)
+				// zero everything
+				int32 vertsToAdd = static_cast<int32>(LodModel.NumVertices) - MorphAccumulatedWeightArray.Num();
+				if (vertsToAdd > 0)
 				{
-					// zero everything
-					int32 vertsToAdd = static_cast<int32>(LodModel.NumVertices) - MorphAccumulatedWeightArray.Num();
-					if (vertsToAdd > 0)
-					{
-						MorphAccumulatedWeightArray.AddUninitialized(vertsToAdd);
-					}
-
-					FMemory::Memzero(MorphAccumulatedWeightArray.GetData(), sizeof(float)*LodModel.NumVertices);
+					MorphAccumulatedWeightArray.AddUninitialized(vertsToAdd);
 				}
 
-				// PackedNormals will be wrong init with 0, but they'll be overwritten later
-				FMemory::Memzero(&Buffer[0], sizeof(FMorphGPUSkinVertex)*LodModel.NumVertices);
+				FMemory::Memzero(MorphAccumulatedWeightArray.GetData(), sizeof(float)*LodModel.NumVertices);
 			}
 
+			// PackedNormals will be wrong init with 0, but they'll be overwritten later
+			FMemory::Memzero(&Buffer[0], sizeof(FMorphGPUSkinVertex)*LodModel.NumVertices);
+		}
+
+		{
+			SCOPE_CYCLE_COUNTER(STAT_MorphVertexBuffer_ApplyDelta);
+
+			// iterate over all active morph targets and accumulate their vertex deltas
+			for (int32 AnimIdx = 0; AnimIdx < ActiveMorphTargets.Num(); AnimIdx++)
 			{
-				SCOPE_CYCLE_COUNTER(STAT_MorphVertexBuffer_ApplyDelta);
+				const FActiveMorphTarget& MorphTarget = ActiveMorphTargets[AnimIdx];
+				checkSlow(MorphTarget.MorphTarget != NULL);
+				checkSlow(MorphTarget.MorphTarget->HasDataForLOD(LODIndex));
+				const float MorphTargetWeight = MorphTargetWeights[MorphTarget.WeightIndex];
+				const float MorphAbsWeight = FMath::Abs(MorphTargetWeight);
+				checkSlow(MorphAbsWeight >= MinMorphTargetBlendWeight && MorphAbsWeight <= MaxMorphTargetBlendWeight);
 
-				// iterate over all active morph targets and accumulate their vertex deltas
-				for (int32 AnimIdx=0; AnimIdx < ActiveMorphTargets.Num(); AnimIdx++)
-				{
-					const FActiveMorphTarget& MorphTarget = ActiveMorphTargets[AnimIdx];
-					checkSlow(MorphTarget.MorphTarget != NULL);
-					checkSlow(MorphTarget.MorphTarget->HasDataForLOD(LODIndex));
-					const float MorphTargetWeight = MorphTargetWeights[MorphTarget.WeightIndex];
-					const float MorphAbsWeight = FMath::Abs(MorphTargetWeight);
-					checkSlow(MorphAbsWeight >= MinMorphTargetBlendWeight && MorphAbsWeight <= MaxMorphTargetBlendWeight);
 
+				// Get deltas
+				int32 NumDeltas;
+				FMorphTargetDelta* Deltas = MorphTarget.MorphTarget->GetMorphTargetDelta(LODIndex, NumDeltas);
 
-					// Get deltas
-					int32 NumDeltas;
-					FMorphTargetDelta* Deltas = MorphTarget.MorphTarget->GetMorphTargetDelta(LODIndex, NumDeltas);
+				// iterate over the vertices that this lod model has changed
+				for (int32 MorphVertIdx = 0; MorphVertIdx < NumDeltas; MorphVertIdx++)
+				{
+					const FMorphTargetDelta& MorphVertex = Deltas[MorphVertIdx];
 
-					// iterate over the vertices that this lod model has changed
-					for (int32 MorphVertIdx=0; MorphVertIdx < NumDeltas; MorphVertIdx++)
+					// @TODO FIXMELH : temp hack until we fix importing issue
+					if ((MorphVertex.SourceIdx < LodModel.NumVertices))
 					{
-						const FMorphTargetDelta& MorphVertex = Deltas[MorphVertIdx];
-
-						// @TODO FIXMELH : temp hack until we fix importing issue
-						if ((MorphVertex.SourceIdx < LodModel.NumVertices))
-						{
-							FMorphGPUSkinVertex& DestVertex = Buffer[MorphVertex.SourceIdx];
+						FMorphGPUSkinVertex& DestVertex = Buffer[MorphVertex.SourceIdx];
 
-							DestVertex.DeltaPosition += MorphVertex.PositionDelta * MorphTargetWeight;
+						DestVertex.DeltaPosition += MorphVertex.PositionDelta * MorphTargetWeight;
 
-							// todo: could be moved out of the inner loop to be more efficient
-							if (bBlendTangentsOnCPU)
-							{
-								DestVertex.DeltaTangentZ += MorphVertex.TangentZDelta * MorphTargetWeight;
-								// accumulate the weight so we can normalized it later
-								MorphAccumulatedWeightArray[MorphVertex.SourceIdx] += MorphAbsWeight;
-							}
+						// todo: could be moved out of the inner loop to be more efficient
+						if (bBlendTangentsOnCPU)
+						{
+							DestVertex.DeltaTangentZ += MorphVertex.TangentZDelta * MorphTargetWeight;
+							// accumulate the weight so we can normalized it later
+							MorphAccumulatedWeightArray[MorphVertex.SourceIdx] += MorphAbsWeight;
 						}
-					} // for all vertices
-				} // for all morph targets
+					}
+				} // for all vertices
+			} // for all morph targets
 
-				if (bBlendTangentsOnCPU)
+			if (bBlendTangentsOnCPU)
+			{
+				// copy back all the tangent values (can't use Memcpy, since we have to pack the normals)
+				for (uint32 iVertex = 0; iVertex < LodModel.NumVertices; ++iVertex)
 				{
-					// copy back all the tangent values (can't use Memcpy, since we have to pack the normals)
-					for (uint32 iVertex = 0; iVertex < LodModel.NumVertices; ++iVertex)
+					FMorphGPUSkinVertex& DestVertex = Buffer[iVertex];
+					float AccumulatedWeight = MorphAccumulatedWeightArray[iVertex];
+
+					// if accumulated weight is >1.f
+					// previous code was applying the weight again in GPU if less than 1, but it doesn't make sense to do so
+					// so instead, we just divide by AccumulatedWeight if it's more than 1.
+					// now DeltaTangentZ isn't FPackedNormal, so you can apply any value to it. 
+					if (AccumulatedWeight > 1.f)
 					{
-						FMorphGPUSkinVertex& DestVertex = Buffer[iVertex];
-						float AccumulatedWeight = MorphAccumulatedWeightArray[iVertex];
-
-						// if accumulated weight is >1.f
-						// previous code was applying the weight again in GPU if less than 1, but it doesn't make sense to do so
-						// so instead, we just divide by AccumulatedWeight if it's more than 1.
-						// now DeltaTangentZ isn't FPackedNormal, so you can apply any value to it. 
-						if (AccumulatedWeight > 1.f)
-						{
-							DestVertex.DeltaTangentZ /= AccumulatedWeight;
-						}
+						DestVertex.DeltaTangentZ /= AccumulatedWeight;
 					}
 				}
-			} // ApplyDelta
-
-			// Lock the real buffer.
-			{
-				SCOPE_CYCLE_COUNTER(STAT_MorphVertexBuffer_RhiLockAndCopy);
-				FMorphGPUSkinVertex* ActualBuffer = (FMorphGPUSkinVertex*)RHILockVertexBuffer(MorphVertexBuffer.VertexBufferRHI, 0, Size, RLM_WriteOnly);
-				FMemory::Memcpy(ActualBuffer, Buffer, Size);
-				FMemory::Free(Buffer);
 			}
+		} // ApplyDelta
 
-			{
-				SCOPE_CYCLE_COUNTER(STAT_MorphVertexBuffer_RhiUnlock);
-				// Unlock the buffer.
-				RHIUnlockVertexBuffer(MorphVertexBuffer.VertexBufferRHI);
-				// set update flag
-				MorphVertexBuffer.bHasBeenUpdated = true;
-			}
+		// Lock the real buffer.
+		{
+			SCOPE_CYCLE_COUNTER(STAT_MorphVertexBuffer_RhiLockAndCopy);
+			FMorphGPUSkinVertex* ActualBuffer = (FMorphGPUSkinVertex*)RHILockVertexBuffer(MorphVertexBuffer.VertexBufferRHI, 0, Size, RLM_WriteOnly);
+			FMemory::Memcpy(ActualBuffer, Buffer, Size);
+			FMemory::Free(Buffer);
+		}
+
+		{
+			SCOPE_CYCLE_COUNTER(STAT_MorphVertexBuffer_RhiUnlock);
+			// Unlock the buffer.
+			RHIUnlockVertexBuffer(MorphVertexBuffer.VertexBufferRHI);
+			// set update flag
+			MorphVertexBuffer.bHasBeenUpdated = true;
 		}
 	}
 }
diff --git a/Engine/Source/Runtime/Engine/Private/SkeletalRenderGPUSkin.h b/Engine/Source/Runtime/Engine/Private/SkeletalRenderGPUSkin.h
--- a/Engine/Source/Runtime/Engine/Private/SkeletalRenderGPUSkin.h
+++ b/Engine/Source/Runtime/Engine/Private/SkeletalRenderGPUSkin.h
@@ -153,6 +153,7 @@ public:
 	*/
 	FMorphVertexBuffer(FSkeletalMeshResource* InSkelMeshResource, int32 InLODIdx)
 		:	bHasBeenUpdated(false)	
+		,	bNeedsInitialClear(true)
 		,	LODIdx(InLODIdx)
 		,	SkelMeshResource(InSkelMeshResource)
 	{
@@ -202,9 +203,31 @@ public:
 		return ResourceSize;
 	}
 
+	/**
+	* Get Resource Size : only get the size of the GPU resource
+	*/
+	SIZE_T GetUAVSize()
+	{
+		SIZE_T ResourceSize = 0;
+
+		if (VertexBufferRHI)
+		{
+			// LOD of the skel mesh is used to find number of vertices in buffer
+			FStaticLODModel& LodModel = SkelMeshResource->LODModels[LODIdx];
+
+			// Create the buffer rendering resource
+			ResourceSize += LodModel.NumVertices * sizeof(FMorphGPUSkinVertex);
+		}
+
+		return ResourceSize;
+	}
+
 	/** Has been updated or not by UpdateMorphVertexBuffer**/
 	bool bHasBeenUpdated;
 
+	/** DX12 cannot clear the buffer in InitDynamicRHI with UAV flag enables, we should really have a Zero initzialized flag instead**/
+	bool bNeedsInitialClear;
+
 	// @param guaranteed only to be valid if the vertex buffer is valid
 	FShaderResourceViewRHIParamRef GetSRV() const
 	{
@@ -292,8 +315,6 @@ public:
 		{
 			LODs[I].GetResourceSizeEx(CumulativeResourceSize);
 		}
-
-		CumulativeResourceSize.AddUnknownMemoryBytes(MorphWeightsVertexBuffer.VertexBufferRHI ? MorphWeightsVertexBuffer.VertexBufferRHI->GetSize() : 0);
 	}
 	//~ End FSkeletalMeshObject Interface
 
@@ -486,7 +507,8 @@ private:
 		 * @param ActiveMorphTargets - Morph to accumulate. assumed to be weighted and have valid targets
 		 * @param MorphTargetWeights - All Morph weights
 		 */
-		void UpdateMorphVertexBuffer(FRHICommandListImmediate& RHICmdList, const TArray<FActiveMorphTarget>& ActiveMorphTargets, const TArray<float>& MorphTargetWeights, FShaderResourceViewRHIRef MorphWeightsSRV, int32 NumInfluencedVerticesByMorph, FShaderResourceViewRHIRef MorphPerVertexInfoSRV, FShaderResourceViewRHIRef MorphFlattenedSRV);
+		void UpdateMorphVertexBufferCPU(const TArray<FActiveMorphTarget>& ActiveMorphTargets, const TArray<float>& MorphTargetWeights);
+		void UpdateMorphVertexBufferGPU(FRHICommandListImmediate& RHICmdList, const TArray<float>& MorphTargetWeights, const FMorphTargetVertexInfoBuffers& MorphTargetVertexInfoBuffers);
 
 		/**
 		 * Determine the current vertex buffers valid for this LOD
@@ -496,7 +518,6 @@ private:
 		void GetVertexBuffers(FVertexFactoryBuffers& OutVertexBuffers,FStaticLODModel& LODModel);
 
 		// Temporary arrays used on UpdateMorphVertexBuffer(); these grow to the max and are not thread safe.
-		static TArray<FVector> MorphDeltaTangentZAccumulationArray;
 		static TArray<float> MorphAccumulatedWeightArray;
 	};
 
@@ -531,52 +552,109 @@ private:
 
 	/** true if the morph resources have been initialized */
 	bool bMorphResourcesInitialized;
-
-	/** Vertex buffer that stores the weights for all morph targets (matches USkeletalMesh::MorphTargets). */
-	FVertexBufferAndSRV MorphWeightsVertexBuffer;
 };
 
 
-class FMorphTargetBaseShader : public FGlobalShader
+class FGPUMorphUpdateCS : public FGlobalShader
 {
 public:
-	FMorphTargetBaseShader() {}
+	DECLARE_SHADER_TYPE(FGPUMorphUpdateCS, Global);
 
-	FMorphTargetBaseShader(const ShaderMetaType::CompiledShaderInitializerType& Initializer)
+	FGPUMorphUpdateCS() {}
+
+	FGPUMorphUpdateCS(const ShaderMetaType::CompiledShaderInitializerType& Initializer)
 		: FGlobalShader(Initializer)
 	{
-		MorphTargetCountsParameter.Bind(Initializer.ParameterMap, TEXT("MorphTargetCounts"));
 		MorphVertexBufferParameter.Bind(Initializer.ParameterMap, TEXT("MorphVertexBuffer"));
-		PerVertexInfoListParameter.Bind(Initializer.ParameterMap, TEXT("PerVertexInfoList"));
-		FlattenedDeltaListParameter.Bind(Initializer.ParameterMap, TEXT("FlattenedDeltaList"));
-		AllWeightsPerMorphsParameter.Bind(Initializer.ParameterMap, TEXT("AllWeightsPerMorphs"));
+
+		MorphTargetWeightParameter.Bind(Initializer.ParameterMap, TEXT("MorphTargetWeight"));
+		OffsetAndSizeParameter.Bind(Initializer.ParameterMap, TEXT("OffsetAndSize"));
+		PositionScaleParameter.Bind(Initializer.ParameterMap, TEXT("PositionScale"));
+
+		VertexIndicesParameter.Bind(Initializer.ParameterMap, TEXT("VertexIndicies"));
+		MorphDeltasParameter.Bind(Initializer.ParameterMap, TEXT("MorphDeltas"));
 	}
 
 	// FShader interface.
 	virtual bool Serialize(FArchive& Ar) override
 	{
 		bool bShaderHasOutdatedParameters = FGlobalShader::Serialize(Ar);
-		Ar << MorphTargetCountsParameter;
 		Ar << MorphVertexBufferParameter;
-		Ar << PerVertexInfoListParameter;
-		Ar << FlattenedDeltaListParameter;
-		Ar << AllWeightsPerMorphsParameter;
+
+		Ar << MorphTargetWeightParameter;
+		Ar << OffsetAndSizeParameter;
+		Ar << PositionScaleParameter;
+
+		Ar << VertexIndicesParameter;
+		Ar << MorphDeltasParameter;
 		return bShaderHasOutdatedParameters;
 	}
 
+	void SetParameters(FRHICommandList& RHICmdList, const FVector4& LocalScale, const FMorphTargetVertexInfoBuffers& MorphTargetVertexInfoBuffers, FMorphVertexBuffer& MorphVertexBuffer);
+	void SetOffsetAndSize(FRHICommandList& RHICmdList, uint32 Offset, uint32 Size, float Weight);
+
+	void Dispatch(FRHICommandList& RHICmdList, uint32 Size);
+	void EndAllDispatches(FRHICommandList& RHICmdList);
+
 	static bool ShouldCache(EShaderPlatform Platform)
 	{
 		return IsFeatureLevelSupported(Platform, ERHIFeatureLevel::SM5);
 	}
 
-	void SetParameters(FRHICommandList& RHICmdList, int32 NumLODVertices, int32 NumMorphVertices, FShaderResourceViewRHIRef MorphWeightsSRV, FShaderResourceViewRHIRef PerVertexInfoListSRV, FShaderResourceViewRHIRef FlattenedDeltaListSRV, FUnorderedAccessViewRHIParamRef UAV);
-	void Dispatch(FRHICommandList& RHICmdList, int32 NumLODVertices, int32 NumMorphVertices, FShaderResourceViewRHIRef MorphWeightsSRV, FShaderResourceViewRHIRef PerVertexInfoListSRV, FShaderResourceViewRHIRef FlattenedDeltaListSRV, FUnorderedAccessViewRHIParamRef UAV);
-
 protected:
-	FShaderParameter MorphTargetCountsParameter;
 	FShaderResourceParameter MorphVertexBufferParameter;
 
-	FShaderResourceParameter PerVertexInfoListParameter;
-	FShaderResourceParameter FlattenedDeltaListParameter;
-	FShaderResourceParameter AllWeightsPerMorphsParameter;
+	FShaderParameter MorphTargetWeightParameter;
+	FShaderParameter OffsetAndSizeParameter;
+	FShaderParameter PositionScaleParameter;
+
+	FShaderResourceParameter VertexIndicesParameter;
+	FShaderResourceParameter MorphDeltasParameter;
 };
+
+class FGPUMorphNormalizeCS : public FGlobalShader
+{
+public:
+	DECLARE_SHADER_TYPE(FGPUMorphNormalizeCS, Global);
+
+	FGPUMorphNormalizeCS() {}
+
+	FGPUMorphNormalizeCS(const ShaderMetaType::CompiledShaderInitializerType& Initializer)
+		: FGlobalShader(Initializer)
+	{
+		MorphVertexBufferParameter.Bind(Initializer.ParameterMap, TEXT("MorphVertexBuffer"));
+
+		MorphTargetWeightParameter.Bind(Initializer.ParameterMap, TEXT("MorphTargetWeight"));
+		MorphWorkItemsParameter.Bind(Initializer.ParameterMap, TEXT("MorphWorkItems"));
+		PositionScaleParameter.Bind(Initializer.ParameterMap, TEXT("PositionScale"));
+	}
+
+	// FShader interface.
+	virtual bool Serialize(FArchive& Ar) override
+	{
+		bool bShaderHasOutdatedParameters = FGlobalShader::Serialize(Ar);
+		Ar << MorphVertexBufferParameter;
+
+		Ar << MorphTargetWeightParameter;
+		Ar << MorphWorkItemsParameter;
+		Ar << PositionScaleParameter;
+
+		return bShaderHasOutdatedParameters;
+	}
+
+	static bool ShouldCache(EShaderPlatform Platform)
+	{
+		return IsFeatureLevelSupported(Platform, ERHIFeatureLevel::SM5);
+	}
+
+	void SetParameters(FRHICommandList& RHICmdList, uint32 NumVerticies, const FVector4& LocalScale, const float AccumulatedWeight, FMorphVertexBuffer& MorphVertexBuffer);
+
+	void Dispatch(FRHICommandList& RHICmdList, uint32 NumVerticies, const FVector4& LocalScale, const float AccumulatedWeight, FMorphVertexBuffer& MorphVertexBuffer);
+
+protected:
+	FShaderResourceParameter MorphVertexBufferParameter;
+
+	FShaderParameter MorphTargetWeightParameter;
+	FShaderParameter MorphWorkItemsParameter;
+	FShaderParameter PositionScaleParameter;
+};
\ No newline at end of file
diff --git a/Engine/Source/Runtime/Engine/Private/StaticMesh.cpp b/Engine/Source/Runtime/Engine/Private/StaticMesh.cpp
--- a/Engine/Source/Runtime/Engine/Private/StaticMesh.cpp
+++ b/Engine/Source/Runtime/Engine/Private/StaticMesh.cpp
@@ -1156,7 +1156,7 @@ static FString BuildStaticMeshDerivedDataKey(UStaticMesh* Mesh, const FStaticMes
 		}
 	}
 
-	KeySuffix.AppendChar(Mesh->bRequiresAreaWeightedSampling ? TEXT('1') : TEXT('0'));
+	KeySuffix.AppendChar(Mesh->bSupportUniformlyDistributedSampling ? TEXT('1') : TEXT('0'));
 
 	return FDerivedDataCacheInterface::BuildCacheKey(
 		TEXT("STATICMESH"),
@@ -1282,7 +1282,7 @@ void FStaticMeshRenderData::Cache(UStaticMesh* Owner, const FStaticMeshLODSettin
 			IMeshUtilities& MeshUtilities = FModuleManager::Get().LoadModuleChecked<IMeshUtilities>(TEXT("MeshUtilities"));
 			MeshUtilities.BuildStaticMesh(*this, Owner->SourceModels, LODGroup, Owner->LightmapUVVersion, Owner->ImportVersion);
 			ComputeUVDensities();
-			if(Owner->bRequiresAreaWeightedSampling)
+			if(Owner->bSupportUniformlyDistributedSampling)
 			{
 				BuildAreaWeighedSamplingData();
 			}
@@ -1393,7 +1393,7 @@ UStaticMesh::UStaticMesh(const FObjectInitializer& ObjectInitializer)
 	LpvBiasMultiplier = 1.0f;
 	MinLOD = 0;
 
-	bRequiresAreaWeightedSampling = false;
+	bSupportUniformlyDistributedSampling = false;
 }
 
 void UStaticMesh::PostInitProperties()
@@ -1802,6 +1802,8 @@ void UStaticMesh::PostEditChangeProperty(FPropertyChangedEvent& PropertyChangedE
 		SetLightingGuid();
 	}
 	
+	UpdateUVChannelData(true);
+
 	Super::PostEditChangeProperty(PropertyChangedEvent);
 }
 
@@ -2649,6 +2651,13 @@ void UStaticMesh::PostLoad()
 	}
 #endif // #if WITH_EDITOR
 
+#if WITH_EDITORONLY_DATA
+	if (GetLinkerCustomVersion(FRenderingObjectVersion::GUID) < FRenderingObjectVersion::FixedMeshUVDensity)
+	{
+		UpdateUVChannelData(true);
+	}
+#endif
+
 	EnforceLightmapRestrictions();
 
 	if (!GVertexElementTypeSupport.IsSupported(VET_Half2))
diff --git a/Engine/Source/Runtime/Engine/Private/Streaming/DynamicTextureInstanceManager.cpp b/Engine/Source/Runtime/Engine/Private/Streaming/DynamicTextureInstanceManager.cpp
--- a/Engine/Source/Runtime/Engine/Private/Streaming/DynamicTextureInstanceManager.cpp
+++ b/Engine/Source/Runtime/Engine/Private/Streaming/DynamicTextureInstanceManager.cpp
@@ -160,7 +160,7 @@ void FDynamicTextureInstanceManager::Remove(const UPrimitiveComponent* Component
 	check(!Component || Component->IsValidLowLevelFast());
 	if (Component && Component->bAttachedToStreamingManagerAsDynamic)
 	{
-		PendingComponents.Remove(Component);
+		PendingComponents.RemoveSwap(Component);
 
 		// If the component is used, stop any task possibly indirecting it, and clear references.
 		if (StateSync.GetState()->HasComponentReferences(Component))
diff --git a/Engine/Source/Runtime/Engine/Private/TextureDerivedData.cpp b/Engine/Source/Runtime/Engine/Private/TextureDerivedData.cpp
--- a/Engine/Source/Runtime/Engine/Private/TextureDerivedData.cpp
+++ b/Engine/Source/Runtime/Engine/Private/TextureDerivedData.cpp
@@ -754,27 +754,22 @@ class FTextureCacheDerivedDataWorker : public FNonAbandonableTask
 
 	/** Build the texture. This function is safe to call from any thread. */
 	void BuildTexture()
-	{
+		{
 		ensure(Compressor);
 		if (Compressor && SourceMips.Num())
-		{
-			// Adding some extra logs to track crash UE-42168
-			FTextureMemoryStats Stats;
-			RHIGetTextureMemoryStats(Stats);
-
+			{
 			FFormatNamedArguments Args;
 			Args.Add( TEXT("TextureName"), FText::FromString( Texture.GetName() ) );
 			Args.Add( TEXT("TextureFormatName"), FText::FromString( BuildSettings.TextureFormatName.GetPlainNameString() ) );
 			Args.Add( TEXT("TextureResolutionX"), FText::FromString( FString::FromInt(SourceMips[0].SizeX) ) );
 			Args.Add( TEXT("TextureResolutionY"), FText::FromString( FString::FromInt(SourceMips[0].SizeY) ) );
-			Args.Add( TEXT("UsedVRAM"), FText::FromString( FString::FromInt(Stats.AllocatedMemorySize / 1024 / 1024 ) ) );
-			FTextureStatusMessageContext StatusMessage( FText::Format( NSLOCTEXT("Engine", "BuildTextureStatus", "Building textures: {TextureName} ({TextureFormatName}, {TextureResolutionX}X{TextureResolutionY}) - {UsedVRAM} MB total VRAM"), Args ) );
+			FTextureStatusMessageContext StatusMessage( FText::Format( NSLOCTEXT("Engine", "BuildTextureStatus", "Building textures: {TextureName} ({TextureFormatName}, {TextureResolutionX}X{TextureResolutionY})"), Args ) );
 
 			check(DerivedData->Mips.Num() == 0);
 			DerivedData->SizeX = 0;
 			DerivedData->SizeY = 0;
 			DerivedData->PixelFormat = PF_Unknown;
-	
+
 			// Compress the texture.
 			TArray<FCompressedImage2D> CompressedMips;
 			if (Compressor->BuildTexture(SourceMips, CompositeSourceMips, BuildSettings, CompressedMips))
@@ -981,7 +976,8 @@ struct FTextureAsyncCacheDerivedDataTask : public FAsyncTask<FTextureCacheDerive
 void FTexturePlatformData::Cache(
 	UTexture& InTexture,
 	const FTextureBuildSettings& InSettings,
-	uint32 InFlags
+	uint32 InFlags,
+	ITextureCompressorModule* Compressor
 	)
 {
 	// Flush any existing async task and ignore results.
@@ -999,7 +995,10 @@ void FTexturePlatformData::Cache(
 	bool bAsync = !bForDDC && (Flags & ETextureCacheFlags::Async) != 0;
 	GetTextureDerivedDataKey(InTexture, InSettings, DerivedDataKey);
 
-	ITextureCompressorModule* Compressor = &FModuleManager::LoadModuleChecked<ITextureCompressorModule>(TEXTURE_COMPRESSOR_MODULENAME);
+	if (!Compressor)
+	{
+		Compressor = &FModuleManager::LoadModuleChecked<ITextureCompressorModule>(TEXTURE_COMPRESSOR_MODULENAME);
+	}
 
 	if (bAsync && !bForceRebuild)
 	{
@@ -1491,7 +1490,7 @@ void UTexture::UpdateCachedLODBias( bool bIncTextureMips )
 }
 
 #if WITH_EDITOR
-void UTexture::CachePlatformData(bool bAsyncCache, bool bAllowAsyncBuild)
+void UTexture::CachePlatformData(bool bAsyncCache, bool bAllowAsyncBuild, ITextureCompressorModule* Compressor)
 {
 	FTexturePlatformData** PlatformDataLinkPtr = GetRunningPlatformData();
 	if (PlatformDataLinkPtr)
@@ -1519,7 +1518,7 @@ void UTexture::CachePlatformData(bool bAsyncCache, bool bAllowAsyncBuild)
 					(bAsyncCache ? ETextureCacheFlags::Async : ETextureCacheFlags::None) |
 					(bAllowAsyncBuild? ETextureCacheFlags::AllowAsyncBuild : ETextureCacheFlags::None);
 
-				PlatformDataLink->Cache(*this, BuildSettings, CacheFlags);
+				PlatformDataLink->Cache(*this, BuildSettings, CacheFlags, Compressor);
 			}
 		}
 		else if (PlatformDataLink == NULL)
@@ -1618,7 +1617,8 @@ void UTexture::BeginCacheForCookedPlatformData( const ITargetPlatform *TargetPla
 				PlatformDataToCache->Cache(
 					*this,
 					BuildSettingsToCache[SettingsIndex],
-					CurrentCacheFlags
+					CurrentCacheFlags,
+					nullptr
 					);
 				CookedPlatformData.Add( DerivedDataKey, PlatformDataToCache );
 			}
@@ -1836,7 +1836,8 @@ void UTexture::ForceRebuildPlatformData()
 		PlatformDataLink->Cache(
 			*this,
 			BuildSettings,
-			ETextureCacheFlags::ForceRebuild
+			ETextureCacheFlags::ForceRebuild,
+			nullptr
 			);
 	}
 }
@@ -1959,7 +1960,7 @@ void UTexture::SerializeCookedPlatformData(FArchive& Ar)
 				if (PlatformDataPtr == NULL)
 				{
 					PlatformDataPtr = new FTexturePlatformData();
-					PlatformDataPtr->Cache(*this, BuildSettings, ETextureCacheFlags::InlineMips | ETextureCacheFlags::Async);
+					PlatformDataPtr->Cache(*this, BuildSettings, ETextureCacheFlags::InlineMips | ETextureCacheFlags::Async, nullptr);
 					
 						CookedPlatformDataPtr->Add(DerivedDataKey, PlatformDataPtr);
 					
diff --git a/Engine/Source/Runtime/Engine/Private/TextureRenderTarget2D.cpp b/Engine/Source/Runtime/Engine/Private/TextureRenderTarget2D.cpp
--- a/Engine/Source/Runtime/Engine/Private/TextureRenderTarget2D.cpp
+++ b/Engine/Source/Runtime/Engine/Private/TextureRenderTarget2D.cpp
@@ -11,6 +11,7 @@
 #include "UnrealEngine.h"
 #include "DeviceProfiles/DeviceProfile.h"
 #include "DeviceProfiles/DeviceProfileManager.h"
+#include "RenderingObjectVersion.h"
 
 int32 GTextureRenderTarget2DMaxSizeX = 999999999;
 int32 GTextureRenderTarget2DMaxSizeY = 999999999;
@@ -22,7 +23,8 @@ int32 GTextureRenderTarget2DMaxSizeY = 999999999;
 UTextureRenderTarget2D::UTextureRenderTarget2D(const FObjectInitializer& ObjectInitializer)
 	: Super(ObjectInitializer)
 {
-	bHDR = true;
+	bHDR_DEPRECATED = true;
+	RenderTargetFormat = RTF_RGBA16f;
 	bAutoGenerateMips = false;
 	NumMips = 0;
 	ClearColor = FLinearColor(0.0f, 0.0f, 0.0f, 1.0f);
@@ -157,6 +159,18 @@ void UTextureRenderTarget2D::PostEditChangeProperty(FPropertyChangedEvent& Prope
 }
 #endif // WITH_EDITOR
 
+void UTextureRenderTarget2D::Serialize(FArchive& Ar)
+{
+	Super::Serialize(Ar);
+
+	Ar.UsingCustomVersion(FRenderingObjectVersion::GUID);
+
+	if (Ar.CustomVer(FRenderingObjectVersion::GUID) < FRenderingObjectVersion::AddedTextureRenderTargetFormats)
+	{
+		RenderTargetFormat = bHDR_DEPRECATED ? RTF_RGBA16f : RTF_RGBA8;
+	}
+}
+
 void UTextureRenderTarget2D::PostLoad()
 {
 	float OriginalSizeX = SizeX;
@@ -172,7 +186,7 @@ void UTextureRenderTarget2D::PostLoad()
 
 	SizeX = FMath::Min<int32>(SizeX, GTextureRenderTarget2DMaxSizeX);
 	SizeY = FMath::Min<int32>(SizeY, GTextureRenderTarget2DMaxSizeY);
-	
+
 	// Maintain aspect ratio if clamped
 	if( SizeX != OriginalSizeX || SizeY != OriginalSizeY )
 	{
diff --git a/Engine/Source/Runtime/Engine/Private/UnrealEngine.cpp b/Engine/Source/Runtime/Engine/Private/UnrealEngine.cpp
--- a/Engine/Source/Runtime/Engine/Private/UnrealEngine.cpp
+++ b/Engine/Source/Runtime/Engine/Private/UnrealEngine.cpp
@@ -1660,6 +1660,11 @@ void UEngine::InitializeObjectReferences()
 		DefaultBokehTexture = LoadObject<UTexture2D>(NULL, *DefaultBokehTextureName.ToString(), NULL, LOAD_None, NULL);
 	}
 
+	if (DefaultBloomKernelTexture == NULL)
+	{
+		DefaultBloomKernelTexture = LoadObject<UTexture2D>(NULL, *DefaultBloomKernelTextureName.ToString(), NULL, LOAD_None, NULL);
+	}
+
 	if( PreIntegratedSkinBRDFTexture == NULL )
 	{
 		PreIntegratedSkinBRDFTexture = LoadObject<UTexture2D>(NULL, *PreIntegratedSkinBRDFTextureName.ToString(), NULL, LOAD_None, NULL);
@@ -3802,6 +3807,7 @@ bool UEngine::HandleListTexturesCommand( const TCHAR* Cmd, FOutputDevice& Ar )
 	const bool bShouldOnlyListNonStreaming = FParse::Command(&Cmd, TEXT("NONSTREAMING")) && !bShouldOnlyListStreaming;
 	const bool bShouldOnlyListForced = FParse::Command(&Cmd, TEXT("FORCED")) && !bShouldOnlyListStreaming && !bShouldOnlyListNonStreaming;
 	const bool bAlphaSort = FParse::Param( Cmd, TEXT("ALPHASORT") );
+	const bool bCSV = FParse::Param( Cmd, TEXT("CSV") );
 
 	Ar.Logf( TEXT("Listing %s textures."), bShouldOnlyListForced ? TEXT("forced") : bShouldOnlyListNonStreaming ? TEXT("non streaming") : bShouldOnlyListStreaming ? TEXT("streaming") : TEXT("all")  );
 
@@ -3895,13 +3901,17 @@ bool UEngine::HandleListTexturesCommand( const TCHAR* Cmd, FOutputDevice& Ar )
 	int32 TotalMaxAllowedSize = 0;
 	int32 TotalCurrentSize	= 0;
 
-	if (!FPlatformProperties::RequiresCookedData())
+	if (bCSV)
+	{
+		Ar.Logf(TEXT(",Max Width,Max Height,Max Size (KB),Bias Authored,Current Width,Current Height,Current Size (KB),Format,LODGroup,Name,Streaming,Usage Count"));
+	}
+	else if (!FPlatformProperties::RequiresCookedData())
 	{
-		Ar.Logf(TEXT("MaxAllowedSize: Width x Height (Size in KB, Bias from Authored), Current/InMem: Width x Height (Size in KB), Format, LODGroup, Name, Streaming, Usage Count"));
+		Ar.Logf(TEXT("MaxAllowedSize: Width x Height (Size in KB, Authored Bias), Current/InMem: Width x Height (Size in KB), Format, LODGroup, Name, Streaming, Usage Count"));
 	}
 	else
 	{
-		Ar.Logf(TEXT("Cooked/OnDisk: Width x Height (Size in KB), Current/InMem: Width x Height (Size in KB), Format, LODGroup, Name, Streaming, Usage Count"));
+		Ar.Logf(TEXT("Cooked/OnDisk: Width x Height (Size in KB, Authored Bias), Current/InMem: Width x Height (Size in KB), Format, LODGroup, Name, Streaming, Usage Count"));
 	}
 
 	for( int32 TextureIndex=0; TextureIndex<SortedTextures.Num(); TextureIndex++ )
@@ -3909,29 +3919,23 @@ bool UEngine::HandleListTexturesCommand( const TCHAR* Cmd, FOutputDevice& Ar )
 		const FSortedTexture& SortedTexture = SortedTextures[TextureIndex];
 		const bool bValidTextureGroup = TextureGroupNames.IsValidIndex(SortedTexture.LODGroup);
 
+		FString AuthoredBiasString(TEXT("?"));
 		if (!FPlatformProperties::RequiresCookedData())
 		{
-			Ar.Logf(TEXT("%ix%i (%i KB, %i), %ix%i (%i KB), %s, %s, %s, %s, %i"),
-				SortedTexture.MaxAllowedSizeX, SortedTexture.MaxAllowedSizeY, SortedTexture.MaxAllowedSize / 1024, SortedTexture.LODBias,
-				SortedTexture.CurSizeX, SortedTexture.CurSizeY, SortedTexture.CurrentSize / 1024,
-				GetPixelFormatString(SortedTexture.Format),
-				bValidTextureGroup ? *TextureGroupNames[SortedTexture.LODGroup] : TEXT("INVALID"),
-				*SortedTexture.Name,
-				SortedTexture.bIsStreaming ? TEXT("YES") : TEXT("NO"),
-				SortedTexture.UsageCount);
-		}
-		else
-		{
-			Ar.Logf(TEXT("%ix%i (%i KB), %ix%i (%i KB), %s, %s, %s, %s, %i"),
-				SortedTexture.MaxAllowedSizeX, SortedTexture.MaxAllowedSizeY, SortedTexture.MaxAllowedSize / 1024,
-				SortedTexture.CurSizeX, SortedTexture.CurSizeY, SortedTexture.CurrentSize / 1024,
-				GetPixelFormatString(SortedTexture.Format),
-				bValidTextureGroup ? *TextureGroupNames[SortedTexture.LODGroup] : TEXT("INVALID"),
-				*SortedTexture.Name,
-				SortedTexture.bIsStreaming ? TEXT("YES") : TEXT("NO"),
-				SortedTexture.UsageCount);
+			AuthoredBiasString.Empty();
+			AuthoredBiasString.AppendInt(SortedTexture.LODBias);
 		}
 
+		Ar.Logf(bCSV ? TEXT(",%i, %i, %i, %s, %i, %i, %i, %s, %s, %s, %s, %i") : TEXT("%ix%i (%i KB, %s), %ix%i (%i KB), %s, %s, %s, %s, %i"),
+			SortedTexture.MaxAllowedSizeX, SortedTexture.MaxAllowedSizeY, SortedTexture.MaxAllowedSize / 1024, 
+			*AuthoredBiasString,
+			SortedTexture.CurSizeX, SortedTexture.CurSizeY, SortedTexture.CurrentSize / 1024,
+			GetPixelFormatString(SortedTexture.Format),
+			bValidTextureGroup ? *TextureGroupNames[SortedTexture.LODGroup] : TEXT("INVALID"),
+			*SortedTexture.Name,
+			SortedTexture.bIsStreaming ? TEXT("YES") : TEXT("NO"),
+			SortedTexture.UsageCount);
+
 		if (bValidTextureGroup)
 		{
 			TextureGroupCurrentSizes[SortedTexture.LODGroup] += SortedTexture.CurrentSize;
@@ -12631,7 +12635,7 @@ static void SetupThreadAffinity(const TArray<FString>& Args)
 	FSimpleDelegateGraphTask::CreateAndDispatchWhenReady(
 		FSimpleDelegateGraphTask::FDelegate::CreateStatic(&SetAffinityOnThread),
 		TStatId(), NULL, ENamedThreads::RenderThread);
-	if (GRHIThread)
+	if (GRHIThread_InternalUseOnly)
 	{
 		FSimpleDelegateGraphTask::CreateAndDispatchWhenReady(
 			FSimpleDelegateGraphTask::FDelegate::CreateStatic(&SetAffinityOnThread),
diff --git a/Engine/Source/Runtime/Engine/Private/UserInterface/Canvas.cpp b/Engine/Source/Runtime/Engine/Private/UserInterface/Canvas.cpp
--- a/Engine/Source/Runtime/Engine/Private/UserInterface/Canvas.cpp
+++ b/Engine/Source/Runtime/Engine/Private/UserInterface/Canvas.cpp
@@ -923,7 +923,7 @@ void FCanvas::Clear(const FLinearColor& ClearColor)
 			{
 				::SetRenderTarget(RHICmdList, CanvasRenderTarget->GetRenderTargetTexture(), FTextureRHIRef(), true);
 				RHICmdList.SetViewport(0, 0, 0.0f, CanvasRenderTarget->GetSizeXY().X, CanvasRenderTarget->GetSizeXY().Y, 1.0f);
-				DrawClearQuad(RHICmdList, GMaxRHIFeatureLevel, ClearColor);
+				DrawClearQuad(RHICmdList, ClearColor);
 			}
 			else
 			{
diff --git a/Engine/Source/Runtime/Engine/Public/BlendableManager.h b/Engine/Source/Runtime/Engine/Public/BlendableManager.h
--- a/Engine/Source/Runtime/Engine/Public/BlendableManager.h
+++ b/Engine/Source/Runtime/Engine/Public/BlendableManager.h
@@ -45,7 +45,7 @@ public:
 		FName BlendableType = T::GetFName();
 
 		// at least 4 byte alignment
-		uint32 Alignment = FMath::Max((uint32)4, (uint32)ALIGNOF(T));
+		uint32 Alignment = FMath::Max((uint32)4, (uint32)alignof(T));
 
 		FBlendableEntry* Entry = PushBlendableDataPtr(InWeight, BlendableType, (const uint8*)&InData, sizeof(T), Alignment);
 
diff --git a/Engine/Source/Runtime/Engine/Public/GPUSkinVertexFactory.h b/Engine/Source/Runtime/Engine/Public/GPUSkinVertexFactory.h
--- a/Engine/Source/Runtime/Engine/Public/GPUSkinVertexFactory.h
+++ b/Engine/Source/Runtime/Engine/Public/GPUSkinVertexFactory.h
@@ -214,64 +214,6 @@ public: // From FTickableObjectRenderThread
 };
 
 
-/** for motion blur skinning */
-class FBoneDataVertexBuffer : public FRenderResource
-{
-public:
-	/** constructor */
-	FBoneDataVertexBuffer();
-
-	/** 
-	* call UnlockData() after this one
-	* @return never 0
-	*/
-	float* LockData();
-	/**
-	* Needs to be called after LockData()
-	*/
-	void UnlockData(uint32 SizeInBytes);
-
-	/** Returns the size of the buffer in pixels. */
-	uint32 GetSizeX() const;
-
-	bool IsValid() const
-	{
-		return IsValidRef(BoneBuffer);
-	}
-
-	// interface FRenderResource ------------------------------------------
-
-	virtual void ReleaseRHI() override
-	{
-		DEC_DWORD_STAT_BY( STAT_SkeletalMeshMotionBlurSkinningMemory, ComputeMemorySize());
-		BoneBuffer.SafeRelease();
-		FRenderResource::ReleaseRHI();
-	}
-
-	virtual void InitDynamicRHI() override
-	{
-		if(SizeX)
-		{
-			INC_DWORD_STAT_BY( STAT_SkeletalMeshMotionBlurSkinningMemory, ComputeMemorySize());
-			const int32 TileBufferSize = ComputeMemorySize();
-			FRHIResourceCreateInfo CreateInfo;
-			// BUF_Dynamic as it can be too large (e.g. 100 characters each 100 bones with 48 bytes per bone is about 1/2 MB) to do it as BUF_Volatile
-			BoneBuffer.VertexBufferRHI = RHICreateVertexBuffer( TileBufferSize, BUF_Dynamic | BUF_ShaderResource, CreateInfo );
-			BoneBuffer.VertexBufferSRV = RHICreateShaderResourceView( BoneBuffer.VertexBufferRHI, sizeof(FVector4), PF_A32B32G32R32F );
-		}
-	}
-	
-	/** Bone buffer reference. */
-	FVertexBufferAndSRV BoneBuffer;
-	
-private: // -------------------------------------------------------
-	/** Buffer size in texels */
-	uint32 SizeX;
-
-	/** @return in bytes */
-	uint32 ComputeMemorySize();
-};
-
 /** Vertex factory with vertex stream components for GPU skinned vertices */
 class FGPUBaseSkinVertexFactory : public FVertexFactory
 {
diff --git a/Engine/Source/Runtime/Engine/Public/MaterialCompiler.h b/Engine/Source/Runtime/Engine/Public/MaterialCompiler.h
--- a/Engine/Source/Runtime/Engine/Public/MaterialCompiler.h
+++ b/Engine/Source/Runtime/Engine/Public/MaterialCompiler.h
@@ -245,6 +245,8 @@ public:
 	virtual int32 PerInstanceRandom() = 0;
 	virtual int32 PerInstanceFadeAmount() = 0;
 	virtual int32 AntialiasedTextureMask(int32 Tex, int32 UV, float Threshold, uint8 Channel) = 0;
+	virtual int32 Sobol(int32 Cell, int32 Index, int32 Seed) = 0;
+	virtual int32 TemporalSobol(int32 Index, int32 Seed) = 0;
 	virtual int32 Noise(int32 Position, float Scale, int32 Quality, uint8 NoiseFunction, bool bTurbulence, int32 Levels, float OutputMin, float OutputMax, float LevelScale, int32 FilterWidth, bool bTiling, uint32 RepeatSize) = 0;
 	virtual int32 VectorNoise(int32 Position, int32 Quality, uint8 NoiseFunction, bool bTiling, uint32 RepeatSize) = 0;
 	virtual int32 BlackBody( int32 Temp ) = 0;
@@ -444,6 +446,8 @@ public:
 	{
 		return Compiler->AntialiasedTextureMask(Tex, UV, Threshold, Channel);
 	}
+	virtual int32 Sobol(int32 Cell, int32 Index, int32 Seed) override {	return Compiler->Sobol(Cell, Index, Seed); }
+	virtual int32 TemporalSobol(int32 Index, int32 Seed) override { return Compiler->TemporalSobol(Index, Seed); }
 	virtual int32 Noise(int32 Position, float Scale, int32 Quality, uint8 NoiseFunction, bool bTurbulence, int32 Levels, float OutputMin, float OutputMax, float LevelScale, int32 FilterWidth, bool bTiling, uint32 TileSize) override
 	{
 		return Compiler->Noise(Position, Scale, Quality, NoiseFunction, bTurbulence, Levels, OutputMin, OutputMax, LevelScale, FilterWidth, bTiling, TileSize);
diff --git a/Engine/Source/Runtime/Engine/Public/MaterialShared.h b/Engine/Source/Runtime/Engine/Public/MaterialShared.h
--- a/Engine/Source/Runtime/Engine/Public/MaterialShared.h
+++ b/Engine/Source/Runtime/Engine/Public/MaterialShared.h
@@ -530,8 +530,9 @@ class FMeshMaterialShaderMap : public TShaderMap<FMeshMaterialShaderType>
 {
 public:
 
-	FMeshMaterialShaderMap(FVertexFactoryType* InVFType) :
-		VertexFactoryType(InVFType)
+	FMeshMaterialShaderMap(EShaderPlatform InPlatform, FVertexFactoryType* InVFType) 
+		: TShaderMap<FMeshMaterialShaderType>(InPlatform)
+		, VertexFactoryType(InVFType)
 	{}
 
 	/**
@@ -620,7 +621,8 @@ public:
 	 */
 	static void LoadFromDerivedDataCache(const FMaterial* Material, const FMaterialShaderMapId& ShaderMapId, EShaderPlatform Platform, TRefCountPtr<FMaterialShaderMap>& InOutShaderMap);
 
-	FMaterialShaderMap();
+	inline FMaterialShaderMap() : FMaterialShaderMap(EShaderPlatform::SP_NumPlatforms) {}
+	FMaterialShaderMap(EShaderPlatform InPlatform);
 
 	// Destructor.
 	~FMaterialShaderMap();
@@ -743,7 +745,6 @@ public:
 	// Accessors.
 	ENGINE_API const FMeshMaterialShaderMap* GetMeshShaderMap(FVertexFactoryType* VertexFactoryType) const;
 	const FMaterialShaderMapId& GetShaderMapId() const { return ShaderMapId; }
-	EShaderPlatform GetShaderPlatform() const { return Platform; }
 	const FString& GetFriendlyName() const { return FriendlyName; }
 	uint32 GetCompilingId() const { return CompilingId; }
 	bool IsCompilationFinalized() const { return bCompilationFinalized; }
@@ -794,9 +795,6 @@ private:
 	/** The material's user friendly name, typically the object name. */
 	FString FriendlyName;
 
-	/** The platform this shader map was compiled with */
-	EShaderPlatform Platform;
-
 	/** The static parameter set that this shader map was compiled with */
 	FMaterialShaderMapId ShaderMapId;
 
@@ -1060,7 +1058,7 @@ public:
 	virtual float GetTranslucentSelfShadowSecondDensityScale() const { return 1.0f; }
 	virtual float GetTranslucentSelfShadowSecondOpacity() const { return 1.0f; }
 	virtual float GetTranslucentBackscatteringExponent() const { return 1.0f; }
-	virtual bool IsSeparateTranslucencyEnabled() const { return false; }
+	virtual bool IsTranslucencyAfterDOFEnabled() const { return false; }
 	virtual bool IsMobileSeparateTranslucencyEnabled() const { return false; }
 	virtual FLinearColor GetTranslucentMultipleScatteringExtinction() const { return FLinearColor::White; }
 	virtual float GetTranslucentShadowStartOffset() const { return 0.0f; }
@@ -1493,6 +1491,7 @@ public:
 
 	ENGINE_API static const TSet<FMaterialRenderProxy*>& GetMaterialRenderProxyMap() 
 	{
+		check(!FPlatformProperties::RequiresCookedData());
 		return MaterialRenderProxyMap;
 	}
 
@@ -1522,6 +1521,13 @@ public:
 #endif
 	}
 
+	ENGINE_API static void UpdateDeferredCachedUniformExpressions();
+
+	static inline bool HasDeferredUniformExpressionCacheRequests() 
+	{
+		return DeferredUniformExpressionCacheRequests.Num() > 0;
+	}
+
 private:
 
 	/** true if the material is selected. */
@@ -1542,6 +1548,8 @@ private:
 	 * This is used to propagate new shader maps to materials being used for rendering.
 	 */
 	ENGINE_API static TSet<FMaterialRenderProxy*> MaterialRenderProxyMap;
+
+	ENGINE_API static TSet<FMaterialRenderProxy*> DeferredUniformExpressionCacheRequests;
 };
 
 /**
@@ -1710,7 +1718,7 @@ public:
 	ENGINE_API virtual float GetTranslucentSelfShadowSecondDensityScale() const override;
 	ENGINE_API virtual float GetTranslucentSelfShadowSecondOpacity() const override;
 	ENGINE_API virtual float GetTranslucentBackscatteringExponent() const override;
-	ENGINE_API virtual bool IsSeparateTranslucencyEnabled() const override;
+	ENGINE_API virtual bool IsTranslucencyAfterDOFEnabled() const override;
 	ENGINE_API virtual bool IsMobileSeparateTranslucencyEnabled() const override;
 	ENGINE_API virtual FLinearColor GetTranslucentMultipleScatteringExtinction() const override;
 	ENGINE_API virtual float GetTranslucentShadowStartOffset() const override;
diff --git a/Engine/Source/Runtime/Engine/Public/PrimitiveViewRelevance.h b/Engine/Source/Runtime/Engine/Public/PrimitiveViewRelevance.h
--- a/Engine/Source/Runtime/Engine/Public/PrimitiveViewRelevance.h
+++ b/Engine/Source/Runtime/Engine/Public/PrimitiveViewRelevance.h
@@ -30,6 +30,10 @@ struct FPrimitiveViewRelevance
 	uint32 bMobileSeparateTranslucencyRelevance : 1;
 	/** The primitive has one or more elements that have normal translucency. */
 	uint32 bNormalTranslucencyRelevance : 1;
+	/** For translucent primitives reading the scene color. */
+	uint32 bUsesSceneColorCopy : 1;
+	/** For primitive that can't render in offscreen buffers (blend modulate). */
+	uint32 bDisableOffscreenRendering : 1;
 	/** */
 	uint32 bUsesGlobalDistanceField : 1;
 
diff --git a/Engine/Source/Runtime/Engine/Public/SceneManagement.h b/Engine/Source/Runtime/Engine/Public/SceneManagement.h
--- a/Engine/Source/Runtime/Engine/Public/SceneManagement.h
+++ b/Engine/Source/Runtime/Engine/Public/SceneManagement.h
@@ -668,9 +668,11 @@ private:
 	TPendingTextureType* PendingTexture;
 	FThreadSafeCounter& Counter;
 	ULevel* LightingScenario;
+	class ITextureCompressorModule* Compressor;
+
 public:
 
-	FAsyncEncode(TPendingTextureType* InPendingTexture, ULevel* InLightingScenario, FThreadSafeCounter& InCounter) : PendingTexture(nullptr), Counter(InCounter)
+	FAsyncEncode(TPendingTextureType* InPendingTexture, ULevel* InLightingScenario, FThreadSafeCounter& InCounter, ITextureCompressorModule* InCompressor) : PendingTexture(nullptr), Counter(InCounter), Compressor(InCompressor)
 	{
 		LightingScenario = InLightingScenario;
 		PendingTexture = InPendingTexture;
@@ -678,13 +680,13 @@ public:
 
 	void Abandon()
 	{
-		PendingTexture->StartEncoding(LightingScenario);
+		PendingTexture->StartEncoding(LightingScenario, Compressor);
 		Counter.Decrement();
 	}
 
 	void DoThreadedWork()
 	{
-		PendingTexture->StartEncoding(LightingScenario);
+		PendingTexture->StartEncoding(LightingScenario, Compressor);
 		Counter.Decrement();
 	}
 };
@@ -1920,7 +1922,7 @@ private:
 // 10x10 tessellated plane at x=-1..1 y=-1...1 z=0
 extern ENGINE_API void DrawPlane10x10(class FPrimitiveDrawInterface* PDI,const FMatrix& ObjectToWorld,float Radii,FVector2D UVMin, FVector2D UVMax,const FMaterialRenderProxy* MaterialRenderProxy,uint8 DepthPriority);
 extern ENGINE_API void DrawBox(class FPrimitiveDrawInterface* PDI,const FMatrix& BoxToWorld,const FVector& Radii,const FMaterialRenderProxy* MaterialRenderProxy,uint8 DepthPriority);
-extern ENGINE_API void DrawSphere(class FPrimitiveDrawInterface* PDI,const FVector& Center,const FVector& Radii,int32 NumSides,int32 NumRings,const FMaterialRenderProxy* MaterialRenderProxy,uint8 DepthPriority,bool bDisableBackfaceCulling=false);
+extern ENGINE_API void DrawSphere(class FPrimitiveDrawInterface* PDI,const FVector& Center,const FRotator& Orientation,const FVector& Radii,int32 NumSides,int32 NumRings,const FMaterialRenderProxy* MaterialRenderProxy,uint8 DepthPriority,bool bDisableBackfaceCulling=false);
 extern ENGINE_API void DrawCone(class FPrimitiveDrawInterface* PDI,const FMatrix& ConeToWorld, float Angle1, float Angle2, int32 NumSides, bool bDrawSideLines, const FLinearColor& SideLineColor, const FMaterialRenderProxy* MaterialRenderProxy, uint8 DepthPriority);
 
 extern ENGINE_API void DrawCylinder(class FPrimitiveDrawInterface* PDI,const FVector& Base, const FVector& XAxis, const FVector& YAxis, const FVector& ZAxis,
diff --git a/Engine/Source/Runtime/Engine/Public/SceneView.h b/Engine/Source/Runtime/Engine/Public/SceneView.h
--- a/Engine/Source/Runtime/Engine/Public/SceneView.h
+++ b/Engine/Source/Runtime/Engine/Public/SceneView.h
@@ -713,6 +713,7 @@ BEGIN_UNIFORM_BUFFER_STRUCT_WITH_CONSTRUCTOR(FViewUniformShaderParameters, ENGIN
 	DECLARE_UNIFORM_BUFFER_STRUCT_MEMBER_SAMPLER(SamplerState, PerlinNoiseGradientTextureSampler)
 	DECLARE_UNIFORM_BUFFER_STRUCT_MEMBER_TEXTURE(Texture3D, PerlinNoise3DTexture)
 	DECLARE_UNIFORM_BUFFER_STRUCT_MEMBER_SAMPLER(SamplerState, PerlinNoise3DTextureSampler)
+	DECLARE_UNIFORM_BUFFER_STRUCT_MEMBER_TEXTURE(Texture2D<uint4>, SobolSamplingTexture)
 
 END_UNIFORM_BUFFER_STRUCT(FViewUniformShaderParameters)
 
@@ -1427,6 +1428,8 @@ public:
 	{
 		return MonoParameters.bEnabled && MonoParameters.Mode != EMonoscopicFarFieldMode::Off;
 	}
+
+	bool AllowTranslucencyAfterDOF() const;
 };
 
 /**
diff --git a/Engine/Source/Runtime/Engine/Public/SkeletalMeshTypes.h b/Engine/Source/Runtime/Engine/Public/SkeletalMeshTypes.h
--- a/Engine/Source/Runtime/Engine/Public/SkeletalMeshTypes.h
+++ b/Engine/Source/Runtime/Engine/Public/SkeletalMeshTypes.h
@@ -1095,48 +1095,78 @@ private:
 class FMorphTargetVertexInfoBuffers : public FRenderResource
 {
 public:
-	FMorphTargetVertexInfoBuffers()
-		: NumInfluencedVerticesByMorphs(0)
+	FMorphTargetVertexInfoBuffers() : NumTotalWorkItems(0)
 	{
 	}
 
 	ENGINE_API virtual void InitRHI() override;
 	ENGINE_API virtual void ReleaseRHI() override;
 
-	uint32 GetNumInfluencedVerticesByMorphs() const
+	uint32 GetNumWorkItems(uint32 index = UINT_MAX) const
 	{
-		return NumInfluencedVerticesByMorphs;
+		check(index == UINT_MAX || index < (uint32)WorkItemsPerMorph.Num());
+		return index != UINT_MAX ? WorkItemsPerMorph[index] : NumTotalWorkItems;
 	}
 
-	FVertexBufferRHIRef PerVertexInfoVB;
-	FShaderResourceViewRHIRef PerVertexInfoSRV;
+	uint32 GetNumMorphs() const
+	{
+		return WorkItemsPerMorph.Num();
+	}
 
-	FVertexBufferRHIRef FlattenedDeltasVB;
-	FShaderResourceViewRHIRef FlattenedDeltasSRV;
+	uint32 GetStartOffset(uint32 Index) const
+	{
+		check(Index < (uint32)StartOffsetPerMorph.Num());
+		return StartOffsetPerMorph[Index];
+	}
 
-	// Changes to this struct must be reflected in MorphTargets.usf
-	struct FPerVertexInfo
+	const FVector4& GetMaximumMorphScale(uint32 Index) const
 	{
-		uint32 DestVertexIndex;
-		uint32 StartDelta;
-		uint32 NumDeltas;
-	};
+		check(Index < (uint32)MaximumValuePerMorph.Num());
+		return MaximumValuePerMorph[Index];
+	}
+
+	const FVector4& GetMinimumMorphScale(uint32 Index) const
+	{
+		check(Index < (uint32)MinimumValuePerMorph.Num());
+		return MinimumValuePerMorph[Index];
+	}
+
+	FVertexBufferRHIRef VertexIndicesVB;
+	FShaderResourceViewRHIRef VertexIndicesSRV;
+
+	FVertexBufferRHIRef MorphDeltasVB;
+	FShaderResourceViewRHIRef MorphDeltasSRV;
 
 	// Changes to this struct must be reflected in MorphTargets.usf
-	struct FFlattenedDelta
+	struct FMorphDelta
 	{
-		FVector PosDelta;
-		FVector TangentDelta;
-		uint32 WeightIndex;
+		FMorphDelta(FVector InPosDelta, FVector InTangentDelta)
+		{
+			PosDelta[0] = FFloat16(InPosDelta.X);
+			PosDelta[1] = FFloat16(InPosDelta.Y);
+			PosDelta[2] = FFloat16(InPosDelta.Z);
+
+			TangentDelta[0] = FFloat16(InTangentDelta.X);
+			TangentDelta[1] = FFloat16(InTangentDelta.Y);
+			TangentDelta[2] = FFloat16(InTangentDelta.Z);
+		}
+
+		FFloat16 PosDelta[3];
+		FFloat16 TangentDelta[3];
 	};
 
 protected:
 	// Transient data used while creating the vertex buffers, gets deleted as soon as VB gets initialized.
-	TArray<FPerVertexInfo> PerVertexInfoList;
-	// Transient data used while creating the vertex buffers, gets deleted as soon as VB gets initialized.
-	TArray<FFlattenedDelta> FlattenedDeltaList;
+	TArray<uint32> VertexIndices;
 	// Transient data used while creating the vertex buffers, gets deleted as soon as VB gets initialized.
-	uint32 NumInfluencedVerticesByMorphs;
+	TArray<FMorphDelta> MorphDeltas;
+
+	//x,y,y separate for position and shared w for tangent
+	TArray<FVector4> MaximumValuePerMorph;
+	TArray<FVector4> MinimumValuePerMorph;
+	TArray<uint32> StartOffsetPerMorph;
+	TArray<uint32> WorkItemsPerMorph;
+	uint32 NumTotalWorkItems;
 
 	friend class FStaticLODModel;
 };
diff --git a/Engine/Source/Runtime/HeadMountedDisplay/Private/DefaultStereoLayers.cpp b/Engine/Source/Runtime/HeadMountedDisplay/Private/DefaultStereoLayers.cpp
--- a/Engine/Source/Runtime/HeadMountedDisplay/Private/DefaultStereoLayers.cpp
+++ b/Engine/Source/Runtime/HeadMountedDisplay/Private/DefaultStereoLayers.cpp
@@ -242,7 +242,7 @@ void FDefaultStereoLayers::PostRenderView_RenderThread(FRHICommandListImmediate&
 	if (OverlayRenderTarget.IsValid())
 	{
 		SetRenderTarget(RHICmdList, OverlayRenderTarget, FTextureRHIRef());
-		DrawClearQuad(RHICmdList, GMaxRHIFeatureLevel, FLinearColor(0.0f, 0.0f, 0.0f, 0.0f));
+		DrawClearQuad(RHICmdList, FLinearColor(0.0f, 0.0f, 0.0f, 0.0f));
 		RHICmdList.SetViewport(RenderParams.Viewport.Min.X, RenderParams.Viewport.Min.Y, 0, RenderParams.Viewport.Max.X, RenderParams.Viewport.Max.Y, 1.0f);
 	}
 
diff --git a/Engine/Source/Runtime/Launch/Private/LaunchEngineLoop.cpp b/Engine/Source/Runtime/Launch/Private/LaunchEngineLoop.cpp
--- a/Engine/Source/Runtime/Launch/Private/LaunchEngineLoop.cpp
+++ b/Engine/Source/Runtime/Launch/Private/LaunchEngineLoop.cpp
@@ -329,7 +329,7 @@ static void RHIExitAndStopRHIThread()
 	RHIExit();
 
 	// Stop the RHI Thread
-	if (GUseRHIThread)
+	if (GRHIThread_InternalUseOnly)
 	{
 		DECLARE_CYCLE_STAT(TEXT("Wait For RHIThread Finish"), STAT_WaitForRHIThreadFinish, STATGROUP_TaskGraphTasks);
 		FGraphEventRef QuitTask = TGraphTask<FReturnGraphTask>::CreateTask(nullptr, ENamedThreads::GameThread).ConstructAndDispatchWhenReady(ENamedThreads::RHIThread);
@@ -1029,6 +1029,9 @@ int32 FEngineLoop::PreInit( const TCHAR* CmdLine )
 	FCString::Strcpy( CommandLineCopy, CommandLineSize, CmdLine );
 	const TCHAR* ParsedCmdLine	= CommandLineCopy;
 
+	// Add the default engine shader dir
+	FGenericPlatformProcess::AddShaderDir(FGenericPlatformProcess::ShaderDir());
+
 	FString Token				= FParse::Token( ParsedCmdLine, 0);
 
 #if WITH_ENGINE
@@ -1718,14 +1721,14 @@ int32 FEngineLoop::PreInit( const TCHAR* CmdLine )
 			if(GRHISupportsRHIThread)
 			{
 				const bool DefaultUseRHIThread = true;
-				GUseRHIThread = DefaultUseRHIThread;
+				GUseRHIThread_InternalUseOnly = DefaultUseRHIThread;
 				if(FParse::Param(FCommandLine::Get(), TEXT("rhithread")))
 				{
-					GUseRHIThread = true;
+					GUseRHIThread_InternalUseOnly = true;
 				}
 				else if(FParse::Param(FCommandLine::Get(), TEXT("norhithread")))
 				{
-					GUseRHIThread = false;
+					GUseRHIThread_InternalUseOnly = false;
 				}
 			}
 
@@ -1852,14 +1855,14 @@ int32 FEngineLoop::PreInit( const TCHAR* CmdLine )
 		if (GRHISupportsRHIThread)
 		{
 			const bool DefaultUseRHIThread = true;
-			GUseRHIThread = DefaultUseRHIThread;
+			GUseRHIThread_InternalUseOnly = DefaultUseRHIThread;
 			if (FParse::Param(FCommandLine::Get(),TEXT("rhithread")))
 			{
-				GUseRHIThread = true;
+				GUseRHIThread_InternalUseOnly = true;
 			}
 			else if (FParse::Param(FCommandLine::Get(),TEXT("norhithread")))
 			{
-				GUseRHIThread = false;
+				GUseRHIThread_InternalUseOnly = false;
 			}
 		}
 		StartRenderingThread();
diff --git a/Engine/Source/Runtime/Niagara/Private/NiagaraDataInterfaceStaticMesh.cpp b/Engine/Source/Runtime/Niagara/Private/NiagaraDataInterfaceStaticMesh.cpp
--- a/Engine/Source/Runtime/Niagara/Private/NiagaraDataInterfaceStaticMesh.cpp
+++ b/Engine/Source/Runtime/Niagara/Private/NiagaraDataInterfaceStaticMesh.cpp
@@ -760,7 +760,7 @@ bool UNiagaraDataInterfaceStaticMesh::ResetRequired()
 	bool bReset = false;
 	if (Mesh)
 	{
-		bIsAreaWeightedSampling = Mesh->bRequiresAreaWeightedSampling;
+		bIsAreaWeightedSampling = Mesh->bSupportUniformlyDistributedSampling;
 		bSupportingVertexColorSampling = bEnableVertexColorRangeSorting && MeshHasColors();
 		bReset = !Mesh->bAllowCPUAccess || bIsAreaWeightedSampling != bPrevAreaWeighted || bSupportingVertexColorSampling != bPrevVCSampling;
 	}
diff --git a/Engine/Source/Runtime/Projects/Private/PluginManager.cpp b/Engine/Source/Runtime/Projects/Private/PluginManager.cpp
--- a/Engine/Source/Runtime/Projects/Private/PluginManager.cpp
+++ b/Engine/Source/Runtime/Projects/Private/PluginManager.cpp
@@ -425,6 +425,17 @@ bool FPluginManager::ConfigureEnabledPlugins()
 			const FPlugin& Plugin = *PluginPair.Value;
 			if (Plugin.bEnabled)
 			{
+				// Plugins can have their own shaders
+				// Add potential plugin shader directory only if the plugin is loaded in PostConfigInit. Not supported otherwise
+				for (const FModuleDescriptor& Module : Plugin.GetDescriptor().Modules)
+				{
+					if (Module.LoadingPhase == ELoadingPhase::PostConfigInit)
+					{
+						FGenericPlatformProcess::AddShaderDir(FPaths::Combine(*Plugin.GetBaseDir(), TEXT("Shaders")));
+						break;
+					}
+				}
+
 				// Build the list of content folders
 				if (Plugin.Descriptor.bCanContainContent)
 				{
diff --git a/Engine/Source/Runtime/Projects/Private/ProjectManager.cpp b/Engine/Source/Runtime/Projects/Private/ProjectManager.cpp
--- a/Engine/Source/Runtime/Projects/Private/ProjectManager.cpp
+++ b/Engine/Source/Runtime/Projects/Private/ProjectManager.cpp
@@ -8,6 +8,7 @@
 #include "Interfaces/IPluginManager.h"
 #include "ProjectDescriptor.h"
 #include "Modules/ModuleManager.h"
+#include "GenericPlatform/GenericPlatformProcess.h"
 
 DEFINE_LOG_CATEGORY_STATIC( LogProjectManager, Log, All );
 
@@ -30,6 +31,10 @@ bool FProjectManager::LoadProjectFile( const FString& InProjectFile )
 	TSharedPtr<FProjectDescriptor> Descriptor = MakeShareable(new FProjectDescriptor());
 	if(Descriptor->Load(InProjectFile, FailureReason))
 	{
+		// Projects can have their own shaders
+		// Add potential project shader directory
+		FGenericPlatformProcess::AddShaderDir(FPaths::Combine(FPaths::GetPath(InProjectFile), TEXT("Shaders")));
+
 		// Create the project
 		CurrentProject = Descriptor;
 		return true;
diff --git a/Engine/Source/Runtime/RHI/Private/DynamicRHI.cpp b/Engine/Source/Runtime/RHI/Private/DynamicRHI.cpp
--- a/Engine/Source/Runtime/RHI/Private/DynamicRHI.cpp
+++ b/Engine/Source/Runtime/RHI/Private/DynamicRHI.cpp
@@ -251,7 +251,7 @@ void FDynamicRHI::EnableIdealGPUCaptureOptions(bool bEnabled)
 
 	const bool bDrawEvents = GEmitDrawEvents != 0;
 	const bool bMaterialDrawEvents = ShowMaterialDrawEventVar ? ShowMaterialDrawEventVar->GetInt() != 0 : false;
-	const bool bRHIThread = GRHIThread != nullptr;	
+	const bool bRHIThread = IsRunningRHIInSeparateThread();
 	const bool bRHIBypass = RHICmdBypassVar ? RHICmdBypassVar->GetInt() != 0 : false;
 
 	UE_LOG(LogRHI, Display, TEXT("Setting GPU Capture Options: %i"), bEnabled ? 1 : 0);
diff --git a/Engine/Source/Runtime/RHI/Private/PipelineStateCache.cpp b/Engine/Source/Runtime/RHI/Private/PipelineStateCache.cpp
--- a/Engine/Source/Runtime/RHI/Private/PipelineStateCache.cpp
+++ b/Engine/Source/Runtime/RHI/Private/PipelineStateCache.cpp
@@ -5,7 +5,7 @@ PipelineStateCache.cpp: Pipeline state cache implementation.
 =============================================================================*/
 
 #include "PipelineStateCache.h"
-#include "Misc/ScopeLock.h"
+#include "Misc/ScopeRWLock.h"
 
 extern RHI_API FRHIComputePipelineState* ExecuteSetComputePipelineState(FComputePipelineState* ComputePipelineState);
 extern RHI_API FRHIGraphicsPipelineState* ExecuteSetGraphicsPipelineState(FGraphicsPipelineState* GraphicsPipelineState);
@@ -47,9 +47,9 @@ static inline uint32 GetTypeHash(const FGraphicsPipelineStateInitializer& Initia
 }
 
 
-static FCriticalSection GComputeLock;
+static FRWLock GComputeLock;
 static TMap<FRHIComputeShader*, FComputePipelineState*> GComputePipelines;
-static FCriticalSection GGraphicsLock;
+static FRWLock GGraphicsLock;
 static TMap <FGraphicsPipelineStateInitializer, FGraphicsPipelineState*> GGraphicsPipelines;
 
 class FPipelineState
@@ -136,7 +136,7 @@ public:
 
 static bool IsAsyncCompilationAllowed(FRHICommandList& RHICmdList)
 {
-	return GCVarAsyncPipelineCompile.GetValueOnAnyThread() && !RHICmdList.Bypass() && GRHIThread;
+	return GCVarAsyncPipelineCompile.GetValueOnAnyThread() && !RHICmdList.Bypass() && IsRunningRHIInSeparateThread();
 }
 
 FComputePipelineState* GetAndOrCreateComputePipelineState(FRHICommandList& RHICmdList, FRHIComputeShader* ComputeShader)
@@ -145,7 +145,7 @@ FComputePipelineState* GetAndOrCreateComputePipelineState(FRHICommandList& RHICm
 
 	{
 		// Should be hitting this case more often once the cache is hot
-		FScopeLock ScopeLock(&GComputeLock);
+		FRWScopeLock ScopeLock(GComputeLock, SLT_ReadOnly);
 
 		FComputePipelineState** Found = GComputePipelines.Find(ComputeShader);
 		if (Found)
@@ -157,22 +157,36 @@ FComputePipelineState* GetAndOrCreateComputePipelineState(FRHICommandList& RHICm
 			}
 			return *Found;
 		}
-
-		FComputePipelineState* PipelineState = new FComputePipelineState(ComputeShader);
-
-		if (IsAsyncCompilationAllowed(RHICmdList))
+		
+		ScopeLock.RaiseLockToWrite();
+		Found = GComputePipelines.Find(ComputeShader);
+		if (Found)
 		{
-			PipelineState->CompletionEvent = TGraphTask<FCompilePipelineStateTask>::CreateTask().ConstructAndDispatchWhenReady(PipelineState);
-			RHICmdList.QueueAsyncPipelineStateCompile(PipelineState->CompletionEvent);
+			if (IsAsyncCompilationAllowed(RHICmdList))
+			{
+				FGraphEventRef& CompletionEvent = (*Found)->CompletionEvent;
+				RHICmdList.QueueAsyncPipelineStateCompile(CompletionEvent);
+			}
+			return *Found;
 		}
 		else
 		{
-			PipelineState->RHIPipeline = RHICreateComputePipelineState(PipelineState->ComputeShader);
-		}
+			FComputePipelineState* PipelineState = new FComputePipelineState(ComputeShader);
 
-		GComputePipelines.Add(ComputeShader, PipelineState);
+			if (IsAsyncCompilationAllowed(RHICmdList))
+			{
+				PipelineState->CompletionEvent = TGraphTask<FCompilePipelineStateTask>::CreateTask().ConstructAndDispatchWhenReady(PipelineState);
+				RHICmdList.QueueAsyncPipelineStateCompile(PipelineState->CompletionEvent);
+			}
+			else
+			{
+				PipelineState->RHIPipeline = RHICreateComputePipelineState(PipelineState->ComputeShader);
+			}
 
-		return PipelineState;
+			GComputePipelines.Add(ComputeShader, PipelineState);
+			
+			return PipelineState;
+		}
 	}
 
 	return nullptr;
@@ -240,7 +254,7 @@ FGraphicsPipelineState* GetAndOrCreateGraphicsPipelineState(FRHICommandList& RHI
 
 	{
 		// Should be hitting this case more often once the cache is hot
-		FScopeLock ScopeLock(&GGraphicsLock);
+		FRWScopeLock ScopeLock(GGraphicsLock, SLT_ReadOnly);
 
 		FGraphicsPipelineState** Found = GGraphicsPipelines.Find(*Initializer);
 		if (Found)
@@ -253,21 +267,34 @@ FGraphicsPipelineState* GetAndOrCreateGraphicsPipelineState(FRHICommandList& RHI
 			return *Found;
 		}
 
-		FGraphicsPipelineState* PipelineState = new FGraphicsPipelineState(*Initializer);
-
-		if (IsAsyncCompilationAllowed(RHICmdList))
+		ScopeLock.RaiseLockToWrite();
+		Found = GGraphicsPipelines.Find(*Initializer);
+		if (Found)
 		{
-			PipelineState->CompletionEvent = TGraphTask<FCompilePipelineStateTask>::CreateTask().ConstructAndDispatchWhenReady(PipelineState);
-			RHICmdList.QueueAsyncPipelineStateCompile(PipelineState->CompletionEvent);
+			if (IsAsyncCompilationAllowed(RHICmdList))
+			{
+				FGraphEventRef& CompletionEvent = (*Found)->CompletionEvent;
+				RHICmdList.QueueAsyncPipelineStateCompile(CompletionEvent);
+			}
+			return *Found;
 		}
 		else
 		{
-			PipelineState->RHIPipeline = RHICreateGraphicsPipelineState(*Initializer);
-		}
+			FGraphicsPipelineState* PipelineState = new FGraphicsPipelineState(*Initializer);
 
-		GGraphicsPipelines.Add(*Initializer, PipelineState);
+			if (IsAsyncCompilationAllowed(RHICmdList))
+			{
+				PipelineState->CompletionEvent = TGraphTask<FCompilePipelineStateTask>::CreateTask().ConstructAndDispatchWhenReady(PipelineState);
+				RHICmdList.QueueAsyncPipelineStateCompile(PipelineState->CompletionEvent);
+			}
+			else
+			{
+				PipelineState->RHIPipeline = RHICreateGraphicsPipelineState(*Initializer);
+			}
 
-		return PipelineState;
+			GGraphicsPipelines.Add(*Initializer, PipelineState);
+			return PipelineState;
+		}
 	}
 
 	return nullptr;
diff --git a/Engine/Source/Runtime/RHI/Private/RHI.cpp b/Engine/Source/Runtime/RHI/Private/RHI.cpp
--- a/Engine/Source/Runtime/RHI/Private/RHI.cpp
+++ b/Engine/Source/Runtime/RHI/Private/RHI.cpp
@@ -91,7 +91,8 @@ const FClearValueBinding FClearValueBinding::DepthZero(0.0f, 0);
 const FClearValueBinding FClearValueBinding::DepthNear((float)ERHIZBuffer::NearPlane, 0);
 const FClearValueBinding FClearValueBinding::DepthFar((float)ERHIZBuffer::FarPlane, 0);
 const FClearValueBinding FClearValueBinding::Green(FLinearColor(0.0f, 1.0f, 0.0f, 1.0f));
-const FClearValueBinding FClearValueBinding::MidGray(FLinearColor(0.5f, 0.5f, 0.5f, 1.0f));
+// Note: this is used as the default normal for DBuffer decals.  It must decode to a value of 0 in DecodeDBufferData.
+const FClearValueBinding FClearValueBinding::DefaultNormal8Bit(FLinearColor(128.0f / 255.0f, 128.0f / 255.0f, 128.0f / 255.0f, 1.0f));
 
 TLockFreePointerListUnordered<FRHIResource, PLATFORM_CACHE_LINE_SIZE> FRHIResource::PendingDeletes;
 FRHIResource* FRHIResource::CurrentlyDeleting = nullptr;
@@ -291,6 +292,7 @@ bool GRHISupportsInstancing = true;
 bool GRHISupportsFirstInstance = false;
 bool GRHIRequiresEarlyBackBufferRenderTarget = true;
 bool GRHISupportsRHIThread = false;
+bool GRHISupportsRHIOnTaskThread = false;
 bool GRHISupportsParallelRHIExecute = false;
 bool GSupportsHDR32bppEncodeModeIntrinsic = false;
 bool GSupportsParallelOcclusionQueries = false;
