Added bSuppressWarningOnEmptyMorphTargetAnimation ini file option, defaults to false, set to true by Orion's DefaultEngine.ini

diff --git a/Engine/Config/BaseEngine.ini b/Engine/Config/BaseEngine.ini
--- a/Engine/Config/BaseEngine.ini
+++ b/Engine/Config/BaseEngine.ini
@@ -205,6 +205,8 @@ bEnableOnScreenDebugMessages=true
 DurationOfErrorsAndWarningsOnHUD=0
 NearClipPlane=10.0
 bUseStreamingPause=false
+bSuppressWarningOnEmptyMorphTargetAnimation=false
+
 
 ; Matinee redirects (starts-with match and replace)
 +MatineeTrackRedirects=(TargetClassName="/Script/Engine.CameraActor",OldFieldName="PostProcessSettings.EyeAdaptationLowPercent",NewFieldName="CameraComponent.PostProcessSettings.AutoExposureLowPercent")
diff --git a/Engine/Source/Runtime/Engine/Private/SkeletalMesh.cpp b/Engine/Source/Runtime/Engine/Private/SkeletalMesh.cpp
--- a/Engine/Source/Runtime/Engine/Private/SkeletalMesh.cpp
+++ b/Engine/Source/Runtime/Engine/Private/SkeletalMesh.cpp
@@ -42,6 +42,7 @@
 #include "Engine/Engine.h"
 #include "Animation/NodeMappingContainer.h"
 #include "GPUSkinCache.h"
+#include "Misc/ConfigCacheIni.h"
 
 #if WITH_EDITOR
 #include "MeshUtilities.h"
@@ -1592,7 +1593,9 @@ void FStaticLODModel::InitResources(bool bNeedsVertexColors, int32 LODIndex, TAr
 			}
 			else
 			{
-				UE_LOG(LogSkeletalMesh, Warning, TEXT("Empty MorphTarget Animation found in %s at index %i"), *MorphTarget->GetName(), AnimIdx);
+				bool bSuppressWarningOnEmptyMorphTargetAnimation = false;
+				GConfig->GetBool(TEXT("/Script/Engine.Engine"), TEXT("bSuppressWarningOnEmptyMorphTargetAnimation"), bSuppressWarningOnEmptyMorphTargetAnimation, GEngineIni);
+				UE_CLOG(!bSuppressWarningOnEmptyMorphTargetAnimation, LogSkeletalMesh, Warning, TEXT("Empty MorphTarget Animation found in %s at index %i"), *MorphTarget->GetName(), AnimIdx);
 			}
 			MorphTargetVertexInfoBuffers.WorkItemsPerMorph.Add(MorphTargetSize);
 			MorphTargetVertexInfoBuffers.MaximumValuePerMorph.Add(FVector4(MaximumValues[0], MaximumValues[1], MaximumValues[2], MaximumValues[3]));
