Copying

//UE4/UT-Staging/...

to //UE4/Main/..

Source: //UE4/UT-Staging @ 3473309

==========================
MAJOR FEATURES + CHANGES
==========================

Peter.Knepley

Expose HRTF functions for UT consumption

diff --git a/Engine/Source/Runtime/Engine/Private/AudioDevice.cpp b/Engine/Source/Runtime/Engine/Private/AudioDevice.cpp
--- a/Engine/Source/Runtime/Engine/Private/AudioDevice.cpp
+++ b/Engine/Source/Runtime/Engine/Private/AudioDevice.cpp
@@ -932,27 +932,14 @@ bool FAudioDevice::HandleResetSoundStateCommand(const TCHAR* Cmd, FOutputDevice&
 
 bool FAudioDevice::HandleToggleSpatializationExtensionCommand(const TCHAR* Cmd, FOutputDevice& Ar)
 {
-	FAudioThread::SuspendAudioThread();
-
-	bSpatializationInterfaceEnabled = !bSpatializationInterfaceEnabled;
-
-	FAudioThread::ResumeAudioThread();
+	SetSpatializationInterfaceEnabled(!bSpatializationInterfaceEnabled);
 
 	return true;
 }
 
 bool FAudioDevice::HandleEnableHRTFForAllCommand(const TCHAR* Cmd, FOutputDevice& Ar)
 {
-	const bool bNewHRTFEnabledForAll = !bHRTFEnabledForAll_OnGameThread;
-
-	bHRTFEnabledForAll_OnGameThread = bNewHRTFEnabledForAll;
-
-	FAudioDevice* AudioDevice = this;
-	FAudioThread::RunCommandOnAudioThread([AudioDevice, bNewHRTFEnabledForAll]()
-	{
-		AudioDevice->bHRTFEnabledForAll = bNewHRTFEnabledForAll;
-
-	});
+	SetHRTFEnabledForAll(!bHRTFEnabledForAll_OnGameThread);
 
 	return true;
 }
diff --git a/Engine/Source/Runtime/Engine/Public/AudioDevice.h b/Engine/Source/Runtime/Engine/Public/AudioDevice.h
--- a/Engine/Source/Runtime/Engine/Public/AudioDevice.h
+++ b/Engine/Source/Runtime/Engine/Public/AudioDevice.h
@@ -920,6 +920,29 @@ public:
 	/** Whether or not HRTF spatialization is enabled for all. */
 	bool IsHRTFEnabledForAll() const;
 
+	void SetHRTFEnabledForAll(bool InbHRTFEnabledForAll)
+	{
+		const bool bNewHRTFEnabledForAll = InbHRTFEnabledForAll;
+
+		bHRTFEnabledForAll_OnGameThread = bNewHRTFEnabledForAll;
+
+		FAudioDevice* AudioDevice = this;
+		FAudioThread::RunCommandOnAudioThread([AudioDevice, bNewHRTFEnabledForAll]()
+		{
+			AudioDevice->bHRTFEnabledForAll = bNewHRTFEnabledForAll;
+
+		});
+	}
+
+	void SetSpatializationInterfaceEnabled(bool InbSpatializationInterfaceEnabled)
+	{
+		FAudioThread::SuspendAudioThread();
+
+		bSpatializationInterfaceEnabled = InbSpatializationInterfaceEnabled;
+
+		FAudioThread::ResumeAudioThread();
+	}
+
 	bool IsAudioDeviceMuted() const;
 
 	void SetDeviceMuted(bool bMuted);
