Merging //UE4/Dev-Main to Release-Staging-4.16 (//UE4/Release-Staging-4.16)

diff --git a/Engine/Config/BaseEngine.ini b/Engine/Config/BaseEngine.ini
--- a/Engine/Config/BaseEngine.ini
+++ b/Engine/Config/BaseEngine.ini
@@ -205,6 +205,8 @@ bEnableOnScreenDebugMessages=true
 DurationOfErrorsAndWarningsOnHUD=0
 NearClipPlane=10.0
 bUseStreamingPause=false
+bSuppressWarningOnEmptyMorphTargetAnimation=false
+
 
 ; Matinee redirects (starts-with match and replace)
 +MatineeTrackRedirects=(TargetClassName="/Script/Engine.CameraActor",OldFieldName="PostProcessSettings.EyeAdaptationLowPercent",NewFieldName="CameraComponent.PostProcessSettings.AutoExposureLowPercent")
diff --git a/Engine/Source/Programs/UnrealBuildTool/System/XGE.cs b/Engine/Source/Programs/UnrealBuildTool/System/XGE.cs
--- a/Engine/Source/Programs/UnrealBuildTool/System/XGE.cs
+++ b/Engine/Source/Programs/UnrealBuildTool/System/XGE.cs
@@ -32,7 +32,7 @@ namespace UnrealBuildTool
 		/// When enabled, XGE will stop compiling targets after a compile error occurs.  Recommended, as it saves computing resources for others.
 		/// </summary>
 		[XmlConfigFile(Category = "BuildConfiguration")]
-		bool bStopXGECompilationAfterErrors = false;
+		bool bStopXGECompilationAfterErrors = true;
 
 		private const string ProgressMarkupPrefix = "@action";
 
diff --git a/Engine/Source/Runtime/Engine/Private/AudioDevice.cpp b/Engine/Source/Runtime/Engine/Private/AudioDevice.cpp
--- a/Engine/Source/Runtime/Engine/Private/AudioDevice.cpp
+++ b/Engine/Source/Runtime/Engine/Private/AudioDevice.cpp
@@ -932,27 +932,14 @@ bool FAudioDevice::HandleResetSoundStateCommand(const TCHAR* Cmd, FOutputDevice&
 
 bool FAudioDevice::HandleToggleSpatializationExtensionCommand(const TCHAR* Cmd, FOutputDevice& Ar)
 {
-	FAudioThread::SuspendAudioThread();
-
-	bSpatializationInterfaceEnabled = !bSpatializationInterfaceEnabled;
-
-	FAudioThread::ResumeAudioThread();
+	SetSpatializationInterfaceEnabled(!bSpatializationInterfaceEnabled);
 
 	return true;
 }
 
 bool FAudioDevice::HandleEnableHRTFForAllCommand(const TCHAR* Cmd, FOutputDevice& Ar)
 {
-	const bool bNewHRTFEnabledForAll = !bHRTFEnabledForAll_OnGameThread;
-
-	bHRTFEnabledForAll_OnGameThread = bNewHRTFEnabledForAll;
-
-	FAudioDevice* AudioDevice = this;
-	FAudioThread::RunCommandOnAudioThread([AudioDevice, bNewHRTFEnabledForAll]()
-	{
-		AudioDevice->bHRTFEnabledForAll = bNewHRTFEnabledForAll;
-
-	});
+	SetHRTFEnabledForAll(!bHRTFEnabledForAll_OnGameThread);
 
 	return true;
 }
diff --git a/Engine/Source/Runtime/Engine/Private/SkeletalMesh.cpp b/Engine/Source/Runtime/Engine/Private/SkeletalMesh.cpp
--- a/Engine/Source/Runtime/Engine/Private/SkeletalMesh.cpp
+++ b/Engine/Source/Runtime/Engine/Private/SkeletalMesh.cpp
@@ -42,6 +42,7 @@
 #include "Engine/Engine.h"
 #include "Animation/NodeMappingContainer.h"
 #include "GPUSkinCache.h"
+#include "Misc/ConfigCacheIni.h"
 
 #if WITH_EDITOR
 #include "MeshUtilities.h"
@@ -1592,7 +1593,9 @@ void FStaticLODModel::InitResources(bool bNeedsVertexColors, int32 LODIndex, TAr
 			}
 			else
 			{
-				UE_LOG(LogSkeletalMesh, Warning, TEXT("Empty MorphTarget Animation found in %s at index %i"), *MorphTarget->GetName(), AnimIdx);
+				bool bSuppressWarningOnEmptyMorphTargetAnimation = false;
+				GConfig->GetBool(TEXT("/Script/Engine.Engine"), TEXT("bSuppressWarningOnEmptyMorphTargetAnimation"), bSuppressWarningOnEmptyMorphTargetAnimation, GEngineIni);
+				UE_CLOG(!bSuppressWarningOnEmptyMorphTargetAnimation, LogSkeletalMesh, Warning, TEXT("Empty MorphTarget Animation found in %s at index %i"), *MorphTarget->GetName(), AnimIdx);
 			}
 			MorphTargetVertexInfoBuffers.WorkItemsPerMorph.Add(MorphTargetSize);
 			MorphTargetVertexInfoBuffers.MaximumValuePerMorph.Add(FVector4(MaximumValues[0], MaximumValues[1], MaximumValues[2], MaximumValues[3]));
diff --git a/Engine/Source/Runtime/Engine/Public/AudioDevice.h b/Engine/Source/Runtime/Engine/Public/AudioDevice.h
--- a/Engine/Source/Runtime/Engine/Public/AudioDevice.h
+++ b/Engine/Source/Runtime/Engine/Public/AudioDevice.h
@@ -920,6 +920,29 @@ public:
 	/** Whether or not HRTF spatialization is enabled for all. */
 	bool IsHRTFEnabledForAll() const;
 
+	void SetHRTFEnabledForAll(bool InbHRTFEnabledForAll)
+	{
+		const bool bNewHRTFEnabledForAll = InbHRTFEnabledForAll;
+
+		bHRTFEnabledForAll_OnGameThread = bNewHRTFEnabledForAll;
+
+		FAudioDevice* AudioDevice = this;
+		FAudioThread::RunCommandOnAudioThread([AudioDevice, bNewHRTFEnabledForAll]()
+		{
+			AudioDevice->bHRTFEnabledForAll = bNewHRTFEnabledForAll;
+
+		});
+	}
+
+	void SetSpatializationInterfaceEnabled(bool InbSpatializationInterfaceEnabled)
+	{
+		FAudioThread::SuspendAudioThread();
+
+		bSpatializationInterfaceEnabled = InbSpatializationInterfaceEnabled;
+
+		FAudioThread::ResumeAudioThread();
+	}
+
 	bool IsAudioDeviceMuted() const;
 
 	void SetDeviceMuted(bool bMuted);
