Fixing inconsistencies between function signatures and adding code comment directives to help Doxygen work correctly. Also restored container.html to the engine.

diff --git a/Engine/Plugins/2D/Paper2D/Source/Paper2D/Private/PaperFlipbookComponent.cpp b/Engine/Plugins/2D/Paper2D/Source/Paper2D/Private/PaperFlipbookComponent.cpp
--- a/Engine/Plugins/2D/Paper2D/Source/Paper2D/Private/PaperFlipbookComponent.cpp
+++ b/Engine/Plugins/2D/Paper2D/Source/Paper2D/Private/PaperFlipbookComponent.cpp
@@ -313,6 +313,8 @@ void UPaperFlipbookComponent::TickFlipbook(float DeltaTime)
 	}
 }
 
+/// @cond DOXYGEN_WARNINGS
+
 void UPaperFlipbookComponent::GetLifetimeReplicatedProps(TArray< FLifetimeProperty > & OutLifetimeProps) const
 {
 	Super::GetLifetimeReplicatedProps(OutLifetimeProps);
@@ -320,6 +322,8 @@ void UPaperFlipbookComponent::GetLifetimeReplicatedProps(TArray< FLifetimeProper
 	DOREPLIFETIME(UPaperFlipbookComponent, SourceFlipbook);
 }
 
+/// @endcond
+
 void UPaperFlipbookComponent::OnRep_SourceFlipbook(class UPaperFlipbook* OldFlipbook)
 {
 	if (OldFlipbook != SourceFlipbook)
diff --git a/Engine/Plugins/Online/OnlineSubsystem/Source/Private/OnlineSessionSettings.cpp b/Engine/Plugins/Online/OnlineSubsystem/Source/Private/OnlineSessionSettings.cpp
--- a/Engine/Plugins/Online/OnlineSubsystem/Source/Private/OnlineSessionSettings.cpp
+++ b/Engine/Plugins/Online/OnlineSubsystem/Source/Private/OnlineSessionSettings.cpp
@@ -203,6 +203,8 @@ void FOnlineSearchSettings::Set(FName Key, const ValueType& Value, EOnlineCompar
 	}
 }
 
+/// @cond DOXYGEN_WARNINGS
+
 /** Explicit instantiation of supported types to Set template above */
 template ONLINESUBSYSTEM_API void FOnlineSearchSettings::Set<int32>(FName Key, const int32& Value, EOnlineComparisonOp::Type InType, int32 InID);
 template ONLINESUBSYSTEM_API void FOnlineSearchSettings::Set<float>(FName Key, const float& Value, EOnlineComparisonOp::Type InType, int32 InID);
@@ -212,6 +214,8 @@ template ONLINESUBSYSTEM_API void FOnlineSearchSettings::Set<FString>(FName Key,
 template ONLINESUBSYSTEM_API void FOnlineSearchSettings::Set< TArray<uint8> >(FName Key, const TArray<uint8>& Value, EOnlineComparisonOp::Type InType, int32 InID);
 template ONLINESUBSYSTEM_API void FOnlineSearchSettings::Set<bool>(FName Key, const bool& Value, EOnlineComparisonOp::Type InType, int32 InID);
 
+/// @endcond
+
 template<typename ValueType> 
 void FOnlineSearchSettings::Set(FName Key, const ValueType& Value, EOnlineComparisonOp::Type InType)
 {
@@ -227,6 +231,8 @@ void FOnlineSearchSettings::Set(FName Key, const ValueType& Value, EOnlineCompar
 	}
 }
 
+/// @cond DOXYGEN_WARNINGS
+
 /** Explicit instantiation of supported types to Set template above */
 template ONLINESUBSYSTEM_API void FOnlineSearchSettings::Set<int32>(FName Key, const int32& Value, EOnlineComparisonOp::Type InType);
 template ONLINESUBSYSTEM_API void FOnlineSearchSettings::Set<float>(FName Key, const float& Value, EOnlineComparisonOp::Type InType);
@@ -236,6 +242,8 @@ template ONLINESUBSYSTEM_API void FOnlineSearchSettings::Set<FString>(FName Key,
 template ONLINESUBSYSTEM_API void FOnlineSearchSettings::Set< TArray<uint8> >(FName Key, const TArray<uint8>& Value, EOnlineComparisonOp::Type InType);
 template ONLINESUBSYSTEM_API void FOnlineSearchSettings::Set<bool>(FName Key, const bool& Value, EOnlineComparisonOp::Type InType);
 
+/// @endcond
+
 template<typename ValueType> 
 bool FOnlineSearchSettings::Get(FName Key, ValueType& Value) const
 {
@@ -249,6 +257,8 @@ bool FOnlineSearchSettings::Get(FName Key, ValueType& Value) const
 	return false;
 }
 
+/// @cond DOXYGEN_WARNINGS
+
 /** Explicit instantiation of supported types to Get template above */
 template ONLINESUBSYSTEM_API bool FOnlineSearchSettings::Get<int32>(FName Key, int32& Value) const;
 template ONLINESUBSYSTEM_API bool FOnlineSearchSettings::Get<float>(FName Key, float& Value) const;
@@ -258,6 +268,8 @@ template ONLINESUBSYSTEM_API bool FOnlineSearchSettings::Get<FString>(FName Key,
 template ONLINESUBSYSTEM_API bool FOnlineSearchSettings::Get< TArray<uint8> >(FName Key, TArray<uint8>& Value) const;
 template ONLINESUBSYSTEM_API bool FOnlineSearchSettings::Get<bool>(FName Key, bool& Value) const;
 
+/// @endcond
+
 EOnlineComparisonOp::Type FOnlineSearchSettings::GetComparisonOp(FName Key) const
 {
 	const FOnlineSessionSearchParam* SearchParam = SearchParams.Find(Key);
diff --git a/Engine/Plugins/Online/OnlineSubsystemNull/Source/Private/OnlineSessionInterfaceNull.h b/Engine/Plugins/Online/OnlineSubsystemNull/Source/Private/OnlineSessionInterfaceNull.h
--- a/Engine/Plugins/Online/OnlineSubsystemNull/Source/Private/OnlineSessionInterfaceNull.h
+++ b/Engine/Plugins/Online/OnlineSubsystemNull/Source/Private/OnlineSessionInterfaceNull.h
@@ -309,7 +309,7 @@ public:
 	virtual bool SendSessionInviteToFriends(int32 LocalUserNum, FName SessionName, const TArray< TSharedRef<const FUniqueNetId> >& Friends) override;
 	virtual bool SendSessionInviteToFriends(const FUniqueNetId& LocalUserId, FName SessionName, const TArray< TSharedRef<const FUniqueNetId> >& Friends) override;
 	virtual bool GetResolvedConnectString(FName SessionName, FString& ConnectInfo, FName PortType) override;
-	virtual bool GetResolvedConnectString(const class FOnlineSessionSearchResult& SearchResult, FName PortType, FString& ConnectInfo) override;
+	virtual bool GetResolvedConnectString(const FOnlineSessionSearchResult& SearchResult, FName PortType, FString& ConnectInfo) override;
 	virtual FOnlineSessionSettings* GetSessionSettings(FName SessionName) override;
 	virtual bool RegisterPlayer(FName SessionName, const FUniqueNetId& PlayerId, bool bWasInvited) override;
 	virtual bool RegisterPlayers(FName SessionName, const TArray< TSharedRef<const FUniqueNetId> >& Players, bool bWasInvited = false) override;
diff --git a/Engine/Plugins/Online/OnlineSubsystemSteam/Source/Private/OnlineSessionInterfaceSteam.h b/Engine/Plugins/Online/OnlineSubsystemSteam/Source/Private/OnlineSessionInterfaceSteam.h
--- a/Engine/Plugins/Online/OnlineSubsystemSteam/Source/Private/OnlineSessionInterfaceSteam.h
+++ b/Engine/Plugins/Online/OnlineSubsystemSteam/Source/Private/OnlineSessionInterfaceSteam.h
@@ -518,7 +518,7 @@ public:
 	virtual bool SendSessionInviteToFriends(int32 LocalUserNum, FName SessionName, const TArray< TSharedRef<const FUniqueNetId> >& Friends) override;
 	virtual bool SendSessionInviteToFriends(const FUniqueNetId& LocalUserId, FName SessionName, const TArray< TSharedRef<const FUniqueNetId> >& Friends) override;
 	virtual bool GetResolvedConnectString(FName SessionName, FString& ConnectInfo, FName PortType) override;
-	virtual bool GetResolvedConnectString(const class FOnlineSessionSearchResult& SearchResult, FName PortType, FString& ConnectInfo) override;
+	virtual bool GetResolvedConnectString(const FOnlineSessionSearchResult& SearchResult, FName PortType, FString& ConnectInfo) override;
 	virtual FOnlineSessionSettings* GetSessionSettings(FName SessionName) override;
 	virtual bool RegisterPlayer(FName SessionName, const FUniqueNetId& PlayerId, bool bWasInvited) override;
 	virtual bool RegisterPlayers(FName SessionName, const TArray< TSharedRef<const FUniqueNetId> >& Players, bool bWasInvited = false) override;
diff --git a/Engine/Plugins/Online/OnlineSubsystemUtils/Source/OnlineSubsystemUtils/Private/OnlineBeaconClient.cpp b/Engine/Plugins/Online/OnlineSubsystemUtils/Source/OnlineSubsystemUtils/Private/OnlineBeaconClient.cpp
--- a/Engine/Plugins/Online/OnlineSubsystemUtils/Source/OnlineSubsystemUtils/Private/OnlineBeaconClient.cpp
+++ b/Engine/Plugins/Online/OnlineSubsystemUtils/Source/OnlineSubsystemUtils/Private/OnlineBeaconClient.cpp
@@ -141,6 +141,9 @@ void AOnlineBeaconClient::OnFailure()
 	Super::OnFailure();
 }
 
+
+/// @cond DOXYGEN_WARNINGS
+
 void AOnlineBeaconClient::ClientOnConnected_Implementation()
 {
 	SetConnectionState(EBeaconConnectionState::Open);
@@ -157,6 +160,8 @@ void AOnlineBeaconClient::ClientOnConnected_Implementation()
 	OnConnected();
 }
 
+/// @endcond
+
 bool AOnlineBeaconClient::UseShortConnectTimeout() const
 {
 	return ConnectionState == EBeaconConnectionState::Open;
diff --git a/Engine/Plugins/Online/OnlineSubsystemUtils/Source/OnlineSubsystemUtils/Private/PartyBeaconClient.cpp b/Engine/Plugins/Online/OnlineSubsystemUtils/Source/OnlineSubsystemUtils/Private/PartyBeaconClient.cpp
--- a/Engine/Plugins/Online/OnlineSubsystemUtils/Source/OnlineSubsystemUtils/Private/PartyBeaconClient.cpp
+++ b/Engine/Plugins/Online/OnlineSubsystemUtils/Source/OnlineSubsystemUtils/Private/PartyBeaconClient.cpp
@@ -317,6 +317,8 @@ void APartyBeaconClient::OnFailure()
 	Super::OnFailure();
 }
 
+/// @cond DOXYGEN_WARNINGS
+
 bool APartyBeaconClient::ServerReservationRequest_Validate(const FString& SessionId, const FPartyReservation& Reservation)
 {
 	return !SessionId.IsEmpty() && Reservation.PartyLeader.IsValid() && Reservation.PartyMembers.Num() > 0;
@@ -397,6 +399,8 @@ void APartyBeaconClient::ClientReservationResponse_Implementation(EPartyReservat
 	}
 }
 
+/// @endcond
+
 void APartyBeaconClient::ProcessReservationResponse(EPartyReservationResult::Type ReservationResponse)
 {
 	if (!bCancelReservation)
@@ -412,6 +416,8 @@ void APartyBeaconClient::ProcessReservationResponse(EPartyReservationResult::Typ
 	}
 }
 
+/// @cond DOXYGEN_WARNINGS
+
 void APartyBeaconClient::ClientCancelReservationResponse_Implementation(EPartyReservationResult::Type ReservationResponse)
 {
 	ensure(bCancelReservation);
@@ -441,6 +447,8 @@ void APartyBeaconClient::ClientCancelReservationResponse_Implementation(EPartyRe
 	}
 }
 
+/// @endcond
+
 void APartyBeaconClient::ProcessCancelReservationResponse(EPartyReservationResult::Type ReservationResponse)
 {
 	ensure(ReservationResponse == EPartyReservationResult::ReservationRequestCanceled || ReservationResponse == EPartyReservationResult::ReservationNotFound);
@@ -448,6 +456,8 @@ void APartyBeaconClient::ProcessCancelReservationResponse(EPartyReservationResul
 	OnCancelledComplete();
 }
 
+/// @cond DOXYGEN_WARNINGS
+
 void APartyBeaconClient::ClientSendReservationUpdates_Implementation(int32 NumRemainingReservations)
 {
 	if (!bCancelReservation)
@@ -476,12 +486,16 @@ void APartyBeaconClient::ClientSendReservationUpdates_Implementation(int32 NumRe
 	}
 }
 
+/// @endcond
+
 void APartyBeaconClient::ProcessReservationUpdate(int32 NumRemainingReservations)
 {
 	UE_LOG(LogBeacon, Verbose, TEXT("Party beacon reservations remaining %d"), NumRemainingReservations);
 	ReservationCountUpdate.ExecuteIfBound(NumRemainingReservations);
 }
 
+/// @cond DOXYGEN_WARNINGS
+
 void APartyBeaconClient::ClientSendReservationFull_Implementation()
 {
 	if (!bCancelReservation)
@@ -510,6 +524,8 @@ void APartyBeaconClient::ClientSendReservationFull_Implementation()
 	}
 }
 
+/// @endcond
+
 void APartyBeaconClient::ProcessReservationFull()
 {
 	UE_LOG(LogBeacon, Verbose, TEXT("Party beacon reservations full"));
diff --git a/Engine/Plugins/Online/OnlineSubsystemUtils/Source/OnlineSubsystemUtils/Private/TestBeaconClient.cpp b/Engine/Plugins/Online/OnlineSubsystemUtils/Source/OnlineSubsystemUtils/Private/TestBeaconClient.cpp
--- a/Engine/Plugins/Online/OnlineSubsystemUtils/Source/OnlineSubsystemUtils/Private/TestBeaconClient.cpp
+++ b/Engine/Plugins/Online/OnlineSubsystemUtils/Source/OnlineSubsystemUtils/Private/TestBeaconClient.cpp
@@ -15,6 +15,8 @@ void ATestBeaconClient::OnFailure()
 	Super::OnFailure();
 }
 
+/// @cond DOXYGEN_WARNINGS
+
 void ATestBeaconClient::ClientPing_Implementation()
 {
 #if !UE_BUILD_SHIPPING
@@ -39,3 +41,5 @@ void ATestBeaconClient::ServerPong_Implementation()
 	ClientPing();
 #endif
 }
+
+/// @endcond
\ No newline at end of file
diff --git a/Engine/Plugins/Runtime/OculusRift/Source/OculusRift/Private/OculusRiftHMD.h b/Engine/Plugins/Runtime/OculusRift/Source/OculusRift/Private/OculusRiftHMD.h
--- a/Engine/Plugins/Runtime/OculusRift/Source/OculusRift/Private/OculusRiftHMD.h
+++ b/Engine/Plugins/Runtime/OculusRift/Source/OculusRift/Private/OculusRiftHMD.h
@@ -445,7 +445,7 @@ private:
 	FAutoConsoleCommand CubemapCommand;
 #endif // !UE_BUILD_SHIPPING
 	// Exec handler that aliases old deprecated Oculus Rift console commands to the new ones.
-	virtual bool Exec(class UWorld* InWorld, const TCHAR* Cmd, FOutputDevice& Ar) override;
+	virtual bool Exec(UWorld* InWorld, const TCHAR* Cmd, FOutputDevice& Ar) override;
 };
 
 
diff --git a/Engine/Plugins/Runtime/OculusRift/Source/OculusRift/Private/OculusRiftHMDConfig.cpp b/Engine/Plugins/Runtime/OculusRift/Source/OculusRift/Private/OculusRiftHMDConfig.cpp
--- a/Engine/Plugins/Runtime/OculusRift/Source/OculusRift/Private/OculusRiftHMDConfig.cpp
+++ b/Engine/Plugins/Runtime/OculusRift/Source/OculusRift/Private/OculusRiftHMDConfig.cpp
@@ -4,7 +4,7 @@
 #include "OculusRiftHMD.h"
 #if !UE_BUILD_SHIPPING
 #include "SceneCubemapCapturer.h"
-static void CubemapCommandHandler(const TArray<FString>&, UWorld*, FOutputDevice&);
+static void CubemapCommandHandler(const TArray<FString>& Args, UWorld* World, FOutputDevice& Ar);
 #endif
 
 
@@ -24,6 +24,7 @@ static void CubemapCommandHandler(const TArray<FString>&, UWorld*, FOutputDevice
 	}\
 	Ar.Logf(ConsoleName TEXT(" = %s"), (FieldExpr) ? TEXT("On") : TEXT("Off"));
 
+/// @cond DOXYGEN_WARNINGS
 
 FOculusRiftConsoleCommands::FOculusRiftConsoleCommands(class FOculusRiftHMD* InHMDPtr)
 	: PixelDensityCommand(TEXT("vr.oculus.PixelDensity"),
@@ -71,6 +72,8 @@ FOculusRiftConsoleCommands::FOculusRiftConsoleCommands(class FOculusRiftHMD* InH
 {
 }
 
+/// @endcond
+
 // Translate between legacy mirror mode values to the new ones
 static EMirrorWindowMode TranslateMirrorMode(int32 ObsoleteMode)
 {
@@ -91,6 +94,8 @@ static EMirrorWindowMode TranslateMirrorMode(int32 ObsoleteMode)
 	}
 }
 
+/// @cond DOXYGEN_WARNINGS
+
 /**
 Clutch to support setting the r.ScreenPercentage and make the equivalent change to PixelDensity
 
@@ -283,6 +288,8 @@ void FOculusRiftHMD::GridCommandHandler(const TArray<FString>& Args, UWorld*, FO
 	BOOLEAN_COMMAND_HANDLER_BODY(TEXT("vr.oculus.Debug.bDrawGrid"), Settings->Flags.bDrawGrid)
 }
 
+/// @endcond
+
 const uint32 CaptureHeight = 2048;
 static void CubemapCommandHandler(const TArray<FString>& Args, UWorld* World, FOutputDevice& Ar)
 {
@@ -314,6 +321,8 @@ static void CubemapCommandHandler(const TArray<FString>& Args, UWorld* World, FO
 	CubemapCapturer->StartCapture(World, bCreateGearVRCubemap ? CaptureHeight / 2 : CaptureHeight);
 }
 
+/// @cond DOXYGEN_WARNINGS
+
 void FOculusRiftHMD::EnforceHeadTrackingCommandHandler(const TArray<FString>& Args, UWorld* World, FOutputDevice& Ar)
 {
 	bool bOldValue = Settings->Flags.bHeadTrackingEnforced;
@@ -569,4 +578,6 @@ void FOculusRiftHMD::SaveToIni()
 #endif // #if !UE_BUILD_SHIPPING
 }
 
+/// @endcond
+
 #endif
diff --git a/Engine/Plugins/Runtime/PhysXVehicles/Source/PhysXVehicles/Private/WheeledVehicleMovementComponent.cpp b/Engine/Plugins/Runtime/PhysXVehicles/Source/PhysXVehicles/Private/WheeledVehicleMovementComponent.cpp
--- a/Engine/Plugins/Runtime/PhysXVehicles/Source/PhysXVehicles/Private/WheeledVehicleMovementComponent.cpp
+++ b/Engine/Plugins/Runtime/PhysXVehicles/Source/PhysXVehicles/Private/WheeledVehicleMovementComponent.cpp
@@ -1024,6 +1024,8 @@ void UWheeledVehicleMovementComponent::UpdateState( float DeltaTime )
 	}
 }
 
+/// @cond DOXYGEN_WARNINGS
+
 bool UWheeledVehicleMovementComponent::ServerUpdateState_Validate(float InSteeringInput, float InThrottleInput, float InBrakeInput, float InHandbrakeInput, int32 InCurrentGear)
 {
 	return true;
@@ -1049,6 +1051,8 @@ void UWheeledVehicleMovementComponent::ServerUpdateState_Implementation(float In
 	ReplicatedState.CurrentGear = InCurrentGear;
 }
 
+/// @endcond
+
 float UWheeledVehicleMovementComponent::CalcSteeringInput()
 {
 	if (bUseRVOAvoidance)
@@ -1718,6 +1722,7 @@ void UWheeledVehicleMovementComponent::PostEditChangeProperty( FPropertyChangedE
 
 #endif // WITH_EDITOR
 
+/// @cond DOXYGEN_WARNINGS
 
 void UWheeledVehicleMovementComponent::GetLifetimeReplicatedProps( TArray< FLifetimeProperty > & OutLifetimeProps ) const
 {
@@ -1726,6 +1731,8 @@ void UWheeledVehicleMovementComponent::GetLifetimeReplicatedProps( TArray< FLife
 	DOREPLIFETIME( UWheeledVehicleMovementComponent, ReplicatedState );
 }
 
+/// @endcond
+
 void UWheeledVehicleMovementComponent::ComputeConstants()
 {
 	DragArea = ChassisWidth * ChassisHeight;
diff --git a/Engine/Source/Developer/GameplayDebugger/Private/GameplayDebuggerCategoryReplicator.cpp b/Engine/Source/Developer/GameplayDebugger/Private/GameplayDebuggerCategoryReplicator.cpp
--- a/Engine/Source/Developer/GameplayDebugger/Private/GameplayDebuggerCategoryReplicator.cpp
+++ b/Engine/Source/Developer/GameplayDebugger/Private/GameplayDebuggerCategoryReplicator.cpp
@@ -469,6 +469,8 @@ bool AGameplayDebuggerCategoryReplicator::IsNetRelevantFor(const AActor* RealVie
 	return (RealViewer == OwnerPC);
 }
 
+/// @cond DOXYGEN_WARNINGS
+
 void AGameplayDebuggerCategoryReplicator::GetLifetimeReplicatedProps(TArray<FLifetimeProperty>& OutLifetimeProps) const
 {
 	Super::GetLifetimeReplicatedProps(OutLifetimeProps);
@@ -529,6 +531,8 @@ void AGameplayDebuggerCategoryReplicator::ServerSendExtensionInputEvent_Implemen
 	SendExtensionInputEvent(ExtensionId, HandlerId);
 }
 
+/// @endcond
+
 void AGameplayDebuggerCategoryReplicator::OnReceivedDataPackPacket(int32 CategoryId, int32 DataPackId, const FGameplayDebuggerDataPack& DataPacket)
 {
 	if (Categories.IsValidIndex(CategoryId) && Categories[CategoryId]->ReplicatedDataPacks.IsValidIndex(DataPackId))
diff --git a/Engine/Source/Editor/BlueprintGraph/Private/BasicTokenParser.cpp b/Engine/Source/Editor/BlueprintGraph/Private/BasicTokenParser.cpp
--- a/Engine/Source/Editor/BlueprintGraph/Private/BasicTokenParser.cpp
+++ b/Engine/Source/Editor/BlueprintGraph/Private/BasicTokenParser.cpp
@@ -40,19 +40,19 @@ void FBasicToken::Clone(const FBasicToken& Other)
 }
 
 //------------------------------------------------------------------------------
-bool FBasicToken::Matches(TCHAR const* Str, ESearchCase::Type SearchCase) const
+bool FBasicToken::Matches(const TCHAR* Str, ESearchCase::Type SearchCase) const
 {
 	return (TokenType==TOKEN_Identifier || TokenType==TOKEN_Symbol) && ((SearchCase == ESearchCase::CaseSensitive) ? !FCString::Strcmp(Identifier, Str) : !FCString::Stricmp(Identifier, Str));
 }
 
 //------------------------------------------------------------------------------
-bool FBasicToken::Matches(FName const& Name) const
+bool FBasicToken::Matches(const FName& Name) const
 {
 	return TokenType==TOKEN_Identifier && TokenName==Name;
 }
 
 //------------------------------------------------------------------------------
-bool FBasicToken::StartsWith(TCHAR const* Str, bool bCaseSensitive) const
+bool FBasicToken::StartsWith(const TCHAR* Str, bool bCaseSensitive) const
 {
 	const int32 StrLength = FCString::Strlen(Str);
 	return (TokenType==TOKEN_Identifier || TokenType==TOKEN_Symbol) && (bCaseSensitive ? (!FCString::Strncmp(Identifier, Str, StrLength)) : (!FCString::Strnicmp(Identifier, Str, StrLength)));
@@ -284,7 +284,7 @@ namespace BasicTokenParserImpl
 }
 
 //------------------------------------------------------------------------------
-void FBasicTokenParser::ResetParser(TCHAR const* SourceBuffer, int32 StartingLineNumber)
+void FBasicTokenParser::ResetParser(const TCHAR* SourceBuffer, int32 StartingLineNumber)
 {
 	Input     = SourceBuffer;
 	InputLen  = FCString::Strlen(Input);
@@ -706,7 +706,7 @@ bool FBasicTokenParser::GetSymbol(FBasicToken& Token)
 }
 
 //------------------------------------------------------------------------------
-bool FBasicTokenParser::GetConstInt(int32& Result, TCHAR const* ErrorContext)
+bool FBasicTokenParser::GetConstInt(int32& Result, const TCHAR* ErrorContext)
 {
 	FBasicToken Token;
 	if (GetToken(Token))
@@ -910,7 +910,7 @@ bool FBasicTokenParser::MatchIdentifier(FName Match)
 }
 
 //------------------------------------------------------------------------------
-bool FBasicTokenParser::MatchIdentifier(TCHAR const* Match)
+bool FBasicTokenParser::MatchIdentifier(const TCHAR* Match)
 {
 	FBasicToken Token;
 	if (GetToken(Token))
@@ -953,7 +953,7 @@ bool FBasicTokenParser::PeekIdentifier(const TCHAR* Match)
 }
 
 //------------------------------------------------------------------------------
-bool FBasicTokenParser::MatchSymbol(TCHAR const* Match)
+bool FBasicTokenParser::MatchSymbol(const TCHAR* Match)
 {
 	FBasicToken Token;
 
@@ -973,7 +973,7 @@ bool FBasicTokenParser::MatchSymbol(TCHAR const* Match)
 }
 
 //------------------------------------------------------------------------------
-bool FBasicTokenParser::PeekSymbol(TCHAR const* Match)
+bool FBasicTokenParser::PeekSymbol(const TCHAR* Match)
 {
 	FBasicToken Token;
 	if (!GetToken(Token, true))
@@ -986,7 +986,7 @@ bool FBasicTokenParser::PeekSymbol(TCHAR const* Match)
 }
 
 //------------------------------------------------------------------------------
-bool FBasicTokenParser::RequireIdentifier(FName Match, TCHAR const* ErrorContext)
+bool FBasicTokenParser::RequireIdentifier(FName Match, const TCHAR* ErrorContext)
 {
 	if (!MatchIdentifier(Match))
 	{
@@ -997,7 +997,7 @@ bool FBasicTokenParser::RequireIdentifier(FName Match, TCHAR const* ErrorContext
 }
 
 //------------------------------------------------------------------------------
-bool FBasicTokenParser::RequireIdentifier(TCHAR const* Match, TCHAR const* ErrorContext)
+bool FBasicTokenParser::RequireIdentifier(const TCHAR* Match, const TCHAR* ErrorContext)
 {
 	if (!MatchIdentifier(Match))
 	{
@@ -1008,7 +1008,7 @@ bool FBasicTokenParser::RequireIdentifier(TCHAR const* Match, TCHAR const* Error
 }
 
 //------------------------------------------------------------------------------
-bool FBasicTokenParser::RequireSymbol(TCHAR const* Match, TCHAR const* ErrorContext)
+bool FBasicTokenParser::RequireSymbol(const TCHAR* Match, const TCHAR* ErrorContext)
 {
 	if (!MatchSymbol(Match))
 	{
diff --git a/Engine/Source/Editor/BlueprintGraph/Private/BasicTokenParser.h b/Engine/Source/Editor/BlueprintGraph/Private/BasicTokenParser.h
--- a/Engine/Source/Editor/BlueprintGraph/Private/BasicTokenParser.h
+++ b/Engine/Source/Editor/BlueprintGraph/Private/BasicTokenParser.h
@@ -44,7 +44,7 @@ public:
 	bool Matches(const TCHAR* Str, ESearchCase::Type SearchCase = ESearchCase::IgnoreCase) const;
 
 	/** Determines if this token has the specified string as a prefix. */
-	bool StartsWith( TCHAR const* Str, bool bCaseSensitive = false) const;
+	bool StartsWith(const TCHAR* Str, bool bCaseSensitive = false) const;
 
 	/** Determines if this token has a boolean ConstantType.*/
 	bool IsBool() const;
@@ -211,7 +211,7 @@ protected:
 	 *
 	 * @return True if a int was successfully processed, false otherwise.
 	 */
-	bool GetConstInt(int32& Result, TCHAR const* ErrorContext = NULL);
+	bool GetConstInt(int32& Result, const TCHAR* ErrorContext = NULL);
 
 	/** 
 	 * Rolls back the line number and moves the parsing pointer back to where 
@@ -262,7 +262,7 @@ protected:
 	 * Determines if the next token in the stream is an identifier that matches 
 	 * the specified string (and consumes it if it does).
 	 */
-	bool MatchIdentifier(TCHAR const* Match);
+	bool MatchIdentifier(const TCHAR* Match);
 
 	/**
 	 * Determines if the next token in the stream is an identifier that matches 
@@ -274,19 +274,19 @@ protected:
 	 * Determines if the next token in the stream is an identifier that matches 
 	 * the specified string (does NOT consume it).
 	 */
-	bool PeekIdentifier(TCHAR const* Match);
+	bool PeekIdentifier(const TCHAR* Match);
 
 	/**
 	 * Determines if the next token in the stream is a symbol that matches 
 	 * the specified string (and consumes it if it does).
 	 */
-	bool MatchSymbol(TCHAR const* Match);
+	bool MatchSymbol(const TCHAR* Match);
 
 	/**
 	 * Determines if the next token in the stream is a symbol that matches 
 	 * the specified string (does NOT consume it).
 	 */
-	bool PeekSymbol(TCHAR const* Match);
+	bool PeekSymbol(const TCHAR* Match);
 
 	//--------------------------------------
 	// Requiring checks
@@ -296,19 +296,19 @@ protected:
 	 * Ensures that the next token in the stream is an identifier that matches 
 	 * the specified name (and errors out if it isn't).
 	 */
-	bool RequireIdentifier(FName Match, TCHAR const* ErrorContext);
+	bool RequireIdentifier(FName Match, const TCHAR* ErrorContext);
 
 	/**
 	 * Ensures that the next token in the stream is an identifier that matches 
 	 * the specified string (and errors out if it isn't).
 	 */
-	bool RequireIdentifier(TCHAR const* Match, TCHAR const* ErrorContext);
+	bool RequireIdentifier(const TCHAR* Match, const TCHAR* ErrorContext);
 
 	/**
 	 * Ensures that the next token in the stream is a symbol that matches 
 	 * the specified string (and errors out if it isn't).
 	 */
-	bool RequireSymbol(TCHAR const* Match, TCHAR const* ErrorContext);
+	bool RequireSymbol(const TCHAR* Match, const TCHAR* ErrorContext);
 
 	/**
 	 * Ensures that the next token in the stream is a semi-colon character (and 
@@ -381,7 +381,7 @@ protected:
 	
 protected:
 	/** Input text */
-	TCHAR const* Input;
+	const TCHAR* Input;
 
 	/** Length of input text */
 	int32 InputLen;
diff --git a/Engine/Source/Editor/Kismet/Public/DiffUtils.h b/Engine/Source/Editor/Kismet/Public/DiffUtils.h
--- a/Engine/Source/Editor/Kismet/Public/DiffUtils.h
+++ b/Engine/Source/Editor/Kismet/Public/DiffUtils.h
@@ -268,7 +268,7 @@ namespace DiffViewUtils
 
 	KISMET_API FText PropertyDiffMessage(FSingleObjectDiffEntry Difference, FText ObjectName);
 	KISMET_API FText SCSDiffMessage(const FSCSDiffEntry& Difference, FText ObjectName);
-	KISMET_API FText GetPanelLabel(const UBlueprint* Blueprint, FRevisionInfo const& Revision, FText Label);
+	KISMET_API FText GetPanelLabel(const UBlueprint* Blueprint, const FRevisionInfo& Revision, FText Label);
 
 	KISMET_API SHorizontalBox::FSlot& Box(bool bIsPresent, FLinearColor Color);
 }
diff --git a/Engine/Source/Editor/NiagaraEditor/Private/NiagaraEditorUtilities.cpp b/Engine/Source/Editor/NiagaraEditor/Private/NiagaraEditorUtilities.cpp
--- a/Engine/Source/Editor/NiagaraEditor/Private/NiagaraEditorUtilities.cpp
+++ b/Engine/Source/Editor/NiagaraEditor/Private/NiagaraEditorUtilities.cpp
@@ -67,7 +67,7 @@ void FNiagaraEditorUtilities::ResetVariableToDefaultValue(FNiagaraVariable& Vari
 	}
 }
 
-void FNiagaraEditorUtilities::InitializeParameterInputNode(UNiagaraNodeInput& InputNode, const FNiagaraTypeDefinition& Type, const class UNiagaraGraph* InGraph, FName InputName)
+void FNiagaraEditorUtilities::InitializeParameterInputNode(UNiagaraNodeInput& InputNode, const FNiagaraTypeDefinition& Type, const UNiagaraGraph* InGraph, FName InputName)
 {
 	InputNode.Usage = ENiagaraInputNodeUsage::Parameter;
 	InputNode.bCanRenameNode = true;
diff --git a/Engine/Source/Editor/SequenceRecorder/Private/Sections/MovieScene3DTransformSectionRecorder.h b/Engine/Source/Editor/SequenceRecorder/Private/Sections/MovieScene3DTransformSectionRecorder.h
--- a/Engine/Source/Editor/SequenceRecorder/Private/Sections/MovieScene3DTransformSectionRecorder.h
+++ b/Engine/Source/Editor/SequenceRecorder/Private/Sections/MovieScene3DTransformSectionRecorder.h
@@ -12,6 +12,7 @@
 class FMovieScene3DTransformSectionRecorder;
 class UMovieScene3DTransformSection;
 class UMovieScene3DTransformTrack;
+struct FActorRecordingSettings;
 
 class FMovieScene3DTransformSectionRecorderFactory : public IMovieSceneSectionRecorderFactory
 {
@@ -24,7 +25,7 @@ public:
 	TSharedPtr<class FMovieScene3DTransformSectionRecorder> CreateSectionRecorder(bool bRecordTransforms, TSharedPtr<class FMovieSceneAnimationSectionRecorder> InAnimRecorder) const;
 
 private:
-	virtual TSharedPtr<IMovieSceneSectionRecorder> CreateSectionRecorder(const struct FActorRecordingSettings& InActorRecordingSettings) const override;
+	virtual TSharedPtr<IMovieSceneSectionRecorder> CreateSectionRecorder(const FActorRecordingSettings& InActorRecordingSettings) const override;
 };
 
 /** Structure used to buffer up transform keys. Keys are inserted into tracks in FinalizeSection() */
diff --git a/Engine/Source/Editor/SequenceRecorder/Private/Sections/MovieSceneAnimationSectionRecorder.h b/Engine/Source/Editor/SequenceRecorder/Private/Sections/MovieSceneAnimationSectionRecorder.h
--- a/Engine/Source/Editor/SequenceRecorder/Private/Sections/MovieSceneAnimationSectionRecorder.h
+++ b/Engine/Source/Editor/SequenceRecorder/Private/Sections/MovieSceneAnimationSectionRecorder.h
@@ -12,6 +12,8 @@
 
 class FMovieSceneAnimationSectionRecorder;
 class UMovieSceneSkeletalAnimationSection;
+class UActorRecording;
+struct FActorRecordingSettings;
 
 class FMovieSceneAnimationSectionRecorderFactory : public IMovieSceneSectionRecorderFactory
 {
@@ -20,10 +22,10 @@ public:
 
 	virtual bool CanRecordObject(class UObject* InObjectToRecord) const override;
 
-	TSharedPtr<class FMovieSceneAnimationSectionRecorder> CreateSectionRecorder(class UActorRecording* InActorRecording, const FAnimationRecordingSettings& InAnimationSettings) const;
+	TSharedPtr<class FMovieSceneAnimationSectionRecorder> CreateSectionRecorder(UActorRecording* InActorRecording, const FAnimationRecordingSettings& InAnimationSettings) const;
 
 private:
-	virtual TSharedPtr<IMovieSceneSectionRecorder> CreateSectionRecorder(const struct FActorRecordingSettings& InActorRecordingSettings) const override;
+	virtual TSharedPtr<IMovieSceneSectionRecorder> CreateSectionRecorder(const FActorRecordingSettings& InActorRecordingSettings) const override;
 };
 
 class FMovieSceneAnimationSectionRecorder : public IMovieSceneSectionRecorder
diff --git a/Engine/Source/Editor/UnrealEd/Classes/Editor/EditorPerProjectUserSettings.h b/Engine/Source/Editor/UnrealEd/Classes/Editor/EditorPerProjectUserSettings.h
--- a/Engine/Source/Editor/UnrealEd/Classes/Editor/EditorPerProjectUserSettings.h
+++ b/Engine/Source/Editor/UnrealEd/Classes/Editor/EditorPerProjectUserSettings.h
@@ -194,7 +194,9 @@ public:
 	FUserSettingChangedEvent& OnUserSettingChanged() { return UserSettingChangedEvent; }
 
 	//~ Begin UObject Interface
-	virtual void PostEditChangeProperty( struct FPropertyChangedEvent& PropertyChangedEvent ) override;
+#if WITH_EDITOR
+	virtual void PostEditChangeProperty( FPropertyChangedEvent& PropertyChangedEvent ) override;
+#endif
 	virtual void PostInitProperties() override;
 	//~ End UObject Interface
 
diff --git a/Engine/Source/Editor/UnrealEd/Private/Commandlets/DumpBlueprintsInfoCommandlet.cpp b/Engine/Source/Editor/UnrealEd/Private/Commandlets/DumpBlueprintsInfoCommandlet.cpp
--- a/Engine/Source/Editor/UnrealEd/Private/Commandlets/DumpBlueprintsInfoCommandlet.cpp
+++ b/Engine/Source/Editor/UnrealEd/Private/Commandlets/DumpBlueprintsInfoCommandlet.cpp
@@ -382,7 +382,7 @@ DumpBlueprintsInfo commandlet params: \n\
 	 * actions available to that pin (as if you dragged and spawned a context 
 	 * menu from it).
 	 */
-	static void DumpContextualPinTypeActions(uint32 Indent, UEdGraph* Graph, FEdGraphPinType const& PinType, FArchive* FileOutWriter);
+	static void DumpContextualPinTypeActions(uint32 Indent, UEdGraph* Graph, const FEdGraphPinType& PinType, FArchive* FileOutWriter);
 	
 	/**
 	 * Generic function that takes a contextual GraphActionList and calls down 
diff --git a/Engine/Source/Editor/UnrealEd/Private/EditorCategoryUtils.cpp b/Engine/Source/Editor/UnrealEd/Private/EditorCategoryUtils.cpp
--- a/Engine/Source/Editor/UnrealEd/Private/EditorCategoryUtils.cpp
+++ b/Engine/Source/Editor/UnrealEd/Private/EditorCategoryUtils.cpp
@@ -40,7 +40,7 @@ namespace FEditorCategoryUtilsImpl
 	 * @param  Key	The key you want a category path for.
 	 * @return The category display string associated with the specified key (an empty string if an entry wasn't found).
 	 */
-	FText const& GetCategory(const FString& Key);
+	const FText& GetCategory(const FString& Key);
 
 	/**
 	 * Performs a lookup into the category key table, retrieving a fully 
@@ -51,7 +51,7 @@ namespace FEditorCategoryUtilsImpl
 	 * @param  DocExcerpt			Name of the excerpt within the document page for this category
 	 * @return						The tooltip (if any) stored at the doc path
 	 */
-	FText GetTooltipForCategory(FString const& CategoryDisplayName, FString const& DocLink, FString const& DocExcerpt);
+	FText GetTooltipForCategory(const FString& CategoryDisplayName, const FString& DocLink, const FString& DocExcerpt);
 
 	/** Metadata tags */
 	const FName ClassHideCategoriesMetaKey(TEXT("HideCategories"));
@@ -103,7 +103,7 @@ FEditorCategoryUtilsImpl::FCategoryInfoMap& FEditorCategoryUtilsImpl::GetCategor
 }
 
 //------------------------------------------------------------------------------
-FText const& FEditorCategoryUtilsImpl::GetCategory(const FString& Key)
+const FText& FEditorCategoryUtilsImpl::GetCategory(const FString& Key)
 {
 	if (FEditorCategoryUtilsImpl::FCategoryInfo const* FoundCategory = GetCategoryTable().Find(Key))
 	{
@@ -113,7 +113,7 @@ FText const& FEditorCategoryUtilsImpl::GetCategory(const FString& Key)
 }
 
 //------------------------------------------------------------------------------
-FText FEditorCategoryUtilsImpl::GetTooltipForCategory(FString const& CategoryDisplayName, FString const& DocLink, FString const& DocExcerpt)
+FText FEditorCategoryUtilsImpl::GetTooltipForCategory(const FString& CategoryDisplayName, const FString& DocLink, const FString& DocExcerpt)
 {
 	FText Tooltip;
 
@@ -157,7 +157,7 @@ FText FEditorCategoryUtilsImpl::GetTooltipForCategory(FString const& CategoryDis
  ******************************************************************************/
 
 //------------------------------------------------------------------------------
-void FEditorCategoryUtils::RegisterCategoryKey(FString const& Key, FText const& Category, FText const& Tooltip)
+void FEditorCategoryUtils::RegisterCategoryKey(const FString& Key, const FText& Category, const FText& Tooltip)
 {
 	FEditorCategoryUtilsImpl::FCategoryInfo& CategoryInfo = FEditorCategoryUtilsImpl::GetCategoryTable().Add(Key);
 
@@ -167,7 +167,7 @@ void FEditorCategoryUtils::RegisterCategoryKey(FString const& Key, FText const&
 	CategoryInfo.Tooltip = (Tooltip.IsEmpty() ? FEditorCategoryUtilsImpl::GetTooltipForCategory(CategoryInfo.DisplayName.ToString(), CategoryInfo.DocLink, CategoryInfo.DocExcerpt) : Tooltip);
 }
 
-void FEditorCategoryUtils::RegisterCategoryKey(FString const& Key, FText const& Category, FString const& DocLink, FString const& DocExcerpt)
+void FEditorCategoryUtils::RegisterCategoryKey(const FString& Key, const FText& Category, const FString& DocLink, const FString& DocExcerpt)
 {
 	FEditorCategoryUtilsImpl::FCategoryInfo& CategoryInfo = FEditorCategoryUtilsImpl::GetCategoryTable().Add(Key);
 
@@ -178,7 +178,7 @@ void FEditorCategoryUtils::RegisterCategoryKey(FString const& Key, FText const&
 }
 
 //------------------------------------------------------------------------------
-FText const& FEditorCategoryUtils::GetCommonCategory(const FCommonEditorCategory::EValue CategoryId)
+const FText& FEditorCategoryUtils::GetCommonCategory(const FCommonEditorCategory::EValue CategoryId)
 {
 	static TMap<FCommonEditorCategory::EValue, FString> CommonCategoryKeys;
 	if (CommonCategoryKeys.Num() == 0)
@@ -220,11 +220,11 @@ FText const& FEditorCategoryUtils::GetCommonCategory(const FCommonEditorCategory
 }
 
 //------------------------------------------------------------------------------
-FText FEditorCategoryUtils::BuildCategoryString(FCommonEditorCategory::EValue RootId, FText const& SubCategory)
+FText FEditorCategoryUtils::BuildCategoryString(FCommonEditorCategory::EValue RootId, const FText& SubCategory)
 {
 	FText ConstructedCategory;
 
-	FText const& RootCategory = GetCommonCategory(RootId);
+	const FText& RootCategory = GetCommonCategory(RootId);
 	if (RootCategory.IsEmpty())
 	{
 		ConstructedCategory = SubCategory;
@@ -243,13 +243,13 @@ FText FEditorCategoryUtils::BuildCategoryString(FCommonEditorCategory::EValue Ro
 }
 
 //------------------------------------------------------------------------------
-FText FEditorCategoryUtils::GetCategoryDisplayString(FText const& UnsanitizedCategory)
+FText FEditorCategoryUtils::GetCategoryDisplayString(const FText& UnsanitizedCategory)
 {
 	return FText::FromString(GetCategoryDisplayString(UnsanitizedCategory.ToString()));
 }
 
 //------------------------------------------------------------------------------
-FString FEditorCategoryUtils::GetCategoryDisplayString(FString const& UnsanitizedCategory)
+FString FEditorCategoryUtils::GetCategoryDisplayString(const FString& UnsanitizedCategory)
 {
 	FString DisplayString = UnsanitizedCategory;
 
@@ -285,14 +285,14 @@ FString FEditorCategoryUtils::GetCategoryDisplayString(FString const& Unsanitize
 }
 
 //------------------------------------------------------------------------------
-void FEditorCategoryUtils::GetClassHideCategories(UClass const* Class, TArray<FString>& CategoriesOut, bool bHomogenize)
+void FEditorCategoryUtils::GetClassHideCategories(const UClass* Class, TArray<FString>& CategoriesOut, bool bHomogenize)
 {
 	CategoriesOut.Empty();
 
 	using namespace FEditorCategoryUtilsImpl;
 	if (Class->HasMetaData(ClassHideCategoriesMetaKey))
 	{
-		FString const& HideCategories = Class->GetMetaData(ClassHideCategoriesMetaKey);
+		const FString& HideCategories = Class->GetMetaData(ClassHideCategoriesMetaKey);
 
 		HideCategories.ParseIntoArray(CategoriesOut, TEXT(" "), /*InCullEmpty =*/true);
 		
@@ -307,14 +307,14 @@ void FEditorCategoryUtils::GetClassHideCategories(UClass const* Class, TArray<FS
 }
 
 //------------------------------------------------------------------------------
-void  FEditorCategoryUtils::GetClassShowCategories(UClass const* Class, TArray<FString>& CategoriesOut)
+void  FEditorCategoryUtils::GetClassShowCategories(const UClass* Class, TArray<FString>& CategoriesOut)
 {
 	CategoriesOut.Empty();
 
 	using namespace FEditorCategoryUtilsImpl;
 	if (Class->HasMetaData(ClassShowCategoriesMetaKey))
 	{
-		FString const& ShowCategories = Class->GetMetaData(ClassShowCategoriesMetaKey);
+		const FString& ShowCategories = Class->GetMetaData(ClassShowCategoriesMetaKey);
 		ShowCategories.ParseIntoArray(CategoriesOut, TEXT(" "), /*InCullEmpty =*/true);
 
 		for (FString& Category : CategoriesOut)
@@ -325,19 +325,19 @@ void  FEditorCategoryUtils::GetClassShowCategories(UClass const* Class, TArray<F
 }
 
 //------------------------------------------------------------------------------
-bool FEditorCategoryUtils::IsCategoryHiddenFromClass(UClass const* Class, FCommonEditorCategory::EValue CategoryId)
+bool FEditorCategoryUtils::IsCategoryHiddenFromClass(const UClass* Class, FCommonEditorCategory::EValue CategoryId)
 {
 	return IsCategoryHiddenFromClass(Class, GetCommonCategory(CategoryId));
 }
 
 //------------------------------------------------------------------------------
-bool FEditorCategoryUtils::IsCategoryHiddenFromClass(UClass const* Class, FText const& Category)
+bool FEditorCategoryUtils::IsCategoryHiddenFromClass(const UClass* Class, const FText& Category)
 {
 	return IsCategoryHiddenFromClass(Class, Category.ToString());
 }
 
 //------------------------------------------------------------------------------
-bool FEditorCategoryUtils::IsCategoryHiddenFromClass(UClass const* Class, FString const& Category)
+bool FEditorCategoryUtils::IsCategoryHiddenFromClass(const UClass* Class, const FString& Category)
 {
 	TArray<FString> ClassHideCategories;
 	GetClassHideCategories(Class, ClassHideCategories);
@@ -345,12 +345,12 @@ bool FEditorCategoryUtils::IsCategoryHiddenFromClass(UClass const* Class, FStrin
 }
 
 //------------------------------------------------------------------------------
-bool FEditorCategoryUtils::IsCategoryHiddenFromClass(const TArray<FString>& ClassHideCategories, UClass const* Class, const FString& Category)
+bool FEditorCategoryUtils::IsCategoryHiddenFromClass(const TArray<FString>& ClassHideCategories, const UClass* Class, const FString& Category)
 {
 	bool bIsHidden = false;
 
 	// run the category through sanitization so we can ensure compares will hit
-	FString const DisplayCategory = GetCategoryDisplayString(Category);
+	const FString DisplayCategory = GetCategoryDisplayString(Category);
 
 	for (const FString& HideCategory : ClassHideCategories)
 	{
@@ -368,7 +368,7 @@ bool FEditorCategoryUtils::IsCategoryHiddenFromClass(const TArray<FString>& Clas
 			DisplayCategory.ParseIntoArray(SubCategoryList, TEXT("|"), /*InCullEmpty =*/true);
 
 			FString FullSubCategoryPath;
-			for (FString const& SubCategory : SubCategoryList)
+			for (const FString& SubCategory : SubCategoryList)
 			{
 				FullSubCategoryPath += SubCategory;
 				if ((HideCategory == SubCategory) || (HideCategory == FullSubCategoryPath))
@@ -410,7 +410,7 @@ void FEditorCategoryUtils::GetCategoryTooltipInfo(const FString& Category, FText
 }
 
 //------------------------------------------------------------------------------
-TSet<FString> FEditorCategoryUtils::GetHiddenCategories(UClass const* Class)
+TSet<FString> FEditorCategoryUtils::GetHiddenCategories(const UClass* Class)
 {
 	TArray<FString> ClassHiddenCategories;
 	GetClassHideCategories(Class, ClassHiddenCategories);
diff --git a/Engine/Source/Editor/UnrealEd/Private/Settings/EditorPerProjectUserSettings.cpp b/Engine/Source/Editor/UnrealEd/Private/Settings/EditorPerProjectUserSettings.cpp
--- a/Engine/Source/Editor/UnrealEd/Private/Settings/EditorPerProjectUserSettings.cpp
+++ b/Engine/Source/Editor/UnrealEd/Private/Settings/EditorPerProjectUserSettings.cpp
@@ -8,6 +8,8 @@
 
 #define LOCTEXT_NAMESPACE "EditorPerProjectUserSettings"
 
+/// @cond DOXYGEN_WARNINGS
+
 UEditorPerProjectUserSettings::UEditorPerProjectUserSettings(const FObjectInitializer& ObjectInitializer)
 	: Super(ObjectInitializer)
 {
@@ -36,6 +38,7 @@ void UEditorPerProjectUserSettings::PostInitProperties()
 	MaterialQualityLevelVar->Set(MaterialQualityLevel, ECVF_SetByScalability);
 }
 
+#if WITH_EDITOR
 void UEditorPerProjectUserSettings::PostEditChangeProperty( FPropertyChangedEvent& PropertyChangedEvent )
 {
 	Super::PostEditChangeProperty(PropertyChangedEvent);
@@ -56,6 +59,8 @@ void UEditorPerProjectUserSettings::PostEditChangeProperty( FPropertyChangedEven
 
 	UserSettingChangedEvent.Broadcast(Name);
 }
+#endif
 
+/// @endcond
 
 #undef LOCTEXT_NAMESPACE
diff --git a/Engine/Source/Editor/UnrealEd/Public/EditorCategoryUtils.h b/Engine/Source/Editor/UnrealEd/Public/EditorCategoryUtils.h
--- a/Engine/Source/Editor/UnrealEd/Public/EditorCategoryUtils.h
+++ b/Engine/Source/Editor/UnrealEd/Public/EditorCategoryUtils.h
@@ -74,7 +74,7 @@ namespace FEditorCategoryUtils
 	 * @param  Category	The qualified category path that you want the key expanded to.
 	 * @param  Tooltip  An optional tooltip text to use for the category.  If not specified an attempt to find it from the NodeCategories UDN file will be made
 	 */
-	UNREALED_API void RegisterCategoryKey(FString const& Key, FText const& Category, FText const& Tooltip = FText::GetEmpty());
+	UNREALED_API void RegisterCategoryKey(const FString& Key, const FText& Category, const FText& Tooltip = FText::GetEmpty());
 
 	/**
 	 * @param  Key			A string key that people will use in metadata to reflect this category mapping.
@@ -82,7 +82,7 @@ namespace FEditorCategoryUtils
 	 * @param  DocLink		Path to the document page that contains the excerpt for this category
 	 * @param  DocExcerpt	Name of the excerpt within the document page for this category
 	 */
-	UNREALED_API void RegisterCategoryKey(FString const& Key, FText const& Category, FString const& DocLink, FString const& DocExcerpt);
+	UNREALED_API void RegisterCategoryKey(const FString& Key, const FText& Category, const FString& DocLink, const FString& DocExcerpt);
 
 	/**
 	 * Retrieves a qualified category path for the desired common category.
@@ -90,7 +90,7 @@ namespace FEditorCategoryUtils
 	 * @param  CategoryId	The common category you want a path for.
 	 * @return A text string, (empty if the common category was not registered)
 	 */
-	UNREALED_API FText const& GetCommonCategory(const FCommonEditorCategory::EValue CategoryId);
+	UNREALED_API const FText& GetCommonCategory(const FCommonEditorCategory::EValue CategoryId);
 
 	/**
 	 * Utility function that concatenates the supplied sub-category with one 
@@ -100,7 +100,7 @@ namespace FEditorCategoryUtils
 	 * @param  SubCategory	A sub-category that you want postfixed to the root category.
 	 * @return A concatenated text string, with the two categories separated by a pipe, '|', character.
 	 */
-	UNREALED_API FText BuildCategoryString(const FCommonEditorCategory::EValue RootCategory, FText const& SubCategory);
+	UNREALED_API FText BuildCategoryString(const FCommonEditorCategory::EValue RootCategory, const FText& SubCategory);
 
 	/**
 	 * Expands any keys found in the category string (any terms found in square 
@@ -110,7 +110,7 @@ namespace FEditorCategoryUtils
 	 * @param  SubCategory	A sub-category that you want postfixing the result.
 	 * @return A concatenated text string, with the two categories separated by a pipe, '|', character.
 	 */
-	UNREALED_API FText GetCategoryDisplayString(FText const& UnsanitizedCategory);
+	UNREALED_API FText GetCategoryDisplayString(const FText& UnsanitizedCategory);
 
 	/**
 	 * Expands any keys found in the category string (any terms found in square 
@@ -120,7 +120,7 @@ namespace FEditorCategoryUtils
 	 * @param  SubCategory	A sub-category that you want postfixing the result.
 	 * @return A concatenated string, with the two categories separated by a pipe, '|', character.
 	 */
-	UNREALED_API FString GetCategoryDisplayString(FString const& UnsanitizedCategory);
+	UNREALED_API FString GetCategoryDisplayString(const FString& UnsanitizedCategory);
 
 	/**
 	 * Parses out the class's "HideCategories" metadata, and returns it 
@@ -130,7 +130,7 @@ namespace FEditorCategoryUtils
 	 * @param  CategoriesOut	An array that will be filled with a list of hidden categories.
 	 * @param  bHomogenize		Determines if the categories should be ran through expansion and display sanitation (useful even when not being displayed, for comparisons)
 	 */
-	UNREALED_API void GetClassHideCategories(UClass const* Class, TArray<FString>& CategoriesOut, bool bHomogenize = true);
+	UNREALED_API void GetClassHideCategories(const UClass* Class, TArray<FString>& CategoriesOut, bool bHomogenize = true);
 
 	/**
 	 * Parses out the class's "ShowCategories" metadata, and returns it 
@@ -139,7 +139,7 @@ namespace FEditorCategoryUtils
 	 * @param  Class			The class you want to pull data from.
 	 * @param  CategoriesOut	An array that will be filled with a list of shown categories.
 	 */
-	UNREALED_API void GetClassShowCategories(UClass const* Class, TArray<FString>& CategoriesOut);
+	UNREALED_API void GetClassShowCategories(const UClass* Class, TArray<FString>& CategoriesOut);
 
 	/**
 	 * Checks to see if the category associated with the supplied common 
@@ -149,7 +149,7 @@ namespace FEditorCategoryUtils
 	 * @param  CategoryId	An id associated with a category that you want to check.
 	 * @return True if the common category is hidden, false if not.
 	 */
-	UNREALED_API bool IsCategoryHiddenFromClass(UClass const* Class, const FCommonEditorCategory::EValue CategoryId);
+	UNREALED_API bool IsCategoryHiddenFromClass(const UClass* Class, const FCommonEditorCategory::EValue CategoryId);
 
 	/**
 	 * Checks to see if the specified category is hidden from the supplied class.
@@ -178,7 +178,7 @@ namespace FEditorCategoryUtils
 	 * @param  Category A category path that you want to check.
 	 * @return True if the category is hidden, false if not.
 	 */
-	UNREALED_API bool IsCategoryHiddenFromClass(const TArray<FString>& ClassHideCategories, UClass const* Class, const FString& Category);
+	UNREALED_API bool IsCategoryHiddenFromClass(const TArray<FString>& ClassHideCategories, const UClass* Class, const FString& Category);
 	/**
 	 * Returns tooltip information for the specified category
 	 * 
@@ -198,5 +198,5 @@ namespace FEditorCategoryUtils
 	 * @param  Class		The class you want to query.
 	 * @return  The set of categories that should be hidden.
 	 */
-	UNREALED_API TSet<FString> GetHiddenCategories(UClass const* Class);
+	UNREALED_API TSet<FString> GetHiddenCategories(const UClass* Class);
 };
diff --git a/Engine/Source/Editor/VREditor/UI/VREditorUISystem.h b/Engine/Source/Editor/VREditor/UI/VREditorUISystem.h
--- a/Engine/Source/Editor/VREditor/UI/VREditorUISystem.h
+++ b/Engine/Source/Editor/VREditor/UI/VREditorUISystem.h
@@ -129,9 +129,9 @@ public:
 	bool IsShowingEditorUIPanel(const VREditorPanelID& InPanelID) const;
 
 	/** Sets whether the specified editor UI panel should be visible.  Any other UI floating off this hand will be dismissed when showing it. */
-	void ShowEditorUIPanel(const class UWidgetComponent* WidgetComponent, UVREditorInteractor* Interactor, const bool bShouldShow, const bool bSpawnInFront = false, const bool bDragFromOpen = false, const bool bPlaySound = true);
+	void ShowEditorUIPanel(const UWidgetComponent* WidgetComponent, UVREditorInteractor* Interactor, const bool bShouldShow, const bool bSpawnInFront = false, const bool bDragFromOpen = false, const bool bPlaySound = true);
 	void ShowEditorUIPanel(const VREditorPanelID& InPanelID, UVREditorInteractor* Interactor, const bool bShouldShow, const bool bSpawnInFront = false, const bool bDragFromOpen = false, const bool bPlaySound = true);
-	void ShowEditorUIPanel(class AVREditorFloatingUI* Panel, UVREditorInteractor* Interactor, const bool bShouldShow, const bool bSpawnInFront = false, const bool bDragFromOpen = false, const bool bPlaySound = true);
+	void ShowEditorUIPanel(AVREditorFloatingUI* Panel, UVREditorInteractor* Interactor, const bool bShouldShow, const bool bSpawnInFront = false, const bool bDragFromOpen = false, const bool bPlaySound = true);
 
 	/** Returns true if the radial menu is visible on this hand */
 	bool IsShowingRadialMenu(const UVREditorInteractor* Interactor ) const;
diff --git a/Engine/Source/Runtime/CinematicCamera/Private/CineCameraComponent.cpp b/Engine/Source/Runtime/CinematicCamera/Private/CineCameraComponent.cpp
--- a/Engine/Source/Runtime/CinematicCamera/Private/CineCameraComponent.cpp
+++ b/Engine/Source/Runtime/CinematicCamera/Private/CineCameraComponent.cpp
@@ -17,6 +17,8 @@
 //////////////////////////////////////////////////////////////////////////
 // UCameraComponent
 
+/// @cond DOXYGEN_WARNINGS
+
 UCineCameraComponent::UCineCameraComponent()
 {
 	// Super 35mm 4 Perf
@@ -57,7 +59,6 @@ UCineCameraComponent::UCineCameraComponent()
 #endif
 }
 
-
 void UCineCameraComponent::PostInitProperties()
 {
 	Super::PostInitProperties();
@@ -109,6 +110,7 @@ void UCineCameraComponent::TickComponent(float DeltaTime, ELevelTick TickType, F
 }
 
 #if WITH_EDITORONLY_DATA
+
 void UCineCameraComponent::PostEditChangeProperty(FPropertyChangedEvent& PropertyChangedEvent)
 {
 	RecalcDerivedData();
@@ -146,8 +148,8 @@ void UCineCameraComponent::ResetProxyMeshTransform()
 		ProxyMeshComponent->SetRelativeLocation(FVector(-46.f, 0, -24.f));
 	}
 }
-#endif	// WITH_EDITORONLY_DATA
 
+#endif	// WITH_EDITORONLY_DATA
 
 float UCineCameraComponent::GetHorizontalFieldOfView() const
 {
@@ -200,6 +202,8 @@ void UCineCameraComponent::RecalcDerivedData()
 #endif
 }
 
+/// @endcond
+
 float UCineCameraComponent::GetDesiredFocusDistance(const FVector& InLocation) const
 {
 	float DesiredFocusDistance = 0.f;
diff --git a/Engine/Source/Runtime/CinematicCamera/Public/CineCameraComponent.h b/Engine/Source/Runtime/CinematicCamera/Public/CineCameraComponent.h
--- a/Engine/Source/Runtime/CinematicCamera/Public/CineCameraComponent.h
+++ b/Engine/Source/Runtime/CinematicCamera/Public/CineCameraComponent.h
@@ -102,9 +102,14 @@ struct FNamedLensPreset
 UENUM()
 enum class ECameraFocusMethod : uint8
 {
-	None,					/** Disables DoF entirely. */
-	Manual,					/** Allows for specifying or animating exact focus distances. */
-	Tracking,				/** Locks focus to specific object. */
+	/** Disables DoF entirely. */
+	None,
+
+	/** Allows for specifying or animating exact focus distances. */
+	Manual,
+
+	/** Locks focus to specific object. */
+	Tracking,
 };
 
 /** Settings to control tracking-focus mode. */
@@ -149,8 +154,8 @@ struct FCameraFocusSettings
 	UPROPERTY(EditAnywhere, BlueprintReadWrite, Category = "Tracking Focus Settings")
 	FCameraTrackingFocusSettings TrackingFocusSettings;
 
-// TODO: Make this editor only again once UE-43122 has been completed.
-//#if WITH_EDITORONLY_DATA
+//~ TODO: Make this editor only again once UE-43122 has been completed.
+//~	#if WITH_EDITORONLY_DATA
 	/** True to draw a translucent plane at the current focus depth, for easy tweaking. */
 	UPROPERTY(Transient, EditAnywhere, Category = "Focus Settings")
 	uint8 bDrawDebugFocusPlane : 1;
@@ -158,7 +163,7 @@ struct FCameraFocusSettings
 	/** For customizing the focus plane color, in case the default doesn't show up well in your scene. */
 	UPROPERTY(EditAnywhere, Category = "Focus Settings", meta = (EditCondition = "bDrawDebugFocusPlane"))
 	FColor DebugFocusPlaneColor;
-//#endif 
+//~	#endif 
 
 	/** True to use interpolation to smooth out changes in focus distance, false for focus distance changes to be instantaneous. */
 	UPROPERTY(EditAnywhere, BlueprintReadWrite, Category = "Focus Settings")
@@ -188,14 +193,7 @@ struct FCameraFocusSettings
 /**
  * A specialized version of a camera component, geared toward cinematic usage.
  */
-UCLASS(
-	HideCategories = (CameraSettings), 
-	HideFunctions = (SetFieldOfView, SetAspectRatio, SetConstraintAspectRatio), 
-	Blueprintable, 
-	ClassGroup = Camera, 
-	meta = (BlueprintSpawnableComponent), 
-	Config = Engine
-	)
+UCLASS(HideCategories = (CameraSettings), HideFunctions = (SetFieldOfView, SetAspectRatio, SetConstraintAspectRatio), Blueprintable, ClassGroup = Camera, meta = (BlueprintSpawnableComponent), Config = Engine)
 class CINEMATICCAMERA_API UCineCameraComponent : public UCameraComponent
 {
 	GENERATED_BODY()
@@ -261,7 +259,12 @@ protected:
 	/** Set to true to skip any interpolations on the next update. Resets to false automatically. */
 	uint8 bResetInterpolation : 1;
 
+	/// @cond DOXYGEN_WARNINGS
+	
 	virtual void PostLoad() override;
+	
+	/// @endcond
+	
 	virtual void PostInitProperties() override;
 	virtual void OnRegister() override;
 	virtual void TickComponent(float DeltaTime, ELevelTick TickType, FActorComponentTickFunction* ThisTickFunction) override;
diff --git a/Engine/Source/Runtime/Core/Private/Misc/SecureHash.cpp b/Engine/Source/Runtime/Core/Private/Misc/SecureHash.cpp
--- a/Engine/Source/Runtime/Core/Private/Misc/SecureHash.cpp
+++ b/Engine/Source/Runtime/Core/Private/Misc/SecureHash.cpp
@@ -266,6 +266,7 @@ void FMD5::Decode( uint32* output, const uint8* input, int32 len )
 	}
 }
 
+/// @cond DOXYGEN_WARNINGS
 //
 // MD5 initialization.  Begins an MD5 operation, writing a new context.
 //
@@ -457,6 +458,7 @@ void FMD5::Decode( uint32* output, const uint8* input, int32 len )
 // 		output[i] = ((uint32)input[j]) | (((uint32)input[j+1]) << 8) |
 // 		(((uint32)input[j+2]) << 16) | (((uint32)input[j+3]) << 24);
 // }
+/// @endcond
 
 namespace Lex
 {
diff --git a/Engine/Source/Runtime/Core/Public/Math/BigInt.h b/Engine/Source/Runtime/Core/Public/Math/BigInt.h
--- a/Engine/Source/Runtime/Core/Public/Math/BigInt.h
+++ b/Engine/Source/Runtime/Core/Public/Math/BigInt.h
@@ -1201,7 +1201,7 @@ namespace FEncryption
 	}
 }
 
-//~ APIDOCTOOL: Document=Off
+/// @cond DOXYGEN_WARNINGS
 
 /**
  * Specialization for int type used in encryption (performance). Avoids using temporary results and most of the operations are inplace.
@@ -1224,7 +1224,7 @@ FORCEINLINE TEncryptionInt FEncryption::ModularPow(TEncryptionInt Base, TEncrypt
 	return Result;
 }
 
-//~ APIDOCTOOL: Document=On
+/// @endcond
 
 template <class TYPE>
 struct FSignatureBase
diff --git a/Engine/Source/Runtime/Core/Public/UObject/NameTypes.h b/Engine/Source/Runtime/Core/Public/UObject/NameTypes.h
--- a/Engine/Source/Runtime/Core/Public/UObject/NameTypes.h
+++ b/Engine/Source/Runtime/Core/Public/UObject/NameTypes.h
@@ -26,6 +26,8 @@
 	#define WITH_CASE_PRESERVING_NAME WITH_EDITORONLY_DATA
 #endif
 
+class FText;
+
 /** Maximum size of name. */
 enum {NAME_SIZE	= 1024};
 
@@ -664,7 +666,7 @@ public:
 	 *
 	 * @return	true if the name is valid
 	 */
-	static bool IsValidXName( const FString& InName, const FString& InInvalidChars, class FText* OutReason = nullptr, const class FText* InErrorCtx = nullptr );
+	static bool IsValidXName( const FString& InName, const FString& InInvalidChars, FText* OutReason = nullptr, const FText* InErrorCtx = nullptr );
 
 	/**
 	 * Checks to see that a FName follows the rules that Unreal requires.
@@ -675,7 +677,7 @@ public:
 	 *
 	 * @return	true if the name is valid
 	 */
-	bool IsValidXName( const FString& InInvalidChars = INVALID_NAME_CHARACTERS, class FText* OutReason = nullptr, const class FText* InErrorCtx = nullptr ) const
+	bool IsValidXName( const FString& InInvalidChars = INVALID_NAME_CHARACTERS, FText* OutReason = nullptr, const FText* InErrorCtx = nullptr ) const
 	{
 		return IsValidXName(ToString(), InInvalidChars, OutReason, InErrorCtx);
 	}
@@ -688,7 +690,7 @@ public:
 	 *
 	 * @return	true if the name is valid
 	 */
-	bool IsValidXName( class FText& OutReason, const FString& InInvalidChars = INVALID_NAME_CHARACTERS ) const
+	bool IsValidXName( FText& OutReason, const FString& InInvalidChars = INVALID_NAME_CHARACTERS ) const
 	{
 		return IsValidXName(ToString(), InInvalidChars, &OutReason);
 	}
@@ -700,7 +702,7 @@ public:
 	 *
 	 * @return	true if the name is valid
 	 */
-	bool IsValidObjectName( class FText& OutReason ) const
+	bool IsValidObjectName( FText& OutReason ) const
 	{
 		return IsValidXName(ToString(), INVALID_OBJECTNAME_CHARACTERS, &OutReason);
 	}
@@ -713,7 +715,7 @@ public:
 	 *
 	 * @return	true if the name is valid
 	 */
-	bool IsValidGroupName( class FText& OutReason, bool bIsGroupName=false ) const
+	bool IsValidGroupName( FText& OutReason, bool bIsGroupName=false ) const
 	{
 		return IsValidXName(ToString(), INVALID_LONGPACKAGE_CHARACTERS, &OutReason);
 	}
diff --git a/Engine/Source/Runtime/D3D12RHI/Private/D3D12Adapter.h b/Engine/Source/Runtime/D3D12RHI/Private/D3D12Adapter.h
--- a/Engine/Source/Runtime/D3D12RHI/Private/D3D12Adapter.h
+++ b/Engine/Source/Runtime/D3D12RHI/Private/D3D12Adapter.h
@@ -45,8 +45,12 @@ enum EMultiGPUMode
 
 static const bool GRedirectDefaultContextForAFR = true;
 
+/// @cond DOXYGEN_WARNINGS
+
 void* FD3D12ThreadLocalObject<FD3D12FastConstantAllocator>::ThisThreadObject = nullptr;
 
+/// @endcond
+
 struct FD3D12AdapterDesc
 {
 	FD3D12AdapterDesc()
diff --git a/Engine/Source/Runtime/D3D12RHI/Private/D3D12Pipelinestate.cpp b/Engine/Source/Runtime/D3D12RHI/Private/D3D12Pipelinestate.cpp
--- a/Engine/Source/Runtime/D3D12RHI/Private/D3D12Pipelinestate.cpp
+++ b/Engine/Source/Runtime/D3D12RHI/Private/D3D12Pipelinestate.cpp
@@ -7,6 +7,8 @@
 //-----------------------------------------------------------------------------
 #include "D3D12RHIPrivate.h"
 
+/// @cond DOXYGEN_WARNINGS
+
 void FD3D12HighLevelGraphicsPipelineStateDesc::GetLowLevelDesc(FD3D12LowLevelGraphicsPipelineStateDesc& Desc)
 {
 	FMemory::Memzero(&Desc, sizeof(Desc));
@@ -45,6 +47,7 @@ void FD3D12HighLevelGraphicsPipelineStateDesc::GetLowLevelDesc(FD3D12LowLevelGra
 	Desc.Desc.DepthStencilState = DepthStencilState ? *DepthStencilState : CD3DX12_DEPTH_STENCIL_DESC(D3D12_DEFAULT);
 }
 
+/// @endcond
 
 FORCEINLINE uint32 SSE4_CRC32(void* Data, SIZE_T NumBytes)
 {
diff --git a/Engine/Source/Runtime/D3D12RHI/Private/D3D12RHI.cpp b/Engine/Source/Runtime/D3D12RHI/Private/D3D12RHI.cpp
--- a/Engine/Source/Runtime/D3D12RHI/Private/D3D12RHI.cpp
+++ b/Engine/Source/Runtime/D3D12RHI/Private/D3D12RHI.cpp
@@ -21,8 +21,12 @@ TAutoConsoleVariable<int32> CVarD3D12ZeroBufferSizeInMB(
 	ECVF_ReadOnly
 	);
 
+/// @cond DOXYGEN_WARNINGS
+
 FD3D12FastAllocator* FD3D12DynamicRHI::HelperThreadDynamicHeapAllocator = nullptr;
 
+/// @endcond
+
 FD3D12DynamicRHI* FD3D12DynamicRHI::SingleD3DRHI = nullptr;
 
 using namespace D3D12RHI;
diff --git a/Engine/Source/Runtime/D3D12RHI/Private/D3D12Texture.cpp b/Engine/Source/Runtime/D3D12RHI/Private/D3D12Texture.cpp
--- a/Engine/Source/Runtime/D3D12RHI/Private/D3D12Texture.cpp
+++ b/Engine/Source/Runtime/D3D12RHI/Private/D3D12Texture.cpp
@@ -39,6 +39,8 @@ template TD3D12Texture2D<FD3D12BaseTextureCube>::TD3D12Texture2D(class FD3D12Dev
 #endif
 	);
 
+/// @cond DOXYGEN_WARNINGS
+
 template void FD3D12TextureStats::D3D12TextureAllocated(TD3D12Texture2D<FD3D12BaseTexture2D>& Texture);
 template void FD3D12TextureStats::D3D12TextureAllocated(TD3D12Texture2D<FD3D12BaseTexture2DArray>& Texture);
 template void FD3D12TextureStats::D3D12TextureAllocated(TD3D12Texture2D<FD3D12BaseTextureCube>& Texture);
@@ -47,6 +49,8 @@ template void FD3D12TextureStats::D3D12TextureDeleted(TD3D12Texture2D<FD3D12Base
 template void FD3D12TextureStats::D3D12TextureDeleted(TD3D12Texture2D<FD3D12BaseTexture2DArray>& Texture);
 template void FD3D12TextureStats::D3D12TextureDeleted(TD3D12Texture2D<FD3D12BaseTextureCube>& Texture);
 
+/// @endcond
+
 struct FRHICommandUpdateTexture : public FRHICommand<FRHICommandUpdateTexture>
 {
 	FD3D12TextureBase* TextureBase;
diff --git a/Engine/Source/Runtime/Engine/Classes/Animation/SkeletalMeshActor.h b/Engine/Source/Runtime/Engine/Classes/Animation/SkeletalMeshActor.h
--- a/Engine/Source/Runtime/Engine/Classes/Animation/SkeletalMeshActor.h
+++ b/Engine/Source/Runtime/Engine/Classes/Animation/SkeletalMeshActor.h
@@ -24,6 +24,8 @@ class ENGINE_API ASkeletalMeshActor : public AActor, public IMatineeAnimInterfac
 {
 	GENERATED_UCLASS_BODY()
 
+	virtual void GetLifetimeReplicatedProps(TArray< FLifetimeProperty > & OutLifetimeProps) const override;
+
 	/** Whether or not this actor should respond to anim notifies - CURRENTLY ONLY AFFECTS PlayParticleEffect NOTIFIES**/
 	UPROPERTY(EditAnywhere, BlueprintReadWrite, Category=Animation, AdvancedDisplay)
 	uint32 bShouldDoAnimNotifies:1;
diff --git a/Engine/Source/Runtime/Engine/Classes/Engine/ExponentialHeightFog.h b/Engine/Source/Runtime/Engine/Classes/Engine/ExponentialHeightFog.h
--- a/Engine/Source/Runtime/Engine/Classes/Engine/ExponentialHeightFog.h
+++ b/Engine/Source/Runtime/Engine/Classes/Engine/ExponentialHeightFog.h
@@ -16,6 +16,8 @@ class AExponentialHeightFog
 {
 	GENERATED_UCLASS_BODY()
 
+	virtual void GetLifetimeReplicatedProps(TArray< FLifetimeProperty > & OutLifetimeProps) const override;
+
 private_subobject:
 	/** @todo document */
 	DEPRECATED_FORGAME(4.6, "Component should not be accessed directly, please use GetComponent() function instead. Component will soon be private and your code will not compile.")
diff --git a/Engine/Source/Runtime/Engine/Classes/Engine/Light.h b/Engine/Source/Runtime/Engine/Classes/Engine/Light.h
--- a/Engine/Source/Runtime/Engine/Classes/Engine/Light.h
+++ b/Engine/Source/Runtime/Engine/Classes/Engine/Light.h
@@ -21,6 +21,8 @@ private_subobject:
 	class ULightComponent* LightComponent;
 public:
 
+	virtual void GetLifetimeReplicatedProps(TArray< FLifetimeProperty > & OutLifetimeProps) const override;
+
 	/** replicated copy of LightComponent's bEnabled property */
 	UPROPERTY(replicatedUsing=OnRep_bEnabled)
 	uint32 bEnabled:1;
diff --git a/Engine/Source/Runtime/Engine/Classes/Engine/SkyLight.h b/Engine/Source/Runtime/Engine/Classes/Engine/SkyLight.h
--- a/Engine/Source/Runtime/Engine/Classes/Engine/SkyLight.h
+++ b/Engine/Source/Runtime/Engine/Classes/Engine/SkyLight.h
@@ -12,6 +12,8 @@ class ENGINE_API ASkyLight : public AInfo
 {
 	GENERATED_UCLASS_BODY()
 
+	virtual void GetLifetimeReplicatedProps(TArray< FLifetimeProperty > & OutLifetimeProps) const override;
+
 private_subobject:
 	/** @todo document */
 	DEPRECATED_FORGAME(4.6, "LightComponent should not be accessed directly, please use GetLightComponent() function instead. LightComponent will soon be private and your code will not compile.")
diff --git a/Engine/Source/Runtime/Engine/Classes/GameFramework/GameStateBase.h b/Engine/Source/Runtime/Engine/Classes/GameFramework/GameStateBase.h
--- a/Engine/Source/Runtime/Engine/Classes/GameFramework/GameStateBase.h
+++ b/Engine/Source/Runtime/Engine/Classes/GameFramework/GameStateBase.h
@@ -82,6 +82,7 @@ class ENGINE_API AGameStateBase : public AInfo
 	UFUNCTION(BlueprintCallable, Category=GameState)
 	virtual float GetPlayerRespawnDelay(AController* Controller) const;
 
+	virtual void GetLifetimeReplicatedProps(TArray< FLifetimeProperty > &OutLifetimeProps) const;
 
 	//~=============================================================================
 	// Interaction with GameModeBase
diff --git a/Engine/Source/Runtime/Engine/Classes/GameFramework/PlayerController.h b/Engine/Source/Runtime/Engine/Classes/GameFramework/PlayerController.h
--- a/Engine/Source/Runtime/Engine/Classes/GameFramework/PlayerController.h
+++ b/Engine/Source/Runtime/Engine/Classes/GameFramework/PlayerController.h
@@ -1332,6 +1332,8 @@ public:
 	virtual void SetPawn(APawn* InPawn) override;
 	//~ End AController Interface
 
+	virtual void GetLifetimeReplicatedProps(TArray< FLifetimeProperty > & OutLifetimeProps) const override;
+
 	/** called on the server when the client sends a message indicating it was unable to initialize an Actor channel,
 	 * most commonly because the desired Actor's archetype couldn't be serialized
 	 * the default is to do nothing (Actor simply won't exist on the client), but this function gives the game code
diff --git a/Engine/Source/Runtime/Engine/Classes/GameFramework/PlayerState.h b/Engine/Source/Runtime/Engine/Classes/GameFramework/PlayerState.h
--- a/Engine/Source/Runtime/Engine/Classes/GameFramework/PlayerState.h
+++ b/Engine/Source/Runtime/Engine/Classes/GameFramework/PlayerState.h
@@ -197,6 +197,8 @@ public:
 	/** return true if PlayerState is primary (ie. non-splitscreen) player */
 	virtual bool IsPrimaryPlayer() const;
 
+	virtual void GetLifetimeReplicatedProps(TArray< FLifetimeProperty > & OutLifetimeProps) const override;
+
 	/** calls OverrideWith and triggers OnOverrideWith for BP extension */
 	void DispatchOverrideWith(APlayerState* PlayerState);
 
diff --git a/Engine/Source/Runtime/Engine/Classes/GameFramework/WorldSettings.h b/Engine/Source/Runtime/Engine/Classes/GameFramework/WorldSettings.h
--- a/Engine/Source/Runtime/Engine/Classes/GameFramework/WorldSettings.h
+++ b/Engine/Source/Runtime/Engine/Classes/GameFramework/WorldSettings.h
@@ -279,6 +279,8 @@ class ENGINE_API AWorldSettings : public AInfo, public IInterface_AssetUserData
 {
 	GENERATED_UCLASS_BODY()
 
+	virtual void GetLifetimeReplicatedProps(TArray< FLifetimeProperty > & OutLifetimeProps) const override;
+
 	/** DEFAULT BASIC PHYSICS SETTINGS **/
 
 	/** If true, enables CheckStillInWorld checks */
diff --git a/Engine/Source/Runtime/Engine/Classes/Kismet/GameplayStatics.h b/Engine/Source/Runtime/Engine/Classes/Kismet/GameplayStatics.h
--- a/Engine/Source/Runtime/Engine/Classes/Kismet/GameplayStatics.h
+++ b/Engine/Source/Runtime/Engine/Classes/Kismet/GameplayStatics.h
@@ -529,9 +529,9 @@ class ENGINE_API UGameplayStatics : public UBlueprintFunctionLibrary
 	 * @param AttenuationSettings - Override attenuation settings package to play sound with
 	 */
 	UFUNCTION(BlueprintCallable, Category="Audio", meta=(WorldContext="WorldContextObject", AdvancedDisplay = "4", UnsafeDuringActorConstruction = "true"))
-	static void PlayDialogueAtLocation(const UObject* WorldContextObject, class UDialogueWave* Dialogue, const struct FDialogueContext& Context, FVector Location, FRotator Rotation, float VolumeMultiplier = 1.f, float PitchMultiplier = 1.f, float StartTime = 0.f, class USoundAttenuation* AttenuationSettings = nullptr);
+	static void PlayDialogueAtLocation(const UObject* WorldContextObject, class UDialogueWave* Dialogue, const FDialogueContext& Context, FVector Location, FRotator Rotation, float VolumeMultiplier = 1.f, float PitchMultiplier = 1.f, float StartTime = 0.f, USoundAttenuation* AttenuationSettings = nullptr);
 
-	static void PlayDialogueAtLocation(const UObject* WorldContextObject, UDialogueWave* Dialogue, const struct FDialogueContext& Context, FVector Location, float VolumeMultiplier = 1.f, float PitchMultiplier = 1.f, float StartTime = 0.f, USoundAttenuation* AttenuationSettings = nullptr)
+	static void PlayDialogueAtLocation(const UObject* WorldContextObject, UDialogueWave* Dialogue, const FDialogueContext& Context, FVector Location, float VolumeMultiplier = 1.f, float PitchMultiplier = 1.f, float StartTime = 0.f, USoundAttenuation* AttenuationSettings = nullptr)
 	{
 		PlayDialogueAtLocation(WorldContextObject, Dialogue, Context, Location, FRotator::ZeroRotator, VolumeMultiplier, PitchMultiplier, StartTime, AttenuationSettings);
 	}
@@ -568,15 +568,15 @@ class ENGINE_API UGameplayStatics : public UBlueprintFunctionLibrary
 	 * @return Audio Component to manipulate the playing dialogue with
 	 */
 	UFUNCTION(BlueprintCallable, Category="Audio", meta=(AdvancedDisplay = "2", UnsafeDuringActorConstruction = "true", Keywords = "play"))
-	static UAudioComponent* SpawnDialogueAttached(UDialogueWave* Dialogue, const struct FDialogueContext& Context, USceneComponent* AttachToComponent, FName AttachPointName = NAME_None, FVector Location = FVector(ForceInit), FRotator Rotation = FRotator::ZeroRotator, EAttachLocation::Type LocationType = EAttachLocation::KeepRelativeOffset, bool bStopWhenAttachedToDestroyed = false, float VolumeMultiplier = 1.f, float PitchMultiplier = 1.f, float StartTime = 0.f, USoundAttenuation* AttenuationSettings = nullptr, bool bAutoDestroy = true);
+	static UAudioComponent* SpawnDialogueAttached(UDialogueWave* Dialogue, const FDialogueContext& Context, USceneComponent* AttachToComponent, FName AttachPointName = NAME_None, FVector Location = FVector(ForceInit), FRotator Rotation = FRotator::ZeroRotator, EAttachLocation::Type LocationType = EAttachLocation::KeepRelativeOffset, bool bStopWhenAttachedToDestroyed = false, float VolumeMultiplier = 1.f, float PitchMultiplier = 1.f, float StartTime = 0.f, USoundAttenuation* AttenuationSettings = nullptr, bool bAutoDestroy = true);
 
 	DEPRECATED(4.9, "PlayDialogueAttached has been renamed SpawnDialogueAttached")
-	static UAudioComponent* PlayDialogueAttached(UDialogueWave* Dialogue, const struct FDialogueContext& Context, USceneComponent* AttachToComponent, FName AttachPointName = NAME_None, FVector Location = FVector(ForceInit), EAttachLocation::Type LocationType = EAttachLocation::KeepRelativeOffset, bool bStopWhenAttachedToDestroyed = false, float VolumeMultiplier = 1.f, float PitchMultiplier = 1.f, float StartTime = 0.f, USoundAttenuation* AttenuationSettings = nullptr)
+	static UAudioComponent* PlayDialogueAttached(UDialogueWave* Dialogue, const FDialogueContext& Context, USceneComponent* AttachToComponent, FName AttachPointName = NAME_None, FVector Location = FVector(ForceInit), EAttachLocation::Type LocationType = EAttachLocation::KeepRelativeOffset, bool bStopWhenAttachedToDestroyed = false, float VolumeMultiplier = 1.f, float PitchMultiplier = 1.f, float StartTime = 0.f, USoundAttenuation* AttenuationSettings = nullptr)
 	{
 		return SpawnDialogueAttached(Dialogue, Context, AttachToComponent, AttachPointName, Location, FRotator::ZeroRotator, LocationType, bStopWhenAttachedToDestroyed, VolumeMultiplier, PitchMultiplier, StartTime, AttenuationSettings);
 	}
 
-	static UAudioComponent* SpawnDialogueAttached(UDialogueWave* Dialogue, const struct FDialogueContext& Context, USceneComponent* AttachToComponent, FName AttachPointName, FVector Location, EAttachLocation::Type LocationType = EAttachLocation::KeepRelativeOffset, bool bStopWhenAttachedToDestroyed = false, float VolumeMultiplier = 1.f, float PitchMultiplier = 1.f, float StartTime = 0.f, USoundAttenuation* AttenuationSettings = nullptr, bool bAutoDestroy = true)
+	static UAudioComponent* SpawnDialogueAttached(UDialogueWave* Dialogue, const FDialogueContext& Context, USceneComponent* AttachToComponent, FName AttachPointName, FVector Location, EAttachLocation::Type LocationType = EAttachLocation::KeepRelativeOffset, bool bStopWhenAttachedToDestroyed = false, float VolumeMultiplier = 1.f, float PitchMultiplier = 1.f, float StartTime = 0.f, USoundAttenuation* AttenuationSettings = nullptr, bool bAutoDestroy = true)
 	{
 		return SpawnDialogueAttached(Dialogue, Context, AttachToComponent, AttachPointName, Location, FRotator::ZeroRotator, LocationType, bStopWhenAttachedToDestroyed, VolumeMultiplier, PitchMultiplier, StartTime, AttenuationSettings, bAutoDestroy);
 	}
diff --git a/Engine/Source/Runtime/Engine/Classes/Matinee/MatineeActor.h b/Engine/Source/Runtime/Engine/Classes/Matinee/MatineeActor.h
--- a/Engine/Source/Runtime/Engine/Classes/Matinee/MatineeActor.h
+++ b/Engine/Source/Runtime/Engine/Classes/Matinee/MatineeActor.h
@@ -60,6 +60,8 @@ class ENGINE_API AMatineeActor : public AActor
 {
 	GENERATED_UCLASS_BODY()
 
+	virtual void GetLifetimeReplicatedProps(TArray< FLifetimeProperty > & OutLifetimeProps) const override;
+
 	/** The matinee data used by this actor*/
 	UPROPERTY(EditAnywhere, BlueprintReadWrite, Category=MatineeActor, replicated, meta=(ForceRebuildProperty = "GroupActorInfos"))
 	class UInterpData* MatineeData;
diff --git a/Engine/Source/Runtime/Engine/Classes/Particles/Emitter.h b/Engine/Source/Runtime/Engine/Classes/Particles/Emitter.h
--- a/Engine/Source/Runtime/Engine/Classes/Particles/Emitter.h
+++ b/Engine/Source/Runtime/Engine/Classes/Particles/Emitter.h
@@ -56,6 +56,8 @@ class ENGINE_API AEmitter : public AActor
 {
 	GENERATED_UCLASS_BODY()
 
+	virtual void GetLifetimeReplicatedProps(TArray< FLifetimeProperty > & OutLifetimeProps) const override;
+
 private_subobject:
 	DEPRECATED_FORGAME(4.6, "ParticleSystemComponent should not be accessed directly, please use GetParticleSystemComponent() function instead. ParticleSystemComponent will soon be private and your code will not compile.")
 	UPROPERTY(Category = Emitter, VisibleAnywhere, BlueprintReadOnly, meta = (ExposeFunctionCategories = "Particles|Beam,Particles|Parameters,Particles,Effects|Components|ParticleSystem,Rendering,Activation,Components|Activation", AllowPrivateAccess = "true"))
diff --git a/Engine/Source/Runtime/Engine/Classes/PhysicsEngine/BodyInstance.h b/Engine/Source/Runtime/Engine/Classes/PhysicsEngine/BodyInstance.h
--- a/Engine/Source/Runtime/Engine/Classes/PhysicsEngine/BodyInstance.h
+++ b/Engine/Source/Runtime/Engine/Classes/PhysicsEngine/BodyInstance.h
@@ -1276,13 +1276,15 @@ private:
 //////////////////////////////////////////////////////////////////////////
 // BodyInstance inlines
 
-//~ APIDOCTOOL: Document=Off
+/// @cond DOXYGEN_WARNINGS
+
 FORCEINLINE_DEBUGGABLE bool FBodyInstance::OverlapMulti(TArray<struct FOverlapResult>& InOutOverlaps, const class UWorld* World, const FTransform* pWorldToComponent, const FVector& Pos, const FRotator& Rot, ECollisionChannel TestChannel, const FComponentQueryParams& Params, const FCollisionResponseParams& ResponseParams, const FCollisionObjectQueryParams& ObjectQueryParams) const
 {
 	// Pass on to FQuat version
 	return OverlapMulti(InOutOverlaps, World, pWorldToComponent, Pos, Rot.Quaternion(), TestChannel, Params, ResponseParams, ObjectQueryParams);
 }
-//~ APIDOCTOOL: Document=On
+
+/// @endcond
 
 FORCEINLINE_DEBUGGABLE bool FBodyInstance::OverlapTestForBodies(const FVector& Position, const FQuat& Rotation, const TArray<FBodyInstance*>& Bodies) const
 {
diff --git a/Engine/Source/Runtime/Engine/Classes/PhysicsEngine/BodySetup.h b/Engine/Source/Runtime/Engine/Classes/PhysicsEngine/BodySetup.h
--- a/Engine/Source/Runtime/Engine/Classes/PhysicsEngine/BodySetup.h
+++ b/Engine/Source/Runtime/Engine/Classes/PhysicsEngine/BodySetup.h
@@ -438,9 +438,13 @@ private:
 	float MaxContactOffset;
 };
 
+/// @cond DOXYGEN_WARNINGS
+
 //Explicit export of template instantiation 
 extern template ENGINE_API void FBodySetupShapeIterator::ForEachShape(const TArray<FKSphereElem>&, TFunctionRef<void(const FKSphereElem&, const physx::PxSphereGeometry&, const physx::PxTransform& , float )>) const;
 extern template ENGINE_API void FBodySetupShapeIterator::ForEachShape(const TArray<FKBoxElem>&, TFunctionRef<void(const FKBoxElem&, const physx::PxBoxGeometry&, const physx::PxTransform&, float)>) const;
 extern template ENGINE_API void FBodySetupShapeIterator::ForEachShape(const TArray<FKSphylElem>&, TFunctionRef<void(const FKSphylElem&, const physx::PxCapsuleGeometry&, const physx::PxTransform&, float)>) const;
 extern template ENGINE_API void FBodySetupShapeIterator::ForEachShape(const TArray<FKConvexElem>&, TFunctionRef<void(const FKConvexElem&, const physx::PxConvexMeshGeometry&, const physx::PxTransform&, float)>) const;
-extern template ENGINE_API void FBodySetupShapeIterator::ForEachShape(const TArray<physx::PxTriangleMesh*>&,TFunctionRef<void (physx::PxTriangleMesh* const &, const physx::PxTriangleMeshGeometry&, const physx::PxTransform&,float)>) const;
\ No newline at end of file
+extern template ENGINE_API void FBodySetupShapeIterator::ForEachShape(const TArray<physx::PxTriangleMesh*>&,TFunctionRef<void (physx::PxTriangleMesh* const &, const physx::PxTriangleMeshGeometry&, const physx::PxTransform&,float)>) const;
+
+/// @endcond
diff --git a/Engine/Source/Runtime/Engine/Classes/Sound/AudioVolume.h b/Engine/Source/Runtime/Engine/Classes/Sound/AudioVolume.h
--- a/Engine/Source/Runtime/Engine/Classes/Sound/AudioVolume.h
+++ b/Engine/Source/Runtime/Engine/Classes/Sound/AudioVolume.h
@@ -244,6 +244,7 @@ public:
 #if WITH_EDITOR
 	virtual void PostEditChangeProperty(FPropertyChangedEvent& PropertyChangedEvent) override;
 #endif // WITH_EDITOR
+	virtual void GetLifetimeReplicatedProps(TArray<FLifetimeProperty>& OutLifetimeProps) const override;
 	//~ End UObject Interface
 
 	//~ Begin AActor Interface
diff --git a/Engine/Source/Runtime/Engine/Private/Animation/AnimNotify.cpp b/Engine/Source/Runtime/Engine/Private/Animation/AnimNotify.cpp
--- a/Engine/Source/Runtime/Engine/Private/Animation/AnimNotify.cpp
+++ b/Engine/Source/Runtime/Engine/Private/Animation/AnimNotify.cpp
@@ -36,6 +36,8 @@ class UWorld* UAnimNotify::GetWorld() const
 	return (MeshContext ? MeshContext->GetWorld() : NULL);
 }
 
+/// @cond DOXYGEN_WARNINGS
+
 FString UAnimNotify::GetNotifyName_Implementation() const
 {
 	UObject* ClassGeneratedBy = GetClass()->ClassGeneratedBy;
@@ -57,6 +59,8 @@ FString UAnimNotify::GetNotifyName_Implementation() const
 	return NotifyName;
 }
 
+/// @endcond
+
 void UAnimNotify::PostLoad()
 {
 	Super::PostLoad();
diff --git a/Engine/Source/Runtime/Engine/Private/Animation/AnimNotifyState.cpp b/Engine/Source/Runtime/Engine/Private/Animation/AnimNotifyState.cpp
--- a/Engine/Source/Runtime/Engine/Private/Animation/AnimNotifyState.cpp
+++ b/Engine/Source/Runtime/Engine/Private/Animation/AnimNotifyState.cpp
@@ -49,6 +49,8 @@ void UAnimNotifyState::BranchingPointNotifyEnd(FBranchingPointNotifyPayload& Bra
 	NotifyEnd(BranchingPointPayload.SkelMeshComponent, BranchingPointPayload.SequenceAsset);
 }
 
+/// @cond DOXYGEN_WARNINGS
+
 FString UAnimNotifyState::GetNotifyName_Implementation() const
 {
 	UObject* ClassGeneratedBy = GetClass()->ClassGeneratedBy;
@@ -70,6 +72,8 @@ FString UAnimNotifyState::GetNotifyName_Implementation() const
 	return NotifyName;
 }
 
+/// @endcond
+
 void UAnimNotifyState::PostLoad()
 {
 	Super::PostLoad();
diff --git a/Engine/Source/Runtime/Engine/Private/Camera/CameraShake.cpp b/Engine/Source/Runtime/Engine/Private/Camera/CameraShake.cpp
--- a/Engine/Source/Runtime/Engine/Private/Camera/CameraShake.cpp
+++ b/Engine/Source/Runtime/Engine/Private/Camera/CameraShake.cpp
@@ -345,11 +345,15 @@ bool UCameraShake::IsFinished() const
 		);
 }
 
+/// @cond DOXYGEN_WARNINGS
+
 bool UCameraShake::ReceiveIsFinished_Implementation() const
 {
 	return true;
 }
 
+/// @endcond
+
 bool UCameraShake::IsLooping() const
 {
 	return OscillationDuration < 0.0f;
diff --git a/Engine/Source/Runtime/Engine/Private/CheatManager.cpp b/Engine/Source/Runtime/Engine/Private/CheatManager.cpp
--- a/Engine/Source/Runtime/Engine/Private/CheatManager.cpp
+++ b/Engine/Source/Runtime/Engine/Private/CheatManager.cpp
@@ -605,6 +605,8 @@ void UCheatManager::BeginDestroy()
 	Super::BeginDestroy();
 }
 
+/// @cond DOXYGEN_WARNINGS
+
 bool UCheatManager::ServerToggleAILogging_Validate()
 {
 	return true;
@@ -647,6 +649,8 @@ void UCheatManager::ServerToggleAILogging_Implementation()
 #endif
 }
 
+/// @endcond
+
 void UCheatManager::ToggleAILogging()
 {
 #if ENABLE_VISUAL_LOG
diff --git a/Engine/Source/Runtime/Engine/Private/Collision/WorldCollision.cpp b/Engine/Source/Runtime/Engine/Private/Collision/WorldCollision.cpp
--- a/Engine/Source/Runtime/Engine/Private/Collision/WorldCollision.cpp
+++ b/Engine/Source/Runtime/Engine/Private/Collision/WorldCollision.cpp
@@ -390,7 +390,7 @@ bool UWorld::OverlapMultiByProfile(TArray<struct FOverlapResult>& OutOverlaps, c
 }
 
 
-bool UWorld::ComponentOverlapMulti(TArray<struct FOverlapResult>& OutOverlaps, const class UPrimitiveComponent* PrimComp, const FVector& Pos, const FQuat& Quat, const struct FComponentQueryParams& Params, const struct FCollisionObjectQueryParams& ObjectQueryParams) const
+bool UWorld::ComponentOverlapMulti(TArray<struct FOverlapResult>& OutOverlaps, const class UPrimitiveComponent* PrimComp, const FVector& Pos, const FQuat& Quat, const FComponentQueryParams& Params, const FCollisionObjectQueryParams& ObjectQueryParams) const
 {
 	if (PrimComp)
 	{
@@ -419,7 +419,7 @@ bool UWorld::ComponentOverlapMultiByChannel(TArray<struct FOverlapResult>& OutOv
 	}
 }
 
-bool UWorld::ComponentSweepMulti(TArray<struct FHitResult>& OutHits, class UPrimitiveComponent* PrimComp, const FVector& Start, const FVector& End, const FQuat& Quat, const struct FComponentQueryParams& Params) const
+bool UWorld::ComponentSweepMulti(TArray<struct FHitResult>& OutHits, class UPrimitiveComponent* PrimComp, const FVector& Start, const FVector& End, const FQuat& Quat, const FComponentQueryParams& Params) const
 {
 	if (GetPhysicsScene() == NULL)
 	{
diff --git a/Engine/Source/Runtime/Engine/Private/Components/PrimitiveComponent.cpp b/Engine/Source/Runtime/Engine/Private/Components/PrimitiveComponent.cpp
--- a/Engine/Source/Runtime/Engine/Private/Components/PrimitiveComponent.cpp
+++ b/Engine/Source/Runtime/Engine/Private/Components/PrimitiveComponent.cpp
@@ -2164,7 +2164,7 @@ bool UPrimitiveComponent::ComputePenetration(FMTDResult& OutMTD, const FCollisio
 	return false;
 }
 
-bool UPrimitiveComponent::IsOverlappingComponent(UPrimitiveComponent const* OtherComp) const
+bool UPrimitiveComponent::IsOverlappingComponent(const UPrimitiveComponent* OtherComp) const
 {
 	for (int32 i=0; i < OverlappingComponents.Num(); ++i)
 	{
@@ -2176,7 +2176,7 @@ bool UPrimitiveComponent::IsOverlappingComponent(UPrimitiveComponent const* Othe
 	return false;
 }
 
-bool UPrimitiveComponent::IsOverlappingComponent( const FOverlapInfo& Overlap ) const
+bool UPrimitiveComponent::IsOverlappingComponent(const FOverlapInfo& Overlap) const
 {
 	return OverlappingComponents.Find(Overlap) != INDEX_NONE;
 }
diff --git a/Engine/Source/Runtime/Engine/Private/Components/StaticMeshComponent.cpp b/Engine/Source/Runtime/Engine/Private/Components/StaticMeshComponent.cpp
--- a/Engine/Source/Runtime/Engine/Private/Components/StaticMeshComponent.cpp
+++ b/Engine/Source/Runtime/Engine/Private/Components/StaticMeshComponent.cpp
@@ -208,6 +208,8 @@ UStaticMeshComponent::~UStaticMeshComponent()
 	// Empty, but required because we don't want to have to include LightMap.h and ShadowMap.h in StaticMeshComponent.h, and they are required to compile FLightMapRef and FShadowMapRef
 }
 
+/// @cond DOXYGEN_WARNINGS
+
 void UStaticMeshComponent::GetLifetimeReplicatedProps(TArray< FLifetimeProperty > & OutLifetimeProps) const
 {
 	Super::GetLifetimeReplicatedProps(OutLifetimeProps);
@@ -217,6 +219,8 @@ void UStaticMeshComponent::GetLifetimeReplicatedProps(TArray< FLifetimeProperty
 	PRAGMA_ENABLE_DEPRECATION_WARNINGS
 }
 
+/// @endcond
+
 void UStaticMeshComponent::OnRep_StaticMesh(class UStaticMesh *OldStaticMesh)
 {
 	PRAGMA_DISABLE_DEPRECATION_WARNINGS
diff --git a/Engine/Source/Runtime/Engine/Private/Controller.cpp b/Engine/Source/Runtime/Engine/Private/Controller.cpp
--- a/Engine/Source/Runtime/Engine/Private/Controller.cpp
+++ b/Engine/Source/Runtime/Engine/Private/Controller.cpp
@@ -345,6 +345,8 @@ void AController::Reset()
 	StartSpot = NULL;
 }
 
+/// @cond DOXYGEN_WARNINGS
+
 void AController::ClientSetLocation_Implementation( FVector NewLocation, FRotator NewRotation )
 {
 	ClientSetRotation(NewRotation);
@@ -363,6 +365,8 @@ void AController::ClientSetRotation_Implementation( FRotator NewRotation, bool b
 	}
 }
 
+/// @endcond
+
 void AController::RemovePawnTickDependency(APawn* InOldPawn)
 {
 	if (InOldPawn != NULL)
diff --git a/Engine/Source/Runtime/Engine/Private/EngineTypes.cpp b/Engine/Source/Runtime/Engine/Private/EngineTypes.cpp
--- a/Engine/Source/Runtime/Engine/Private/EngineTypes.cpp
+++ b/Engine/Source/Runtime/Engine/Private/EngineTypes.cpp
@@ -346,6 +346,8 @@ FVector FRepMovement::RebaseOntoZeroOrigin(const struct FVector& Location, const
 	return RebaseOntoZeroOrigin(Location, WorldContextActor->GetWorld()->OriginLocation);
 }
 
+/// @cond DOXYGEN_WARNINGS
+
 /** Rebase zero-origin position onto local world origin value based on an actor component's world. */
 FVector FRepMovement::RebaseOntoLocalOrigin(const struct FVector& Location, const UActorComponent* const WorldContextActorComponent)
 {
@@ -367,3 +369,5 @@ FVector FRepMovement::RebaseOntoZeroOrigin(const struct FVector& Location, const
 
 	return RebaseOntoZeroOrigin(Location, WorldContextActorComponent->GetWorld()->OriginLocation);
 }
+
+/// @endcond
diff --git a/Engine/Source/Runtime/Engine/Private/GameMode.cpp b/Engine/Source/Runtime/Engine/Private/GameMode.cpp
--- a/Engine/Source/Runtime/Engine/Private/GameMode.cpp
+++ b/Engine/Source/Runtime/Engine/Private/GameMode.cpp
@@ -169,6 +169,8 @@ void AGameMode::HandleMatchIsWaitingToStart()
 	}
 }
 
+/// @cond DOXYGEN_WARNINGS
+
 bool AGameMode::ReadyToStartMatch_Implementation()
 {
 	// If bDelayed Start is set, wait for a manual match start
@@ -188,6 +190,8 @@ bool AGameMode::ReadyToStartMatch_Implementation()
 	return false;
 }
 
+/// @endcond
+
 void AGameMode::StartMatch()
 {
 	if (HasMatchStarted())
@@ -249,12 +253,16 @@ void AGameMode::HandleMatchHasStarted()
 	}
 }
 
+/// @cond DOXYGEN_WARNINGS
+
 bool AGameMode::ReadyToEndMatch_Implementation()
 {
 	// By default don't explicitly end match
 	return false;
 }
 
+/// @endcond
+
 void AGameMode::EndMatch()
 {
 	if (!IsMatchInProgress())
diff --git a/Engine/Source/Runtime/Engine/Private/GameModeBase.cpp b/Engine/Source/Runtime/Engine/Private/GameModeBase.cpp
--- a/Engine/Source/Runtime/Engine/Private/GameModeBase.cpp
+++ b/Engine/Source/Runtime/Engine/Private/GameModeBase.cpp
@@ -127,11 +127,15 @@ TSubclassOf<AGameSession> AGameModeBase::GetGameSessionClass() const
 	return AGameSession::StaticClass();
 }
 
+/// @cond DOXYGEN_WARNINGS
+
 UClass* AGameModeBase::GetDefaultPawnClassForController_Implementation(AController* InController)
 {
 	return DefaultPawnClass;
 }
 
+/// @endcond
+
 int32 AGameModeBase::GetNumPlayers()
 {
 	int32 PlayerCount = 0;
@@ -290,11 +294,15 @@ void AGameModeBase::Reset()
 	InitGameState();
 }
 
+/// @cond DOXYGEN_WARNINGS
+
 bool AGameModeBase::ShouldReset_Implementation(AActor* ActorToReset)
 {
 	return true;
 }
 
+/// @endcond
+
 void AGameModeBase::ResetLevel()
 {
 	UE_LOG(LogGameMode, Verbose, TEXT("Reset %s"), *GetName());
@@ -799,12 +807,16 @@ bool AGameModeBase::ShouldStartInCinematicMode(APlayerController* Player, bool&
 	return false;
 }
 
+/// @cond DOXYGEN_WARNINGS
+
 void AGameModeBase::InitializeHUDForPlayer_Implementation(APlayerController* NewPlayer)
 {
 	// Tell client what HUD class to use
 	NewPlayer->ClientSetHUD(HUDClass);
 }
 
+/// @endcond
+
 void AGameModeBase::UpdateGameplayMuteList(APlayerController* aPlayer)
 {
 	if (aPlayer)
@@ -966,6 +978,8 @@ void AGameModeBase::Logout(AController* Exiting)
 	}
 }
 
+/// @cond DOXYGEN_WARNINGS
+
 void AGameModeBase::HandleStartingNewPlayer_Implementation(APlayerController* NewPlayer)
 {
 	// If players should start as spectators, leave them in the spectator state
@@ -1037,11 +1051,15 @@ AActor* AGameModeBase::ChoosePlayerStart_Implementation(AController* Player)
 	return FoundPlayerStart;
 }
 
+/// @endcond
+
 bool AGameModeBase::ShouldSpawnAtStartSpot(AController* Player)
 {
 	return (Player != nullptr && Player->StartSpot != nullptr);
 }
 
+/// @cond DOXYGEN_WARNINGS
+
 AActor* AGameModeBase::FindPlayerStart_Implementation(AController* Player, const FString& IncomingName)
 {
 	UWorld* World = GetWorld();
@@ -1087,11 +1105,15 @@ AActor* AGameModeBase::FindPlayerStart_Implementation(AController* Player, const
 	return BestStart;
 }
 
+/// @endcond
+
 AActor* AGameModeBase::K2_FindPlayerStart(AController* Player, const FString& IncomingName)
 {
 	return FindPlayerStart(Player, IncomingName);
 }
 
+/// @cond DOXYGEN_WARNINGS
+
 bool AGameModeBase::PlayerCanRestart_Implementation(APlayerController* Player)
 {
 	if (Player == nullptr || Player->IsPendingKillPending())
@@ -1128,6 +1150,8 @@ APawn* AGameModeBase::SpawnDefaultPawnAtTransform_Implementation(AController* Ne
 	return ResultPawn;
 }
 
+/// @endcond
+
 void AGameModeBase::RestartPlayer(AController* NewPlayer)
 {
 	if (NewPlayer == nullptr || NewPlayer->IsPendingKillPending())
@@ -1260,11 +1284,16 @@ void AGameModeBase::FinishRestartPlayer(AController* NewPlayer, const FRotator&
 	}
 }
 
+/// @cond DOXYGEN_WARNINGS
+
 void AGameModeBase::InitStartSpot_Implementation(AActor* StartSpot, AController* NewPlayer)
 {
 
 }
 
+/// @endcond
+
+
 void AGameModeBase::SetPlayerDefaults(APawn* PlayerPawn)
 {
 	PlayerPawn->SetPlayerDefaults();
diff --git a/Engine/Source/Runtime/Engine/Private/GameState.cpp b/Engine/Source/Runtime/Engine/Private/GameState.cpp
--- a/Engine/Source/Runtime/Engine/Private/GameState.cpp
+++ b/Engine/Source/Runtime/Engine/Private/GameState.cpp
@@ -181,6 +181,8 @@ float AGameState::GetPlayerRespawnDelay(class AController* Controller) const
 	return Super::GetPlayerRespawnDelay(Controller);
 }
 
+/// @cond DOXYGEN_WARNINGS
+
 void AGameState::GetLifetimeReplicatedProps( TArray< FLifetimeProperty > & OutLifetimeProps ) const
 {
 	Super::GetLifetimeReplicatedProps( OutLifetimeProps );
@@ -188,3 +190,5 @@ void AGameState::GetLifetimeReplicatedProps( TArray< FLifetimeProperty > & OutLi
 	DOREPLIFETIME( AGameState, MatchState );
 	DOREPLIFETIME_CONDITION( AGameState, ElapsedTime,	COND_InitialOnly );
 }
+
+/// @endcond
diff --git a/Engine/Source/Runtime/Engine/Private/HUD.cpp b/Engine/Source/Runtime/Engine/Private/HUD.cpp
--- a/Engine/Source/Runtime/Engine/Private/HUD.cpp
+++ b/Engine/Source/Runtime/Engine/Private/HUD.cpp
@@ -590,6 +590,8 @@ void AHUD::DrawDebugTextList()
 	}
 }
 
+/// @cond DOXYGEN_WARNINGS
+
 void AHUD::AddDebugText_Implementation(const FString& DebugText,
 										 AActor* SrcActor,
 										 float Duration,
@@ -682,6 +684,8 @@ void AHUD::RemoveAllDebugStrings_Implementation()
 	DebugTextList.Reset();
 }
 
+/// @endcond
+
 void AHUD::NotifyHitBoxClick(FName BoxName)
 {
 	// dispatch BP event
diff --git a/Engine/Source/Runtime/Engine/Private/HighResScreenshot.cpp b/Engine/Source/Runtime/Engine/Private/HighResScreenshot.cpp
--- a/Engine/Source/Runtime/Engine/Private/HighResScreenshot.cpp
+++ b/Engine/Source/Runtime/Engine/Private/HighResScreenshot.cpp
@@ -290,10 +290,14 @@ bool FHighResScreenshotConfig::SaveImage(const FString& File, const TArray<TPixe
 	return bSuccess;
 }
 
+/// @cond DOXYGEN_WARNINGS
+
 template ENGINE_API bool FHighResScreenshotConfig::SaveImage<FColor>(const FString& File, const TArray<FColor>& Bitmap, const FIntPoint& BitmapSize, FString* OutFilename) const;
 template ENGINE_API bool FHighResScreenshotConfig::SaveImage<FFloat16Color>(const FString& File, const TArray<FFloat16Color>& Bitmap, const FIntPoint& BitmapSize, FString* OutFilename) const;
 template ENGINE_API bool FHighResScreenshotConfig::SaveImage<FLinearColor>(const FString& File, const TArray<FLinearColor>& Bitmap, const FIntPoint& BitmapSize, FString* OutFilename) const;
 
+/// @endcond
+
 FHighResScreenshotConfig::FImageWriter::FImageWriter(const TSharedPtr<class IImageWrapper>& InWrapper)
 	: ImageWrapper(InWrapper)
 {
diff --git a/Engine/Source/Runtime/Engine/Private/HitProxies.cpp b/Engine/Source/Runtime/Engine/Private/HitProxies.cpp
--- a/Engine/Source/Runtime/Engine/Private/HitProxies.cpp
+++ b/Engine/Source/Runtime/Engine/Private/HitProxies.cpp
@@ -7,9 +7,13 @@
 #include "HitProxies.h"
 #include "Misc/ScopeLock.h"
 
+/// @cond DOXYGEN_WARNINGS
+
 IMPLEMENT_HIT_PROXY_BASE( HHitProxy, NULL );
 IMPLEMENT_HIT_PROXY(HObject,HHitProxy);
 
+/// @endcond
+
 const FHitProxyId FHitProxyId::InvisibleHitProxyId( INDEX_NONE - 1 );
 
 
diff --git a/Engine/Source/Runtime/Engine/Private/PhysicsEngine/BodyInstance.cpp b/Engine/Source/Runtime/Engine/Private/PhysicsEngine/BodyInstance.cpp
--- a/Engine/Source/Runtime/Engine/Private/PhysicsEngine/BodyInstance.cpp
+++ b/Engine/Source/Runtime/Engine/Private/PhysicsEngine/BodyInstance.cpp
@@ -4779,7 +4779,7 @@ FTransform RootSpaceToWeldedSpace(const FBodyInstance* BI, const FTransform& Roo
 	return RootTM;
 }
 
-bool FBodyInstance::OverlapMulti(TArray<struct FOverlapResult>& InOutOverlaps, const class UWorld* World, const FTransform* pWorldToComponent, const FVector& Pos, const FQuat& Quat, ECollisionChannel TestChannel, const struct FComponentQueryParams& Params, const struct FCollisionResponseParams& ResponseParams, const struct FCollisionObjectQueryParams& ObjectQueryParams) const
+bool FBodyInstance::OverlapMulti(TArray<struct FOverlapResult>& InOutOverlaps, const class UWorld* World, const FTransform* pWorldToComponent, const FVector& Pos, const FQuat& Quat, ECollisionChannel TestChannel, const struct FComponentQueryParams& Params, const struct FCollisionResponseParams& ResponseParams, const FCollisionObjectQueryParams& ObjectQueryParams) const
 {
 	SCOPE_CYCLE_COUNTER(STAT_Collision_SceneQueryTotal);
 	SCOPE_CYCLE_COUNTER(STAT_Collision_FBodyInstance_OverlapMulti);
diff --git a/Engine/Source/Runtime/Engine/Private/PhysicsEngine/BodySetup.cpp b/Engine/Source/Runtime/Engine/Private/PhysicsEngine/BodySetup.cpp
--- a/Engine/Source/Runtime/Engine/Private/PhysicsEngine/BodySetup.cpp
+++ b/Engine/Source/Runtime/Engine/Private/PhysicsEngine/BodySetup.cpp
@@ -2067,8 +2067,13 @@ TEnumAsByte<enum ECollisionTraceFlag> UBodySetup::GetCollisionTraceFlag() const
 	TEnumAsByte<enum ECollisionTraceFlag> DefaultFlag = UPhysicsSettings::Get()->DefaultShapeComplexity;
 	return CollisionTraceFlag == ECollisionTraceFlag::CTF_UseDefault ? DefaultFlag : CollisionTraceFlag;
 }
+
+/// @cond DOXYGEN_WARNINGS
+
 template void FBodySetupShapeIterator::ForEachShape(const TArray<FKSphereElem>&, TFunctionRef<void(const FKSphereElem&, const physx::PxSphereGeometry&, const physx::PxTransform&, float)>) const;
 template void FBodySetupShapeIterator::ForEachShape(const TArray<FKBoxElem>&, TFunctionRef<void(const FKBoxElem&, const physx::PxBoxGeometry&, const physx::PxTransform&, float)>) const;
 template void FBodySetupShapeIterator::ForEachShape(const TArray<FKSphylElem>&, TFunctionRef<void(const FKSphylElem&, const physx::PxCapsuleGeometry&, const physx::PxTransform&, float)>) const;
 template void FBodySetupShapeIterator::ForEachShape(const TArray<FKConvexElem>&, TFunctionRef<void(const FKConvexElem&, const physx::PxConvexMeshGeometry&, const physx::PxTransform&, float)>) const;
-template void FBodySetupShapeIterator::ForEachShape(const TArray<physx::PxTriangleMesh*>&, TFunctionRef<void(physx::PxTriangleMesh* const &, const physx::PxTriangleMeshGeometry&, const physx::PxTransform&, float)>) const;
\ No newline at end of file
+template void FBodySetupShapeIterator::ForEachShape(const TArray<physx::PxTriangleMesh*>&, TFunctionRef<void(physx::PxTriangleMesh* const &, const physx::PxTriangleMeshGeometry&, const physx::PxTransform&, float)>) const;
+
+/// @endcond
diff --git a/Engine/Source/Runtime/Engine/Private/PlayerCameraManager.cpp b/Engine/Source/Runtime/Engine/Private/PlayerCameraManager.cpp
--- a/Engine/Source/Runtime/Engine/Private/PlayerCameraManager.cpp
+++ b/Engine/Source/Runtime/Engine/Private/PlayerCameraManager.cpp
@@ -59,6 +59,8 @@ APlayerCameraManager::APlayerCameraManager(const FObjectInitializer& ObjectIniti
 	DefaultModifiers.Add(UCameraModifier_CameraShake::StaticClass());
 }
 
+/// @cond DOXYGEN_WARNINGS
+
 void APlayerCameraManager::PhotographyCameraModify_Implementation(const FVector NewCameraLocation, const FVector PreviousCameraLocation, const FVector OriginalCameraLocation, FVector& OutCameraLocation)
 {	// let proposed camera through unmodified by default
 	OutCameraLocation = NewCameraLocation;
@@ -80,6 +82,8 @@ void APlayerCameraManager::OnPhotographyMultiPartCaptureEnd_Implementation()
 {	// do nothing by default
 }
 
+/// @endcond
+
 
 APlayerController* APlayerCameraManager::GetOwningPlayerController() const
 {
diff --git a/Engine/Source/Runtime/Engine/Private/PlayerController.cpp b/Engine/Source/Runtime/Engine/Private/PlayerController.cpp
--- a/Engine/Source/Runtime/Engine/Private/PlayerController.cpp
+++ b/Engine/Source/Runtime/Engine/Private/PlayerController.cpp
@@ -207,6 +207,7 @@ void APlayerController::FailedToSpawnPawn()
 	ClientGotoState(NAME_Inactive);
 }
 
+/// @cond DOXYGEN_WARNINGS
 
 void APlayerController::ClientUpdateLevelStreamingStatus_Implementation(FName PackageName, bool bNewShouldBeLoaded, bool bNewShouldBeVisible, bool bNewShouldBlockOnLoad, int32 LODIndex )
 {
@@ -375,6 +376,8 @@ void APlayerController::ClientAddTextureStreamingLoc_Implementation(FVector InLo
 	}
 }
 
+/// @endcond
+
 void APlayerController::SetNetSpeed(int32 NewSpeed)
 {
 	UNetDriver* Driver = GetWorld()->GetNetDriver();
@@ -539,6 +542,8 @@ UInterpTrackInstDirector* APlayerController::GetControllingDirector()
 	return ControllingDirTrackInst;
 }
 
+/// @cond DOXYGEN_WARNINGS
+
 bool APlayerController::ServerNotifyLoadedWorld_Validate(FName WorldPackageName)
 {
 	RPC_VALIDATE( WorldPackageName.IsValid() );
@@ -575,6 +580,8 @@ void APlayerController::ServerNotifyLoadedWorld_Implementation(FName WorldPackag
 	}
 }
 
+/// @cond DOXYGEN_WARNINGS
+
 bool APlayerController::HasClientLoadedCurrentWorld()
 {
 	UNetConnection* Connection = Cast<UNetConnection>(Player);
@@ -685,6 +692,8 @@ void APlayerController::SafeRetryClientRestart()
 }
 
 
+/// @cond DOXYGEN_WARNINGS
+
 /** Avoid calling ClientRestart if we have already accepted this pawn */
 void APlayerController::ClientRetryClientRestart_Implementation(APawn* NewPawn)
 {
@@ -746,6 +755,8 @@ void APlayerController::ClientRestart_Implementation(APawn* NewPawn)
 	}
 }
 
+/// @endcond
+
 void APlayerController::Possess(APawn* PawnToPossess)
 {
 	if (!HasAuthority())
@@ -966,6 +977,8 @@ void APlayerController::PostInitializeComponents()
 	StateName = NAME_Spectating; // Don't use ChangeState, because we want to defer spawning the SpectatorPawn until the Player is received
 }
 
+/// @cond DOXYGEN_WARNINGS
+
 bool APlayerController::ServerShortTimeout_Validate()
 {
 	return true;
@@ -1016,6 +1029,8 @@ void APlayerController::ServerShortTimeout_Implementation()
 	}
 }
 
+/// @cond DOXYGEN_WARNINGS
+
 void APlayerController::AddCheats(bool bForce)
 {
 	UWorld* World = GetWorld();
@@ -1146,6 +1161,8 @@ void APlayerController::Reset()
 	ChangeState(NAME_Spectating);
 }
 
+/// @cond DOXYGEN_WARNINGS
+
 void APlayerController::ClientReset_Implementation()
 {
 	ResetCameraMode();
@@ -1160,6 +1177,9 @@ void APlayerController::ClientGotoState_Implementation(FName NewState)
 	ChangeState(NewState);
 }
 
+/// @endcond
+
+
 void APlayerController::UnFreeze() {}
 
 bool APlayerController::IsFrozen()
@@ -1167,6 +1187,8 @@ bool APlayerController::IsFrozen()
 	return GetWorldTimerManager().IsTimerActive(TimerHandle_UnFreeze);
 }
 
+/// @cond DOXYGEN_WARNINGS
+
 void APlayerController::ServerAcknowledgePossession_Implementation(APawn* P)
 {
 	UE_LOG(LogPlayerController, Verbose, TEXT("ServerAcknowledgePossession_Implementation %s"), *GetNameSafe(P));
@@ -1183,6 +1205,8 @@ bool APlayerController::ServerAcknowledgePossession_Validate(APawn* P)
 	return true;
 }
 
+/// @endcond
+
 void APlayerController::UnPossess()
 {
 	if (GetPawn() != NULL)
@@ -1201,6 +1225,8 @@ void APlayerController::UnPossess()
 	SetPawn(NULL);
 }
 
+/// @cond DOXYGEN_WARNINGS
+
 void APlayerController::ClientSetHUD_Implementation(TSubclassOf<AHUD> NewHUDClass)
 {
 	if ( MyHUD != NULL )
@@ -1218,6 +1244,8 @@ void APlayerController::ClientSetHUD_Implementation(TSubclassOf<AHUD> NewHUDClas
 	MyHUD = GetWorld()->SpawnActor<AHUD>(NewHUDClass, SpawnInfo );
 }
 
+/// @endcond
+
 void APlayerController::CleanupPlayerState()
 {
 	if (PlayerState)
@@ -1291,6 +1319,8 @@ void APlayerController::OnNetCleanup(UNetConnection* Connection)
 	UNetConnection::GNetConnectionBeingCleanedUp = NULL;
 }
 
+/// @cond DOXYGEN_WARNINGS
+
 void APlayerController::ClientReceiveLocalizedMessage_Implementation( TSubclassOf<ULocalMessage> Message, int32 Switch, APlayerState* RelatedPlayerState_1, APlayerState* RelatedPlayerState_2, UObject* OptionalObject )
 {
 	// Wait for player to be up to date with replication when joining a server, before stacking up messages
@@ -1368,6 +1398,8 @@ void APlayerController::ServerToggleAILogging_Implementation()
 	}
 }
 
+/// @endcond
+
 void APlayerController::PawnLeavingGame()
 {
 	if (GetPawn() != NULL)
@@ -1497,6 +1529,8 @@ void APlayerController::Camera( FName NewMode )
 	ServerCamera(NewMode);
 }
 
+/// @cond DOXYGEN_WARNINGS
+
 void APlayerController::ServerCamera_Implementation( FName NewMode )
 {
 	SetCameraMode(NewMode);
@@ -1516,6 +1550,7 @@ void APlayerController::ClientSetCameraMode_Implementation( FName NewCamMode )
 	}
 }
 
+/// @endcond
 
 void APlayerController::SetCameraMode( FName NewCamMode )
 {
@@ -1541,6 +1576,8 @@ void APlayerController::ResetCameraMode()
 	SetCameraMode(DefaultMode);
 }
 
+/// @cond DOXYGEN_WARNINGS
+
 void APlayerController::ClientSetCameraFade_Implementation(bool bEnableFading, FColor FadeColor, FVector2D FadeAlpha, float FadeTime, bool bFadeAudio)
 {
 	if (PlayerCameraManager != nullptr)
@@ -1556,6 +1593,8 @@ void APlayerController::ClientSetCameraFade_Implementation(bool bEnableFading, F
 	}
 }
 
+/// @endcond
+
 void APlayerController::SendClientAdjustment()
 {
 	if (AcknowledgedPawn != GetPawn() && !GetSpectatorPawn())
@@ -1576,6 +1615,7 @@ void APlayerController::SendClientAdjustment()
 	}
 }
 
+/// @cond DOXYGEN_WARNINGS
 
 void APlayerController::ClientCapBandwidth_Implementation(int32 Cap)
 {
@@ -1586,6 +1626,7 @@ void APlayerController::ClientCapBandwidth_Implementation(int32 Cap)
 	}
 }
 
+/// @endcond
 
 void APlayerController::UpdatePing(float InPing)
 {
@@ -1613,6 +1654,7 @@ void APlayerController::SetInitialLocationAndRotation(const FVector& NewLocation
 	}
 }
 
+/// @cond DOXYGEN_WARNINGS
 
 bool APlayerController::ServerUpdateCamera_Validate(FVector_NetQuantize CamLoc, int32 CamPitchAndYaw)
 {
@@ -1654,6 +1696,8 @@ void APlayerController::ServerUpdateCamera_Implementation(FVector_NetQuantize Ca
 	}
 }
 
+/// @endcond
+
 void APlayerController::RestartLevel()
 {
 	if( GetNetMode()==NM_Standalone )
@@ -1670,6 +1714,8 @@ void APlayerController::LocalTravel( const FString& FURL )
 	}
 }
 
+/// @cond DOXYGEN_WARNINGS
+
 void APlayerController::ClientReturnToMainMenu_Implementation(const FString& ReturnReason)
 {
 	UWorld* World = GetWorld();
@@ -1683,6 +1729,7 @@ void APlayerController::ClientReturnToMainMenu_Implementation(const FString& Ret
 	}
 }
 
+/// @endcond
 
 bool APlayerController::SetPause( bool bPause, FCanUnpause CanUnpauseDelegate)
 {
@@ -1717,6 +1764,8 @@ void APlayerController::Pause()
 	ServerPause();
 }
 
+/// @cond DOXYGEN_WARNINGS
+
 bool APlayerController::ServerPause_Validate()
 {
 #if UE_BUILD_SHIPPING
@@ -1732,6 +1781,8 @@ void APlayerController::ServerPause_Implementation()
 	SetPause(!IsPaused());
 }
 
+/// @endcond
+
 void APlayerController::SetName(const FString& S)
 {
 	if (!S.IsEmpty())
@@ -1741,6 +1792,8 @@ void APlayerController::SetName(const FString& S)
 	}
 }
 
+/// @cond DOXYGEN_WARNINGS
+
 void APlayerController::ServerChangeName_Implementation( const FString& S )
 {
 	AGameModeBase* GameMode = GetWorld()->GetAuthGameMode();
@@ -1756,6 +1809,8 @@ bool APlayerController::ServerChangeName_Validate( const FString& S )
 	return true;
 }
 
+/// @endcond
+
 void APlayerController::SwitchLevel(const FString& FURL)
 {
 	const ENetMode NetMode = GetNetMode();
@@ -1788,12 +1843,15 @@ void APlayerController::GameHasEnded(AActor* EndGameFocus, bool bIsWinner)
 	ClientGameEnded(EndGameFocus, bIsWinner);
 }
 
+/// @cond DOXYGEN_WARNINGS
 
 void APlayerController::ClientGameEnded_Implementation(AActor* EndGameFocus, bool bIsWinner)
 {
 	SetViewTarget(EndGameFocus);
 }
 
+/// @endcond
+
 bool APlayerController::GetHitResultUnderCursor(ECollisionChannel TraceChannel, bool bTraceComplex, FHitResult& HitResult) const
 {
 	ULocalPlayer* LocalPlayer = Cast<ULocalPlayer>(Player);
@@ -2442,6 +2500,8 @@ void APlayerController::SetViewTargetWithBlend(AActor* NewViewTarget, float Blen
 	SetViewTarget(NewViewTarget, TransitionParams);
 }
 
+/// @cond DOXYGEN_WARNINGS
+
 void APlayerController::ClientSetViewTarget_Implementation( AActor* A, FViewTargetTransitionParams TransitionParams )
 {
 	if (PlayerCameraManager && !PlayerCameraManager->bClientSimulatingViewTarget)
@@ -2475,6 +2535,8 @@ void APlayerController::ServerVerifyViewTarget_Implementation()
 	ClientSetViewTarget( TheViewTarget );
 }
 
+/// @endcond
+
 void APlayerController::SpawnPlayerCameraManager()
 {
 	// servers and owning clients get cameras
@@ -2547,6 +2609,8 @@ void APlayerController::ClearAudioListenerOverride()
 	AudioListenerComponent = nullptr;
 }
 
+/// @cond DOXYGEN_WARNINGS
+
 bool APlayerController::ServerCheckClientPossession_Validate()
 {
 	return true;
@@ -2572,6 +2636,7 @@ void APlayerController::ServerCheckClientPossessionReliable_Implementation()
 	ServerCheckClientPossession_Implementation();
 }
 
+/// @endcond
 
 void APlayerController::SafeServerCheckClientPossession()
 {
@@ -2597,6 +2662,8 @@ void APlayerController::SafeServerUpdateSpectatorState()
 	}
 }
 
+/// @cond DOXYGEN_WARNINGS
+
 bool APlayerController::ServerSetSpectatorLocation_Validate(FVector NewLoc, FRotator NewRot)
 {
 	return true;
@@ -2631,7 +2698,6 @@ void APlayerController::ServerSetSpectatorLocation_Implementation(FVector NewLoc
 	}
 }
 
-
 bool APlayerController::ServerSetSpectatorWaiting_Validate(bool bWaiting)
 {
 	return true;
@@ -2653,7 +2719,6 @@ void APlayerController::ClientSetSpectatorWaiting_Implementation(bool bWaiting)
 	}
 }
 
-
 bool APlayerController::ServerViewNextPlayer_Validate()
 {
 	return true;
@@ -2680,6 +2745,7 @@ void APlayerController::ServerViewPrevPlayer_Implementation()
 	}
 }
 
+/// @endcond
 
 APlayerState* APlayerController::GetNextViewablePlayer(int32 dir)
 {
@@ -2728,7 +2794,6 @@ APlayerState* APlayerController::GetNextViewablePlayer(int32 dir)
 	return nullptr;
 }
 
-
 void APlayerController::ViewAPlayer(int32 dir)
 {
 	APlayerState* const NextPlayerState = GetNextViewablePlayer(dir);
@@ -2739,6 +2804,8 @@ void APlayerController::ViewAPlayer(int32 dir)
 	}
 }
 
+/// @cond DOXYGEN_WARNINGS
+
 bool APlayerController::ServerViewSelf_Validate(FViewTargetTransitionParams TransitionParams)
 {
 	return true;
@@ -2754,6 +2821,8 @@ void APlayerController::ServerViewSelf_Implementation(FViewTargetTransitionParam
 	}
 }
 
+/// @endcond
+
 void APlayerController::StartFire( uint8 FireModeNum ) 
 {
 	if ( ((IsInState(NAME_Spectating) && bPlayerIsWaiting) || IsInState(NAME_Inactive)) && !IsFrozen() )
@@ -2770,7 +2839,6 @@ void APlayerController::StartFire( uint8 FireModeNum )
 	}
 }
 
-
 bool APlayerController::NotifyServerReceivedClientData(APawn* InPawn, float TimeStamp)
 {
 	if (GetPawn() != InPawn || (GetNetMode() == NM_Client))
@@ -2787,6 +2855,8 @@ bool APlayerController::NotifyServerReceivedClientData(APawn* InPawn, float Time
 	return true;
 }
 
+/// @cond DOXYGEN_WARNINGS
+
 bool APlayerController::ServerRestartPlayer_Validate()
 {
 	return true;
@@ -2822,11 +2892,15 @@ void APlayerController::ServerRestartPlayer_Implementation()
 	}
 }
 
+/// @endcond
+
 bool APlayerController::CanRestartPlayer()
 {
 	return PlayerState && !PlayerState->bOnlySpectator && HasClientLoadedCurrentWorld() && PendingSwapConnection == NULL;
 }
 
+/// @cond DOXYGEN_WARNINGS
+
 void APlayerController::ClientIgnoreMoveInput_Implementation(bool bIgnore)
 {
 	SetIgnoreMoveInput(bIgnore);
@@ -2837,6 +2911,7 @@ void APlayerController::ClientIgnoreLookInput_Implementation(bool bIgnore)
 	SetIgnoreLookInput(bIgnore);
 }
 
+/// @endcond
 
 void APlayerController::DisplayDebug(class UCanvas* Canvas, const FDebugDisplayInfo& DebugDisplay, float& YL, float& YPos)
 {
@@ -2978,6 +3053,8 @@ void APlayerController::SetCinematicMode(bool bInCinematicMode, bool bHidePlayer
 	ClientSetCinematicMode(bCinematicMode, bAffectsMovement, bAffectsTurning, bAffectsHUD);
 }
 
+/// @cond DOXYGEN_WARNINGS
+
 void APlayerController::ClientSetCinematicMode_Implementation(bool bInCinematicMode, bool bAffectsMovement, bool bAffectsTurning, bool bAffectsHUD)
 {
 	bCinematicMode = bInCinematicMode;
@@ -3002,11 +3079,15 @@ void APlayerController::ClientForceGarbageCollection_Implementation()
 	GetWorld()->ForceGarbageCollection();
 }
 
+/// @endcond
+
 void APlayerController::LevelStreamingStatusChanged(ULevelStreaming* LevelObject, bool bNewShouldBeLoaded, bool bNewShouldBeVisible, bool bNewShouldBlockOnLoad, int32 LODIndex )
 {
 	ClientUpdateLevelStreamingStatus(LevelObject->GetWorldAssetPackageFName(),bNewShouldBeLoaded,bNewShouldBeVisible,bNewShouldBlockOnLoad,LODIndex);
 }
 
+/// @cond DOXYGEN_WARNINGS
+
 void APlayerController::ClientPrepareMapChange_Implementation(FName LevelName, bool bFirst, bool bLast)
 {
 	// Only call on the first local player controller to handle it being called on multiple PCs for splitscreen.
@@ -3033,6 +3114,8 @@ void APlayerController::ClientPrepareMapChange_Implementation(FName LevelName, b
 	}
 }
 
+/// @endcond
+
 void APlayerController::DelayedPrepareMapChange()
 {
 	if (GetWorld()->IsPreparingMapChange())
@@ -3046,6 +3129,7 @@ void APlayerController::DelayedPrepareMapChange()
 	}
 }
 
+/// @cond DOXYGEN_WARNINGS
 
 void APlayerController::ClientCommitMapChange_Implementation()
 {
@@ -3075,12 +3159,12 @@ void APlayerController::ClientCancelPendingMapChange_Implementation()
 	GetWorld()->CancelPendingMapChange();
 }
 
-
 void APlayerController::ClientSetBlockOnAsyncLoading_Implementation()
 {
 	GetWorld()->bRequestedBlockOnAsyncLoading = true;
 }
 
+/// @endcond
 
 void APlayerController::GetSeamlessTravelActorList(bool bToEntry, TArray<AActor*>& ActorList)
 {
@@ -3127,11 +3211,15 @@ void APlayerController::PostSeamlessTravel()
 
 }
 
+/// @cond DOXYGEN_WARNINGS
+
 void APlayerController::ClientEnableNetworkVoice_Implementation(bool bEnable)
 {
 	ToggleSpeaking(bEnable);
 }
 
+/// @endcond
+
 void APlayerController::StartTalking()
 {
 	ToggleSpeaking(true);
@@ -3159,11 +3247,15 @@ void APlayerController::ToggleSpeaking(bool bSpeaking)
 	}
 }
 
+/// @cond DOXYGEN_WARNINGS
+
 void APlayerController::ClientVoiceHandshakeComplete_Implementation()
 {
 	MuteList.bHasVoiceHandshakeCompleted = true;
 }
 
+/// @endcond
+
 void APlayerController::GameplayMutePlayer(const FUniqueNetIdRepl& PlayerNetId)
 {
 	if (PlayerNetId.IsValid())
@@ -3180,6 +3272,8 @@ void APlayerController::GameplayUnmutePlayer(const FUniqueNetIdRepl& PlayerNetId
 	}
 }
 
+/// @cond DOXYGEN_WARNINGS
+
 void APlayerController::ServerMutePlayer_Implementation(FUniqueNetIdRepl PlayerId)
 {
 	MuteList.ServerMutePlayer(this, PlayerId);
@@ -3220,6 +3314,8 @@ void APlayerController::ClientUnmutePlayer_Implementation(FUniqueNetIdRepl Playe
 	MuteList.ClientUnmutePlayer(this, PlayerId);
 }
 
+/// @endcond
+
 bool APlayerController::IsPlayerMuted(const FUniqueNetId& PlayerId)
 {
 	return MuteList.IsPlayerMuted(PlayerId);
@@ -3235,6 +3331,8 @@ void APlayerController::NotifyDirectorControl(bool bNowControlling, AMatineeActo
 	}
 }
 
+/// @cond DOXYGEN_WARNINGS
+
 void APlayerController::ClientWasKicked_Implementation(const FText& KickReason)
 {
 }
@@ -3255,6 +3353,8 @@ void APlayerController::ClientEndOnlineSession_Implementation()
 	}
 }
 
+/// @endcond
+
 void APlayerController::ConsoleKey(FKey Key)
 {
 #if ALLOW_CONSOLE
@@ -3474,6 +3574,7 @@ int32 APlayerController::GetSplitscreenPlayerCount() const
 	return Result;
 }
 
+/// @cond DOXYGEN_WARNINGS
 
 void APlayerController::ClientSetForceMipLevelsToBeResident_Implementation( UMaterialInterface* Material, float ForceDuration, int32 CinematicTextureGroups )
 {
@@ -3483,7 +3584,6 @@ void APlayerController::ClientSetForceMipLevelsToBeResident_Implementation( UMat
 	}
 }
 
-
 void APlayerController::ClientPrestreamTextures_Implementation( AActor* ForcedActor, float ForceDuration, bool bEnableStreaming, int32 CinematicTextureGroups)
 {
 	if ( ForcedActor != NULL && IsPrimaryPlayer() )
@@ -3534,6 +3634,8 @@ void APlayerController::ClientStopForceFeedback_Implementation( UForceFeedbackEf
 	}
 }
 
+/// @endcond
+
 /** Action that interpolates a component over time to a desired position */
 class FDynamicForceFeedbackAction : public FPendingLatentAction
 {
@@ -3855,6 +3957,8 @@ void APlayerController::ProcessForceFeedbackAndHaptics(const float DeltaTime, co
 	}
 }
 
+/// @cond DOXYGEN_WARNINGS
+
 void APlayerController::ClientPlayCameraShake_Implementation( TSubclassOf<class UCameraShake> Shake, float Scale, ECameraAnimPlaySpace::Type PlaySpace, FRotator UserPlaySpaceRot )
 {
 	if (PlayerCameraManager != NULL)
@@ -3889,7 +3993,6 @@ void APlayerController::ClientStopCameraAnim_Implementation(UCameraAnim* AnimToS
 	}
 }
 
-
 void APlayerController::ClientSpawnCameraLensEffect_Implementation( TSubclassOf<AEmitterCameraLensEffectBase> LensEffectEmitterClass )
 {
 	if (PlayerCameraManager != NULL)
@@ -3906,6 +4009,8 @@ void APlayerController::ClientClearCameraLensEffects_Implementation()
 	}
 }
 
+/// @endcond
+
 void APlayerController::ReceivedGameModeClass(TSubclassOf<AGameModeBase> GameModeClass)
 {
 }
@@ -3921,19 +4026,6 @@ void APlayerController::ReceivedSpectatorClass(TSubclassOf<ASpectatorPawn> Spect
 	}
 }
 
-void APlayerController::GetLifetimeReplicatedProps( TArray< FLifetimeProperty > & OutLifetimeProps ) const
-{
-	Super::GetLifetimeReplicatedProps( OutLifetimeProps );
-
-	// These used to only replicate if PlayerCameraManager->GetViewTargetPawn() != GetPawn()
-	// But, since they also don't update unless that condition is true, these values won't change, thus won't send
-	// This is a little less efficient, but fits into the new condition system well, and shouldn't really add much overhead
-	DOREPLIFETIME_CONDITION( APlayerController, TargetViewRotation,		COND_OwnerOnly );
-
-	// Replicate SpawnLocation for remote spectators
-	DOREPLIFETIME_CONDITION( APlayerController, SpawnLocation, COND_OwnerOnly );
-}
-
 void APlayerController::SetPawn(APawn* InPawn)
 {
 	if (InPawn == NULL)
@@ -3957,6 +4049,18 @@ void APlayerController::SetPawn(APawn* InPawn)
 	}
 }
 
+void APlayerController::GetLifetimeReplicatedProps(TArray< FLifetimeProperty > & OutLifetimeProps) const
+{
+	Super::GetLifetimeReplicatedProps(OutLifetimeProps);
+
+	// These used to only replicate if PlayerCameraManager->GetViewTargetPawn() != GetPawn()
+	// But, since they also don't update unless that condition is true, these values won't change, thus won't send
+	// This is a little less efficient, but fits into the new condition system well, and shouldn't really add much overhead
+	DOREPLIFETIME_CONDITION(APlayerController, TargetViewRotation, COND_OwnerOnly);
+
+	// Replicate SpawnLocation for remote spectators
+	DOREPLIFETIME_CONDITION(APlayerController, SpawnLocation, COND_OwnerOnly);
+}
 
 void APlayerController::SetPlayer( UPlayer* InPlayer )
 {
@@ -4214,6 +4318,8 @@ void APlayerController::ClientTravel(const FString& URL, ETravelType TravelType,
 	ClientTravelInternal(URL, TravelType, bSeamless, MapPackageGuid);
 }
 
+/// @cond DOXYGEN_WARNINGS
+
 void APlayerController::ClientTravelInternal_Implementation(const FString& URL, ETravelType TravelType, bool bSeamless, FGuid MapPackageGuid)
 {
 	UWorld* World = GetWorld();
@@ -4236,6 +4342,8 @@ void APlayerController::ClientTravelInternal_Implementation(const FString& URL,
 	}
 }
 
+/// @endcond
+
 FString APlayerController::GetPlayerNetworkAddress()
 {
 	if( Player && Player->IsA(UNetConnection::StaticClass()) )
@@ -4879,11 +4987,15 @@ void APlayerController::BuildHiddenComponentList(const FVector& ViewLocation, TS
 	UpdateHiddenComponents(ViewLocation, HiddenComponentsOut);
 }
 
+/// @cond DOXYGEN_WARNINGS
+
 void APlayerController::ClientRepObjRef_Implementation(UObject *Object)
 {
 	UE_LOG(LogPlayerController, Warning, TEXT("APlayerController::ClientRepObjRef repped: %s"), Object ? *Object->GetName() : TEXT("NULL") );
 }
 
+/// @endcond
+
 void FDynamicForceFeedbackDetails::Update(FForceFeedbackValues& Values) const
 {
 	if (bAffectsLeftLarge)
@@ -4904,6 +5016,8 @@ void FDynamicForceFeedbackDetails::Update(FForceFeedbackValues& Values) const
 	}
 }
 
+/// @cond DOXYGEN_WARNINGS
+
 void APlayerController::OnServerStartedVisualLogger_Implementation(bool bIsLogging)
 {
 #if ENABLE_VISUAL_LOG
@@ -4912,6 +5026,8 @@ void APlayerController::OnServerStartedVisualLogger_Implementation(bool bIsLoggi
 #endif
 }
 
+/// @endcond
+
 bool APlayerController::ShouldPerformFullTickWhenPaused() const
 {
 	return bShouldPerformFullTickWhenPaused || (/*bIsInVr =*/GEngine->HMDDevice.IsValid() && GEngine->HMDDevice->IsStereoEnabled() && GEngine->HMDDevice->IsHMDConnected());
diff --git a/Engine/Source/Runtime/Engine/Private/Timeline.cpp b/Engine/Source/Runtime/Engine/Private/Timeline.cpp
--- a/Engine/Source/Runtime/Engine/Private/Timeline.cpp
+++ b/Engine/Source/Runtime/Engine/Private/Timeline.cpp
@@ -939,6 +939,8 @@ void UTimelineComponent::OnRep_Timeline()
 	}
 }
 
+/// @cond DOXYGEN_WARNINGS
+
 void UTimelineComponent::GetLifetimeReplicatedProps( TArray< FLifetimeProperty > & OutLifetimeProps ) const
 {
 	Super::GetLifetimeReplicatedProps( OutLifetimeProps );
@@ -946,6 +948,8 @@ void UTimelineComponent::GetLifetimeReplicatedProps( TArray< FLifetimeProperty >
 	DOREPLIFETIME( UTimelineComponent, TheTimeline );
 }
 
+/// @endcond
+
 void UTimelineComponent::GetAllCurves(TSet<class UCurveBase*>& InOutCurves) const
 {
 	TheTimeline.GetAllCurves(InOutCurves);
diff --git a/Engine/Source/Runtime/Engine/Public/BatchedElements.h b/Engine/Source/Runtime/Engine/Public/BatchedElements.h
--- a/Engine/Source/Runtime/Engine/Public/BatchedElements.h
+++ b/Engine/Source/Runtime/Engine/Public/BatchedElements.h
@@ -16,6 +16,7 @@
 #include "StaticBoundShaderState.h"
 
 struct FBatchedPoint;
+struct FDrawingPolicyRenderState;
 
 namespace EBlendModeFilter
 {
@@ -203,7 +204,7 @@ public:
 	 * @param DepthTexture	DepthTexture for manual depth testing with editor compositing in the pixel shader
 	 */
 	DEPRECATED(4.14, "Deprecated. Use the FBatchedElements::Draw method that takes a non-optional FSceneView parameter instead")
-	bool Draw(FRHICommandList& RHICmdList, const struct FDrawingPolicyRenderState& DrawRenderState, ERHIFeatureLevel::Type FeatureLevel, bool bNeedToSwitchVerticalAxis, const FMatrix& Transform, uint32 ViewportSizeX, uint32 ViewportSizeY, bool bHitTesting, float Gamma = 1.0f, const FSceneView* View = NULL, FTexture2DRHIRef DepthTexture = FTexture2DRHIRef(), EBlendModeFilter::Type Filter = EBlendModeFilter::All) const;
+	bool Draw(FRHICommandList& RHICmdList, const FDrawingPolicyRenderState& DrawRenderState, ERHIFeatureLevel::Type FeatureLevel, bool bNeedToSwitchVerticalAxis, const FMatrix& Transform, uint32 ViewportSizeX, uint32 ViewportSizeY, bool bHitTesting, float Gamma = 1.0f, const FSceneView* View = nullptr, FTexture2DRHIRef DepthTexture = FTexture2DRHIRef(), EBlendModeFilter::Type Filter = EBlendModeFilter::All) const;
 	
 	/**
 	 * Draws the batch
@@ -213,7 +214,7 @@ public:
 	 * @param Gamma			Optional gamma override
 	 * @param DepthTexture	DepthTexture for manual depth testing with editor compositing in the pixel shader
 	 */
-	bool Draw(FRHICommandList& RHICmdList, const struct FDrawingPolicyRenderState& DrawRenderState, ERHIFeatureLevel::Type FeatureLevel, bool bNeedToSwitchVerticalAxis, const FSceneView& View, bool bHitTesting, float Gamma = 1.0f, FTexture2DRHIRef DepthTexture = FTexture2DRHIRef(), EBlendModeFilter::Type Filter = EBlendModeFilter::All) const;
+	bool Draw(FRHICommandList& RHICmdList, const FDrawingPolicyRenderState& DrawRenderState, ERHIFeatureLevel::Type FeatureLevel, bool bNeedToSwitchVerticalAxis, const FSceneView& View, bool bHitTesting, float Gamma = 1.0f, FTexture2DRHIRef DepthTexture = FTexture2DRHIRef(), EBlendModeFilter::Type Filter = EBlendModeFilter::All) const;
 
 	/**
 	 * Creates a proxy FSceneView for operations that are not tied directly to a scene but still require batched elements to be drawn.
diff --git a/Engine/Source/Runtime/Renderer/Private/ScenePrivate.h b/Engine/Source/Runtime/Renderer/Private/ScenePrivate.h
--- a/Engine/Source/Runtime/Renderer/Private/ScenePrivate.h
+++ b/Engine/Source/Runtime/Renderer/Private/ScenePrivate.h
@@ -698,7 +698,10 @@ public:
 
 	// Pre-computed filter in spectral (i.e. FFT) domain along with data to determine if we need to up date it
 	struct {
+		/// @cond DOXYGEN_WARNINGS
 		void SafeRelease() { Spectral.SafeRelease(); CenterWeight.SafeRelease(); }
+		/// @endcond
+
 		// The 2d fourier transform of the physical space texture.
 		TRefCountPtr<IPooledRenderTarget> Spectral;
 		TRefCountPtr<IPooledRenderTarget> CenterWeight; // a 1-pixel buffer that holds blend weights for half-resolution fft.
diff --git a/Engine/Source/Runtime/SlateCore/Public/Fonts/FontCache.h b/Engine/Source/Runtime/SlateCore/Public/Fonts/FontCache.h
--- a/Engine/Source/Runtime/SlateCore/Public/Fonts/FontCache.h
+++ b/Engine/Source/Runtime/SlateCore/Public/Fonts/FontCache.h
@@ -228,7 +228,7 @@ public:
 		int32 TextLen;
 	};
 
-	FShapedGlyphSequence(TArray<FShapedGlyphEntry> InGlyphsToRender, const int16 InTextBaseline, const uint16 InMaxTextHeight, const UObject* InFontMaterial, const struct FFontOutlineSettings& InOutlineSettings, const FSourceTextRange& InSourceTextRange);
+	FShapedGlyphSequence(TArray<FShapedGlyphEntry> InGlyphsToRender, const int16 InTextBaseline, const uint16 InMaxTextHeight, const UObject* InFontMaterial, const FFontOutlineSettings& InOutlineSettings, const FSourceTextRange& InSourceTextRange);
 	~FShapedGlyphSequence();
 
 	/** Get the amount of memory allocated to this sequence */
diff --git a/Engine/Source/Runtime/UMG/Private/Components/TextBlock.cpp b/Engine/Source/Runtime/UMG/Private/Components/TextBlock.cpp
--- a/Engine/Source/Runtime/UMG/Private/Components/TextBlock.cpp
+++ b/Engine/Source/Runtime/UMG/Private/Components/TextBlock.cpp
@@ -163,6 +163,8 @@ void UTextBlock::SynchronizeProperties()
 	}
 }
 
+/// @cond DOXYGEN_WARNINGS
+
 FText UTextBlock::GetText() const
 {
 	if (MyTextBlock.IsValid())
@@ -173,6 +175,8 @@ FText UTextBlock::GetText() const
 	return Text;
 }
 
+/// @endcond
+
 void UTextBlock::SetText(FText InText)
 {
 	Text = InText;
diff --git a/Engine/Source/Runtime/UMG/Private/DragDropOperation.cpp b/Engine/Source/Runtime/UMG/Private/DragDropOperation.cpp
--- a/Engine/Source/Runtime/UMG/Private/DragDropOperation.cpp
+++ b/Engine/Source/Runtime/UMG/Private/DragDropOperation.cpp
@@ -11,6 +11,8 @@ UDragDropOperation::UDragDropOperation(const FObjectInitializer& ObjectInitializ
 	Pivot = EDragPivot::CenterCenter;
 }
 
+/// @cond DOXYGEN_WARNINGS
+
 void UDragDropOperation::Drop_Implementation(const FPointerEvent& PointerEvent)
 {
 	OnDrop.Broadcast(this);
@@ -25,3 +27,5 @@ void UDragDropOperation::Dragged_Implementation(const FPointerEvent& PointerEven
 {
 	OnDragged.Broadcast(this);
 }
+
+/// @endcond
