Copying //UE4/Dev-Console to //UE4/Dev-Main (Source: //UE4/Dev-Console @ 3483086)

diff --git a/Engine/Source/Runtime/Engine/Private/CsvProfiler.cpp b/Engine/Source/Runtime/Engine/Private/CsvProfiler.cpp
new file mode 100644
--- /dev/null
+++ b/Engine/Source/Runtime/Engine/Private/CsvProfiler.cpp
@@ -0,0 +1,794 @@
+// Copyright 1998-2016 Epic Games, Inc. All Rights Reserved.
+
+/**
+*
+* A lightweight single-threaded CSV profiler which can be used for profiling in Test/Shipping builds
+*/
+
+#include "CsvProfiler.h"
+#include "CoreGlobals.h"
+#include "HAL/RunnableThread.h"
+#include "HAL/ThreadManager.h"
+#include "HAL/FileManager.h"
+#include "HAL/PlatformProcess.h"
+#include "Misc/Paths.h"
+#include "Misc/CommandLine.h"
+#include "Misc/ScopeLock.h"
+#include "RenderingThread.h"
+#include "CoreDelegates.h"
+#include "RenderCore.h"
+#include "Modules/ModuleManager.h"
+#include "EngineGlobals.h"
+
+#if CSV_PROFILER
+
+#define RECORD_TIMESTAMPS 1 
+
+DEFINE_LOG_CATEGORY_STATIC(LogCsvProfiler, Log, All);
+
+#define USE_THREAD_LOCAL 0
+
+FCsvProfiler* FCsvProfiler::Instance = nullptr;
+
+static void HandleCSVProfileCommand(const TArray<FString>& Args)
+{
+	if (Args.Num() < 1)
+	{
+		return;
+	}
+	FString Param = Args[0];
+	if (Param == TEXT("START"))
+	{
+		FCsvProfiler::Get()->BeginCapture();
+	}
+	else if (Param == TEXT("STOP"))
+	{
+		FCsvProfiler::Get()->EndCapture();
+	}
+	else
+	{
+		int32 CaptureFrames = 0;
+		if (FParse::Value(*Param, TEXT("FRAMES="), CaptureFrames))
+		{
+			FCsvProfiler::Get()->BeginCapture(CaptureFrames);
+		}
+	}
+}
+
+static void CsvProfilerBeginFrame()
+{
+	FCsvProfiler::Get()->BeginFrame();
+}
+
+static void CsvProfilerEndFrame()
+{
+	FCsvProfiler::Get()->EndFrame();
+}
+
+
+static FAutoConsoleCommand HandleCSVProfileCmd(
+	TEXT("CsvProfile"),
+	TEXT("Starts or stops Csv Profiles"),
+	FConsoleCommandWithArgsDelegate::CreateStatic(&HandleCSVProfileCommand)
+);
+
+ 
+//-----------------------------------------------------------------------------
+//	TSingleProducerSingleConsumerList : fast lock-free single producer/single 
+//  consumer list implementation. 
+//  Uses a linked list of blocks for allocations. Note that one block will always 
+//  leak, because removing the tail cannot be done without locking
+//-----------------------------------------------------------------------------
+template <class T, int BlockSize>
+class TSingleProducerSingleConsumerList
+{
+	// A block of BlockSize entries
+	struct FBlock
+	{
+		FBlock() : Next(nullptr)
+		{
+		}
+		T Entries[BlockSize];
+		FBlock* Next;
+	};
+
+	//MS_ALIGN(PLATFORM_CACHE_LINE_SIZE) 
+	struct FCacheLineAlignedCounter
+	{
+		#define NUM_PADDING_WORDS ( PLATFORM_CACHE_LINE_SIZE/4 - 2 )
+		uint32 Pad[NUM_PADDING_WORDS];
+
+		volatile uint64 Value;
+	};
+	//GCC_ALIGN(PLATFORM_CACHE_LINE_SIZE);
+
+public:
+	TSingleProducerSingleConsumerList()
+	{
+		HeadBlock = nullptr;
+		TailBlock = nullptr;
+#if DO_CHECK
+		bElementReserved = false;
+#endif
+		Counter.Value = 0;
+		ConsumerThreadLastReadIndex = 0;
+	}
+
+	// Reserve an element prior to writing it
+	// Must be called from the Producer thread
+	T* ReserveElement()
+	{
+#if DO_CHECK
+		check(!bElementReserved);
+		bElementReserved = true;
+#endif
+		uint32 TailBlockSize = Counter.Value % BlockSize;
+		if (TailBlockSize == 0)
+		{
+			AddTailBlock();
+		}
+		return &TailBlock->Entries[TailBlockSize];
+	}
+
+	// Commit an element after writing it
+	// Must be called from the Producer thread after a call to ReserveElement
+	void CommitElement()
+	{
+#if DO_CHECK
+		check(bElementReserved);
+		bElementReserved = false;
+#endif
+		FPlatformMisc::MemoryBarrier();
+		
+		// Keep track of the count of all the elements we ever committed. This value is never reset, even on a PopAll
+		Counter.Value++;
+	}
+
+	// Called from the consumer thread
+	void PopAll(TArray<T>& ElementsOut)
+	{
+		volatile uint64 CurrentCounterValue = Counter.Value;
+		FPlatformMisc::MemoryBarrier();
+
+		uint32 ElementCount = uint32( CurrentCounterValue - ConsumerThreadLastReadIndex );
+		ConsumerThreadLastReadIndex = CurrentCounterValue;
+		ElementsOut.AddUninitialized(ElementCount);
+
+		uint32 IndexInBlock = 0;
+		for (uint32 i = 0; i < ElementCount; i++)
+		{
+			// if this block is full and it's completed, retire it and move to the next block (update the head)
+			if (i > 0 && ( i % BlockSize == 0) )
+			{
+				// Both threads are done with the head block now, so we can safely delete it 
+				// Note that the Producer thread only reads/writes to the HeadBlock pointer on startup, so it's safe to update it at this point
+				FBlock* PrevBlock = HeadBlock;
+				HeadBlock = HeadBlock->Next;
+				IndexInBlock = 0;
+				delete PrevBlock;
+			}
+			check(HeadBlock != nullptr);
+			ElementsOut[i] = HeadBlock->Entries[IndexInBlock++];
+		}
+	}
+private:
+	void AddTailBlock()
+	{
+		FBlock* NewTail = new FBlock;
+		if (TailBlock == nullptr)
+		{
+			// This must only happen on startup, otherwise it's not thread-safe
+			check(Counter.Value == 0);
+			check(HeadBlock == nullptr);
+			HeadBlock = NewTail;
+		}
+		else
+		{
+			TailBlock->Next = NewTail;
+		}
+		TailBlock = NewTail;
+	}
+
+
+	FBlock* HeadBlock;
+	FBlock* TailBlock;
+
+	FCacheLineAlignedCounter Counter;
+
+	// Used from the consumer thread
+	uint64 ConsumerThreadLastReadIndex;
+
+#if DO_CHECK
+	bool bElementReserved;
+#endif
+
+};
+
+//-----------------------------------------------------------------------------
+//	FCsvTimingMarker : records timestamps. Uses StatName pointer as a unique ID
+//-----------------------------------------------------------------------------
+namespace ECsvMarkerType
+{
+	enum Type 
+	{
+		TimestampStart,
+		TimestampEnd,
+		CustomStat
+	};
+};
+
+struct FCsvTimingMarker
+{
+	const char*			 StatName;
+	uint64				 Timestamp;
+	float				 CustomValue;
+	ECsvMarkerType::Type MarkerType;
+};
+
+//-----------------------------------------------------------------------------
+//	FCsvProfilerThread class : records all timings for a particular thread
+//-----------------------------------------------------------------------------
+class FCsvProfilerThread
+{
+public:
+	FCsvProfilerThread( uint32 InThreadId ) : ThreadId(InThreadId)
+	{
+		CurrentCaptureStartCycles = FPlatformTime::Cycles64();
+	}
+
+	void FlushResults(TArray<FCsvTimingMarker>& OutMarkers )
+	{
+		uint64 ValidTimestampStart = CurrentCaptureStartCycles;
+		CurrentCaptureStartCycles = FPlatformTime::Cycles64();
+
+		TimingMarkers.PopAll(OutMarkers);
+	}
+
+	void AddTimestamp(const char* StatName, const bool bBegin)
+	{
+		// TODO: fast pool allocator for timing markers
+		FCsvTimingMarker* Marker = TimingMarkers.ReserveElement();
+
+		Marker->MarkerType = bBegin ? ECsvMarkerType::TimestampStart : ECsvMarkerType::TimestampEnd;
+		Marker->StatName = StatName;
+		Marker->CustomValue = 0.0f;
+		Marker->Timestamp = FPlatformTime::Cycles64();
+
+		TimingMarkers.CommitElement();
+	}
+
+	void AddCustomStat(const char* StatName, const float Value )
+	{
+		FCsvTimingMarker* Marker = TimingMarkers.ReserveElement();
+
+		Marker->MarkerType = ECsvMarkerType::CustomStat;
+		Marker->StatName = StatName;
+		Marker->CustomValue = Value;
+		Marker->Timestamp = FPlatformTime::Cycles64();
+
+		TimingMarkers.CommitElement();
+	}
+
+	uint32 ThreadId;
+	uint64 CurrentCaptureStartCycles;
+	TSingleProducerSingleConsumerList<FCsvTimingMarker, 128> TimingMarkers;
+};
+
+
+FCsvProfiler* FCsvProfiler::Get()
+{
+	if (Instance == nullptr)
+	{
+		Instance = new FCsvProfiler;
+	}
+	return Instance;
+}
+
+FCsvProfiler::FCsvProfiler() 
+	: NumFramesToCapture(-1)
+	, CaptureFrameNumber(0)
+	, bRequestStartCapture(false)
+	, bRequestStopCapture(false)
+	, bCapturing(false)
+{
+	check(IsInGameThread());
+	GetProfilerThread();
+
+	FCoreDelegates::OnBeginFrame.AddStatic(CsvProfilerBeginFrame);
+	FCoreDelegates::OnEndFrame.AddStatic(CsvProfilerEndFrame);
+
+}
+
+/** Per-frame update */
+void FCsvProfiler::BeginFrame()
+{
+	check(IsInGameThread());
+	if (bRequestStartCapture)
+	{
+		bCapturing = true;
+		bRequestStartCapture = false;
+
+		UE_LOG(LogCsvProfiler, Display, TEXT("Capture Starting"));
+
+		LastEndFrameTimestamp = FPlatformTime::Cycles64();
+	}
+
+	if (bCapturing)
+	{
+		FrameBeginTimestamps.Add(FPlatformTime::Cycles64());
+	}
+
+	ENQUEUE_UNIQUE_RENDER_COMMAND(
+		CSVBeginFrame,
+		{
+			FCsvProfiler::Get()->BeginFrameRT();
+		});
+
+}
+
+void FCsvProfiler::EndFrame()
+{
+	check(IsInGameThread());
+
+	ENQUEUE_UNIQUE_RENDER_COMMAND(
+		CSVEndFrame,
+		{
+			FCsvProfiler::Get()->EndFrameRT();
+		});
+
+	if (!bCapturing)
+	{
+		return;
+	}
+
+
+	// CSV profiler 
+	{
+		CSV_CUSTOM_STAT(RenderThreadTime, FPlatformTime::ToMilliseconds(GRenderThreadTime));
+		CSV_CUSTOM_STAT(GameThreadTime, FPlatformTime::ToMilliseconds(GGameThreadTime));
+		CSV_CUSTOM_STAT(GPUTime, FPlatformTime::ToMilliseconds(GGPUFrameTime));
+		FPlatformMemoryStats MemoryStats = FPlatformMemory::GetStats();
+		float PhysicalMBFree = float(MemoryStats.AvailablePhysical / 1024) / 1024.0f;
+		CSV_CUSTOM_STAT(MemoryFreeMB, PhysicalMBFree);
+	}
+
+	{
+		// Record the frametime (measured since the last EndFrame)
+		uint64 CurrentTimeStamp = FPlatformTime::Cycles64();
+		uint64 ElapsedCycles = CurrentTimeStamp - LastEndFrameTimestamp;
+		float ElapsedMs = FPlatformTime::ToMilliseconds64(ElapsedCycles);
+		CSV_CUSTOM_STAT(FrameTime, ElapsedMs);
+		LastEndFrameTimestamp = CurrentTimeStamp;
+	}
+
+	if (NumFramesToCapture >= 0)
+	{
+		NumFramesToCapture--;
+		if (NumFramesToCapture == 0)
+		{
+			bRequestStopCapture = true;
+		}
+	}
+
+	if (bRequestStopCapture)
+	{
+		FScopeLock Lock(&GetResultsLock);
+
+		bCapturing = false;
+
+		// We'd need to move this into engine to do this:
+		FlushRenderingCommands();
+
+		// Sleep for a bit to give the threads a chance to finish writing stats
+		// Doesn't matter massively if they haven't, since we're using a lock-free queue
+		//FPlatformProcess::Sleep(0.002f);
+
+		WriteCaptureToFile();
+
+		FrameBeginTimestamps.Empty();
+		FrameBeginTimestampsRT.Empty();
+
+		bRequestStopCapture = false;
+	}
+
+	CaptureFrameNumber++;
+}
+
+
+/** Per-frame update */
+void FCsvProfiler::BeginFrameRT()
+{
+	check(IsInRenderingThread());
+	if (bCapturing)
+	{
+		// Mark where the renderthread frames begin
+		FScopeLock Lock(&GetResultsLock);
+		FrameBeginTimestampsRT.Add(FPlatformTime::Cycles64());
+	}
+}
+
+void FCsvProfiler::EndFrameRT()
+{
+	check(IsInRenderingThread());
+}
+
+
+/** Final cleanup */
+void FCsvProfiler::Release()
+{
+
+}
+
+void FCsvProfiler::BeginCapture(int InNumFramesToCapture)
+{
+	NumFramesToCapture = InNumFramesToCapture;
+	bRequestStartCapture = true;
+}
+
+void FCsvProfiler::EndCapture()
+{
+	bRequestStopCapture = true;
+}
+
+class FCsvColumn
+{
+public:
+	FCsvColumn(const FString& InName) : Name(InName) {}
+	
+	float GetValue(uint32 Row)
+	{
+		if (Row < (uint32)Values.Num())
+		{
+			return Values[Row];
+		}
+		return 0.0f;
+	}
+
+	void SetValue(uint32 Row, float Value)
+	{
+		// Grow the column if needed (pad with zeroes)
+		if (Row >= (uint32)Values.Num())
+		{
+			Values.Reserve((int32)Row + 1);
+			int32 GrowBy = (int32)Row + 1 - Values.Num();
+			for (int i = 0; i < GrowBy; i++)
+			{
+				Values.Add(0.0f);
+			}
+		}
+		Values[Row] = Value;
+	}
+
+	FString Name;
+	TArray<float> Values;
+};
+
+class FCSVTable
+{
+public:
+	FCSVTable() : NumRows(0)
+	{
+	}
+
+	~FCSVTable()
+	{
+		for (int i = 0; i < Columns.Num(); i++)
+		{
+			delete Columns[i];
+		}
+	}
+
+	uint32 AddColumn(const FString& ColumnName)
+	{
+		Columns.Add(new FCsvColumn(ColumnName));
+		return (uint32)Columns.Num() - 1;
+	}
+
+	void SetValue(uint32 Row, uint32 Column, float Value)
+	{
+		if (Column < (uint32)Columns.Num())
+		{
+			Columns[Column]->SetValue(Row, Value);
+		}
+
+		if (Row >= NumRows)
+		{
+			NumRows = Row + 1;
+		}
+		
+	}
+
+	void AccumulateValue(uint32 Row, uint32 Column, float Value)
+	{
+		if (Column < (uint32)Columns.Num())
+		{
+			Columns[Column]->SetValue(Row, Columns[Column]->GetValue(Row) + Value);
+		}
+		if (Row >= NumRows)
+		{
+			NumRows = Row + 1;
+		}
+	}
+
+	void WriteToFile(const FString& FileName)
+	{
+		char Comma = ',';
+		char NewLine = '\n';
+		FArchive* OutputFile = IFileManager::Get().CreateDebugFileWriter(*FileName);
+		for (int c = 0; c< Columns.Num(); c++)
+		{
+			FCsvColumn& Column = *Columns[c];
+			
+			auto AnsiStr = StringCast<ANSICHAR>(*Column.Name);
+			OutputFile->Serialize( (void*)AnsiStr.Get(), AnsiStr.Length() );
+			if (c < Columns.Num() - 1)
+			{
+				OutputFile->Serialize( (void*)&Comma, sizeof(ANSICHAR));
+			}
+		}
+		OutputFile->Serialize((void*)&NewLine, sizeof(ANSICHAR));
+
+		// Write the data
+		char StringBuffer[256];
+		for (int r = 0; r < (int)NumRows; r++)
+		{
+			for (int c = 0; c < Columns.Num(); c++)
+			{
+				FCsvColumn& Column = *Columns[c];
+				float Value = Column.GetValue(r);
+				FCStringAnsi::Snprintf(StringBuffer, 256, "%.4f", Value);
+				OutputFile->Serialize((void*)StringBuffer, sizeof(ANSICHAR)*FCStringAnsi::Strlen(StringBuffer));
+
+				if (c < Columns.Num() - 1)
+				{
+					OutputFile->Serialize((void*)&Comma, sizeof(ANSICHAR));
+				}
+			}
+			OutputFile->Serialize((void*)&NewLine, sizeof(ANSICHAR));
+		}
+
+		OutputFile->Close();
+	}
+private:
+	TArray<FCsvColumn*> Columns;
+	uint32 NumRows;
+};
+
+void FCsvProfiler::WriteCaptureToFile()
+{
+
+	// Write the results out to a file
+	FString Filename = FString::Printf(TEXT("Profile(%s)"), *FDateTime::Now().ToString(TEXT("%Y%m%d_%H%M%S")));
+	FString CSVRootPath = FPaths::ProfilingDir() + TEXT("CSV/");
+	FString OutputFilename = CSVRootPath + Filename + TEXT(".csv");
+	FCSVTable Csv;
+
+	for (int t = 0; t < ProfilerThreads.Num(); t++)
+	{
+		FCsvProfilerThread* CurrentThread = ProfilerThreads[t];
+
+		// Get the thread name
+		FString ThreadName = FThreadManager::Get().GetThreadName(CurrentThread->ThreadId);
+		if (CurrentThread->ThreadId == GGameThreadId)
+		{
+			ThreadName = TEXT("GameThread");
+		}
+		else if (CurrentThread->ThreadId == GRenderThreadId )
+		{
+			ThreadName = TEXT("RenderThread");
+		}
+
+		// Read the results
+		TArray<FCsvTimingMarker> ThreadMarkers;
+		CurrentThread->FlushResults(ThreadMarkers);
+
+		TMap<FString, uint32> StatNameToColumnIndex;
+		TArray<uint32> MarkerFrameIndices;
+		TArray<uint32> MarkerColumnIndices;
+
+		FString Slash(TEXT("/"));
+
+		MarkerFrameIndices.Reserve(ThreadMarkers.Num());
+		MarkerColumnIndices.Reserve(ThreadMarkers.Num());
+
+		bool bUseRenderthreadTimeline = CurrentThread->ThreadId == GRenderThreadId;
+		const TArray<uint64>& FrameBeginTimestampsForThread = bUseRenderthreadTimeline ? FrameBeginTimestampsRT : FrameBeginTimestamps;
+
+		// Create a column in the CSV for each unique stat name
+		//static bool bCreateEndToEndMarkers = false;
+		uint32 CurrentFrameIndex = 0;
+		for (int i = 0; i < ThreadMarkers.Num(); i++)
+		{
+			FCsvTimingMarker& Marker = ThreadMarkers[i];
+			FString StatName = (TCHAR*)StringCast<TCHAR>((Marker.StatName)).Get();
+			uint32 ColumnIndex = 0;
+			uint32* IndexPtr = StatNameToColumnIndex.Find(StatName);
+
+			if (IndexPtr)
+			{
+				ColumnIndex = *IndexPtr;
+			}
+			else
+			{
+				FString ColumnName = (Marker.MarkerType == ECsvMarkerType::CustomStat) ? StatName : ThreadName + Slash + StatName;
+				ColumnIndex = Csv.AddColumn(ColumnName);
+				StatNameToColumnIndex.Add(StatName, ColumnIndex);
+#if 0
+				if (bCreateEndToEndMarkers && Marker.MarkerType != ECsvMarkerType::CustomStat)
+				{
+					// Add another column for the end to end markers
+					FString ColumnName2 = ThreadName + Slash + TEXT("EndToEnd_") + StatName;
+					Csv.AddColumn(ColumnName2);
+				}
+#endif
+			}
+			MarkerColumnIndices.Add(ColumnIndex);
+
+			while (CurrentFrameIndex + 1 < (uint32)FrameBeginTimestampsForThread.Num() && Marker.Timestamp > FrameBeginTimestampsForThread[CurrentFrameIndex + 1])
+			{
+				CurrentFrameIndex++;
+			}
+			MarkerFrameIndices.Add(CurrentFrameIndex);
+		}
+
+		// Process the markers
+		TArray<FCsvTimingMarker*> MarkerStack;
+		for (int i = 0; i < ThreadMarkers.Num(); i++)
+		{
+			FCsvTimingMarker& Marker = ThreadMarkers[i];
+			uint32 ColumnIndex = MarkerColumnIndices[i];
+			uint32 FrameIndex = MarkerFrameIndices[i];
+
+			switch (Marker.MarkerType)
+			{
+				case ECsvMarkerType::TimestampStart:
+				{
+					MarkerStack.Push(&Marker);
+				}
+				break;
+				case ECsvMarkerType::TimestampEnd:
+				{
+					// Markers might not match up if they were truncated mid-frame, so we need to be robust to that
+					if (MarkerStack.Num() > 0)
+					{
+						FCsvTimingMarker* StartMarker = MarkerStack.Pop();
+						if (Marker.Timestamp > StartMarker->Timestamp)
+						{
+							uint64 ElapsedCycles = Marker.Timestamp - StartMarker->Timestamp;
+							float ElapsedMs = FPlatformTime::ToMilliseconds64(ElapsedCycles);
+
+							// Add the elapsed time to the table entry for this frame/stat 
+							Csv.AccumulateValue(FrameIndex, ColumnIndex, ElapsedMs);
+						}
+					}
+				}
+				break;
+				case ECsvMarkerType::CustomStat:
+				{
+					// Add the elapsed time to the table entry for this frame/stat 
+					Csv.AccumulateValue(FrameIndex, ColumnIndex, Marker.CustomValue);
+				}
+				break;
+			}
+
+#if 0
+			if (bCreateEndToEndMarkers)
+			{
+				// Array sized by num CSV Columns
+				// Track frame of each
+				// Process the markers
+				TArray<FCsvTimingMarker*> MarkerStack;
+				for (int i = 0; i < ThreadMarkers.Num(); i++)
+				{
+					FCsvTimingMarker& Marker = ThreadMarkers[i];
+					uint32 ColumnIndex = MarkerColumnIndices[i]+1;
+					uint32 FrameIndex = MarkerFrameIndices[i];
+
+					switch (Marker.MarkerType)
+					{
+						case ECsvMarkerType::TimestampStart:
+						{
+							MarkerStack.Push(&Marker);
+						}
+						break;
+						case ECsvMarkerType::TimestampEnd:
+						{
+							// Markers might not match up if they were truncated mid-frame, so we need to be robust to that
+							if (MarkerStack.Num() > 0)
+							{
+								FCsvTimingMarker* StartMarker = MarkerStack.Pop();
+								if (Marker.Timestamp > StartMarker->Timestamp)
+								{
+									uint64 ElapsedCycles = Marker.Timestamp - StartMarker->Timestamp;
+									float ElapsedMs = FPlatformTime::ToMilliseconds64(ElapsedCycles);
+
+									// Add the elapsed time to the table entry for this frame/stat 
+									Csv.AccumulateValue(FrameIndex, ColumnIndex, ElapsedMs);
+								}
+							}
+						}
+						break;
+					}
+				}
+#endif
+
+		}
+
+
+	} // Thread loop
+
+	UE_LOG(LogCsvProfiler, Display, TEXT("Capture Ended. Writing CSV to file : %s"), *OutputFilename);
+
+	Csv.WriteToFile(OutputFilename);
+}
+
+
+// Create the TLS profiler thread lazily
+FCsvProfilerThread* FCsvProfiler::GetProfilerThread()
+{
+#if USE_THREAD_LOCAL == 1
+	static thread_local FCsvProfilerThread* ProfilerThread = nullptr;
+#else 
+	static uint32 TlsSlot = FPlatformTLS::AllocTlsSlot();
+	FCsvProfilerThread* ProfilerThread = (FCsvProfilerThread*)FPlatformTLS::GetTlsValue(TlsSlot);
+#endif
+	if (!ProfilerThread)
+	{
+		ProfilerThread = new FCsvProfilerThread(FPlatformTLS::GetCurrentThreadId());
+#if USE_THREAD_LOCAL == 0
+		FPlatformTLS::SetTlsValue(TlsSlot, ProfilerThread);
+#endif
+		{
+			FScopeLock Lock(&GetResultsLock);
+			ProfilerThreads.Add(ProfilerThread);
+		}
+	}
+	return ProfilerThread;
+}
+
+/** Push/pop events */
+void FCsvProfiler::BeginStat(const char * StatName)
+{
+#if RECORD_TIMESTAMPS
+	if (!bCapturing)
+	{
+		return;
+	}
+	GetProfilerThread()->AddTimestamp(StatName, true);
+#endif
+}
+
+void FCsvProfiler::EndStat(const char * StatName)
+{
+#if RECORD_TIMESTAMPS
+	if (!bCapturing)
+	{
+		return;
+	}
+	GetProfilerThread()->AddTimestamp(StatName, false);
+#endif
+}
+
+void FCsvProfiler::RecordCustomStat(const char * StatName, float Value)
+{
+	if (!bCapturing)
+	{
+		return;
+	}
+	GetProfilerThread()->AddCustomStat(StatName, Value);
+}
+
+
+void FCsvProfiler::Init()
+{
+	int32 NumCsvFrames = 0;
+	if (FParse::Value(FCommandLine::Get(), TEXT("csvCaptureFrames="), NumCsvFrames))
+	{
+		check(IsInGameThread());
+		BeginCapture(NumCsvFrames);
+		BeginFrame();
+	}
+}
+
+
+#endif // CSV_PROFILER
\ No newline at end of file
diff --git a/Engine/Source/Runtime/Engine/Private/DynamicMeshBuilder.cpp b/Engine/Source/Runtime/Engine/Private/DynamicMeshBuilder.cpp
--- a/Engine/Source/Runtime/Engine/Private/DynamicMeshBuilder.cpp
+++ b/Engine/Source/Runtime/Engine/Private/DynamicMeshBuilder.cpp
@@ -160,10 +160,8 @@ public:
 	FVertexBufferRHIRef CreateResource(FGlobalDynamicMeshPoolPolicy::CreationArguments Args)
 	{
 		FGlobalDynamicMeshPoolPolicy::CreationArguments BufferSize = GetPoolBucketSize(GetPoolBucketIndex(Args));
-		// The use of BUF_Static is deliberate - on OS X the buffer backing-store orphaning & reallocation will dominate execution time
-		// so to avoid this we don't reuse a buffer for several frames, thereby avoiding the pipeline stall and the reallocation cost.
 		FRHIResourceCreateInfo CreateInfo;
-		FVertexBufferRHIRef VertexBuffer = RHICreateVertexBuffer(BufferSize, BUF_Static, CreateInfo);
+		FVertexBufferRHIRef VertexBuffer = RHICreateVertexBuffer(BufferSize, BUF_Volatile, CreateInfo);
 		return VertexBuffer;
 	}
 	
diff --git a/Engine/Source/Runtime/Engine/Private/Rendering/SubsurfaceProfile.cpp b/Engine/Source/Runtime/Engine/Private/Rendering/SubsurfaceProfile.cpp
--- a/Engine/Source/Runtime/Engine/Private/Rendering/SubsurfaceProfile.cpp
+++ b/Engine/Source/Runtime/Engine/Private/Rendering/SubsurfaceProfile.cpp
@@ -14,12 +14,6 @@ ENGINE_API TGlobalResource<FSubsurfaceProfileTexture> GSubsurfaceProfileTextureO
 // Texture with one or more SubSurfaceProfiles or 0 if there is no user
 static TRefCountPtr<IPooledRenderTarget> GSSProfiles;
 
-static TAutoConsoleVariable<int32> CVarFastVRamSubsurfaceProfile(
-	TEXT("r.FastVRamSubsurfaceProfile"),
-	0,
-	TEXT("Whether to store subsurfaceprofile in fast VRAM"),
-	ECVF_Scalability | ECVF_RenderThreadSafe);
-
 FSubsurfaceProfileTexture::FSubsurfaceProfileTexture()
 {
 	check(IsInGameThread());
@@ -141,10 +135,6 @@ void FSubsurfaceProfileTexture::CreateTexture(FRHICommandListImmediate& RHICmdLi
 
 	// at minimum 64 lines (less reallocations)
 	FPooledRenderTargetDesc Desc(FPooledRenderTargetDesc::Create2DDesc(FIntPoint(Width, FMath::Max(Height, (uint32)64)), PF_B8G8R8A8, FClearValueBinding::None, 0, TexCreate_None, false));
-	if (CVarFastVRamSubsurfaceProfile.GetValueOnRenderThread() >= 1)
-	{
-		Desc.Flags |= TexCreate_FastVRAM;
-	}
 	if (b16Bit)
 	{
 		Desc.Format = PF_A16B16G16R16;
diff --git a/Engine/Source/Runtime/Engine/Private/StaticMesh.cpp b/Engine/Source/Runtime/Engine/Private/StaticMesh.cpp
--- a/Engine/Source/Runtime/Engine/Private/StaticMesh.cpp
+++ b/Engine/Source/Runtime/Engine/Private/StaticMesh.cpp
@@ -2703,17 +2703,37 @@ void UStaticMesh::PostLoad()
 		CalculateExtendedBounds();
 	}
 
-	if (SectionInfoMap.Map.Num() == 0)
+	//Always redo the whole SectionInfoMap to be sure it contain only valid data
+	//This will reuse everything valid from the just serialize SectionInfoMap.
+	FMeshSectionInfoMap TempOldSectionInfoMap = SectionInfoMap;
+	SectionInfoMap.Clear();
+	for (int32 LODResourceIndex = 0; LODResourceIndex < RenderData->LODResources.Num(); ++LODResourceIndex)
 	{
-		// Before this serialization issue was fixed, some assets were resaved and permanently lost their section info map.
-		// This attempts to recreate it based on the render data.
-		SectionInfoMap.Clear();
-		for (int32 LODResourceIndex = 0; LODResourceIndex < RenderData->LODResources.Num(); ++LODResourceIndex)
+		FStaticMeshLODResources& LOD = RenderData->LODResources[LODResourceIndex];
+		for (int32 SectionIndex = 0; SectionIndex < LOD.Sections.Num(); ++SectionIndex)
 		{
-			FStaticMeshLODResources& LOD = RenderData->LODResources[LODResourceIndex];
-			const int32 NumSections = LOD.Sections.Num();
-			for (int32 SectionIndex = 0; SectionIndex < NumSections; ++SectionIndex)
+			if (TempOldSectionInfoMap.IsValidSection(LODResourceIndex, SectionIndex))
+			{
+				FMeshSectionInfo Info = TempOldSectionInfoMap.Get(LODResourceIndex, SectionIndex);
+				if (StaticMaterials.IsValidIndex(Info.MaterialIndex))
+				{
+					//Reuse the valid data that come from the serialize
+					SectionInfoMap.Set(LODResourceIndex, SectionIndex, Info);
+				}
+				else
+				{
+					//Use the render data material index, but keep the flags (collision, shadow...)
+					const int32 MaterialIndex = LOD.Sections[SectionIndex].MaterialIndex;
+					if (StaticMaterials.IsValidIndex(MaterialIndex))
+					{
+						Info.MaterialIndex = MaterialIndex;
+						SectionInfoMap.Set(LODResourceIndex, SectionIndex, Info);
+					}
+				}
+			}
+			else
 			{
+				//Create a new SectionInfoMap from the render data
 				const int32 MaterialIndex = LOD.Sections[SectionIndex].MaterialIndex;
 				if (StaticMaterials.IsValidIndex(MaterialIndex))
 				{
diff --git a/Engine/Source/Runtime/Engine/Private/UnrealEngine.cpp b/Engine/Source/Runtime/Engine/Private/UnrealEngine.cpp
--- a/Engine/Source/Runtime/Engine/Private/UnrealEngine.cpp
+++ b/Engine/Source/Runtime/Engine/Private/UnrealEngine.cpp
@@ -1381,6 +1381,9 @@ void UEngine::UpdateTimeAndHandleMaxTickRate()
 		double ActualWaitTime = 0.f;
 		if( WaitTime > 0 )
 		{
+			// track all this waiting so that Game Thread is correct
+			FThreadIdleStats::FScopeIdle Scope;
+			
 			FSimpleScopeSecondsCounter ActualWaitTimeCounter(ActualWaitTime);
 			double WaitEndTime = FApp::GetCurrentTime() + WaitTime;
 
@@ -7199,13 +7202,13 @@ void UEngine::AddOnScreenDebugMessage(uint64 Key, float TimeToDisplay, FColor Di
 		{
 			if (bNewerOnTop)
 			{
-				FScreenMessageString* NewMessage = new(PriorityScreenMessages)FScreenMessageString();
-				check(NewMessage);
-				NewMessage->Key = Key;
-				NewMessage->ScreenMessage = DebugMessage;
-				NewMessage->DisplayColor = DisplayColor;
-				NewMessage->TimeToDisplay = TimeToDisplay;
-				NewMessage->CurrentTimeDisplayed = 0.0f;
+			FScreenMessageString* NewMessage = new(PriorityScreenMessages)FScreenMessageString();
+			check(NewMessage);
+			NewMessage->Key = Key;
+			NewMessage->ScreenMessage = DebugMessage;
+			NewMessage->DisplayColor = DisplayColor;
+			NewMessage->TimeToDisplay = TimeToDisplay;
+			NewMessage->CurrentTimeDisplayed = 0.0f;
 				NewMessage->TextScale = TextScale;
 			}
 			else
diff --git a/Engine/Source/Runtime/Engine/Private/World.cpp b/Engine/Source/Runtime/Engine/Private/World.cpp
--- a/Engine/Source/Runtime/Engine/Private/World.cpp
+++ b/Engine/Source/Runtime/Engine/Private/World.cpp
@@ -2970,7 +2970,7 @@ void UWorld::ConditionallyBuildStreamingData()
 
 bool UWorld::IsVisibilityRequestPending() const
 {
-	return (CurrentLevelPendingVisibility != nullptr && CurrentLevelPendingInvisibility != nullptr);
+	return (CurrentLevelPendingVisibility != nullptr || CurrentLevelPendingInvisibility != nullptr);
 }
 
 bool UWorld::AreAlwaysLoadedLevelsLoaded() const
diff --git a/Engine/Source/Runtime/Engine/Public/AnimationRuntime.h b/Engine/Source/Runtime/Engine/Public/AnimationRuntime.h
--- a/Engine/Source/Runtime/Engine/Public/AnimationRuntime.h
+++ b/Engine/Source/Runtime/Engine/Public/AnimationRuntime.h
@@ -37,10 +37,7 @@ namespace ETransformBlendMode
 }
 
 template<int32>
-void BlendTransform(const FTransform& Source, FTransform& Dest, const float BlendWeight)
-{
-	check(false); /// should never call this
-}
+FORCEINLINE void BlendTransform(const FTransform& Source, FTransform& Dest, const float BlendWeight);
 
 template<>
 FORCEINLINE void BlendTransform<ETransformBlendMode::Overwrite>(const FTransform& Source, FTransform& Dest, const float BlendWeight)
diff --git a/Engine/Source/Runtime/Engine/Public/CsvProfiler.h b/Engine/Source/Runtime/Engine/Public/CsvProfiler.h
new file mode 100644
--- /dev/null
+++ b/Engine/Source/Runtime/Engine/Public/CsvProfiler.h
@@ -0,0 +1,96 @@
+// Copyright 1998-2016 Epic Games, Inc. All Rights Reserved.
+
+/**
+*
+* A lightweight single-threaded CSV profiler which can be used for profiling in Test/Shipping builds 
+*/
+
+#pragma once
+
+#include "Modules/ModuleInterface.h"
+#include "Modules/ModuleManager.h"
+#include "CoreTypes.h"
+#include "CoreMinimal.h"
+
+#define CSV_PROFILER (WITH_ENGINE && !UE_BUILD_SHIPPING)
+
+#if CSV_PROFILER
+  #define CSV_SCOPED_STAT(StatName)							FScopedCsvStat ScopedStat_ ## StatName (#StatName);
+  #define CSV_CUSTOM_STAT(StatName,Value) 					FCsvProfiler::Get()->RecordCustomStat(#StatName, Value)
+#else
+  #define CSV_SCOPED_STAT(StatName) 
+  #define CSV_CUSTOM_STAT(StatName,Value)							
+#endif
+
+#if CSV_PROFILER
+class FCsvProfilerFrame;
+class FCsvProfilerThread;
+
+/**
+* FCsvProfiler class. This manages recording and reporting all for CSV stats
+*/
+class FCsvProfiler
+{
+private:
+	static FCsvProfiler* Instance;
+	FCsvProfiler();
+public:
+	static ENGINE_API FCsvProfiler* Get();
+
+	ENGINE_API void Init();
+
+	/** Push/pop events */
+	ENGINE_API void BeginStat(const char * StatName);
+	ENGINE_API void EndStat(const char * StatName);
+
+	ENGINE_API void RecordCustomStat(const char * StatName, float Value);
+
+	/** Per-frame update */
+	void BeginFrame();
+	void EndFrame();
+
+	/** Begin/End Capture */
+	void BeginCapture(int InNumFramesToCapture = -1);
+	void EndCapture();
+
+	/** Final cleanup */
+	void Release();
+
+private:
+	/** Renderthread begin/end frame */
+	void BeginFrameRT();
+	void EndFrameRT();
+
+	void WriteCaptureToFile();
+	FCsvProfilerThread* GetProfilerThread();
+
+	int32 NumFramesToCapture;
+	int32 CaptureFrameNumber;
+	bool bRequestStartCapture;
+	bool bRequestStopCapture;
+	volatile bool bCapturing;
+	uint64 LastEndFrameTimestamp;
+
+	FCriticalSection GetResultsLock;
+	TArray<FCsvProfilerThread*> ProfilerThreads;
+	TArray<uint64> FrameBeginTimestamps;
+	TArray<uint64> FrameBeginTimestampsRT;
+
+};
+
+class FScopedCsvStat
+{
+public:
+	FScopedCsvStat(const char * InStatName)
+		: StatName(InStatName)
+	{
+		FCsvProfiler::Get()->BeginStat(StatName);
+	}
+
+	~FScopedCsvStat()
+	{
+		FCsvProfiler::Get()->EndStat(StatName);
+	}
+	const char * StatName;
+};
+#endif //CSV_PROFILER
\ No newline at end of file
diff --git a/Engine/Source/Runtime/Engine/Public/SceneView.h b/Engine/Source/Runtime/Engine/Public/SceneView.h
--- a/Engine/Source/Runtime/Engine/Public/SceneView.h
+++ b/Engine/Source/Runtime/Engine/Public/SceneView.h
@@ -24,7 +24,6 @@ class FSceneViewStateInterface;
 class FViewElementDrawer;
 class ISceneViewExtension;
 class FSceneViewFamily;
-class FForwardLightingViewResources;
 class FVolumetricFogViewResources;
 
 // Projection data for a FSceneView
diff --git a/Engine/Source/Runtime/Engine/Public/UpdateTextureShaders.h b/Engine/Source/Runtime/Engine/Public/UpdateTextureShaders.h
--- a/Engine/Source/Runtime/Engine/Public/UpdateTextureShaders.h
+++ b/Engine/Source/Runtime/Engine/Public/UpdateTextureShaders.h
@@ -95,13 +95,14 @@ public:
 	{
 		SrcTexture.Bind(Initializer.ParameterMap, TEXT("SrcTexture"), SPF_Mandatory);
 		DestTexture.Bind(Initializer.ParameterMap, TEXT("DestTexture"), SPF_Mandatory);
+		DestPosSize.Bind(Initializer.ParameterMap, TEXT("DestPosSize"), SPF_Mandatory);
 	}
 
 	// FShader interface.
 	virtual bool Serialize(FArchive& Ar) override
 	{
 		bool bShaderHasOutdatedParameters = FGlobalShader::Serialize(Ar);
-		Ar << SrcTexture << DestTexture;
+		Ar << SrcTexture << DestTexture << DestPosSize;
 		return bShaderHasOutdatedParameters;
 	}
 
@@ -113,6 +114,7 @@ public:
 //protected:
 	FShaderResourceParameter SrcTexture;
 	FShaderResourceParameter DestTexture;
+	FShaderParameter DestPosSize;
 };
 
 template<uint32 ElementsPerThread>
diff --git a/Engine/Source/Runtime/Launch/Private/LaunchEngineLoop.cpp b/Engine/Source/Runtime/Launch/Private/LaunchEngineLoop.cpp
--- a/Engine/Source/Runtime/Launch/Private/LaunchEngineLoop.cpp
+++ b/Engine/Source/Runtime/Launch/Private/LaunchEngineLoop.cpp
@@ -177,6 +177,7 @@ class FFeedbackContext;
 	#if ENABLE_VISUAL_LOG
 		#include "VisualLogger/VisualLogger.h"
 	#endif
+	#include "CsvProfiler.h"
 #endif
 
 #if defined(WITH_LAUNCHERCHECK) && WITH_LAUNCHERCHECK
@@ -1391,6 +1392,10 @@ int32 FEngineLoop::PreInit( const TCHAR* CmdLine )
 
 	LoadPreInitModules();
 
+#if WITH_ENGINE && CSV_PROFILER 
+	FCsvProfiler::Get()->Init();
+#endif
+
 	// Start the application
 	if(!AppInit())
 	{
@@ -2321,6 +2326,7 @@ void FEngineLoop::LoadPreInitModules()
 	// compress asynchronously and that can lead to a race condition.
 	FModuleManager::Get().LoadModule(TEXT("TextureCompressor"));
 #endif
+
 #endif // WITH_ENGINE
 
 #if (WITH_EDITOR && !(UE_BUILD_SHIPPING || UE_BUILD_TEST))
@@ -3034,6 +3040,8 @@ void FEngineLoop::Tick()
 			RHICmdList.BeginFrame();
 		});
 
+		FCoreDelegates::OnBeginFrame.Broadcast();
+
 		{
 			QUICK_SCOPE_CYCLE_COUNTER(STAT_FEngineLoop_FlushThreadedLogs);
 			// Flush debug output which has been buffered by other threads.
@@ -3298,6 +3306,8 @@ void FEngineLoop::Tick()
 			GEngine->TickDeferredCommands();		
 		}
 
+		FCoreDelegates::OnEndFrame.Broadcast();
+
 		ENQUEUE_UNIQUE_RENDER_COMMAND(
 			EndFrame,
 		{
diff --git a/Engine/Source/Runtime/RHI/Private/RHI.cpp b/Engine/Source/Runtime/RHI/Private/RHI.cpp
--- a/Engine/Source/Runtime/RHI/Private/RHI.cpp
+++ b/Engine/Source/Runtime/RHI/Private/RHI.cpp
@@ -304,6 +304,7 @@ bool GRHISupportsParallelRHIExecute = false;
 bool GSupportsHDR32bppEncodeModeIntrinsic = false;
 bool GSupportsParallelOcclusionQueries = false;
 bool GSupportsRenderTargetWriteMask = false;
+bool GSupportsTransientResourceAliasing = false;
 
 bool GRHISupportsMSAADepthSampleAccess = false;
 bool GRHISupportsResolveCubemapFaces = false;
diff --git a/Engine/Source/Runtime/RHI/Private/RHICommandList.cpp b/Engine/Source/Runtime/RHI/Private/RHICommandList.cpp
--- a/Engine/Source/Runtime/RHI/Private/RHICommandList.cpp
+++ b/Engine/Source/Runtime/RHI/Private/RHICommandList.cpp
@@ -1710,6 +1710,13 @@ FVertexBufferRHIRef FDynamicRHI::CreateVertexBuffer_RenderThread(class FRHIComma
 	FScopedRHIThreadStaller StallRHIThread(RHICmdList);
 	return GDynamicRHI->RHICreateVertexBuffer(Size, InUsage, CreateInfo);
 }
+
+FStructuredBufferRHIRef FDynamicRHI::CreateStructuredBuffer_RenderThread(class FRHICommandListImmediate& RHICmdList, uint32 Stride, uint32 Size, uint32 InUsage, FRHIResourceCreateInfo& CreateInfo)
+{
+	FScopedRHIThreadStaller StallRHIThread(RHICmdList);
+	return GDynamicRHI->RHICreateStructuredBuffer(Stride, Size, InUsage, CreateInfo);
+}
+
 FIndexBufferRHIRef FDynamicRHI::CreateIndexBuffer_RenderThread(class FRHICommandListImmediate& RHICmdList, uint32 Stride, uint32 Size, uint32 InUsage, FRHIResourceCreateInfo& CreateInfo)
 {
 	FScopedRHIThreadStaller StallRHIThread(RHICmdList);
diff --git a/Engine/Source/Runtime/RHI/Public/DynamicRHI.h b/Engine/Source/Runtime/RHI/Public/DynamicRHI.h
--- a/Engine/Source/Runtime/RHI/Public/DynamicRHI.h
+++ b/Engine/Source/Runtime/RHI/Public/DynamicRHI.h
@@ -1317,6 +1317,7 @@ public:
 	virtual FIndexBufferRHIRef CreateAndLockIndexBuffer_RenderThread(class FRHICommandListImmediate& RHICmdList, uint32 Stride, uint32 Size, uint32 InUsage, FRHIResourceCreateInfo& CreateInfo, void*& OutDataBuffer);
 
 	virtual FVertexBufferRHIRef CreateVertexBuffer_RenderThread(class FRHICommandListImmediate& RHICmdList, uint32 Size, uint32 InUsage, FRHIResourceCreateInfo& CreateInfo);
+	virtual FStructuredBufferRHIRef CreateStructuredBuffer_RenderThread(class FRHICommandListImmediate& RHICmdList, uint32 Stride, uint32 Size, uint32 InUsage, FRHIResourceCreateInfo& CreateInfo);
 	virtual FShaderResourceViewRHIRef CreateShaderResourceView_RenderThread(class FRHICommandListImmediate& RHICmdList, FVertexBufferRHIParamRef VertexBuffer, uint32 Stride, uint8 Format);
 	virtual FShaderResourceViewRHIRef CreateShaderResourceView_RenderThread(class FRHICommandListImmediate& RHICmdList, FIndexBufferRHIParamRef Buffer);
 	virtual void* LockVertexBuffer_RenderThread(class FRHICommandListImmediate& RHICmdList, FVertexBufferRHIParamRef VertexBuffer, uint32 Offset, uint32 SizeRHI, EResourceLockMode LockMode);
@@ -1365,6 +1366,13 @@ public:
 	virtual FTextureCubeRHIRef RHICreateTextureCube_RenderThread(class FRHICommandListImmediate& RHICmdList, uint32 Size, uint8 Format, uint32 NumMips, uint32 Flags, FRHIResourceCreateInfo& CreateInfo);
 	virtual FTextureCubeRHIRef RHICreateTextureCubeArray_RenderThread(class FRHICommandListImmediate& RHICmdList, uint32 Size, uint32 ArraySize, uint8 Format, uint32 NumMips, uint32 Flags, FRHIResourceCreateInfo& CreateInfo);
 	virtual FRenderQueryRHIRef RHICreateRenderQuery_RenderThread(class FRHICommandListImmediate& RHICmdList, ERenderQueryType QueryType);
+	
+	virtual void RHIAcquireTransientResource_RenderThread(FTextureRHIParamRef Texture) { }
+	virtual void RHIDiscardTransientResource_RenderThread(FTextureRHIParamRef Texture) { }
+	virtual void RHIAcquireTransientResource_RenderThread(FVertexBufferRHIParamRef Buffer) { }
+	virtual void RHIDiscardTransientResource_RenderThread(FVertexBufferRHIParamRef Buffer) { }
+	virtual void RHIAcquireTransientResource_RenderThread(FStructuredBufferRHIParamRef Buffer) { }
+	virtual void RHIDiscardTransientResource_RenderThread(FStructuredBufferRHIParamRef Buffer) { }
 
 	//Utilities
 	virtual void EnableIdealGPUCaptureOptions(bool bEnable);
@@ -1572,4 +1580,3 @@ public:
 *	Called to create the instance of the dynamic RHI.
 */
 FDynamicRHI* PlatformCreateDynamicRHI();
-
diff --git a/Engine/Source/Runtime/RHI/Public/RHI.h b/Engine/Source/Runtime/RHI/Public/RHI.h
--- a/Engine/Source/Runtime/RHI/Public/RHI.h
+++ b/Engine/Source/Runtime/RHI/Public/RHI.h
@@ -180,6 +180,9 @@ extern RHI_API bool GSupportsHDR32bppEncodeModeIntrinsic;
 /** True if the RHI supports getting the result of occlusion queries when on a thread other than the renderthread */
 extern RHI_API bool GSupportsParallelOcclusionQueries;
 
+/** true if the RHI supports aliasing of transient resources */
+extern RHI_API bool GSupportsTransientResourceAliasing;
+
 /** The minimum Z value in clip space for the RHI. */
 extern RHI_API float GMinClipZ;
 
@@ -823,7 +826,7 @@ struct FVRamAllocation
 	{
 	}
 
-	bool IsValid() { return AllocationSize > 0; }
+	bool IsValid() const { return AllocationSize > 0; }
 
 	// in bytes
 	uint32 AllocationStart;
@@ -946,6 +949,7 @@ struct FRHIResourceCreateInfo
 		: BulkData(nullptr)
 		, ResourceArray(nullptr)
 		, ClearValueBinding(FLinearColor::Transparent)
+		, DebugName(NULL)
 	{}
 
 	// for CreateTexture calls
@@ -953,6 +957,7 @@ struct FRHIResourceCreateInfo
 		: BulkData(InBulkData)
 		, ResourceArray(nullptr)
 		, ClearValueBinding(FLinearColor::Transparent)
+		, DebugName(NULL)
 	{}
 
 	// for CreateVertexBuffer/CreateStructuredBuffer calls
@@ -960,12 +965,14 @@ struct FRHIResourceCreateInfo
 		: BulkData(nullptr)
 		, ResourceArray(InResourceArray)
 		, ClearValueBinding(FLinearColor::Transparent)
+		, DebugName(NULL)
 	{}
 
 	FRHIResourceCreateInfo(const FClearValueBinding& InClearValueBinding)
 		: BulkData(nullptr)
 		, ResourceArray(nullptr)
 		, ClearValueBinding(InClearValueBinding)
+		, DebugName(NULL)
 	{
 	}
 
@@ -976,6 +983,7 @@ struct FRHIResourceCreateInfo
 
 	// for binding clear colors to rendertargets.
 	FClearValueBinding ClearValueBinding;
+	const TCHAR* DebugName;
 };
 
 // Forward-declaration.
diff --git a/Engine/Source/Runtime/RHI/Public/RHICommandList.h b/Engine/Source/Runtime/RHI/Public/RHICommandList.h
--- a/Engine/Source/Runtime/RHI/Public/RHICommandList.h
+++ b/Engine/Source/Runtime/RHI/Public/RHICommandList.h
@@ -1028,18 +1028,14 @@ struct FRHICommandCopyToResolveTarget : public FRHICommand<FRHICommandCopyToReso
 
 struct FRHICommandTransitionTextures : public FRHICommand<FRHICommandTransitionTextures>
 {
-	static const int32 MaxTexturesToTransition = 16;
 	int32 NumTextures;
-	FTextureRHIParamRef Textures[MaxTexturesToTransition];
+	FTextureRHIParamRef* Textures; // Pointer to an array of textures, allocated inline with the command list
 	EResourceTransitionAccess TransitionType;
 	FORCEINLINE_DEBUGGABLE FRHICommandTransitionTextures(EResourceTransitionAccess InTransitionType, FTextureRHIParamRef* InTextures, int32 InNumTextures)
 		: NumTextures(InNumTextures)
+		, Textures(InTextures)
 		, TransitionType(InTransitionType)
 	{
-		for (int32 i = 0; i < NumTextures; ++i)
-		{
-			Textures[i] = InTextures[i];
-		}
 	}
 	RHI_API void Execute(FRHICommandListBase& CmdList);
 };
@@ -2349,12 +2345,17 @@ public:
 	FORCEINLINE_DEBUGGABLE void TransitionResource(EResourceTransitionAccess TransitionType, FTextureRHIParamRef InTexture)
 	{
 		FTextureRHIParamRef Texture = InTexture;
+		check(Texture == nullptr || Texture->IsCommitted());
 		if (Bypass())
 		{
 			CMD_CONTEXT(RHITransitionResources)(TransitionType, &Texture, 1);
 			return;
 		}
-		new (AllocCommand<FRHICommandTransitionTextures>()) FRHICommandTransitionTextures(TransitionType, &Texture, 1);
+
+		// Allocate space to hold the single texture pointer inline in the command list itself.
+		FTextureRHIParamRef* TextureArray = (FTextureRHIParamRef*)Alloc(sizeof(FTextureRHIParamRef), alignof(FTextureRHIParamRef));
+		TextureArray[0] = Texture;
+		new (AllocCommand<FRHICommandTransitionTextures>()) FRHICommandTransitionTextures(TransitionType, TextureArray, 1);
 	}
 
 	FORCEINLINE_DEBUGGABLE void TransitionResources(EResourceTransitionAccess TransitionType, FTextureRHIParamRef* InTextures, int32 NumTextures)
@@ -2364,7 +2365,15 @@ public:
 			CMD_CONTEXT(RHITransitionResources)(TransitionType, InTextures, NumTextures);
 			return;
 		}
-		new (AllocCommand<FRHICommandTransitionTextures>()) FRHICommandTransitionTextures(TransitionType, InTextures, NumTextures);
+
+		// Allocate space to hold the list of textures inline in the command list itself.
+		FTextureRHIParamRef* InlineTextureArray = (FTextureRHIParamRef*)Alloc(sizeof(FTextureRHIParamRef) * NumTextures, alignof(FTextureRHIParamRef));
+		for (int32 Index = 0; Index < NumTextures; ++Index)
+		{
+			InlineTextureArray[Index] = InTextures[Index];
+		}
+
+		new (AllocCommand<FRHICommandTransitionTextures>()) FRHICommandTransitionTextures(TransitionType, InlineTextureArray, NumTextures);
 	}
 
 	FORCEINLINE_DEBUGGABLE void TransitionResourceArrayNoCopy(EResourceTransitionAccess TransitionType, TArray<FTextureRHIParamRef>& InTextures)
@@ -2380,6 +2389,7 @@ public:
 	FORCEINLINE_DEBUGGABLE void TransitionResource(EResourceTransitionAccess TransitionType, EResourceTransitionPipeline TransitionPipeline, FUnorderedAccessViewRHIParamRef InUAV, FComputeFenceRHIParamRef WriteFence)
 	{
 		FUnorderedAccessViewRHIParamRef UAV = InUAV;
+		check(InUAV == nullptr || InUAV->IsCommitted());
 		if (Bypass())
 		{
 			CMD_CONTEXT(RHITransitionResources)(TransitionType, TransitionPipeline, &UAV, 1, WriteFence);
@@ -2390,6 +2400,7 @@ public:
 
 	FORCEINLINE_DEBUGGABLE void TransitionResource(EResourceTransitionAccess TransitionType, EResourceTransitionPipeline TransitionPipeline, FUnorderedAccessViewRHIParamRef InUAV)
 	{
+		check(InUAV == nullptr || InUAV->IsCommitted());
 		TransitionResource(TransitionType, TransitionPipeline, InUAV, nullptr);
 	}
 
@@ -2915,8 +2926,7 @@ public:
 	
 	FORCEINLINE FStructuredBufferRHIRef CreateStructuredBuffer(uint32 Stride, uint32 Size, uint32 InUsage, FRHIResourceCreateInfo& CreateInfo)
 	{
-		FScopedRHIThreadStaller StallRHIThread(*this);
-		return GDynamicRHI->RHICreateStructuredBuffer(Stride, Size, InUsage, CreateInfo);
+		return GDynamicRHI->CreateStructuredBuffer_RenderThread(*this, Stride, Size, InUsage, CreateInfo);
 	}
 	
 	FORCEINLINE void* LockStructuredBuffer(FStructuredBufferRHIParamRef StructuredBuffer, uint32 Offset, uint32 SizeRHI, EResourceLockMode LockMode)
@@ -3207,7 +3217,79 @@ public:
 		FScopedRHIThreadStaller StallRHIThread(*this);
 		return GDynamicRHI->RHICreateRenderQuery(QueryType);
 	}
-	
+
+	FORCEINLINE void AcquireTransientResource_RenderThread(FTextureRHIParamRef Texture)
+	{
+		if (!Texture->IsCommitted() )
+		{
+			if (GSupportsTransientResourceAliasing)
+			{
+				GDynamicRHI->RHIAcquireTransientResource_RenderThread(Texture);
+			}
+			Texture->SetCommitted(true);
+		}
+	}
+
+	FORCEINLINE void DiscardTransientResource_RenderThread(FTextureRHIParamRef Texture)
+	{
+		if (Texture->IsCommitted())
+		{
+			if (GSupportsTransientResourceAliasing)
+			{
+				GDynamicRHI->RHIDiscardTransientResource_RenderThread(Texture);
+			}
+			Texture->SetCommitted(false);
+		}
+	}
+
+	FORCEINLINE void AcquireTransientResource_RenderThread(FVertexBufferRHIParamRef Buffer)
+	{
+		if (!Buffer->IsCommitted())
+		{
+			if (GSupportsTransientResourceAliasing)
+			{
+				GDynamicRHI->RHIAcquireTransientResource_RenderThread(Buffer);
+			}
+			Buffer->SetCommitted(true);
+		}
+	}
+
+	FORCEINLINE void DiscardTransientResource_RenderThread(FVertexBufferRHIParamRef Buffer)
+	{
+		if (Buffer->IsCommitted())
+		{
+			if (GSupportsTransientResourceAliasing)
+			{
+				GDynamicRHI->RHIDiscardTransientResource_RenderThread(Buffer);
+			}
+			Buffer->SetCommitted(false);
+		}
+	}
+
+	FORCEINLINE void AcquireTransientResource_RenderThread(FStructuredBufferRHIParamRef Buffer)
+	{
+		if (!Buffer->IsCommitted())
+		{
+			if (GSupportsTransientResourceAliasing)
+			{
+				GDynamicRHI->RHIAcquireTransientResource_RenderThread(Buffer);
+			}
+			Buffer->SetCommitted(true);
+		}
+	}
+
+	FORCEINLINE void DiscardTransientResource_RenderThread(FStructuredBufferRHIParamRef Buffer)
+	{
+		if (Buffer->IsCommitted())
+		{
+			if (GSupportsTransientResourceAliasing)
+			{
+				GDynamicRHI->RHIDiscardTransientResource_RenderThread(Buffer);
+			}
+			Buffer->SetCommitted(false);
+		}
+	}
+
 	FORCEINLINE bool GetRenderQueryResult(FRenderQueryRHIParamRef RenderQuery, uint64& OutResult, bool bWait)
 	{
 		return RHIGetRenderQueryResult(RenderQuery, OutResult, bWait);
@@ -3812,6 +3894,36 @@ FORCEINLINE FRenderQueryRHIRef RHICreateRenderQuery(ERenderQueryType QueryType)
 	return FRHICommandListExecutor::GetImmediateCommandList().CreateRenderQuery(QueryType);
 }
 
+FORCEINLINE void RHIAcquireTransientResource(FTextureRHIParamRef Resource)
+{
+	FRHICommandListExecutor::GetImmediateCommandList().AcquireTransientResource_RenderThread(Resource);
+}
+
+FORCEINLINE void RHIDiscardTransientResource(FTextureRHIParamRef Resource)
+{
+	FRHICommandListExecutor::GetImmediateCommandList().DiscardTransientResource_RenderThread(Resource);
+}
+
+FORCEINLINE void RHIAcquireTransientResource(FVertexBufferRHIParamRef Resource)
+{
+	FRHICommandListExecutor::GetImmediateCommandList().AcquireTransientResource_RenderThread(Resource);
+}
+
+FORCEINLINE void RHIDiscardTransientResource(FVertexBufferRHIParamRef Resource)
+{
+	FRHICommandListExecutor::GetImmediateCommandList().DiscardTransientResource_RenderThread(Resource);
+}
+
+FORCEINLINE void RHIAcquireTransientResource(FStructuredBufferRHIParamRef Resource)
+{
+	FRHICommandListExecutor::GetImmediateCommandList().AcquireTransientResource_RenderThread(Resource);
+}
+
+FORCEINLINE void RHIDiscardTransientResource(FStructuredBufferRHIParamRef Resource)
+{
+	FRHICommandListExecutor::GetImmediateCommandList().DiscardTransientResource_RenderThread(Resource);
+}
+
 FORCEINLINE void RHIAcquireThreadOwnership()
 {
 	return FRHICommandListExecutor::GetImmediateCommandList().AcquireThreadOwnership();
diff --git a/Engine/Source/Runtime/RHI/Public/RHIDefinitions.h b/Engine/Source/Runtime/RHI/Public/RHIDefinitions.h
--- a/Engine/Source/Runtime/RHI/Public/RHIDefinitions.h
+++ b/Engine/Source/Runtime/RHI/Public/RHIDefinitions.h
@@ -492,6 +492,8 @@ static_assert(PT_Num <= (1 << PT_NumBits), "PT_NumBits is too small");
  */
 enum EBufferUsageFlags
 {
+	BUF_None			  = 0x0000,
+
 	// Mutually exclusive write-frequency flags
 	BUF_Static            = 0x0001, // The buffer will be written to once.
 	BUF_Dynamic           = 0x0002, // The buffer will be written to occasionally, GPU read only, CPU write only.  The data lifetime is until the next update, or the buffer is destroyed.
@@ -526,9 +528,12 @@ enum EBufferUsageFlags
 	 */
 	BUF_ZeroStride        = 0x0800,
 
-	/** Buffer should go in fast vram (hint only) */
+	/** Buffer should go in fast vram (hint only). Requires BUF_Transient */
 	BUF_FastVRAM          = 0x1000,
 
+	/** Buffer should be allocated from transient memory. */
+	BUF_Transient		  = 0x2000,
+
 	// Helper bit-masks
 	BUF_AnyDynamic = (BUF_Dynamic | BUF_Volatile),
 };
@@ -630,7 +635,9 @@ enum ETextureCreateFlags
 	// Hint to the driver that this resource is managed properly by the engine for Alternate-Frame-Rendering in mGPU usage.
 	TexCreate_AFRManual = 1 << 29,
 	// Workaround for 128^3 volume textures getting bloated 4x due to tiling mode on PS4
-	TexCreate_ReduceMemoryWithTilingMode = 1 << 30
+	TexCreate_ReduceMemoryWithTilingMode = 1 << 30,
+	/** Texture should be allocated from transient memory. */
+	TexCreate_Transient = 1 << 31
 };
 
 enum EAsyncComputePriority
diff --git a/Engine/Source/Runtime/RHI/Public/RHIResources.h b/Engine/Source/Runtime/RHI/Public/RHIResources.h
--- a/Engine/Source/Runtime/RHI/Public/RHIResources.h
+++ b/Engine/Source/Runtime/RHI/Public/RHIResources.h
@@ -32,6 +32,7 @@ public:
 	FRHIResource(bool InbDoNotDeferDelete = false)
 		: MarkedForDelete(0)
 		, bDoNotDeferDelete(InbDoNotDeferDelete)
+		, bCommitted(true)
 	{
 	}
 	virtual ~FRHIResource() 
@@ -87,10 +88,25 @@ public:
 
 	static bool Bypass();
 
+	// Transient resource tracking
+	// We do this at a high level so we can catch errors even when transient resources are not supported
+	void SetCommitted(bool bInCommitted) 
+	{ 
+		check(IsInRenderingThread()); 
+		bCommitted = bInCommitted;
+	}
+	bool IsCommitted() const 
+	{ 
+		check(IsInRenderingThread());
+		return bCommitted;
+	}
+
 private:
 	mutable FThreadSafeCounter NumRefs;
 	mutable int32 MarkedForDelete;
 	bool bDoNotDeferDelete;
+	bool bCommitted;
+
 	static TLockFreePointerListUnordered<FRHIResource, PLATFORM_CACHE_LINE_SIZE> PendingDeletes;
 	static FRHIResource* CurrentlyDeleting;
 
diff --git a/Engine/Source/Runtime/RHI/Public/RHIUtilities.h b/Engine/Source/Runtime/RHI/Public/RHIUtilities.h
--- a/Engine/Source/Runtime/RHI/Public/RHIUtilities.h
+++ b/Engine/Source/Runtime/RHI/Public/RHIUtilities.h
@@ -49,21 +49,47 @@ struct FRWBuffer
 	FShaderResourceViewRHIRef SRV;
 	uint32 NumBytes;
 
-	FRWBuffer(): NumBytes(0) {}
+	FRWBuffer()
+		: NumBytes(0)
+	{}
+
+	~FRWBuffer()
+	{
+		Release();
+	}
 
 	// @param AdditionalUsage passed down to RHICreateVertexBuffer(), get combined with "BUF_UnorderedAccess | BUF_ShaderResource" e.g. BUF_Static
-	void Initialize(uint32 BytesPerElement, uint32 NumElements, EPixelFormat Format, uint32 AdditionalUsage = 0)
+	void Initialize(uint32 BytesPerElement, uint32 NumElements, EPixelFormat Format, uint32 AdditionalUsage = 0, const TCHAR* InDebugName = NULL)
 	{
 		check(GMaxRHIFeatureLevel == ERHIFeatureLevel::SM5);
+		// Provide a debug name if using Fast VRAM so the allocators diagnostics will work
+		ensure(!((AdditionalUsage & BUF_FastVRAM) && !InDebugName));
 		NumBytes = BytesPerElement * NumElements;
 		FRHIResourceCreateInfo CreateInfo;
+		CreateInfo.DebugName = InDebugName;
 		Buffer = RHICreateVertexBuffer(NumBytes, BUF_UnorderedAccess | BUF_ShaderResource | AdditionalUsage, CreateInfo);
 		UAV = RHICreateUnorderedAccessView(Buffer, Format);
 		SRV = RHICreateShaderResourceView(Buffer, BytesPerElement, Format);
 	}
 
+	void AcquireTransientResource()
+	{
+		RHIAcquireTransientResource(Buffer);
+	}
+	void DiscardTransientResource()
+	{
+		RHIDiscardTransientResource(Buffer);
+	}
+
 	void Release()
 	{
+		int32 BufferRefCount = Buffer ? Buffer->GetRefCount() : -1;
+
+		if (BufferRefCount == 1)
+		{
+			DiscardTransientResource();
+		}
+
 		NumBytes = 0;
 		Buffer.SafeRelease();
 		UAV.SafeRelease();
@@ -107,11 +133,20 @@ struct FRWBufferStructured
 
 	FRWBufferStructured(): NumBytes(0) {}
 
-	void Initialize(uint32 BytesPerElement, uint32 NumElements, uint32 AdditionalUsage = 0, bool bUseUavCounter = false, bool bAppendBuffer = false)
+	~FRWBufferStructured()
+	{
+		Release();
+	}
+
+	void Initialize(uint32 BytesPerElement, uint32 NumElements, uint32 AdditionalUsage = 0, const TCHAR* InDebugName = NULL, bool bUseUavCounter = false, bool bAppendBuffer = false)
 	{
 		check(GMaxRHIFeatureLevel == ERHIFeatureLevel::SM5);
+		// Provide a debug name if using Fast VRAM so the allocators diagnostics will work
+		ensure(!((AdditionalUsage & BUF_FastVRAM) && !InDebugName));
+
 		NumBytes = BytesPerElement * NumElements;
 		FRHIResourceCreateInfo CreateInfo;
+		CreateInfo.DebugName = InDebugName;
 		Buffer = RHICreateStructuredBuffer(BytesPerElement, NumBytes, BUF_UnorderedAccess | BUF_ShaderResource | AdditionalUsage, CreateInfo);
 		UAV = RHICreateUnorderedAccessView(Buffer, bUseUavCounter, bAppendBuffer);
 		SRV = RHICreateShaderResourceView(Buffer);
@@ -119,11 +154,27 @@ struct FRWBufferStructured
 
 	void Release()
 	{
+		int32 BufferRefCount = Buffer ? Buffer->GetRefCount() : -1;
+
+		if (BufferRefCount == 1)
+		{
+			DiscardTransientResource();
+		}
+
 		NumBytes = 0;
 		Buffer.SafeRelease();
 		UAV.SafeRelease();
 		SRV.SafeRelease();
 	}
+
+	void AcquireTransientResource()
+	{
+		RHIAcquireTransientResource(Buffer);
+	}
+	void DiscardTransientResource()
+	{
+		RHIDiscardTransientResource(Buffer);
+	}
 };
 
 /** Encapsulates a GPU read/write ByteAddress buffer with its UAV and SRV. */
@@ -797,7 +848,6 @@ private:
 
 extern RHI_API void EnableDepthBoundsTest(FRHICommandList& RHICmdList, float WorldSpaceDepthNear, float WorldSpaceDepthFar, const FMatrix& ProjectionMatrix);
 extern RHI_API void DisableDepthBoundsTest(FRHICommandList& RHICmdList);
-
 struct FRHILockTracker
 {
 	struct FLockParams
diff --git a/Engine/Source/Runtime/RenderCore/Private/RenderCore.cpp b/Engine/Source/Runtime/RenderCore/Private/RenderCore.cpp
--- a/Engine/Source/Runtime/RenderCore/Private/RenderCore.cpp
+++ b/Engine/Source/Runtime/RenderCore/Private/RenderCore.cpp
@@ -172,16 +172,25 @@ DEFINE_STAT(STAT_RenderWholeSceneReflectiveShadowMapsTime);
 DEFINE_STAT(STAT_ShadowmapAtlasMemory);
 DEFINE_STAT(STAT_CachedShadowmapMemory);
 
-#if !(UE_BUILD_SHIPPING || UE_BUILD_TEST)
+#define EXPOSE_FORCE_LOD !(UE_BUILD_SHIPPING || UE_BUILD_TEST)
+
+#if EXPOSE_FORCE_LOD
 
 static TAutoConsoleVariable<int32> CVarForceLOD(
 	TEXT("r.ForceLOD"),
 	-1,
 	TEXT("LOD level to force, -1 is off."),
-	ECVF_Cheat | ECVF_RenderThreadSafe
+	ECVF_Scalability | ECVF_Default | ECVF_RenderThreadSafe
 	);
 
-#endif // !(UE_BUILD_SHIPPING || UE_BUILD_TEST)
+static TAutoConsoleVariable<int32> CVarForceLODShadow(
+	TEXT("r.ForceLODShadow"),
+	-1,
+	TEXT("LOD level to force for the shadow map generation only, -1 is off."),
+	ECVF_Scalability | ECVF_Default | ECVF_RenderThreadSafe
+);
+
+#endif // EXPOSE_FORCE_LOD
 
 /** Whether to pause the global realtime clock for the rendering thread (read and write only on main thread). */
 bool GPauseRenderingRealtimeClock;
@@ -242,11 +251,24 @@ RENDERCORE_API int32 GetCVarForceLOD()
 {
 	int32 Ret = -1;
 
-#if !(UE_BUILD_SHIPPING || UE_BUILD_TEST)
+#if EXPOSE_FORCE_LOD
 	{
 		Ret = CVarForceLOD.GetValueOnRenderThread();
 	}
-#endif // !(UE_BUILD_SHIPPING || UE_BUILD_TEST)
+#endif // EXPOSE_FORCE_LOD
+
+	return Ret;
+}
+
+RENDERCORE_API int32 GetCVarForceLODShadow()
+{
+	int32 Ret = -1;
+
+#if EXPOSE_FORCE_LOD
+	{
+		Ret = CVarForceLODShadow.GetValueOnRenderThread();
+	}
+#endif // EXPOSE_FORCE_LOD
 
 	return Ret;
 }
diff --git a/Engine/Source/Runtime/RenderCore/Public/RenderCore.h b/Engine/Source/Runtime/RenderCore/Public/RenderCore.h
--- a/Engine/Source/Runtime/RenderCore/Public/RenderCore.h
+++ b/Engine/Source/Runtime/RenderCore/Public/RenderCore.h
@@ -278,3 +278,5 @@ extern RENDERCORE_API uint32 GSwapBufferTime;
 // shared by renderer and engine, compiles down to a constant in final release
 RENDERCORE_API int32 GetCVarForceLOD();
 
+RENDERCORE_API int32 GetCVarForceLODShadow();
+
diff --git a/Engine/Source/Runtime/RenderCore/Public/RendererInterface.h b/Engine/Source/Runtime/RenderCore/Public/RendererInterface.h
--- a/Engine/Source/Runtime/RenderCore/Public/RendererInterface.h
+++ b/Engine/Source/Runtime/RenderCore/Public/RendererInterface.h
@@ -175,7 +175,7 @@ public:
 		auto LhsFlags = Flags;
 		auto RhsFlags = rhs.Flags;
 
-		if (!bExact || !FPlatformProperties::SupportsFastVRAMMemory())
+		if (!bExact || !FPlatformMemory::SupportsFastVRAMMemory())
 		{
 			LhsFlags &= (~TexCreate_FastVRAM);
 			RhsFlags &= (~TexCreate_FastVRAM);
@@ -276,6 +276,11 @@ public:
 			FlagsString += TEXT(" VRam");
 		}
 
+		if (LocalFlags & TexCreate_Transient)
+		{
+			FlagsString += TEXT(" Transient");
+		}
+
 		FString ArrayString;
 
 		if(IsArray())
@@ -342,7 +347,7 @@ public:
 	/** Whether the shader-resource and targetable texture must be separate textures. */
 	bool bForceSeparateTargetAndShaderResource;
 	/** only set a pointer to memory that never gets released */
-	const TCHAR *DebugName;
+	const TCHAR* DebugName;
 	/** automatically set to writable via barrier during */
 	bool AutoWritable;
 	/** create render target write mask (supported only on specific platforms) */
@@ -398,10 +403,12 @@ struct FSceneRenderTargetItem
 };
 
 /**
- * Render thread side, use TRefCountPtr<IPooledRenderTarget>, allows to use sharing and VisualizeTexture
+ * Render thread side, use TRefCountPtr<IPooledRenderTarget>, allows sharing and VisualizeTexture
  */
-struct IPooledRenderTarget : public IRefCountedObject
+struct IPooledRenderTarget
 {
+	virtual ~IPooledRenderTarget() {}
+
 	/** Checks if the reference count indicated that the rendertarget is unused and can be reused. */
 	virtual bool IsFree() const = 0;
 	/** Get all the data that is needed to create the render target. */
@@ -418,8 +425,22 @@ struct IPooledRenderTarget : public IRefCountedObject
 	/** Get the low level internals (texture/surface) */
 	inline const FSceneRenderTargetItem& GetRenderTargetItem() const { return RenderTargetItem; }
 
+	// Refcounting
+	virtual uint32 AddRef() const = 0;
+	virtual uint32 Release() = 0;
+	virtual uint32 GetRefCount() const = 0;
+
 protected:
 
+	/** For pool management (only if NumRef == 0 the element can be reused) */
+	mutable int32 NumRefs;
+
+	/** Snapshots are sortof fake pooled render targets, they don't own anything and can outlive the things that created them. These are for threaded rendering. */
+	bool bSnapshot;
+
+	/** Pointer back to the pool for render targets which are actually pooled, otherwise NULL. */
+	class FRenderTargetPool* RenderTargetPool;
+
 	/** The internal references to the created render target */
 	FSceneRenderTargetItem RenderTargetItem;
 };
diff --git a/Engine/Source/Runtime/Renderer/Private/CapsuleShadowRendering.cpp b/Engine/Source/Runtime/Renderer/Private/CapsuleShadowRendering.cpp
--- a/Engine/Source/Runtime/Renderer/Private/CapsuleShadowRendering.cpp
+++ b/Engine/Source/Runtime/Renderer/Private/CapsuleShadowRendering.cpp
@@ -670,8 +670,9 @@ void AllocateCapsuleTileIntersectionCountsBuffer(FIntPoint GroupSize, FSceneView
 }
 
 bool FDeferredShadingSceneRenderer::RenderCapsuleDirectShadows(
-	const FLightSceneInfo& LightSceneInfo,
 	FRHICommandListImmediate& RHICmdList, 
+	const FLightSceneInfo& LightSceneInfo,
+	IPooledRenderTarget* ScreenShadowMaskTexture,
 	const TArray<FProjectedShadowInfo*, SceneRenderingAllocator>& CapsuleShadows, 
 	bool bProjectingForForwardShading) const
 {
@@ -837,7 +838,7 @@ bool FDeferredShadingSceneRenderer::RenderCapsuleDirectShadows(
 					SCOPED_DRAW_EVENTF(RHICmdList, Upsample, TEXT("Upsample %dx%d"),
 						ScissorRect.Width(), ScissorRect.Height());
 						
-					FSceneRenderTargets::Get(RHICmdList).BeginRenderingLightAttenuation(RHICmdList);
+					SetRenderTarget(RHICmdList, ScreenShadowMaskTexture->GetRenderTargetItem().TargetableTexture, FSceneRenderTargets::Get(RHICmdList).GetSceneDepthSurface(), ESimpleRenderTargetMode::EExistingColorAndDepth, FExclusiveDepthStencil::DepthRead_StencilWrite, true);
 
 					FGraphicsPipelineStateInitializer GraphicsPSOInit;
 					RHICmdList.ApplyCachedRenderTargets(GraphicsPSOInit);
diff --git a/Engine/Source/Runtime/Renderer/Private/CapsuleShadowRendering.h b/Engine/Source/Runtime/Renderer/Private/CapsuleShadowRendering.h
--- a/Engine/Source/Runtime/Renderer/Private/CapsuleShadowRendering.h
+++ b/Engine/Source/Runtime/Renderer/Private/CapsuleShadowRendering.h
@@ -14,7 +14,7 @@ extern int32 GCapsuleShadows;
 inline bool DoesPlatformSupportCapsuleShadows(EShaderPlatform Platform)
 {
 	// Hasn't been tested elsewhere yet
-	return Platform == SP_PCD3D_SM5 || Platform == SP_PS4 || Platform == SP_METAL_SM5;
+	return Platform == SP_PCD3D_SM5 || Platform == SP_PS4 || Platform == SP_XBOXONE_D3D11 || Platform == SP_XBOXONE_D3D12 || Platform == SP_METAL_SM5;
 }
 
 inline bool SupportsCapsuleShadows(ERHIFeatureLevel::Type FeatureLevel, EShaderPlatform ShaderPlatform)
diff --git a/Engine/Source/Runtime/Renderer/Private/DeferredShadingRenderer.cpp b/Engine/Source/Runtime/Renderer/Private/DeferredShadingRenderer.cpp
--- a/Engine/Source/Runtime/Renderer/Private/DeferredShadingRenderer.cpp
+++ b/Engine/Source/Runtime/Renderer/Private/DeferredShadingRenderer.cpp
@@ -1096,6 +1096,8 @@ void FDeferredShadingSceneRenderer::Render(FRHICommandListImmediate& RHICmdList)
 		RenderDeferredReflections(RHICmdList, DynamicBentNormalAO, VelocityRT);
 		ServiceLocalQueue();
 
+		DynamicBentNormalAO = NULL;
+
 		// Post-lighting composition lighting stage
 		// e.g. ScreenSpaceSubsurfaceScattering
 		for(int32 ViewIndex = 0; ViewIndex < Views.Num(); ++ViewIndex)
diff --git a/Engine/Source/Runtime/Renderer/Private/DeferredShadingRenderer.h b/Engine/Source/Runtime/Renderer/Private/DeferredShadingRenderer.h
--- a/Engine/Source/Runtime/Renderer/Private/DeferredShadingRenderer.h
+++ b/Engine/Source/Runtime/Renderer/Private/DeferredShadingRenderer.h
@@ -298,8 +298,9 @@ private:
 
 	/** Renders capsule shadows for all per-object shadows using it for the given light. */
 	bool RenderCapsuleDirectShadows(
-		const FLightSceneInfo& LightSceneInfo,
 		FRHICommandListImmediate& RHICmdList, 
+		const FLightSceneInfo& LightSceneInfo, 
+		IPooledRenderTarget* ScreenShadowMaskTexture,
 		const TArray<FProjectedShadowInfo*, SceneRenderingAllocator>& CapsuleShadows, 
 		bool bProjectingForForwardShading) const;
 
@@ -316,7 +317,7 @@ private:
 	void RenderCapsuleShadowsForMovableSkylight(FRHICommandListImmediate& RHICmdList, TRefCountPtr<IPooledRenderTarget>& BentNormalOutput) const;
 
 	/** Render deferred projections of shadows for a given light into the light attenuation buffer. */
-	bool RenderShadowProjections(FRHICommandListImmediate& RHICmdList, const FLightSceneInfo* LightSceneInfo, bool& bInjectedTranslucentVolume);
+	bool RenderShadowProjections(FRHICommandListImmediate& RHICmdList, const FLightSceneInfo* LightSceneInfo, IPooledRenderTarget* ScreenShadowMaskTexture, bool& bInjectedTranslucentVolume);
 
 	/** Render shadow projections when forward rendering. */
 	void RenderForwardShadingShadowProjections(FRHICommandListImmediate& RHICmdList);
@@ -327,15 +328,16 @@ private:
 	  * @param LightSceneInfo Represents the current light
 	  * @param LightIndex The light's index into FScene::Lights
 	  */
-	bool RenderLightFunction(FRHICommandListImmediate& RHICmdList, const FLightSceneInfo* LightSceneInfo, bool bLightAttenuationCleared, bool bProjectingForForwardShading);
+	bool RenderLightFunction(FRHICommandListImmediate& RHICmdList, const FLightSceneInfo* LightSceneInfo, IPooledRenderTarget* ScreenShadowMaskTexture, bool bLightAttenuationCleared, bool bProjectingForForwardShading);
 
 	/** Renders a light function indicating that whole scene shadowing being displayed is for previewing only, and will go away in game. */
-	bool RenderPreviewShadowsIndicator(FRHICommandListImmediate& RHICmdList, const FLightSceneInfo* LightSceneInfo, bool bLightAttenuationCleared);
+	bool RenderPreviewShadowsIndicator(FRHICommandListImmediate& RHICmdList, const FLightSceneInfo* LightSceneInfo, IPooledRenderTarget* ScreenShadowMaskTexture, bool bLightAttenuationCleared);
 
 	/** Renders a light function with the given material. */
 	bool RenderLightFunctionForMaterial(
 		FRHICommandListImmediate& RHICmdList, 
 		const FLightSceneInfo* LightSceneInfo, 
+		IPooledRenderTarget* ScreenShadowMaskTexture,
 		const FMaterialRenderProxy* MaterialProxy, 
 		bool bLightAttenuationCleared,
 		bool bProjectingForForwardShading, 
@@ -348,7 +350,7 @@ private:
 	  * @param LightIndex The light's index into FScene::Lights
 	  * @return true if anything got rendered
 	  */
-	void RenderLight(FRHICommandList& RHICmdList, const FLightSceneInfo* LightSceneInfo, bool bRenderOverlap, bool bIssueDrawEvent);
+	void RenderLight(FRHICommandList& RHICmdList, const FLightSceneInfo* LightSceneInfo, IPooledRenderTarget* ScreenShadowMaskTexture, bool bRenderOverlap, bool bIssueDrawEvent);
 
 	/** Renders an array of simple lights using standard deferred shading. */
 	void RenderSimpleLightsStandardDeferred(FRHICommandListImmediate& RHICmdList, const FSimpleLightArray& SimpleLights);
diff --git a/Engine/Source/Runtime/Renderer/Private/DistanceFieldAmbientOcclusion.cpp b/Engine/Source/Runtime/Renderer/Private/DistanceFieldAmbientOcclusion.cpp
--- a/Engine/Source/Runtime/Renderer/Private/DistanceFieldAmbientOcclusion.cpp
+++ b/Engine/Source/Runtime/Renderer/Private/DistanceFieldAmbientOcclusion.cpp
@@ -812,6 +812,7 @@ bool FDeferredShadingSceneRenderer::RenderDistanceFieldLighting(
 			{
 				const FIntPoint BufferSize = GetBufferSizeForAO();
 				FPooledRenderTargetDesc Desc(FPooledRenderTargetDesc::Create2DDesc(BufferSize, PF_FloatRGBA, FClearValueBinding::Transparent, TexCreate_None, TexCreate_RenderTargetable | TexCreate_UAV, false));
+				Desc.Flags |= GetTextureFastVRamFlag_DynamicLayout();
 				GRenderTargetPool.FindFreeElement(RHICmdList, Desc, DistanceFieldNormal, TEXT("DistanceFieldNormal"));
 			}
 
@@ -846,6 +847,14 @@ bool FDeferredShadingSceneRenderer::RenderDistanceFieldLighting(
 				BentNormalOutput, 
 				IrradianceOutput);
 
+			if ( IsTransientResourceBufferAliasingEnabled() )
+			{
+				GAOCulledObjectBuffers.Buffers.DiscardTransientResource();
+
+				FTileIntersectionResources* TileIntersectionResources = ((FSceneViewState*)View.State)->AOTileIntersectionResources;
+				TileIntersectionResources->DiscardTransientResource();
+			}
+
 			RenderCapsuleShadowsForMovableSkylight(RHICmdList, BentNormalOutput);
 
 			GRenderTargetPool.VisualizeTexture.SetCheckPoint(RHICmdList, BentNormalOutput);
@@ -857,8 +866,10 @@ bool FDeferredShadingSceneRenderer::RenderDistanceFieldLighting(
 			else
 			{
 				FPooledRenderTargetDesc Desc = SceneContext.GetSceneColor()->GetDesc();
-				// Make sure we get a signed format
-				Desc.Format = PF_FloatRGBA;
+				Desc.Flags &= ~(TexCreate_FastVRAM | TexCreate_Transient);
+				// Bent normals are signed so we will have to pack / unpack
+				Desc.Format = PF_FloatR11G11B10;
+				Desc.Flags |= GetTextureFastVRamFlag_DynamicLayout();
 				GRenderTargetPool.FindFreeElement(RHICmdList, Desc, OutDynamicBentNormalAO, TEXT("DynamicBentNormalAO"));
 
 				if (bUseDistanceFieldGI)
diff --git a/Engine/Source/Runtime/Renderer/Private/DistanceFieldAmbientOcclusion.h b/Engine/Source/Runtime/Renderer/Private/DistanceFieldAmbientOcclusion.h
--- a/Engine/Source/Runtime/Renderer/Private/DistanceFieldAmbientOcclusion.h
+++ b/Engine/Source/Runtime/Renderer/Private/DistanceFieldAmbientOcclusion.h
@@ -77,6 +77,8 @@ public:
 
 	virtual void InitDynamicRHI() override;
 
+	
+
 	virtual void ReleaseDynamicRHI() override
 	{
 		TileConeAxisAndCos.Release();
@@ -88,6 +90,24 @@ public:
 		ObjectTilesIndirectArguments.Release();
 	}
 
+	void AcquireTransientResource()
+	{
+		TileConeAxisAndCos.AcquireTransientResource();
+		TileConeDepthRanges.AcquireTransientResource();
+		NumCulledTilesArray.AcquireTransientResource();
+		CulledTilesStartOffsetArray.AcquireTransientResource();
+		CulledTileDataArray.AcquireTransientResource();
+	}
+
+	void DiscardTransientResource()
+	{
+		TileConeAxisAndCos.DiscardTransientResource();
+		TileConeDepthRanges.DiscardTransientResource();
+		NumCulledTilesArray.DiscardTransientResource();
+		CulledTilesStartOffsetArray.DiscardTransientResource();
+		CulledTileDataArray.DiscardTransientResource();
+	}
+
 	size_t GetSizeBytes() const
 	{
 		return TileConeAxisAndCos.NumBytes + TileConeDepthRanges.NumBytes
@@ -213,6 +233,28 @@ public:
 		HeightfieldIrradiance.Release();
 	}
 
+	void AcquireTransientResource()
+	{
+		ScreenGridConeVisibility.AcquireTransientResource();
+		if (bAllocateResourceForGI)
+		{
+			StepBentNormal.AcquireTransientResource();
+			SurfelIrradiance.AcquireTransientResource();
+			HeightfieldIrradiance.AcquireTransientResource();
+		}
+	}
+
+	void DiscardTransientResource()
+	{
+		ScreenGridConeVisibility.DiscardTransientResource();
+		if (bAllocateResourceForGI)
+		{
+			StepBentNormal.DiscardTransientResource();
+			SurfelIrradiance.DiscardTransientResource();
+			HeightfieldIrradiance.DiscardTransientResource();
+		}
+	}
+
 	FIntPoint ScreenGridDimensions;
 
 	FRWBuffer ScreenGridConeVisibility;
diff --git a/Engine/Source/Runtime/Renderer/Private/DistanceFieldLightingPost.cpp b/Engine/Source/Runtime/Renderer/Private/DistanceFieldLightingPost.cpp
--- a/Engine/Source/Runtime/Renderer/Private/DistanceFieldLightingPost.cpp
+++ b/Engine/Source/Runtime/Renderer/Private/DistanceFieldLightingPost.cpp
@@ -396,13 +396,13 @@ IMPLEMENT_SHADER_TYPE(template<>,TFilterHistoryPS<true>,TEXT("DistanceFieldLight
 IMPLEMENT_SHADER_TYPE(template<>,TFilterHistoryPS<false>,TEXT("DistanceFieldLightingPost"),TEXT("FilterHistoryPS"),SF_Pixel);
 
 
-void AllocateOrReuseAORenderTarget(FRHICommandList& RHICmdList, TRefCountPtr<IPooledRenderTarget>& Target, const TCHAR* Name, EPixelFormat Format)
+void AllocateOrReuseAORenderTarget(FRHICommandList& RHICmdList, TRefCountPtr<IPooledRenderTarget>& Target, const TCHAR* Name, EPixelFormat Format, uint32 Flags) 
 {
 	if (!Target)
 	{
 		FIntPoint BufferSize = GetBufferSizeForAO();
 
-		FPooledRenderTargetDesc Desc(FPooledRenderTargetDesc::Create2DDesc(BufferSize, Format, FClearValueBinding::None, TexCreate_None, TexCreate_RenderTargetable | TexCreate_UAV, false));
+		FPooledRenderTargetDesc Desc(FPooledRenderTargetDesc::Create2DDesc(BufferSize, Format, FClearValueBinding::None, Flags, TexCreate_RenderTargetable | TexCreate_UAV, false));
 		Desc.AutoWritable = false;
 		GRenderTargetPool.FindFreeElement(RHICmdList, Desc, Target, Name);
 	}
@@ -444,12 +444,13 @@ void UpdateHistory(
 			// If the scene render targets reallocate, toss the history so we don't read uninitialized data
 			&& (*BentNormalHistoryState)->GetDesc().Extent == BufferSize)
 		{
+			uint32 HistoryPassOutputFlags = GAOHistoryStabilityPass ? GetTextureFastVRamFlag_DynamicLayout() : 0;
 			// Reuse a render target from the pool with a consistent name, for vis purposes
 			TRefCountPtr<IPooledRenderTarget> NewBentNormalHistory;
-			AllocateOrReuseAORenderTarget(RHICmdList, NewBentNormalHistory, BentNormalHistoryRTName, PF_FloatRGBA);
+			AllocateOrReuseAORenderTarget(RHICmdList, NewBentNormalHistory, BentNormalHistoryRTName, PF_FloatRGBA, HistoryPassOutputFlags);
 
 			TRefCountPtr<IPooledRenderTarget> NewConfidenceHistory;
-			AllocateOrReuseAORenderTarget(RHICmdList, NewConfidenceHistory, ConfidenceHistoryRTName, PF_G8);
+			AllocateOrReuseAORenderTarget(RHICmdList, NewConfidenceHistory, ConfidenceHistoryRTName, PF_G8, HistoryPassOutputFlags);
 
 			TRefCountPtr<IPooledRenderTarget> NewIrradianceHistory;
 
diff --git a/Engine/Source/Runtime/Renderer/Private/DistanceFieldLightingPost.h b/Engine/Source/Runtime/Renderer/Private/DistanceFieldLightingPost.h
--- a/Engine/Source/Runtime/Renderer/Private/DistanceFieldLightingPost.h
+++ b/Engine/Source/Runtime/Renderer/Private/DistanceFieldLightingPost.h
@@ -11,7 +11,7 @@
 
 class FDistanceFieldAOParameters;
 
-extern void AllocateOrReuseAORenderTarget(FRHICommandList& RHICmdList, TRefCountPtr<IPooledRenderTarget>& Target, const TCHAR* Name, EPixelFormat Format);
+extern void AllocateOrReuseAORenderTarget(FRHICommandList& RHICmdList, TRefCountPtr<IPooledRenderTarget>& Target, const TCHAR* Name, EPixelFormat Format, uint32 Flags = 0);
 
 extern void UpdateHistory(
 	FRHICommandList& RHICmdList,
diff --git a/Engine/Source/Runtime/Renderer/Private/DistanceFieldLightingShared.h b/Engine/Source/Runtime/Renderer/Private/DistanceFieldLightingShared.h
--- a/Engine/Source/Runtime/Renderer/Private/DistanceFieldLightingShared.h
+++ b/Engine/Source/Runtime/Renderer/Private/DistanceFieldLightingShared.h
@@ -14,6 +14,7 @@
 #include "RHIStaticStates.h"
 #include "DistanceFieldAtlas.h"
 #include "UniquePtr.h"
+#include "SceneRendering.h"
 
 class FLightSceneProxy;
 class FMaterialRenderProxy;
@@ -291,18 +292,40 @@ public:
 	{
 		if (MaxObjects > 0)
 		{
+			const uint32 FastVRamFlag = IsTransientResourceBufferAliasingEnabled() ? ( BUF_FastVRAM | BUF_Transient ) : BUF_None;
+
 			ObjectIndirectArguments.Initialize(sizeof(uint32), 5, PF_R32_UINT, BUF_Static | BUF_DrawIndirect);
 			ObjectIndirectDispatch.Initialize(sizeof(uint32), 3, PF_R32_UINT, BUF_Static | BUF_DrawIndirect);
-			Bounds.Initialize(sizeof(FVector4), MaxObjects, BUF_Static);
-			Data.Initialize(sizeof(FVector4), MaxObjects * ObjectDataStride, BUF_Static);
+			Bounds.Initialize(sizeof(FVector4), MaxObjects, BUF_Static | FastVRamFlag, TEXT("FDistanceFieldCulledObjectBuffers::Bounds"));
+			Data.Initialize(sizeof(FVector4), MaxObjects * ObjectDataStride, BUF_Static | FastVRamFlag, TEXT("FDistanceFieldCulledObjectBuffers::Data"));
 
 			if (bWantBoxBounds)
 			{
-				BoxBounds.Initialize(sizeof(FVector4), MaxObjects * ObjectBoxBoundsStride, BUF_Static);
+				BoxBounds.Initialize(sizeof(FVector4), MaxObjects * ObjectBoxBoundsStride, BUF_Static | FastVRamFlag, TEXT("FDistanceFieldCulledObjectBuffers::BoxBounds"));
 			}
 		}
 	}
 
+	void AcquireTransientResource()
+	{
+		Bounds.AcquireTransientResource();
+		Data.AcquireTransientResource();
+		if (bWantBoxBounds)
+		{
+			BoxBounds.AcquireTransientResource();
+		}
+	}
+
+	void DiscardTransientResource()
+	{
+		Bounds.DiscardTransientResource();
+		Data.DiscardTransientResource();
+		if (bWantBoxBounds)
+		{
+			BoxBounds.DiscardTransientResource();
+		}
+	}
+
 	void Release()
 	{
 		ObjectIndirectArguments.Release();
diff --git a/Engine/Source/Runtime/Renderer/Private/DistanceFieldObjectCulling.cpp b/Engine/Source/Runtime/Renderer/Private/DistanceFieldObjectCulling.cpp
--- a/Engine/Source/Runtime/Renderer/Private/DistanceFieldObjectCulling.cpp
+++ b/Engine/Source/Runtime/Renderer/Private/DistanceFieldObjectCulling.cpp
@@ -76,16 +76,22 @@ TGlobalResource<FDistanceFieldObjectBufferResource> GAOCulledObjectBuffers;
 
 void FTileIntersectionResources::InitDynamicRHI()
 {
-	TileConeAxisAndCos.Initialize(sizeof(float)* 4, TileDimensions.X * TileDimensions.Y, PF_A32B32G32R32F, BUF_Static);
-	TileConeDepthRanges.Initialize(sizeof(float)* 4, TileDimensions.X * TileDimensions.Y, PF_A32B32G32R32F, BUF_Static);
+	const uint32 FastVRamFlag = IsTransientResourceBufferAliasingEnabled() ? (BUF_FastVRAM | BUF_Transient) : BUF_None;
+	TileConeAxisAndCos.Initialize(sizeof(FVector4), TileDimensions.X * TileDimensions.Y, PF_A32B32G32R32F, BUF_Static | FastVRamFlag, TEXT("TileConeAxisAndCos"));
+	TileConeDepthRanges.Initialize(sizeof(FVector4), TileDimensions.X * TileDimensions.Y, PF_A32B32G32R32F, BUF_Static | FastVRamFlag, TEXT("TileConeDepthRanges"));
 
-	NumCulledTilesArray.Initialize(sizeof(uint32), MaxSceneObjects, PF_R32_UINT, BUF_Static);
-	CulledTilesStartOffsetArray.Initialize(sizeof(uint32), MaxSceneObjects, PF_R32_UINT, BUF_Static);
+	NumCulledTilesArray.Initialize(sizeof(uint32), MaxSceneObjects, PF_R32_UINT, BUF_Static | FastVRamFlag, TEXT("NumCulledTilesArray"));
+	CulledTilesStartOffsetArray.Initialize(sizeof(uint32), MaxSceneObjects, PF_R32_UINT, BUF_Static | FastVRamFlag, TEXT("CulledTilesStartOffsetArray"));
 
 	// Can only use 16 bit for CulledTileDataArray if few enough objects and tiles
 	const bool b16BitObjectIndices = MaxSceneObjects < (1 << 16);
 	const bool b16BitCulledTileIndexBuffer = bAllow16BitIndices && b16BitObjectIndices && TileDimensions.X * TileDimensions.Y < (1 << 16);
-	CulledTileDataArray.Initialize(b16BitCulledTileIndexBuffer ? sizeof(uint16) : sizeof(uint32), GMaxDistanceFieldObjectsPerCullTile * TileDimensions.X * TileDimensions.Y * CulledTileDataStride, b16BitCulledTileIndexBuffer ? PF_R16_UINT : PF_R32_UINT);
+	CulledTileDataArray.Initialize(
+		b16BitCulledTileIndexBuffer ? sizeof(uint16) : sizeof(uint32), 
+		GMaxDistanceFieldObjectsPerCullTile * TileDimensions.X * TileDimensions.Y * CulledTileDataStride, 
+		b16BitCulledTileIndexBuffer ? PF_R16_UINT : PF_R32_UINT, 
+		BUF_Static | FastVRamFlag, 
+		TEXT("CulledTileDataArray"));
 	ObjectTilesIndirectArguments.Initialize(sizeof(uint32), 3, PF_R32_UINT, BUF_Static | BUF_DrawIndirect);
 }
 
@@ -187,13 +193,15 @@ void CullObjectsToView(FRHICommandListImmediate& RHICmdList, FScene* Scene, cons
 {
 	SCOPED_DRAW_EVENT(RHICmdList, ObjectFrustumCulling);
 
-	if (CulledObjectBuffers.Buffers.MaxObjects < Scene->DistanceFieldSceneData.NumObjectsInBuffer
+	if (!CulledObjectBuffers.IsInitialized()
+		|| CulledObjectBuffers.Buffers.MaxObjects < Scene->DistanceFieldSceneData.NumObjectsInBuffer
 		|| CulledObjectBuffers.Buffers.MaxObjects > 3 * Scene->DistanceFieldSceneData.NumObjectsInBuffer)
 	{
 		CulledObjectBuffers.Buffers.MaxObjects = Scene->DistanceFieldSceneData.NumObjectsInBuffer * 5 / 4;
-		CulledObjectBuffers.Buffers.Release();
-		CulledObjectBuffers.Buffers.Initialize();
+		CulledObjectBuffers.ReleaseResource();
+		CulledObjectBuffers.InitResource();
 	}
+	CulledObjectBuffers.Buffers.AcquireTransientResource();
 
 	{
 		ClearUAV(RHICmdList, CulledObjectBuffers.Buffers.ObjectIndirectArguments, 0);
@@ -598,7 +606,9 @@ FIntPoint BuildTileObjectLists(FRHICommandListImmediate& RHICmdList, FScene* Sce
 
 		FTileIntersectionResources*& TileIntersectionResources = ((FSceneViewState*)View.State)->AOTileIntersectionResources;
 
-		if (!TileIntersectionResources || !TileIntersectionResources->HasAllocatedEnoughFor(TileListGroupSize, Scene->DistanceFieldSceneData.NumObjectsInBuffer))
+		if (!TileIntersectionResources 
+			|| !TileIntersectionResources->IsInitialized() 
+			|| !TileIntersectionResources->HasAllocatedEnoughFor(TileListGroupSize, Scene->DistanceFieldSceneData.NumObjectsInBuffer))
 		{
 			if (TileIntersectionResources)
 			{
@@ -608,10 +618,11 @@ FIntPoint BuildTileObjectLists(FRHICommandListImmediate& RHICmdList, FScene* Sce
 			{
 				TileIntersectionResources = new FTileIntersectionResources(!IsMetalPlatform(GShaderPlatformForFeatureLevel[View.FeatureLevel]));
 			}
-
+			
 			TileIntersectionResources->SetupParameters(TileListGroupSize, Scene->DistanceFieldSceneData.NumObjectsInBuffer);
 			TileIntersectionResources->InitResource();
 		}
+		TileIntersectionResources->AcquireTransientResource();
 
 		if (GAOScatterTileCulling)
 		{
diff --git a/Engine/Source/Runtime/Renderer/Private/DistanceFieldObjectManagement.cpp b/Engine/Source/Runtime/Renderer/Private/DistanceFieldObjectManagement.cpp
--- a/Engine/Source/Runtime/Renderer/Private/DistanceFieldObjectManagement.cpp
+++ b/Engine/Source/Runtime/Renderer/Private/DistanceFieldObjectManagement.cpp
@@ -931,7 +931,7 @@ void ProcessPrimitiveUpdate(
 
 					if (GAOLogGlobalDistanceFieldModifiedPrimitives)
 					{
-						UE_LOG(LogDistanceField,Warning,TEXT("Global Distance Field primitive %s %s %s bounding radius %.1f"), (bIsAddOperation ? TEXT("add") : TEXT("update")), *PrimitiveSceneInfo->Proxy->GetOwnerName().ToString(), *PrimitiveSceneInfo->Proxy->GetResourceName().ToString(), BoundingRadius);
+						UE_LOG(LogDistanceField,Log,TEXT("Global Distance Field %s primitive %s %s %s bounding radius %.1f"), PrimitiveSceneInfo->Proxy->IsOftenMoving() ? TEXT("CACHED") : TEXT("Movable"), (bIsAddOperation ? TEXT("add") : TEXT("update")), *PrimitiveSceneInfo->Proxy->GetOwnerName().ToString(), *PrimitiveSceneInfo->Proxy->GetResourceName().ToString(), BoundingRadius);
 					}
 				}
 				else if (bIsAddOperation)
diff --git a/Engine/Source/Runtime/Renderer/Private/DistanceFieldScreenGridLighting.cpp b/Engine/Source/Runtime/Renderer/Private/DistanceFieldScreenGridLighting.cpp
--- a/Engine/Source/Runtime/Renderer/Private/DistanceFieldScreenGridLighting.cpp
+++ b/Engine/Source/Runtime/Renderer/Private/DistanceFieldScreenGridLighting.cpp
@@ -65,15 +65,16 @@ FVector2D GetJitterOffset(int32 SampleIndex)
 void FAOScreenGridResources::InitDynamicRHI()
 {
 	//@todo - 2d textures
-	ScreenGridConeVisibility.Initialize(sizeof(uint32), NumConeSampleDirections * ScreenGridDimensions.X * ScreenGridDimensions.Y, PF_R32_UINT, BUF_Static);
+	const uint32 FastVRamFlag = IsTransientResourceBufferAliasingEnabled() ? (BUF_FastVRAM | BUF_Transient) : BUF_None;
+	ScreenGridConeVisibility.Initialize(sizeof(uint32), NumConeSampleDirections * ScreenGridDimensions.X * ScreenGridDimensions.Y, PF_R32_UINT, BUF_Static | FastVRamFlag, TEXT("ScreenGridConeVisibility"));
 
 	if (bAllocateResourceForGI)
 	{
 		ConeDepthVisibilityFunction.Initialize(sizeof(float), NumConeSampleDirections * NumVisibilitySteps * ScreenGridDimensions.X * ScreenGridDimensions.Y, PF_R32_FLOAT, BUF_Static);
 		//@todo - fp16
 		StepBentNormal.Initialize(sizeof(float) * 4, NumVisibilitySteps * ScreenGridDimensions.X * ScreenGridDimensions.Y, PF_A32B32G32R32F, BUF_Static);
-		SurfelIrradiance.Initialize(sizeof(FFloat16Color), ScreenGridDimensions.X * ScreenGridDimensions.Y, PF_FloatRGBA, BUF_Static);
-		HeightfieldIrradiance.Initialize(sizeof(FFloat16Color), ScreenGridDimensions.X * ScreenGridDimensions.Y, PF_FloatRGBA, BUF_Static);
+		SurfelIrradiance.Initialize(sizeof(FFloat16Color), ScreenGridDimensions.X * ScreenGridDimensions.Y, PF_FloatRGBA, BUF_Static | FastVRamFlag, TEXT("SurfelIrradiance"));
+		HeightfieldIrradiance.Initialize(sizeof(FFloat16Color), ScreenGridDimensions.X * ScreenGridDimensions.Y, PF_FloatRGBA, BUF_Static | FastVRamFlag, TEXT("HeightfieldIrradiance"));
 	}
 }
 
@@ -755,7 +756,8 @@ void FDeferredShadingSceneRenderer::RenderDistanceFieldAOScreenGrid(
 
 	if (!ScreenGridResources 
 		|| ScreenGridResources->ScreenGridDimensions != ConeTraceBufferSize 
-		|| ScreenGridResources->bAllocateResourceForGI != bUseDistanceFieldGI)
+		|| ScreenGridResources->bAllocateResourceForGI != bUseDistanceFieldGI
+		|| !ScreenGridResources->IsInitialized())
 	{
 		if (ScreenGridResources)
 		{
@@ -771,6 +773,7 @@ void FDeferredShadingSceneRenderer::RenderDistanceFieldAOScreenGrid(
 
 		ScreenGridResources->InitResource();
 	}
+	ScreenGridResources->AcquireTransientResource();
 
 	SetRenderTarget(RHICmdList, NULL, NULL);
 
@@ -886,6 +889,7 @@ void FDeferredShadingSceneRenderer::RenderDistanceFieldAOScreenGrid(
 
 	{
 		FPooledRenderTargetDesc Desc(FPooledRenderTargetDesc::Create2DDesc(ConeTraceBufferSize, PF_FloatRGBA, FClearValueBinding::None, TexCreate_None, TexCreate_RenderTargetable | TexCreate_UAV, false));
+		Desc.Flags |= GetTextureFastVRamFlag_DynamicLayout();
 		GRenderTargetPool.FindFreeElement(RHICmdList, Desc, DownsampledBentNormal, TEXT("DownsampledBentNormal"));
 	}
 
@@ -902,6 +906,11 @@ void FDeferredShadingSceneRenderer::RenderDistanceFieldAOScreenGrid(
 		ComputeShader->UnsetParameters(RHICmdList, DownsampledBentNormal->GetRenderTargetItem());
 	}
 
+	if ( IsTransientResourceBufferAliasingEnabled() )
+	{
+		ScreenGridResources->DiscardTransientResource();
+	}
+
 	GRenderTargetPool.VisualizeTexture.SetCheckPoint(RHICmdList, DownsampledBentNormal);
 
 	PostProcessBentNormalAOScreenGrid(
diff --git a/Engine/Source/Runtime/Renderer/Private/DistanceFieldShadowing.cpp b/Engine/Source/Runtime/Renderer/Private/DistanceFieldShadowing.cpp
--- a/Engine/Source/Runtime/Renderer/Private/DistanceFieldShadowing.cpp
+++ b/Engine/Source/Runtime/Renderer/Private/DistanceFieldShadowing.cpp
@@ -669,14 +669,16 @@ void CullDistanceFieldObjectsForLight(
 	SCOPED_DRAW_EVENT(RHICmdList, CullObjectsForLight);
 
 	{
-		if (GShadowCulledObjectBuffers.Buffers.MaxObjects < Scene->DistanceFieldSceneData.NumObjectsInBuffer
+		if (!GShadowCulledObjectBuffers.IsInitialized()
+			|| GShadowCulledObjectBuffers.Buffers.MaxObjects < Scene->DistanceFieldSceneData.NumObjectsInBuffer
 			|| GShadowCulledObjectBuffers.Buffers.MaxObjects > 3 * Scene->DistanceFieldSceneData.NumObjectsInBuffer)
 		{
 			GShadowCulledObjectBuffers.Buffers.bWantBoxBounds = true;
 			GShadowCulledObjectBuffers.Buffers.MaxObjects = Scene->DistanceFieldSceneData.NumObjectsInBuffer * 5 / 4;
-			GShadowCulledObjectBuffers.Buffers.Release();
-			GShadowCulledObjectBuffers.Buffers.Initialize();
+			GShadowCulledObjectBuffers.ReleaseResource();
+			GShadowCulledObjectBuffers.InitResource();
 		}
+		GShadowCulledObjectBuffers.Buffers.AcquireTransientResource();
 
 		{
 			ClearUAV(RHICmdList, GShadowCulledObjectBuffers.Buffers.ObjectIndirectArguments, 0);
@@ -906,18 +908,19 @@ void FProjectedShadowInfo::BeginRenderRayTracedDistanceFieldProjection(FRHIComma
 			{
 				const FIntPoint BufferSize = GetBufferSizeForDFShadows();
 				FPooledRenderTargetDesc Desc(FPooledRenderTargetDesc::Create2DDesc(BufferSize, PF_G16R16F, FClearValueBinding::None, TexCreate_None, TexCreate_RenderTargetable | TexCreate_UAV, false));
+				Desc.Flags |= GetTextureFastVRamFlag_DynamicLayout();
 				GRenderTargetPool.FindFreeElement(RHICmdList, Desc, RayTracedShadowsRT, TEXT("RayTracedShadows"));
 			}
 
 			SCOPED_DRAW_EVENT(RHICmdList, RayTraceShadows);
 			SetRenderTarget(RHICmdList, NULL, NULL);
 
-			RayTraceShadows(RHICmdList, View, this, TileIntersectionResources);
+				RayTraceShadows(RHICmdList, View, this, TileIntersectionResources);
+			}
 		}
-	}
 }
 
-void FProjectedShadowInfo::RenderRayTracedDistanceFieldProjection(FRHICommandListImmediate& RHICmdList, const FViewInfo& View, bool bProjectingForForwardShading) 
+void FProjectedShadowInfo::RenderRayTracedDistanceFieldProjection(FRHICommandListImmediate& RHICmdList, const FViewInfo& View, IPooledRenderTarget* ScreenShadowMaskTexture, bool bProjectingForForwardShading )
 {
 	BeginRenderRayTracedDistanceFieldProjection(RHICmdList, View);
 
@@ -933,14 +936,19 @@ void FProjectedShadowInfo::RenderRayTracedDistanceFieldProjection(FRHICommandLis
 			ScissorRect = View.ViewRect;
 		}
 
+		if ( IsTransientResourceBufferAliasingEnabled() )
+		{
+			GShadowCulledObjectBuffers.Buffers.DiscardTransientResource();
+		}
+
 		{
-			FSceneRenderTargets::Get(RHICmdList).BeginRenderingLightAttenuation(RHICmdList);
+			SetRenderTarget(RHICmdList, ScreenShadowMaskTexture->GetRenderTargetItem().TargetableTexture, FSceneRenderTargets::Get(RHICmdList).GetSceneDepthSurface(), ESimpleRenderTargetMode::EExistingColorAndDepth, FExclusiveDepthStencil::DepthRead_StencilWrite, true);
 
 			SCOPED_DRAW_EVENT(RHICmdList, Upsample);
 
 			FGraphicsPipelineStateInitializer GraphicsPSOInit;
 			RHICmdList.ApplyCachedRenderTargets(GraphicsPSOInit);
-
+    
 			RHICmdList.SetViewport(ScissorRect.Min.X, ScissorRect.Min.Y, 0.0f, ScissorRect.Max.X, ScissorRect.Max.Y, 1.0f);
 			GraphicsPSOInit.RasterizerState = TStaticRasterizerState<FM_Solid, CM_None>::GetRHI();
 			GraphicsPSOInit.DepthStencilState = TStaticDepthStencilState<false, CF_Always>::GetRHI();
diff --git a/Engine/Source/Runtime/Renderer/Private/DistanceFieldVisualization.cpp b/Engine/Source/Runtime/Renderer/Private/DistanceFieldVisualization.cpp
--- a/Engine/Source/Runtime/Renderer/Private/DistanceFieldVisualization.cpp
+++ b/Engine/Source/Runtime/Renderer/Private/DistanceFieldVisualization.cpp
@@ -237,6 +237,11 @@ void FDeferredShadingSceneRenderer::RenderMeshDistanceFieldVisualization(FRHICom
 				}
 			}
 
+			if ( IsTransientResourceBufferAliasingEnabled())
+			{
+				GAOCulledObjectBuffers.Buffers.DiscardTransientResource();
+			}
+
 			{
 				FSceneRenderTargets::Get(RHICmdList).BeginRenderingSceneColor(RHICmdList, ESimpleRenderTargetMode::EExistingColorAndDepth, FExclusiveDepthStencil::DepthRead_StencilRead);
 
diff --git a/Engine/Source/Runtime/Renderer/Private/DistortionRendering.cpp b/Engine/Source/Runtime/Renderer/Private/DistortionRendering.cpp
--- a/Engine/Source/Runtime/Renderer/Private/DistortionRendering.cpp
+++ b/Engine/Source/Runtime/Renderer/Private/DistortionRendering.cpp
@@ -37,12 +37,6 @@ static TAutoConsoleVariable<int32> CVarDisableDistortion(
 														 TEXT("Prevents distortion effects from rendering.  Saves a full-screen framebuffer's worth of memory."),
 														 ECVF_Default);
 
-static TAutoConsoleVariable<int32> CVarFastVRamDistortion(
-	TEXT("r.FastVRamDistortion"),
-	0,
-	TEXT("Whether to store distortion in fast VRAM"),
-	ECVF_Scalability | ECVF_RenderThreadSafe);
-
 /**
 * A pixel shader for rendering the full screen refraction pass
 */
@@ -1062,10 +1056,7 @@ void FSceneRenderer::RenderDistortion(FRHICommandListImmediate& RHICmdList)
 		// Create a texture to store the resolved light attenuation values, and a render-targetable surface to hold the unresolved light attenuation values.
 		{
 			FPooledRenderTargetDesc Desc(FPooledRenderTargetDesc::Create2DDesc(SceneContext.GetBufferSizeXY(), PF_B8G8R8A8, FClearValueBinding::Transparent, TexCreate_None, TexCreate_RenderTargetable, false));
-		    if (CVarFastVRamDistortion.GetValueOnRenderThread() >= 1)
-		    {
-			    Desc.Flags |= TexCreate_FastVRAM;
-		    }
+			Desc.Flags |= GetTextureFastVRamFlag_DynamicLayout();
 			Desc.NumSamples = MSAACount;
 			GRenderTargetPool.FindFreeElement(RHICmdList, Desc, DistortionRT, TEXT("Distortion"));
 
@@ -1127,12 +1118,11 @@ void FSceneRenderer::RenderDistortion(FRHICommandListImmediate& RHICmdList)
 
 		RHICmdList.TransitionResource(EResourceTransitionAccess::EReadable, SceneContext.GetSceneColor()->GetRenderTargetItem().TargetableTexture);
 		
-// OCULUS BEGIN: select ONE render target for all views (eyes)
 		TRefCountPtr<IPooledRenderTarget> NewSceneColor;
-		// we don't create a new name to make it easier to use "vis SceneColor" and get the last HDRSceneColor
-		GRenderTargetPool.FindFreeElement(RHICmdList, SceneContext.GetSceneColor()->GetDesc(), NewSceneColor, TEXT("SceneColor"));
+		FPooledRenderTargetDesc Desc = SceneContext.GetSceneColor()->GetDesc();
+		Desc.Flags &= ~(TexCreate_FastVRAM | TexCreate_Transient);
+		GRenderTargetPool.FindFreeElement(RHICmdList, Desc, NewSceneColor, TEXT("DistortedSceneColor"));
 		const FSceneRenderTargetItem& DestRenderTarget = NewSceneColor->GetRenderTargetItem();
-// OCULUS END
 
 		// Apply distortion and store off-screen
 		SetRenderTarget(RHICmdList, DestRenderTarget.TargetableTexture, SceneContext.GetSceneDepthSurface(), ESimpleRenderTargetMode::EExistingColorAndDepth, FExclusiveDepthStencil::DepthRead_StencilWrite);
@@ -1209,7 +1199,9 @@ void FSceneRenderer::RenderDistortionES2(FRHICommandListImmediate& RHICmdList)
 		RHICmdList.CopyToResolveTarget(SceneContext.GetSceneColorSurface(), SceneContext.GetSceneColorTexture(), true, FResolveRect(0, 0, ViewFamily.FamilySizeX, ViewFamily.FamilySizeY));
 
 		TRefCountPtr<IPooledRenderTarget> SceneColorDistorted;
-		GRenderTargetPool.FindFreeElement(RHICmdList, SceneContext.GetSceneColor()->GetDesc(), SceneColorDistorted, TEXT("SceneColorDistorted"));
+		FPooledRenderTargetDesc Desc = SceneContext.GetSceneColor()->GetDesc();
+		Desc.Flags &= ~(TexCreate_FastVRAM | TexCreate_Transient);
+		GRenderTargetPool.FindFreeElement(RHICmdList, Desc, SceneColorDistorted, TEXT("SceneColorDistorted"));
 		const FSceneRenderTargetItem& DistortedRenderTarget = SceneColorDistorted->GetRenderTargetItem();
 
 		FGraphicsPipelineStateInitializer GraphicsPSOInit;
diff --git a/Engine/Source/Runtime/Renderer/Private/GlobalDistanceField.cpp b/Engine/Source/Runtime/Renderer/Private/GlobalDistanceField.cpp
--- a/Engine/Source/Runtime/Renderer/Private/GlobalDistanceField.cpp
+++ b/Engine/Source/Runtime/Renderer/Private/GlobalDistanceField.cpp
@@ -11,7 +11,7 @@
 
 int32 GAOGlobalDistanceField = 1;
 FAutoConsoleVariableRef CVarAOGlobalDistanceField(
-	TEXT("r.AOGlobalDistanceField"),
+	TEXT("r.AOGlobalDistanceField"), 
 	GAOGlobalDistanceField,
 	TEXT("Whether to use a global distance field to optimize occlusion cone traces.\n")
 	TEXT("The global distance field is created by compositing object distance fields into clipmaps as the viewer moves through the level."),
@@ -1288,13 +1288,15 @@ void UpdateGlobalDistanceFieldVolume(
 		{
 			SCOPED_DRAW_EVENT(RHICmdList, UpdateGlobalDistanceFieldVolume);
 
-			if (GGlobalDistanceFieldCulledObjectBuffers.Buffers.MaxObjects < Scene->DistanceFieldSceneData.NumObjectsInBuffer
+			if (!GGlobalDistanceFieldCulledObjectBuffers.IsInitialized()
+				|| GGlobalDistanceFieldCulledObjectBuffers.Buffers.MaxObjects < Scene->DistanceFieldSceneData.NumObjectsInBuffer
 				|| GGlobalDistanceFieldCulledObjectBuffers.Buffers.MaxObjects > 3 * Scene->DistanceFieldSceneData.NumObjectsInBuffer)
 			{
 				GGlobalDistanceFieldCulledObjectBuffers.Buffers.MaxObjects = Scene->DistanceFieldSceneData.NumObjectsInBuffer * 5 / 4;
-				GGlobalDistanceFieldCulledObjectBuffers.Buffers.Release();
-				GGlobalDistanceFieldCulledObjectBuffers.Buffers.Initialize();
+				GGlobalDistanceFieldCulledObjectBuffers.ReleaseResource();
+				GGlobalDistanceFieldCulledObjectBuffers.InitResource();
 			}
+			GGlobalDistanceFieldCulledObjectBuffers.Buffers.AcquireTransientResource();
 
 			const uint32 MaxCullGridDimension = GAOGlobalDFResolution / GCullGridTileSize;
 
@@ -1410,11 +1412,11 @@ void UpdateGlobalDistanceFieldVolume(
 									{
 										check(FlattenedDimension == Flatten_ZAxis);
 										TShaderMapRef<TCompositeObjectDistanceFieldsCS<true, Flatten_ZAxis>> ComputeShader(View.ShaderMap);
-										RHICmdList.SetComputeShader(ComputeShader->GetComputeShader());
-										ComputeShader->SetParameters(RHICmdList, Scene, View, MaxOcclusionDistance, GlobalDistanceFieldInfo.ParameterData, Clipmap, ParentDistanceField, ClipmapIndex, UpdateRegion);
-										DispatchComputeShader(RHICmdList, *ComputeShader, NumGroupsX, NumGroupsY, NumGroupsZ);
-										ComputeShader->UnsetParameters(RHICmdList, Clipmap);
-									}
+									RHICmdList.SetComputeShader(ComputeShader->GetComputeShader());
+									ComputeShader->SetParameters(RHICmdList, Scene, View, MaxOcclusionDistance, GlobalDistanceFieldInfo.ParameterData, Clipmap, ParentDistanceField, ClipmapIndex, UpdateRegion);
+									DispatchComputeShader(RHICmdList, *ComputeShader, NumGroupsX, NumGroupsY, NumGroupsZ);
+									ComputeShader->UnsetParameters(RHICmdList, Clipmap);
+								}
 								}
 								else
 								{
@@ -1442,18 +1444,18 @@ void UpdateGlobalDistanceFieldVolume(
 										DispatchComputeShader(RHICmdList, *ComputeShader, NumGroupsX, NumGroupsY, NumGroupsZ);
 										ComputeShader->UnsetParameters(RHICmdList, Clipmap);
 									}
-									else
-									{
+								else
+								{
 										check(FlattenedDimension == Flatten_ZAxis);
 										TShaderMapRef<TCompositeObjectDistanceFieldsCS<false, Flatten_ZAxis>> ComputeShader(View.ShaderMap);
-										RHICmdList.SetComputeShader(ComputeShader->GetComputeShader());
-										ComputeShader->SetParameters(RHICmdList, Scene, View, MaxOcclusionDistance, GlobalDistanceFieldInfo.ParameterData, Clipmap, NULL, ClipmapIndex, UpdateRegion);
-										DispatchComputeShader(RHICmdList, *ComputeShader, NumGroupsX, NumGroupsY, NumGroupsZ);
-										ComputeShader->UnsetParameters(RHICmdList, Clipmap);
-									}
+									RHICmdList.SetComputeShader(ComputeShader->GetComputeShader());
+									ComputeShader->SetParameters(RHICmdList, Scene, View, MaxOcclusionDistance, GlobalDistanceFieldInfo.ParameterData, Clipmap, NULL, ClipmapIndex, UpdateRegion);
+									DispatchComputeShader(RHICmdList, *ComputeShader, NumGroupsX, NumGroupsY, NumGroupsZ);
+									ComputeShader->UnsetParameters(RHICmdList, Clipmap);
 								}
 							}
 						}
+						}
 
 						if (UpdateRegion.UpdateType & VUT_Heightfields)
 						{
@@ -1462,6 +1464,11 @@ void UpdateGlobalDistanceFieldVolume(
 					}
 				}
 			}
+
+			if ( IsTransientResourceBufferAliasingEnabled() )
+			{
+				GGlobalDistanceFieldCulledObjectBuffers.Buffers.DiscardTransientResource();
+			}
 		}
 	}
 }
diff --git a/Engine/Source/Runtime/Renderer/Private/LightFunctionRendering.cpp b/Engine/Source/Runtime/Renderer/Private/LightFunctionRendering.cpp
--- a/Engine/Source/Runtime/Renderer/Private/LightFunctionRendering.cpp
+++ b/Engine/Source/Runtime/Renderer/Private/LightFunctionRendering.cpp
@@ -250,21 +250,21 @@ bool FDeferredShadingSceneRenderer::CheckForLightFunction( const FLightSceneInfo
  *
  * @param LightSceneInfo Represents the current light
  */
-bool FDeferredShadingSceneRenderer::RenderLightFunction(FRHICommandListImmediate& RHICmdList, const FLightSceneInfo* LightSceneInfo, bool bLightAttenuationCleared, bool bProjectingForForwardShading)
+bool FDeferredShadingSceneRenderer::RenderLightFunction(FRHICommandListImmediate& RHICmdList, const FLightSceneInfo* LightSceneInfo, IPooledRenderTarget* ScreenShadowMaskTexture, bool bLightAttenuationCleared, bool bProjectingForForwardShading)
 {
 	if (ViewFamily.EngineShowFlags.LightFunctions)
 	{
-		return RenderLightFunctionForMaterial(RHICmdList, LightSceneInfo, LightSceneInfo->Proxy->GetLightFunctionMaterial(), bLightAttenuationCleared, bProjectingForForwardShading, false);
+		return RenderLightFunctionForMaterial(RHICmdList, LightSceneInfo, ScreenShadowMaskTexture, LightSceneInfo->Proxy->GetLightFunctionMaterial(), bLightAttenuationCleared, bProjectingForForwardShading, false);
 	}
 	
 	return false;
 }
 
-bool FDeferredShadingSceneRenderer::RenderPreviewShadowsIndicator(FRHICommandListImmediate& RHICmdList, const FLightSceneInfo* LightSceneInfo, bool bLightAttenuationCleared)
+bool FDeferredShadingSceneRenderer::RenderPreviewShadowsIndicator(FRHICommandListImmediate& RHICmdList, const FLightSceneInfo* LightSceneInfo, IPooledRenderTarget* ScreenShadowMaskTexture, bool bLightAttenuationCleared)
 {
 	if (GEngine->PreviewShadowsIndicatorMaterial)
 	{
-		return RenderLightFunctionForMaterial(RHICmdList, LightSceneInfo, GEngine->PreviewShadowsIndicatorMaterial->GetRenderProxy(false), bLightAttenuationCleared, false, true);
+		return RenderLightFunctionForMaterial(RHICmdList, LightSceneInfo, ScreenShadowMaskTexture, GEngine->PreviewShadowsIndicatorMaterial->GetRenderProxy(false), bLightAttenuationCleared, false, true);
 	}
 
 	return false;
@@ -273,6 +273,7 @@ bool FDeferredShadingSceneRenderer::RenderPreviewShadowsIndicator(FRHICommandLis
 bool FDeferredShadingSceneRenderer::RenderLightFunctionForMaterial(
 	FRHICommandListImmediate& RHICmdList, 
 	const FLightSceneInfo* LightSceneInfo, 
+	IPooledRenderTarget* ScreenShadowMaskTexture,
 	const FMaterialRenderProxy* MaterialProxy, 
 	bool bLightAttenuationCleared, 
 	bool bProjectingForForwardShading, 
@@ -280,9 +281,10 @@ bool FDeferredShadingSceneRenderer::RenderLightFunctionForMaterial(
 {
 	bool bRenderedLightFunction = false;
 
-	if (MaterialProxy && MaterialProxy->GetMaterial(Scene->GetFeatureLevel())->IsLightFunction())
+	check(ScreenShadowMaskTexture);
+	if (MaterialProxy && MaterialProxy->GetMaterial(Scene->GetFeatureLevel())->IsLightFunction() )
 	{
-		FSceneRenderTargets::Get(RHICmdList).BeginRenderingLightAttenuation(RHICmdList);
+		SetRenderTarget(RHICmdList, ScreenShadowMaskTexture->GetRenderTargetItem().TargetableTexture, FSceneRenderTargets::Get(RHICmdList).GetSceneDepthSurface(), ESimpleRenderTargetMode::EExistingColorAndDepth, FExclusiveDepthStencil::DepthRead_StencilWrite, true);
 
 		bRenderedLightFunction = true;
 
