Fix regression where BSP would disappear when taking a high res screenshot. Exposed a bug where -1 LODIndex was possible

diff --git a/Engine/Source/Runtime/Engine/Private/ModelRender.cpp b/Engine/Source/Runtime/Engine/Private/ModelRender.cpp
--- a/Engine/Source/Runtime/Engine/Private/ModelRender.cpp
+++ b/Engine/Source/Runtime/Engine/Private/ModelRender.cpp
@@ -392,6 +392,7 @@ public:
 												MeshElement.bCanApplyViewModeOverrides = true;
 												MeshElement.bUseWireframeSelectionColoring = false;
 												MeshElement.bUseSelectionOutline = bOnlySelectedSurfaces;
+												MeshElement.LODIndex = 0;
 												Collector.AddMesh(ViewIndex, MeshElement);
 												FirstIndex += NumIndices;
 											}
@@ -426,6 +427,7 @@ public:
 								MeshElement.DepthPriorityGroup = DepthPriorityGroup;
 								MeshElement.bCanApplyViewModeOverrides = true;
 								MeshElement.bUseWireframeSelectionColoring = false;
+								MeshElement.LODIndex = 0;
 								Collector.AddMesh(ViewIndex, MeshElement);
 							}
 						}
@@ -464,6 +466,7 @@ public:
 					BatchElement.MaxVertexIndex = ModelElement.MaxVertexIndex;
 					MeshElement.Type = PT_TriangleList;
 					MeshElement.DepthPriorityGroup = PrimitiveDPG;
+					MeshElement.LODIndex = 0;
 					PDI->DrawMesh(MeshElement, FLT_MAX);
 				}
 			}
