[NO NON-PII COMMIT MESSAGE]

diff --git a/Engine/Source/Runtime/Engine/Classes/Components/InstancedStaticMeshComponent.h b/Engine/Source/Runtime/Engine/Classes/Components/InstancedStaticMeshComponent.h
--- a/Engine/Source/Runtime/Engine/Classes/Components/InstancedStaticMeshComponent.h
+++ b/Engine/Source/Runtime/Engine/Classes/Components/InstancedStaticMeshComponent.h
@@ -252,6 +252,7 @@ public:
 	virtual void Serialize(FArchive& Ar) override;
 	virtual void GetResourceSizeEx(FResourceSizeEx& CumulativeResourceSize) override;
 	void BeginDestroy() override;
+	virtual void PostDuplicate(bool bDuplicateForPIE) override;
 #if WITH_EDITOR
 	virtual void PostEditChangeChainProperty(FPropertyChangedChainEvent& PropertyChangedEvent) override;
 	virtual void PostEditUndo() override;
diff --git a/Engine/Source/Runtime/Engine/Private/HierarchicalInstancedStaticMesh.cpp b/Engine/Source/Runtime/Engine/Private/HierarchicalInstancedStaticMesh.cpp
--- a/Engine/Source/Runtime/Engine/Private/HierarchicalInstancedStaticMesh.cpp
+++ b/Engine/Source/Runtime/Engine/Private/HierarchicalInstancedStaticMesh.cpp
@@ -1821,7 +1821,7 @@ void UHierarchicalInstancedStaticMeshComponent::Serialize(FArchive& Ar)
 void UHierarchicalInstancedStaticMeshComponent::PostDuplicate(bool bDuplicateForPIE)
 {
 	Super::PostDuplicate(bDuplicateForPIE);
-	InitPerInstanceRenderData();
+
 	BuildTreeIfOutdated(false, false);
 }
 
diff --git a/Engine/Source/Runtime/Engine/Private/InstancedStaticMesh.cpp b/Engine/Source/Runtime/Engine/Private/InstancedStaticMesh.cpp
--- a/Engine/Source/Runtime/Engine/Private/InstancedStaticMesh.cpp
+++ b/Engine/Source/Runtime/Engine/Private/InstancedStaticMesh.cpp
@@ -1860,12 +1860,7 @@ void UInstancedStaticMeshComponent::InitPerInstanceRenderData(FStaticMeshInstanc
 {
 	UWorld* World = GetWorld();
 
-	if (World == nullptr)
-	{
-		return;
-	}
-
-	ERHIFeatureLevel::Type FeatureLevel = World->FeatureLevel;
+	ERHIFeatureLevel::Type FeatureLevel = World != nullptr ? World->FeatureLevel : GMaxRHIFeatureLevel;
 
 	if (!PerInstanceRenderData.IsValid())
 	{
@@ -1883,12 +1878,12 @@ void UInstancedStaticMeshComponent::InitPerInstanceRenderData(FStaticMeshInstanc
 
 		if (InSharedInstanceBufferData != nullptr)
 		{
-			PerInstanceRenderData = MakeShareable(new FPerInstanceRenderData(this, *InSharedInstanceBufferData, World->FeatureLevel, IsDynamic));
+			PerInstanceRenderData = MakeShareable(new FPerInstanceRenderData(this, *InSharedInstanceBufferData, FeatureLevel, IsDynamic));
 			bPerInstanceRenderDataWasPrebuilt = true;
 		}
 		else
 		{
-			PerInstanceRenderData = MakeShareable(new FPerInstanceRenderData(this, World->FeatureLevel, IsDynamic));
+			PerInstanceRenderData = MakeShareable(new FPerInstanceRenderData(this, FeatureLevel, IsDynamic));
 		}
 	}
 }
@@ -1900,12 +1895,13 @@ void UInstancedStaticMeshComponent::PostLoad()
 	if (!HasAnyFlags(RF_ClassDefaultObject))
 	{
 		InitPerInstanceRenderData();
-		
-		// Force update all the Render Data
-		if (CVarASyncInstaneBufferConversion.GetValueOnGameThread() > 0)
+
+		if (PerInstanceSMData.Num() > 0)
 		{
 			UWorld* World = GetWorld();
-			if (World && World->IsGameWorld() && PerInstanceSMData.Num() > 0)
+
+			// Force update all the Render Data
+			if (CVarASyncInstaneBufferConversion.GetValueOnGameThread() > 0 && World != nullptr && World->IsGameWorld())
 			{
 				World->AsyncPreRegisterLevelStreamingTasks.Increment();
 				while (InstancingRandomSeed == 0)
@@ -1915,10 +1911,10 @@ void UInstancedStaticMeshComponent::PostLoad()
 				AsyncBuildInstanceBufferTask = new FAsyncTask<FAsyncBuildInstanceBuffer>(this, World);
 				AsyncBuildInstanceBufferTask->StartBackgroundTask();
 			}
-		}
-		else
-		{
-			PerInstanceRenderData->UpdateInstanceData(this, 0, PerInstanceSMData.Num());
+			else
+			{
+				PerInstanceRenderData->UpdateInstanceData(this, 0, PerInstanceSMData.Num());
+			}
 		}
 	}
 }
@@ -2036,6 +2032,18 @@ void UInstancedStaticMeshComponent::BeginDestroy()
 	ReleasePerInstanceRenderData();
 }
 
+void UInstancedStaticMeshComponent::PostDuplicate(bool bDuplicateForPIE)
+{
+	Super::PostDuplicate(bDuplicateForPIE);
+
+	InitPerInstanceRenderData();
+
+	if (PerInstanceSMData.Num() > 0)
+	{
+		PerInstanceRenderData->UpdateInstanceData(this, 0, PerInstanceSMData.Num());
+	}
+}
+
 #if WITH_EDITOR
 void UInstancedStaticMeshComponent::PostEditChangeChainProperty(FPropertyChangedChainEvent& PropertyChangedEvent)
 {
