[NO NON-PII COMMIT MESSAGE]

diff --git a/Engine/Source/Runtime/Engine/Private/LightMap.cpp b/Engine/Source/Runtime/Engine/Private/LightMap.cpp
--- a/Engine/Source/Runtime/Engine/Private/LightMap.cpp
+++ b/Engine/Source/Runtime/Engine/Private/LightMap.cpp
@@ -9,6 +9,7 @@
 #include "Interfaces/ITargetPlatform.h"
 #include "StaticLighting.h"
 #include "Components/InstancedStaticMeshComponent.h"
+#include "InstancedStaticMesh.h"
 #include "UObject/UObjectHash.h"
 #include "UObject/UObjectIterator.h"
 #include "Misc/FeedbackContext.h"
@@ -270,6 +271,8 @@ struct FLightMapAllocation
 			FMeshMapBuildData* MeshBuildData = Registry->GetMeshBuildData(MapBuildDataId);
 			check(MeshBuildData);
 
+			UInstancedStaticMeshComponent* Component = CastChecked<UInstancedStaticMeshComponent>(Primitive);
+
 			// Instances may have been removed since LM allocation.
 			// Instances may have also been shuffled from removes. We do not handle this case.
 			if( InstanceIndex < MeshBuildData->PerInstanceLightmapData.Num() )
@@ -277,12 +280,10 @@ struct FLightMapAllocation
 				// TODO: We currently only support one LOD of static lighting in foliage
 				// Need to create per-LOD instance data to fix that
 				MeshBuildData->PerInstanceLightmapData[InstanceIndex].LightmapUVBias = LightMap->GetCoordinateBias();
-			}
 
-			UInstancedStaticMeshComponent* Component = CastChecked<UInstancedStaticMeshComponent>(Primitive);
-
-			Component->ReleasePerInstanceRenderData();
-			Component->MarkRenderStateDirty();
+				Component->PerInstanceRenderData->UpdateInstanceData(Component, InstanceIndex);
+				Component->MarkRenderStateDirty();
+			}
 		}
 	}
 
diff --git a/Engine/Source/Runtime/Engine/Private/ShadowMap.cpp b/Engine/Source/Runtime/Engine/Private/ShadowMap.cpp
--- a/Engine/Source/Runtime/Engine/Private/ShadowMap.cpp
+++ b/Engine/Source/Runtime/Engine/Private/ShadowMap.cpp
@@ -8,6 +8,7 @@
 #include "Interfaces/ITargetPlatform.h"
 #include "Engine/ShadowMapTexture2D.h"
 #include "Components/InstancedStaticMeshComponent.h"
+#include "InstancedStaticMesh.h"
 #include "LightMap.h"
 #include "UObject/Package.h"
 #include "Misc/FeedbackContext.h"
@@ -104,6 +105,8 @@ struct FShadowMapAllocation
 			FMeshMapBuildData* MeshBuildData = Registry->GetMeshBuildData(MapBuildDataId);
 			check(MeshBuildData);
 
+			UInstancedStaticMeshComponent* Component = CastChecked<UInstancedStaticMeshComponent>(Primitive);
+
 			// Instances may have been removed since LM allocation.
 			// Instances may have also been shuffled from removes. We do not handle this case.
 			if( InstanceIndex < MeshBuildData->PerInstanceLightmapData.Num() )
@@ -111,12 +114,10 @@ struct FShadowMapAllocation
 				// TODO: We currently only support one LOD of static lighting in foliage
 				// Need to create per-LOD instance data to fix that
 				MeshBuildData->PerInstanceLightmapData[InstanceIndex].ShadowmapUVBias = ShadowMap->GetCoordinateBias();
-			}
 
-			UInstancedStaticMeshComponent* Component = CastChecked<UInstancedStaticMeshComponent>(Primitive);
-
-			Component->ReleasePerInstanceRenderData();
-			Component->MarkRenderStateDirty();
+				Component->PerInstanceRenderData->UpdateInstanceData(Component, InstanceIndex);
+				Component->MarkRenderStateDirty();
+			}
 		}
 	}
 };
