[NO NON-PII COMMIT MESSAGE]

diff --git a/Engine/Source/Runtime/Engine/Private/HierarchicalInstancedStaticMesh.cpp b/Engine/Source/Runtime/Engine/Private/HierarchicalInstancedStaticMesh.cpp
--- a/Engine/Source/Runtime/Engine/Private/HierarchicalInstancedStaticMesh.cpp
+++ b/Engine/Source/Runtime/Engine/Private/HierarchicalInstancedStaticMesh.cpp
@@ -2160,6 +2160,11 @@ void UHierarchicalInstancedStaticMeshComponent::BuildTree()
 		SortedInstances = MoveTemp(Builder.Result->SortedInstances);
 		CacheMeshExtendedBounds = GetStaticMesh()->GetBounds();
 
+		if (!PerInstanceRenderData.IsValid())
+		{
+			InitPerInstanceRenderData();
+		}
+
 		// Resync RenderData with newly built cluster tree so we take into account the newly generated InstanceReorderTable generated from the cluster tree
 		PerInstanceRenderData->UpdateInstanceData(this, 0, PerInstanceSMData.Num(), false);
 
