Copying //UE4/Dev-Framework to //UE4/Dev-Main (Source: //UE4/Dev-Framework @ 3510040)

diff --git a/Engine/Source/Runtime/Engine/Private/ComponentInstanceDataCache.cpp b/Engine/Source/Runtime/Engine/Private/ComponentInstanceDataCache.cpp
--- a/Engine/Source/Runtime/Engine/Private/ComponentInstanceDataCache.cpp
+++ b/Engine/Source/Runtime/Engine/Private/ComponentInstanceDataCache.cpp
@@ -314,8 +314,19 @@ FComponentInstanceDataCache::FComponentInstanceDataCache(const AActor* Actor)
 				{
 					if (SceneComponent->GetAttachParent() && SceneComponent->GetAttachParent()->IsCreatedByConstructionScript())
 					{
+						// In rare cases the root component can be unset so walk the hierarchy and find what is probably the root component for the purposes of storing off the relative transform
+						USceneComponent* RelativeToComponent = Actor->GetRootComponent();
+						if (RelativeToComponent == nullptr)
+						{
+							RelativeToComponent = SceneComponent->GetAttachParent();
+							while (RelativeToComponent->GetAttachParent() && RelativeToComponent->GetAttachParent()->GetOwner() == Actor)
+							{
+								RelativeToComponent = RelativeToComponent->GetAttachParent();
+							}
+						}
+
 						SceneComponent->ConditionalUpdateComponentToWorld();
-						InstanceComponentTransformToRootMap.Add(SceneComponent, SceneComponent->GetComponentTransform().GetRelativeTransform(Actor->GetRootComponent()->GetComponentTransform()));
+						InstanceComponentTransformToRootMap.Add(SceneComponent, SceneComponent->GetComponentTransform().GetRelativeTransform(RelativeToComponent->GetComponentTransform()));
 					}
 				}
 			}
diff --git a/Engine/Source/Runtime/Engine/Private/Components/ChildActorComponent.cpp b/Engine/Source/Runtime/Engine/Private/Components/ChildActorComponent.cpp
--- a/Engine/Source/Runtime/Engine/Private/Components/ChildActorComponent.cpp
+++ b/Engine/Source/Runtime/Engine/Private/Components/ChildActorComponent.cpp
@@ -150,7 +150,12 @@ void UChildActorComponent::PostEditChangeProperty(FPropertyChangedEvent& Propert
 	{
 		ChildActorName = NAME_None;
 
-		if (!IsTemplate())
+		if (IsTemplate())
+		{
+			// This case is necessary to catch the situation where we are propogating the change down to child blueprints
+			SetChildActorClass(ChildActorClass);
+		}
+		else
 		{
 			UChildActorComponent* Archetype = CastChecked<UChildActorComponent>(GetArchetype());
 			ChildActorTemplate = (Archetype->ChildActorClass == ChildActorClass ? Archetype->ChildActorTemplate : nullptr);
diff --git a/Engine/Source/Runtime/Engine/Private/Components/InterpToMovementComponent.cpp b/Engine/Source/Runtime/Engine/Private/Components/InterpToMovementComponent.cpp
--- a/Engine/Source/Runtime/Engine/Private/Components/InterpToMovementComponent.cpp
+++ b/Engine/Source/Runtime/Engine/Private/Components/InterpToMovementComponent.cpp
@@ -83,17 +83,20 @@ void UInterpToMovementComponent::TickComponent(float DeltaTime, enum ELevelTick
 	{
 		WaitPos = UpdatedComponent->GetComponentLocation(); //-V595
 	}
-	while (RemainingTime >= MIN_TICK_TIME && (Iterations < MaxSimulationIterations) && !ActorOwner->IsPendingKill() && UpdatedComponent)
+	while (RemainingTime >= MIN_TICK_TIME && (Iterations < MaxSimulationIterations) && !ActorOwner->IsPendingKill() && UpdatedComponent && bIsActive)
 	{
 		Iterations++;
 
 		const float TimeTick = ShouldUseSubStepping() ? GetSimulationTimeStep(RemainingTime, Iterations) : RemainingTime;
 		RemainingTime -= TimeTick;
 
-		// Calculate the current time with this tick iteration
-		float Time = FMath::Clamp(CurrentTime + ((DeltaTime*TimeMultiplier)*CurrentDirection),0.0f,1.0f);		
-		FVector MoveDelta = ComputeMoveDelta(Time);
+		// Calculate the current alpha with this tick iteration
+		const float TargetTime = FMath::Clamp(CurrentTime + ((TimeTick*TimeMultiplier)*CurrentDirection), 0.0f, 1.0f);		
+		FVector MoveDelta = ComputeMoveDelta(TargetTime);
 		
+		// Update velocity
+		Velocity = MoveDelta / TimeTick;
+
 		// Update the rotation on the spline if required
 		FRotator CurrentRotation = UpdatedComponent->GetComponentRotation(); //-V595
 		
@@ -111,27 +114,48 @@ void UInterpToMovementComponent::TickComponent(float DeltaTime, enum ELevelTick
 		}
 		//DrawDebugPoint(GetWorld(), UpdatedComponent->GetComponentLocation(), 16, FColor::White,true,5.0f);
 		// If we hit a trigger that destroyed us, abort.
-		if (ActorOwner->IsPendingKill() || !UpdatedComponent)
+		if (ActorOwner->IsPendingKill() || !UpdatedComponent || !bIsActive)
 		{
 			return;
 		}
 
+		// Update current time
+		float AlphaRemainder = 0.0f;
+		if (bIsWaiting == false)
+		{
+			// Compute time used out of tick time to get to the hit
+			const float TimeDeltaAtHit = TimeTick * Hit.Time;
+			// Compute new time lerp alpha based on how far we moved
+			CurrentTime = CalculateNewTime(CurrentTime, TimeDeltaAtHit, Hit, true, bStopped, AlphaRemainder);
+		}
+
+		// See if we moved at all
+		if (Hit.Time != 0.f)
+		{
+			// If we were 'waiting' we are not any more - broadcast we are moving again
+			if (bIsWaiting == true)
+			{
+				OnWaitEndDelegate.Broadcast(Hit, CurrentTime);
+				bIsWaiting = false;
+			}
+		}
+
 		// Handle hit result after movement
+		float SubTickTimeRemaining = 0.0f;
 		if (!Hit.bBlockingHit)
 		{
-			// If we were 'waiting' were not any more - broadcast we are off again
-			if( bIsWaiting == true )
+			if (bStopped == true)
 			{
-				OnWaitEndDelegate.Broadcast(Hit, Time);
-				bIsWaiting = false;
+				Velocity = FVector::ZeroVector;
+				break;
 			}
-			else
-			{				
-				CalculateNewTime(CurrentTime, TimeTick, Hit, true, bStopped);
-				if (bStopped == true)
-				{
-					return;
-				}
+
+			// Handle remainder of alpha after it goes off the end, for instance if ping-pong is set and it hit the end,
+			// continue with the time remaining off the end but in the reverse direction. It is similar to hitting an object in this respect.
+			if (AlphaRemainder != 0.0f)
+			{
+				NumBounces++;
+				SubTickTimeRemaining = (AlphaRemainder * Duration);
 			}
 		}
 		else
@@ -142,26 +166,23 @@ void UInterpToMovementComponent::TickComponent(float DeltaTime, enum ELevelTick
 			}
 
 			NumBounces++;
-			float SubTickTimeRemaining = TimeTick * (1.f - Hit.Time);
-			
-			// A few initial bounces should add more time and iterations to complete most of the simulation.
-			if (NumBounces <= 2 && SubTickTimeRemaining >= MIN_TICK_TIME)
-			{
-				RemainingTime += SubTickTimeRemaining;
-				Iterations--;
-			}
+			SubTickTimeRemaining = TimeTick * (1.f - Hit.Time);
+		}
+
+		// A few initial bounces should add more time and iterations to complete most of the simulation.
+		if (NumBounces <= 2 && SubTickTimeRemaining >= MIN_TICK_TIME)
+		{
+			RemainingTime += SubTickTimeRemaining;
+			Iterations--;
 		}
 	}
-	if( bIsWaiting == false )
-	{		
-		FHitResult DummyHit;
-		CurrentTime = CalculateNewTime(CurrentTime, DeltaTime, DummyHit, false, bStopped);		
-	}
+
 	UpdateComponentVelocity();
 }
 
-float UInterpToMovementComponent::CalculateNewTime( float TimeNow, float Delta, FHitResult& HitResult, bool InBroadcastEvent, bool& OutStopped )
+float UInterpToMovementComponent::CalculateNewTime( float TimeNow, float Delta, FHitResult& HitResult, bool InBroadcastEvent, bool& OutStopped, float& OutTimeRemainder )
 {
+	OutTimeRemainder = 0.0f;
 	float NewTime = TimeNow;
 	OutStopped = false;
 	if (bIsWaiting == false)
@@ -169,6 +190,7 @@ float UInterpToMovementComponent::CalculateNewTime( float TimeNow, float Delta,
 		NewTime += ((Delta*TimeMultiplier)*CurrentDirection);
 		if (NewTime >= 1.0f)
 		{
+			OutTimeRemainder = NewTime - 1.0f;
 			if (BehaviourType == EInterpToBehaviourType::OneShot)
 			{
 				NewTime = 1.0f;
@@ -189,11 +211,13 @@ float UInterpToMovementComponent::CalculateNewTime( float TimeNow, float Delta,
 			else
 			{
 				FHitResult DummyHit;
-				NewTime = ReverseDirection(DummyHit,NewTime, InBroadcastEvent);
+				NewTime = 1.0f;
+				ReverseDirection(DummyHit, NewTime, InBroadcastEvent);
 			}	
 		}
 		else if (NewTime < 0.0f)
 		{
+			OutTimeRemainder = -NewTime;
 			if (BehaviourType == EInterpToBehaviourType::OneShot_Reverse)
 			{
 				NewTime = 0.0f;
@@ -205,9 +229,9 @@ float UInterpToMovementComponent::CalculateNewTime( float TimeNow, float Delta,
 			}
 			else if (BehaviourType == EInterpToBehaviourType::PingPong)
 			{
-
 				FHitResult DummyHit;
-				NewTime = ReverseDirection(DummyHit, NewTime, InBroadcastEvent);
+				NewTime = 0.0f;
+				ReverseDirection(DummyHit, NewTime, InBroadcastEvent);
 			}			
 		}
 	}
@@ -485,20 +509,15 @@ void UInterpToMovementComponent::UpdateControlPoints(bool InForceUpdate)
 	}
 }
 
-float UInterpToMovementComponent::ReverseDirection(const FHitResult& Hit, float Time, bool InBroadcastEvent)
+void UInterpToMovementComponent::ReverseDirection(const FHitResult& Hit, float Time, bool InBroadcastEvent)
 {
-	float NewTime = Time;
 	// Invert the direction we are moving 
 	if (InBroadcastEvent == true)
 	{
-		OnInterpToReverse.Broadcast(Hit, NewTime);
+		OnInterpToReverse.Broadcast(Hit, Time);
 	}
 	// flip dir
 	CurrentDirection = -CurrentDirection;
-	// remove extra time from current time
-	float Remainder = NewTime - 1.0f;
-	NewTime = 1.0f - Remainder;
-	return NewTime;
 }
 
 
diff --git a/Engine/Source/Runtime/Engine/Private/Components/SceneComponent.cpp b/Engine/Source/Runtime/Engine/Private/Components/SceneComponent.cpp
--- a/Engine/Source/Runtime/Engine/Private/Components/SceneComponent.cpp
+++ b/Engine/Source/Runtime/Engine/Private/Components/SceneComponent.cpp
@@ -1509,7 +1509,7 @@ FVector USceneComponent::K2_GetComponentScale() const
 
 void USceneComponent::GetParentComponents(TArray<class USceneComponent*>& Parents) const
 {
-	Parents.Empty();
+	Parents.Reset();
 
 	USceneComponent* ParentIterator = GetAttachParent();
 	while (ParentIterator != nullptr)
diff --git a/Engine/Source/Runtime/Engine/Private/EdGraph/EdGraph.cpp b/Engine/Source/Runtime/Engine/Private/EdGraph/EdGraph.cpp
--- a/Engine/Source/Runtime/Engine/Private/EdGraph/EdGraph.cpp
+++ b/Engine/Source/Runtime/Engine/Private/EdGraph/EdGraph.cpp
@@ -295,16 +295,16 @@ void UEdGraph::MoveNodesToAnotherGraph(UEdGraph* DestinationGraph, bool bIsLoadi
 			if (bInIsCompiling && !Node->IsNodeEnabled())
 			{
 				// Pass existing connections through non-enabled nodes
-				for (auto Pin : Node->Pins)
+				for (UEdGraphPin* Pin : Node->Pins)
 				{
 					if (Pin->Direction == EGPD_Input && Pin->LinkedTo.Num() > 0)
 					{
 						UEdGraphPin* PassThroughPin = Node->GetPassThroughPin(Pin);
 						if (PassThroughPin != nullptr && PassThroughPin->LinkedTo.Num() > 0)
 						{
-							for (auto OutputPin : Pin->LinkedTo)
+							for (UEdGraphPin* OutputPin : Pin->LinkedTo)
 							{
-								for (auto InputPin : PassThroughPin->LinkedTo)
+								for (UEdGraphPin* InputPin : PassThroughPin->LinkedTo)
 								{
 									InputPin->LinkedTo.Add(OutputPin);
 									OutputPin->LinkedTo.Add(InputPin);
diff --git a/Engine/Source/Runtime/Engine/Private/EdGraph/EdGraphNode.cpp b/Engine/Source/Runtime/Engine/Private/EdGraph/EdGraphNode.cpp
--- a/Engine/Source/Runtime/Engine/Private/EdGraph/EdGraphNode.cpp
+++ b/Engine/Source/Runtime/Engine/Private/EdGraph/EdGraphNode.cpp
@@ -250,9 +250,8 @@ void UEdGraphNode::BreakAllNodeLinks()
 	}
 
 	// Send all nodes that received a new pin connection a notification
-	for (auto It = NodeList.CreateConstIterator(); It; ++It)
+	for (UEdGraphNode* Node : NodeList)
 	{
-		UEdGraphNode* Node = (*It);
 		Node->NodeConnectionListChanged();
 	}
 }
@@ -272,7 +271,7 @@ void UEdGraphNode::SnapToGrid(float GridSnapSize)
 class UEdGraph* UEdGraphNode::GetGraph() const
 {
 	UEdGraph* Graph = Cast<UEdGraph>(GetOuter());
-	if (Graph == NULL && !IsPendingKill())
+	if (Graph == nullptr && !IsPendingKill())
 	{
 		ensureMsgf(false, TEXT("EdGraphNode::GetGraph : '%s' does not have a UEdGraph as an Outer."), *GetPathName());
 	}
diff --git a/Engine/Source/Runtime/Engine/Private/EdGraph/EdGraphPin.cpp b/Engine/Source/Runtime/Engine/Private/EdGraph/EdGraphPin.cpp
--- a/Engine/Source/Runtime/Engine/Private/EdGraph/EdGraphPin.cpp
+++ b/Engine/Source/Runtime/Engine/Private/EdGraph/EdGraphPin.cpp
@@ -305,6 +305,8 @@ namespace PinHelpers
 	// Don't need bIsDiffing's name because it is transient
 	static FString bAdvancedViewName = GET_MEMBER_NAME_STRING_CHECKED(UEdGraphPin, bAdvancedView);
 	// Don't need bDisplayAsMutableRef's name because it is transient
+	static FString bOrphanedPinName = GET_MEMBER_NAME_STRING_CHECKED(UEdGraphPin, bOrphanedPin);
+	// Don't need bSavePinIfOrphaned's name because it is transient
 	// Don't need bWasTrashed's name because it is transient
 #endif
 	static EPinResolveResult ImportText_PinReference(const TCHAR*& Buffer, UEdGraphPin*& PinRef, int32 ArrayIdx, UEdGraphPin* RequestingPin, EPinResolveType ResolveType);
@@ -371,13 +373,22 @@ void UEdGraphPin::BreakLinkTo(UEdGraphPin* ToPin)
 	}
 }
 
-void UEdGraphPin::BreakAllPinLinks()
+void UEdGraphPin::BreakAllPinLinks(const bool bNotifyNodes)
 {
 	TArray<UEdGraphPin*> LinkedToCopy = LinkedTo;
 
-	for (int32 LinkIdx = 0; LinkIdx < LinkedToCopy.Num(); LinkIdx++)
+	for (UEdGraphPin* LinkedToPin : LinkedToCopy)
 	{
-		BreakLinkTo(LinkedToCopy[LinkIdx]);
+		BreakLinkTo(LinkedToPin);
+#if WITH_EDITOR
+		if (bNotifyNodes)
+		{
+			if (UEdGraphNode* LinkedToNode = LinkedToPin->GetOwningNodeUnchecked())
+			{
+				LinkedToNode->PinConnectionListChanged(LinkedToPin);
+			}
+		}
+#endif
 	}
 }
 
@@ -395,6 +406,10 @@ void UEdGraphPin::CopyPersistentDataFromOldPin(const UEdGraphPin& SourcePin)
 		DefaultValue = SourcePin.DefaultValue;
 		DefaultTextValue = SourcePin.DefaultTextValue;
 	}
+	else
+	{
+		GetSchema()->ResetPinToAutogeneratedDefaultValue(this, false);
+	}
 
 	// Copy the links
 	for (int32 LinkIndex = 0; LinkIndex < SourcePin.LinkedTo.Num(); ++LinkIndex)
@@ -503,8 +518,11 @@ FString UEdGraphPin::GetDefaultAsString() const
 
 bool UEdGraphPin::DoesDefaultValueMatchAutogenerated() const
 {
-	// Ignoring case on purpose to match default behavior
-	return GetDefaultAsString().Equals(AutogeneratedDefaultValue, ESearchCase::IgnoreCase);
+	if (const UEdGraphSchema* Schema = GetSchema())
+	{
+		return Schema->DoesDefaultValueMatchAutogenerated(*this);
+	}
+	return false;
 }
 
 #if WITH_EDITORONLY_DATA
@@ -783,6 +801,13 @@ bool UEdGraphPin::ExportTextItem(FString& ValueStr, int32 PortFlags) const
 		ValueStr += PinHelpers::ExportTextPropDelimiter;
 
 		// Intentionally not exporting bDisplayAsMutableRef as it is transient
+
+		ValueStr += PinHelpers::bOrphanedPinName + TEXT("=");
+		bool LocalOrphanedPin = bOrphanedPin;
+		BoolPropCDO->ExportTextItem(ValueStr, &LocalOrphanedPin, nullptr, nullptr, PortFlags, nullptr);
+		ValueStr += PinHelpers::ExportTextPropDelimiter;
+
+		// Intentionally not exporting bSavePinIfOrphaned as it is transient
 		// Intentionally not exporting bWasTrashed as it is transient
 	}
 #endif //WITH_EDITORONLY_DATA
@@ -811,6 +836,8 @@ bool UEdGraphPin::ImportTextItem(const TCHAR*& Buffer, int32 PortFlags, class UO
 	SkipWhitespace(Buffer);
 	if (*Buffer != '(')
 	{
+		ErrorText->Logf(ELogVerbosity::Warning, TEXT("%s: Missing opening \'(\' while importing property values."), *GetName());
+
 		// Parse error
 		return false;
 	}
@@ -823,6 +850,8 @@ bool UEdGraphPin::ImportTextItem(const TCHAR*& Buffer, int32 PortFlags, class UO
 		const TCHAR* StartBuffer = Buffer;
 		if (*Buffer == 0)
 		{
+			ErrorText->Logf(ELogVerbosity::Warning, TEXT("%s: Unexpected end-of-stream while importing property values."), *GetName());
+
 			// Parse error
 			return false;
 		}
@@ -834,6 +863,8 @@ bool UEdGraphPin::ImportTextItem(const TCHAR*& Buffer, int32 PortFlags, class UO
 
 		if (*Buffer == 0)
 		{
+			ErrorText->Logf(ELogVerbosity::Warning, TEXT("%s: Unexpected end-of-stream while importing property values."), *GetName());
+
 			// Parse error
 			return false;
 		}
@@ -841,6 +872,8 @@ bool UEdGraphPin::ImportTextItem(const TCHAR*& Buffer, int32 PortFlags, class UO
 		int32 NumCharsInToken = Buffer - StartBuffer;
 		if (NumCharsInToken <= 0)
 		{
+			ErrorText->Logf(ELogVerbosity::Warning, TEXT("%s: Encountered a missing property token while importing properties."), *GetName());
+
 			// Parse error
 			return false;
 		}
@@ -886,6 +919,8 @@ bool UEdGraphPin::ImportTextItem(const TCHAR*& Buffer, int32 PortFlags, class UO
 				}
 				else
 				{
+					ErrorText->Logf(ELogVerbosity::Warning, TEXT("Unknown pin direction enum value: %s."), *DirectionString);
+
 					// Invalid enum entry
 					bParseSuccess = false;
 				}
@@ -906,12 +941,16 @@ bool UEdGraphPin::ImportTextItem(const TCHAR*& Buffer, int32 PortFlags, class UO
 				}
 				else
 				{
+					ErrorText->Logf(ELogVerbosity::Warning, TEXT("Unknown pin type member: %s."), *Right);
+
 					// Couldn't find the prop
 					bParseSuccess = false;
 				}
 			}
 			else
 			{
+				ErrorText->Logf(ELogVerbosity::Warning, TEXT("Unable to parse pin type member name from token: %s."), *PropertyToken);
+
 				// Eh, split failed to find the '.' we saw in the token?
 				bParseSuccess = false;
 			}
@@ -1007,11 +1046,24 @@ bool UEdGraphPin::ImportTextItem(const TCHAR*& Buffer, int32 PortFlags, class UO
 			bAdvancedView = LocalAdvancedView;
 		}
 		// Intentionally not importing bDisplayAsMutableRef as it is transient
+		else if (PropertyToken == PinHelpers::bOrphanedPinName)
+		{
+			bool LocalOrphanedPin = bOrphanedPin;
+			Buffer = BoolPropCDO->ImportText(Buffer, &LocalOrphanedPin, PortFlags, Parent, ErrorText);
+			bParseSuccess = (Buffer != nullptr);
+			if (EnableOrphanPins())
+			{
+				bOrphanedPin = LocalOrphanedPin;
+			}
+		}
+		// Intentionally not including bSavePinIfOrphaned. It is transient.
 		// Intentionally not importing bWasTrashed as it is transient
 #endif
 
 		if (!bParseSuccess)
 		{
+			ErrorText->Logf(ELogVerbosity::Warning, TEXT("%s: Parse error while importing property values (PinName = %s)."), *GetName(), *PinName);
+
 			// Parse error
 			return false;
 		}
@@ -1031,6 +1083,8 @@ bool UEdGraphPin::ImportTextItem(const TCHAR*& Buffer, int32 PortFlags, class UO
 
 	if (!PinId.IsValid())
 	{
+		ErrorText->Logf(ELogVerbosity::Warning, TEXT("%s: Invalid PinId after importing property values (PinName = %s)."), *GetName(), *PinName);
+
 		// Did not contain a valid PinId
 		return false;
 	}
@@ -1100,6 +1154,8 @@ UEdGraphPin::UEdGraphPin(UEdGraphNode* InOwningNode, const FGuid& PinIdGuid)
 	, bIsDiffing(false)
 	, bAdvancedView(false)
 	, bDisplayAsMutableRef(false)
+	, bOrphanedPin(false)
+	, bSavePinIfOrphaned(true)
 #endif
 	, bWasTrashed(false)
 {
@@ -1300,7 +1356,19 @@ void UEdGraphPin::DestroyImpl(bool bClearLinks)
 		LinkedTo.Empty();
 	}
 	OwningNode = nullptr;
-	SubPins.Empty();
+	for (int32 SubPinIndex = SubPins.Num() - 1; SubPinIndex >= 0; --SubPinIndex)
+	{
+		UEdGraphPin* SubPin = SubPins[SubPinIndex];
+		if (!SubPin->bWasTrashed)
+		{
+			SubPins[SubPinIndex]->DestroyImpl(bClearLinks);
+		}
+		else
+		{
+			SubPins.RemoveAt(SubPinIndex, 1, false);
+		}
+	}
+	ensure(SubPins.Num() == 0);
 	ParentPin = nullptr;
 	ReferencePassThroughConnection = nullptr;
 	bWasTrashed = true;
@@ -1329,7 +1397,7 @@ bool UEdGraphPin::Serialize(FArchive& Ar)
 	Ar << DefaultObject;
 	Ar << DefaultTextValue;
 
-	// There are several options to solve the problem wwe're fixing here. First,
+	// There are several options to solve the problem we're fixing here. First,
 	// the problem: we have two UObjects (A and B) that own Pin arrays, and within
 	// those Pin arrays there are pins with references to pins owned by the other UObject
 	// in LinkedTo. We could:
@@ -1350,28 +1418,49 @@ bool UEdGraphPin::Serialize(FArchive& Ar)
 	{
 		Ar << PersistentGuid;
 
-		// Do not reorder these! If you are adding a new one, add it to the end, and if you want to remove one, leave the others at the same offset in the uint32
 
+		// Do not reorder the persistent properties! 
+		// If you are adding a new one, add it to the end, and if you want to remove one, leave the others at the same offset in the uint32
 		uint32 BitField = 0;
 		BitField |= bHidden << 0;
 		BitField |= bNotConnectable << 1;
 		BitField |= bDefaultValueIsReadOnly << 2;
 		BitField |= bDefaultValueIsIgnored << 3;
-		// Intentionally not including bIsDiffing. It is transient.
 		BitField |= bAdvancedView << 4;
-		// Intentionally not including bDisplayAsMutableRef. It is transient.
-		// Intentionally not including bWasTrashed. It is transient.
+		BitField |= bOrphanedPin << 5;
+
+		const int32 PersistentBits = 6;
+
+		// While these properties are transient from a serialized point of view, they need to be maintained through undo/redo so add them to
+		// the bitfield after the persistent bits.
+		if (Ar.IsTransacting())
+		{
+			BitField |= bDisplayAsMutableRef << (PersistentBits + 0);
+			BitField |= bIsDiffing << (PersistentBits + 1);
+			BitField |= bSavePinIfOrphaned << (PersistentBits + 2);
+			BitField |= bWasTrashed << (PersistentBits + 3);
+		}
 
 		Ar << BitField;
 
-		bHidden = !!(BitField & 0x01);
-		bNotConnectable = !!(BitField & 0x02);
-		bDefaultValueIsReadOnly = !!(BitField & 0x04);
-		bDefaultValueIsIgnored = !!(BitField & 0x08);
-		// Intentionally not including bIsDiffing. It is transient.
-		bAdvancedView = !!(BitField & 0x10);
-		// Intentionally not including bDisplayAsMutableRef. It is transient.
-		// Intentionally not including bWasTrashed. It is transient.
+		bHidden = !!(BitField & (1 << 0));
+		bNotConnectable = !!(BitField & (1 << 1));
+		bDefaultValueIsReadOnly = !!(BitField & (1 << 2));
+		bDefaultValueIsIgnored = !!(BitField & (1 << 3));
+		bAdvancedView = !!(BitField & (1 << 4));
+		if (EnableOrphanPins())
+		{
+			bOrphanedPin = !!(BitField & (1 << 5));
+		}
+
+		if (Ar.IsTransacting())
+		{
+			bDisplayAsMutableRef = !!(BitField & (1 << (PersistentBits + 0)));
+			bIsDiffing = !!(BitField & (1 << (PersistentBits + 1)));
+			bSavePinIfOrphaned = !!(BitField & (1 << (PersistentBits + 2)));
+			bWasTrashed = !!(BitField & (1 << (PersistentBits + 3)));
+		}
+
 	}
 #endif
 
@@ -1771,6 +1860,14 @@ bool UEdGraphPin::ImportText_PinArray(const TCHAR*& Buffer, TArray<UEdGraphPin*>
 	return false;
 }
 
+TAutoConsoleVariable<int32> CVarDisableOrphanPins(TEXT("DisableOrphanPins"), 0, TEXT("0=Enable, 1=Disable"), ECVF_ReadOnly);
+
+bool UEdGraphPin::EnableOrphanPins()
+{
+	return (CVarDisableOrphanPins.GetValueOnGameThread() == 0);
+}
+
+
 /////////////////////////////////////////////////////
 // UEdGraphPin
 
diff --git a/Engine/Source/Runtime/Engine/Private/EdGraph/EdGraphSchema.cpp b/Engine/Source/Runtime/Engine/Private/EdGraph/EdGraphSchema.cpp
--- a/Engine/Source/Runtime/Engine/Private/EdGraph/EdGraphSchema.cpp
+++ b/Engine/Source/Runtime/Engine/Private/EdGraph/EdGraphSchema.cpp
@@ -415,7 +415,7 @@ bool UEdGraphSchema::TryCreateConnection(UEdGraphPin* PinA, UEdGraphPin* PinB) c
 	case CONNECT_RESPONSE_BREAK_OTHERS_A:
 		PinA->Modify();
 		PinB->Modify();
-		PinA->BreakAllPinLinks();
+		PinA->BreakAllPinLinks(true);
 		PinA->MakeLinkTo(PinB);
 		bModified = true;
 		break;
@@ -423,7 +423,7 @@ bool UEdGraphSchema::TryCreateConnection(UEdGraphPin* PinA, UEdGraphPin* PinB) c
 	case CONNECT_RESPONSE_BREAK_OTHERS_B:
 		PinA->Modify();
 		PinB->Modify();
-		PinB->BreakAllPinLinks();
+		PinB->BreakAllPinLinks(true);
 		PinA->MakeLinkTo(PinB);
 		bModified = true;
 		break;
@@ -431,8 +431,8 @@ bool UEdGraphSchema::TryCreateConnection(UEdGraphPin* PinA, UEdGraphPin* PinB) c
 	case CONNECT_RESPONSE_BREAK_OTHERS_AB:
 		PinA->Modify();
 		PinB->Modify();
-		PinA->BreakAllPinLinks();
-		PinB->BreakAllPinLinks();
+		PinA->BreakAllPinLinks(true);
+		PinB->BreakAllPinLinks(true);
 		PinA->MakeLinkTo(PinB);
 		bModified = true;
 		break;
@@ -496,6 +496,13 @@ void UEdGraphSchema::TrySetDefaultText(UEdGraphPin& InPin, const FText& InNewDef
 #endif	//#if WITH_EDITOR
 }
 
+bool UEdGraphSchema::DoesDefaultValueMatchAutogenerated(const UEdGraphPin& InPin) const
+{
+	// Same logic as UEdGraphPin::DoesDefaultValueMatchAutogenerated
+	// Ignoring case on purpose to match default behavior
+	return InPin.GetDefaultAsString().Equals(InPin.AutogeneratedDefaultValue, ESearchCase::IgnoreCase);
+}
+
 void UEdGraphSchema::BreakNodeLinks(UEdGraphNode& TargetNode) const
 {
 #if WITH_EDITOR
@@ -547,7 +554,7 @@ bool UEdGraphSchema::SetNodeMetaData(UEdGraphNode* Node, FName const& KeyValue)
 	return false;
 }
 
-void UEdGraphSchema::BreakPinLinks(UEdGraphPin& TargetPin, bool bSendsNodeNotifcation) const
+void UEdGraphSchema::BreakPinLinks(UEdGraphPin& TargetPin, bool bSendsNodeNotification) const
 {
 #if WITH_EDITOR
 	// Copy the old pin links
@@ -568,23 +575,20 @@ void UEdGraphSchema::BreakPinLinks(UEdGraphPin& TargetPin, bool bSendsNodeNotifc
 	}
 
 	// As well as all other nodes that were connected
-	for (TArray<UEdGraphPin*>::TIterator PinIt(OldLinkedTo); PinIt; ++PinIt)
+	for (UEdGraphPin* OtherPin : OldLinkedTo)
 	{
-		UEdGraphPin* OtherPin = *PinIt;
-		UEdGraphNode* OtherNode = OtherPin->GetOwningNode();
-		if (OtherNode != nullptr)
+		if (UEdGraphNode* OtherNode = OtherPin->GetOwningNode())
 		{
 			OtherNode->PinConnectionListChanged(OtherPin);
 			NodeList.Add(OtherNode);
 		}
 	}
 
-	if (bSendsNodeNotifcation)
+	if (bSendsNodeNotification)
 	{
 		// Send all nodes that received a new pin connection a notification
-		for (auto It = NodeList.CreateConstIterator(); It; ++It)
+		for (UEdGraphNode* Node : NodeList)
 		{
-			UEdGraphNode* Node = (*It);
 			Node->NodeConnectionListChanged();
 		}
 	}
@@ -625,14 +629,20 @@ FPinConnectionResponse UEdGraphSchema::MovePinLinks(UEdGraphPin& MoveFromPin, UE
 	for (int32 i=0; i<CurrentLinks.Num(); i++)
 	{
 		UEdGraphPin* NewLink = CurrentLinks[i];
-		FPinConnectionResponse Response = CanCreateConnection(&MoveToPin, NewLink);
-		if(Response.CanSafeConnect())
-		{
-			MoveToPin.MakeLinkTo(NewLink);
-		}
-		else
+#if WITH_EDITORONLY_DATA
+		// Don't move connections to removed pins
+		if (!NewLink->bOrphanedPin)
+#endif
 		{
-			FinalResponse = Response;
+			FPinConnectionResponse Response = CanCreateConnection(&MoveToPin, NewLink);
+			if (Response.CanSafeConnect())
+			{
+				MoveToPin.MakeLinkTo(NewLink);
+			}
+			else
+			{
+				FinalResponse = Response;
+			}
 		}
 	}
 	// Move over the default values
@@ -654,13 +664,19 @@ FPinConnectionResponse UEdGraphSchema::CopyPinLinks(UEdGraphPin& CopyFromPin, UE
 	{
 		UEdGraphPin* NewLink = CopyFromPin.LinkedTo[i];
 		FPinConnectionResponse Response = CanCreateConnection(&CopyToPin, NewLink);
-		if (Response.CanSafeConnect())
-		{
-			CopyToPin.MakeLinkTo(NewLink);
-		}
-		else
+#if WITH_EDITORONLY_DATA
+		// Don't copy connections to removed pins
+		if (!NewLink->bOrphanedPin)
+#endif
 		{
-			FinalResponse = Response;
+			if (Response.CanSafeConnect())
+			{
+				CopyToPin.MakeLinkTo(NewLink);
+			}
+			else
+			{
+				FinalResponse = Response;
+			}
 		}
 	}
 
diff --git a/Engine/Source/Runtime/Engine/Private/GameFramework/SpringArmComponent.cpp b/Engine/Source/Runtime/Engine/Private/GameFramework/SpringArmComponent.cpp
--- a/Engine/Source/Runtime/Engine/Private/GameFramework/SpringArmComponent.cpp
+++ b/Engine/Source/Runtime/Engine/Private/GameFramework/SpringArmComponent.cpp
@@ -198,6 +198,13 @@ FVector USpringArmComponent::BlendLocations(const FVector& DesiredArmLocation, c
 	return bHitSomething ? TraceHitLocation : DesiredArmLocation;
 }
 
+void USpringArmComponent::ApplyWorldOffset(const FVector & InOffset, bool bWorldShift)
+{
+	Super::ApplyWorldOffset(InOffset, bWorldShift);
+	PreviousDesiredLoc += InOffset;
+	PreviousArmOrigin += InOffset;
+}
+
 void USpringArmComponent::OnRegister()
 {
 	Super::OnRegister();
diff --git a/Engine/Source/Runtime/Engine/Private/KismetFunctionLibrary.cpp b/Engine/Source/Runtime/Engine/Private/KismetFunctionLibrary.cpp
--- a/Engine/Source/Runtime/Engine/Private/KismetFunctionLibrary.cpp
+++ b/Engine/Source/Runtime/Engine/Private/KismetFunctionLibrary.cpp
@@ -49,14 +49,14 @@ int32 UBlueprintFunctionLibrary::GetFunctionCallspace(UFunction* Function, void*
 		// check which world the cosmetic/authoritative call would run in
 		if (PlayInSettings->GetRunUnderOneProcess(bIsOnePIEProcessEnabled) && bIsOnePIEProcessEnabled)
 		{
-			UObject* WorldContextObj = nullptr;
+			UWorld* WorldContext = nullptr;
 			if (Stack != nullptr)
 			{
 //				const FName WorldContextTag = TEXT("WorldContext");
-
-				if (Stack->Object && Stack->Object->ImplementsGetWorld())
+				
+				if (Stack->Object)
 				{
-					WorldContextObj = Stack->Object;
+					WorldContext = GEngine->GetWorldFromContextObject(Stack->Object, EGetWorldErrorMode::ReturnNull);
 				}
 // 				else if (Function->HasMetaData(WorldContextTag))
 // 				{
@@ -71,15 +71,11 @@ int32 UBlueprintFunctionLibrary::GetFunctionCallspace(UFunction* Function, void*
 // 				}
 			}
 
-			if (WorldContextObj != nullptr)
+			if (WorldContext != nullptr)
 			{
-				UWorld* WorldContext = WorldContextObj->GetWorld();
-				if (ensure(WorldContext != nullptr))
-				{
-					ENetMode WorldNetMode = WorldContext->GetNetMode();
-					bAbsorbCosmeticCalls  = (WorldNetMode == NM_DedicatedServer);
-					bAbsorbAuthorityCalls = (WorldNetMode == NM_Client);
-				}
+				ENetMode WorldNetMode = WorldContext->GetNetMode();
+				bAbsorbCosmeticCalls  = (WorldNetMode == NM_DedicatedServer);
+				bAbsorbAuthorityCalls = (WorldNetMode == NM_Client);
 			}
 			else
 			{
diff --git a/Engine/Source/Runtime/Engine/Private/KismetMaterialLibrary.cpp b/Engine/Source/Runtime/Engine/Private/KismetMaterialLibrary.cpp
--- a/Engine/Source/Runtime/Engine/Private/KismetMaterialLibrary.cpp
+++ b/Engine/Source/Runtime/Engine/Private/KismetMaterialLibrary.cpp
@@ -25,21 +25,22 @@ void UKismetMaterialLibrary::SetScalarParameterValue(UObject* WorldContextObject
 {
 	if (Collection)
 	{
-		check(WorldContextObject);
-		UWorld* World = GEngine->GetWorldFromContextObject(WorldContextObject);
-		UMaterialParameterCollectionInstance* Instance = World->GetParameterCollectionInstance(Collection);
-
-		const bool bFoundParameter = Instance->SetScalarParameterValue(ParameterName, ParameterValue);
-
-		if (!bFoundParameter && !Instance->bLoggedMissingParameterWarning)
+		if (UWorld* World = GEngine->GetWorldFromContextObject(WorldContextObject, EGetWorldErrorMode::LogAndReturnNull))
 		{
-			FFormatNamedArguments Arguments;
-			Arguments.Add(TEXT("ParamName"), FText::FromName(ParameterName));
-			FMessageLog("PIE").Warning()
-				->AddToken(FTextToken::Create(LOCTEXT("SetScalarParamOn", "SetScalarParameterValue called on")))
-				->AddToken(FUObjectToken::Create(Collection))
-				->AddToken(FTextToken::Create(FText::Format(LOCTEXT("WithInvalidParam", "with invalid ParameterName '{ParamName}'. This is likely due to a Blueprint error."), Arguments)));
-			Instance->bLoggedMissingParameterWarning = true;
+			UMaterialParameterCollectionInstance* Instance = World->GetParameterCollectionInstance(Collection);
+
+			const bool bFoundParameter = Instance->SetScalarParameterValue(ParameterName, ParameterValue);
+
+			if (!bFoundParameter && !Instance->bLoggedMissingParameterWarning)
+			{
+				FFormatNamedArguments Arguments;
+				Arguments.Add(TEXT("ParamName"), FText::FromName(ParameterName));
+				FMessageLog("PIE").Warning()
+					->AddToken(FTextToken::Create(LOCTEXT("SetScalarParamOn", "SetScalarParameterValue called on")))
+					->AddToken(FUObjectToken::Create(Collection))
+					->AddToken(FTextToken::Create(FText::Format(LOCTEXT("WithInvalidParam", "with invalid ParameterName '{ParamName}'. This is likely due to a Blueprint error."), Arguments)));
+				Instance->bLoggedMissingParameterWarning = true;
+			}
 		}
 	}
 }
@@ -48,21 +49,22 @@ void UKismetMaterialLibrary::SetVectorParameterValue(UObject* WorldContextObject
 {
 	if (Collection)
 	{
-		check(WorldContextObject);
-		UWorld* World = GEngine->GetWorldFromContextObject(WorldContextObject);
-		UMaterialParameterCollectionInstance* Instance = World->GetParameterCollectionInstance(Collection);
-
-		const bool bFoundParameter = Instance->SetVectorParameterValue(ParameterName, ParameterValue);
-
-		if (!bFoundParameter && !Instance->bLoggedMissingParameterWarning)
+		if (UWorld* World = GEngine->GetWorldFromContextObject(WorldContextObject, EGetWorldErrorMode::LogAndReturnNull))
 		{
-			FFormatNamedArguments Arguments;
-			Arguments.Add(TEXT("ParamName"), FText::FromName(ParameterName));
-			FMessageLog("PIE").Warning()
-				->AddToken(FTextToken::Create(LOCTEXT("SetVectorParamOn", "SetVectorParameterValue called on")))
-				->AddToken(FUObjectToken::Create(Collection))
-				->AddToken(FTextToken::Create(FText::Format(LOCTEXT("WithInvalidParam", "with invalid ParameterName '{ParamName}'. This is likely due to a Blueprint error."), Arguments)));
-			Instance->bLoggedMissingParameterWarning = true;
+			UMaterialParameterCollectionInstance* Instance = World->GetParameterCollectionInstance(Collection);
+
+			const bool bFoundParameter = Instance->SetVectorParameterValue(ParameterName, ParameterValue);
+
+			if (!bFoundParameter && !Instance->bLoggedMissingParameterWarning)
+			{
+				FFormatNamedArguments Arguments;
+				Arguments.Add(TEXT("ParamName"), FText::FromName(ParameterName));
+				FMessageLog("PIE").Warning()
+					->AddToken(FTextToken::Create(LOCTEXT("SetVectorParamOn", "SetVectorParameterValue called on")))
+					->AddToken(FUObjectToken::Create(Collection))
+					->AddToken(FTextToken::Create(FText::Format(LOCTEXT("WithInvalidParam", "with invalid ParameterName '{ParamName}'. This is likely due to a Blueprint error."), Arguments)));
+				Instance->bLoggedMissingParameterWarning = true;
+			}
 		}
 	}
 }
@@ -73,21 +75,22 @@ float UKismetMaterialLibrary::GetScalarParameterValue(UObject* WorldContextObjec
 
 	if (Collection)
 	{
-		check(WorldContextObject);
-		UWorld* World = GEngine->GetWorldFromContextObject(WorldContextObject);
-		UMaterialParameterCollectionInstance* Instance = World->GetParameterCollectionInstance(Collection);
-
-		bool bFoundParameter = Instance->GetScalarParameterValue(ParameterName, ParameterValue);
-
-		if (!bFoundParameter && !Instance->bLoggedMissingParameterWarning)
+		if (UWorld* World = GEngine->GetWorldFromContextObject(WorldContextObject, EGetWorldErrorMode::LogAndReturnNull))
 		{
-			FFormatNamedArguments Arguments;
-			Arguments.Add(TEXT("ParamName"), FText::FromName(ParameterName));
-			FMessageLog("PIE").Warning()
-				->AddToken(FTextToken::Create(LOCTEXT("GetScalarParamOn", "GetScalarParameterValue called on")))
-				->AddToken(FUObjectToken::Create(Collection))
-				->AddToken(FTextToken::Create(FText::Format(LOCTEXT("WithInvalidParam", "with invalid ParameterName '{ParamName}'. This is likely due to a Blueprint error."), Arguments)));
-			Instance->bLoggedMissingParameterWarning = true;
+			UMaterialParameterCollectionInstance* Instance = World->GetParameterCollectionInstance(Collection);
+
+			const bool bFoundParameter = Instance->GetScalarParameterValue(ParameterName, ParameterValue);
+
+			if (!bFoundParameter && !Instance->bLoggedMissingParameterWarning)
+			{
+				FFormatNamedArguments Arguments;
+				Arguments.Add(TEXT("ParamName"), FText::FromName(ParameterName));
+				FMessageLog("PIE").Warning()
+					->AddToken(FTextToken::Create(LOCTEXT("GetScalarParamOn", "GetScalarParameterValue called on")))
+					->AddToken(FUObjectToken::Create(Collection))
+					->AddToken(FTextToken::Create(FText::Format(LOCTEXT("WithInvalidParam", "with invalid ParameterName '{ParamName}'. This is likely due to a Blueprint error."), Arguments)));
+				Instance->bLoggedMissingParameterWarning = true;
+			}
 		}
 	}
 
@@ -100,21 +103,22 @@ FLinearColor UKismetMaterialLibrary::GetVectorParameterValue(UObject* WorldConte
 
 	if (Collection)
 	{
-		check(WorldContextObject);
-		UWorld* World = GEngine->GetWorldFromContextObject(WorldContextObject);
-		UMaterialParameterCollectionInstance* Instance = World->GetParameterCollectionInstance(Collection);
-
-		bool bFoundParameter = Instance->GetVectorParameterValue(ParameterName, ParameterValue);
-
-		if (!bFoundParameter && !Instance->bLoggedMissingParameterWarning)
+		if (UWorld* World = GEngine->GetWorldFromContextObject(WorldContextObject, EGetWorldErrorMode::LogAndReturnNull))
 		{
-			FFormatNamedArguments Arguments;
-			Arguments.Add(TEXT("ParamName"), FText::FromName(ParameterName));
-			FMessageLog("PIE").Warning()
-				->AddToken(FTextToken::Create(LOCTEXT("GetVectorParamOn", "GetVectorParameterValue called on")))
-				->AddToken(FUObjectToken::Create(Collection))
-				->AddToken(FTextToken::Create(FText::Format(LOCTEXT("WithInvalidParam", "with invalid ParameterName '{ParamName}'. This is likely due to a Blueprint error."), Arguments)));
-			Instance->bLoggedMissingParameterWarning = true;
+			UMaterialParameterCollectionInstance* Instance = World->GetParameterCollectionInstance(Collection);
+
+			const bool bFoundParameter = Instance->GetVectorParameterValue(ParameterName, ParameterValue);
+
+			if (!bFoundParameter && !Instance->bLoggedMissingParameterWarning)
+			{
+				FFormatNamedArguments Arguments;
+				Arguments.Add(TEXT("ParamName"), FText::FromName(ParameterName));
+				FMessageLog("PIE").Warning()
+					->AddToken(FTextToken::Create(LOCTEXT("GetVectorParamOn", "GetVectorParameterValue called on")))
+					->AddToken(FUObjectToken::Create(Collection))
+					->AddToken(FTextToken::Create(FText::Format(LOCTEXT("WithInvalidParam", "with invalid ParameterName '{ParamName}'. This is likely due to a Blueprint error."), Arguments)));
+				Instance->bLoggedMissingParameterWarning = true;
+			}
 		}
 	}
 
diff --git a/Engine/Source/Runtime/Engine/Private/KismetMathLibrary.cpp b/Engine/Source/Runtime/Engine/Private/KismetMathLibrary.cpp
--- a/Engine/Source/Runtime/Engine/Private/KismetMathLibrary.cpp
+++ b/Engine/Source/Runtime/Engine/Private/KismetMathLibrary.cpp
@@ -8,6 +8,8 @@
 #include "Blueprint/BlueprintSupport.h"
 #include "Math/DualQuat.h"
 
+#include "Misc/RuntimeErrors.h"
+
 #define LOCTEXT_NAMESPACE "UKismetMathLibrary"
 
 /** Interpolate a linear alpha value using an ease mode and function. */
@@ -106,6 +108,11 @@ void UKismetMathLibrary::ReportError_Divide_VectorVector()
 	FFrame::KismetExecutionMessage(TEXT("Divide by zero: Divide_VectorVector"), ELogVerbosity::Warning, DivideByZeroWarning);
 }
 
+void UKismetMathLibrary::ReportError_Divide_Vector2DVector2D()
+{
+	FFrame::KismetExecutionMessage(TEXT("Divide by zero: Divide_Vector2DVector2D"), ELogVerbosity::Warning, DivideByZeroWarning);
+}
+
 void UKismetMathLibrary::ReportError_ProjectVectorOnToVector()
 {
 	FFrame::KismetExecutionMessage(TEXT("Divide by zero: ProjectVectorOnToVector with zero Target vector"), ELogVerbosity::Warning, ZeroLengthProjectionWarning);
@@ -618,6 +625,106 @@ void UKismetMathLibrary::BreakTimespan(FTimespan InTimespan, int32& Days, int32&
 	Milliseconds = InTimespan.GetMilliseconds();
 }
 
+FTimespan UKismetMathLibrary::FromDays(float Days)
+{
+	if (Days < FTimespan::MinValue().GetTotalDays())
+	{
+		FFormatNamedArguments Args;
+		Args.Add(TEXT("DaysValue"), Days);
+		LogRuntimeError(FText::Format(LOCTEXT("ClampDaysToMinTimespan", "Days value {DaysValue} is less than minimum days TimeSpan can represent. Clamping to MinValue."), Args));
+		return FTimespan::MinValue();
+	}
+	else if (Days > FTimespan::MaxValue().GetTotalDays())
+	{
+		FFormatNamedArguments Args;
+		Args.Add(TEXT("DaysValue"), Days);
+		LogRuntimeError(FText::Format(LOCTEXT("ClampDaysToMaxTimespan", "Days value {DaysValue} is greater than maximum days TimeSpan can represent. Clamping to MaxValue."), Args));
+		return FTimespan::MaxValue();
+	}
+
+	return FTimespan::FromDays(Days);
+}
+
+FTimespan UKismetMathLibrary::FromHours(float Hours)
+{
+	if (Hours < FTimespan::MinValue().GetTotalHours())
+	{
+		FFormatNamedArguments Args;
+		Args.Add(TEXT("HoursValue"), Hours);
+		LogRuntimeError(FText::Format(LOCTEXT("ClampHoursToMinTimespan", "Hours value {HoursValue} is less than minimum hours TimeSpan can represent. Clamping to MinValue."), Args));
+		return FTimespan::MinValue();
+	}
+	else if (Hours > FTimespan::MaxValue().GetTotalHours())
+	{
+		FFormatNamedArguments Args;
+		Args.Add(TEXT("HoursValue"), Hours);
+		LogRuntimeError(FText::Format(LOCTEXT("ClampHoursToMaxTimespan", "Hours value {HoursValue} is greater than maximum hours TimeSpan can represent. Clamping to MaxValue."), Args));
+		return FTimespan::MaxValue();
+	}
+
+	return FTimespan::FromHours(Hours);
+}
+
+FTimespan UKismetMathLibrary::FromMinutes(float Minutes)
+{
+	if (Minutes < FTimespan::MinValue().GetTotalMinutes())
+	{
+		FFormatNamedArguments Args;
+		Args.Add(TEXT("MinutesValue"), Minutes);
+		LogRuntimeError(FText::Format(LOCTEXT("ClampMinutesToMinTimespan", "Minutes value {MinutesValue} is less than minimum minutes TimeSpan can represent. Clamping to MinValue."), Args));
+		return FTimespan::MinValue();
+	}
+	else if (Minutes > FTimespan::MaxValue().GetTotalMinutes())
+	{
+		FFormatNamedArguments Args;
+		Args.Add(TEXT("MinutesValue"), Minutes);
+		LogRuntimeError(FText::Format(LOCTEXT("ClampMinutesToMaxTimespan", "Minutes value {MinutesValue} is greater than maximum minutes TimeSpan can represent. Clamping to MaxValue."), Args));
+		return FTimespan::MaxValue();
+	}
+
+	return FTimespan::FromMinutes(Minutes);
+}
+
+FTimespan UKismetMathLibrary::FromSeconds(float Seconds)
+{
+	if (Seconds < FTimespan::MinValue().GetTotalSeconds())
+	{
+		FFormatNamedArguments Args;
+		Args.Add(TEXT("SecondsValue"), Seconds);
+		LogRuntimeError(FText::Format(LOCTEXT("ClampSecondsToMinTimespan", "Seconds value {SecondsValue} is less than minimum seconds TimeSpan can represent. Clamping to MinValue."), Args));
+		return FTimespan::MinValue();
+	}
+	else if (Seconds > FTimespan::MaxValue().GetTotalSeconds())
+	{
+		FFormatNamedArguments Args;
+		Args.Add(TEXT("SecondsValue"), Seconds);
+		LogRuntimeError(FText::Format(LOCTEXT("ClampSecondsToMaxTimespan", "Seconds value {SecondsValue} is greater than maximum seconds TimeSpan can represent. Clamping to MaxValue."), Args));
+		return FTimespan::MaxValue();
+	}
+
+	return FTimespan::FromSeconds(Seconds);
+}
+
+FTimespan UKismetMathLibrary::FromMilliseconds(float Milliseconds)
+{
+	if (Milliseconds < FTimespan::MinValue().GetTotalMilliseconds())
+	{
+		FFormatNamedArguments Args;
+		Args.Add(TEXT("MillisecondsValue"), Milliseconds);
+		LogRuntimeError(FText::Format(LOCTEXT("ClampMillisecondsToMinTimespan", "Milliseconds value {MillisecondsValue} is less than minimum milliseconds TimeSpan can represent. Clamping to MinValue."), Args));
+		return FTimespan::MinValue();
+	}
+	else if (Milliseconds > FTimespan::MaxValue().GetTotalMilliseconds())
+	{
+		FFormatNamedArguments Args;
+		Args.Add(TEXT("MillisecondsValue"), Milliseconds);
+		LogRuntimeError(FText::Format(LOCTEXT("ClampMillisecondsToMaxTimespan", "Milliseconds value {MillisecondsValue} is greater than maximum milliseconds TimeSpan can represent. Clamping to MaxValue."), Args));
+		return FTimespan::MaxValue();
+	}
+
+	return FTimespan::FromMilliseconds(Milliseconds);
+}
+
 /* END Timespan functions */
 
 
@@ -826,20 +933,15 @@ void UKismetMathLibrary::MinimumAreaRectangle(class UObject* WorldContextObject,
 	OutSideLengthY = RectSideB.Size();
 
 #if ENABLE_DRAW_DEBUG
-	if( bDebugDraw )
+	if (bDebugDraw)
 	{
-		UWorld* World = (WorldContextObject) ? GEngine->GetWorldFromContextObject(WorldContextObject) : nullptr;
-		if (World != nullptr)
+		if (UWorld* World = GEngine->GetWorldFromContextObject(WorldContextObject, EGetWorldErrorMode::LogAndReturnNull))
 		{
 			DrawDebugSphere(World, OutRectCenter, 10.f, 12, FColor::Yellow, true);
 			DrawDebugCoordinateSystem(World, OutRectCenter, SurfaceNormalMatrix.Rotator(), 100.f, true);
 			DrawDebugLine(World, OutRectCenter - RectSideA * 0.5f + FVector(0, 0, 10.f), OutRectCenter + RectSideA * 0.5f + FVector(0, 0, 10.f), FColor::Green, true, -1, 0, 5.f);
 			DrawDebugLine(World, OutRectCenter - RectSideB * 0.5f + FVector(0, 0, 10.f), OutRectCenter + RectSideB * 0.5f + FVector(0, 0, 10.f), FColor::Blue, true, -1, 0, 5.f);
 		}
-		else
-		{
-			FFrame::KismetExecutionMessage(TEXT("WorldContext required for MinimumAreaRectangle to draw a debug visualization."), ELogVerbosity::Warning);
-		}
 	}
 #endif
 }
diff --git a/Engine/Source/Runtime/Engine/Private/KismetNodeHelperLibrary.cpp b/Engine/Source/Runtime/Engine/Private/KismetNodeHelperLibrary.cpp
--- a/Engine/Source/Runtime/Engine/Private/KismetNodeHelperLibrary.cpp
+++ b/Engine/Source/Runtime/Engine/Private/KismetNodeHelperLibrary.cpp
@@ -2,6 +2,7 @@
 
 #include "Kismet/KismetNodeHelperLibrary.h"
 #include "EngineLogs.h"
+#include "Misc/RuntimeErrors.h"
 
 UKismetNodeHelperLibrary::UKismetNodeHelperLibrary(const FObjectInitializer& ObjectInitializer)
 	: Super(ObjectInitializer)
@@ -158,7 +159,7 @@ FString UKismetNodeHelperLibrary::GetEnumeratorUserFriendlyName(const UEnum* Enu
 
 uint8 UKismetNodeHelperLibrary::GetValidValue(const UEnum* Enum, uint8 EnumeratorValue)
 {
-	if (ensure(Enum))
+	if (ensureAsRuntimeWarning(Enum != nullptr))
 	{
 		if (Enum->IsValidEnumValue(EnumeratorValue))
 		{
@@ -166,7 +167,6 @@ uint8 UKismetNodeHelperLibrary::GetValidValue(const UEnum* Enum, uint8 Enumerato
 		}
 		return Enum->GetMaxEnumValue();
 	}
-	UE_LOG(LogBlueprintUserMessages, Error, TEXT("UKismetNodeHelperLibrary::GetValidValue is called with Enum == null"));
 	return INDEX_NONE;
 }
 
diff --git a/Engine/Source/Runtime/Engine/Private/KismetRenderingLibrary.cpp b/Engine/Source/Runtime/Engine/Private/KismetRenderingLibrary.cpp
--- a/Engine/Source/Runtime/Engine/Private/KismetRenderingLibrary.cpp
+++ b/Engine/Source/Runtime/Engine/Private/KismetRenderingLibrary.cpp
@@ -32,8 +32,7 @@ UKismetRenderingLibrary::UKismetRenderingLibrary(const FObjectInitializer& Objec
 
 void UKismetRenderingLibrary::ClearRenderTarget2D(UObject* WorldContextObject, UTextureRenderTarget2D* TextureRenderTarget, FLinearColor ClearColor)
 {
-	check(WorldContextObject);
-	UWorld* World = GEngine->GetWorldFromContextObject(WorldContextObject);
+	UWorld* World = GEngine->GetWorldFromContextObject(WorldContextObject, EGetWorldErrorMode::LogAndReturnNull);
 
 	if (TextureRenderTarget
 		&& TextureRenderTarget->Resource
@@ -51,8 +50,7 @@ void UKismetRenderingLibrary::ClearRenderTarget2D(UObject* WorldContextObject, U
 
 UTextureRenderTarget2D* UKismetRenderingLibrary::CreateRenderTarget2D(UObject* WorldContextObject, int32 Width, int32 Height)
 {
-	check(WorldContextObject);
-	UWorld* World = GEngine->GetWorldFromContextObject(WorldContextObject);
+	UWorld* World = GEngine->GetWorldFromContextObject(WorldContextObject, EGetWorldErrorMode::LogAndReturnNull);
 
 	if (Width > 0 && Height > 0 && World && FApp::CanEverRender())
 	{
@@ -69,8 +67,7 @@ UTextureRenderTarget2D* UKismetRenderingLibrary::CreateRenderTarget2D(UObject* W
 
 void UKismetRenderingLibrary::DrawMaterialToRenderTarget(UObject* WorldContextObject, UTextureRenderTarget2D* TextureRenderTarget, UMaterialInterface* Material)
 {
-	check(WorldContextObject);
-	UWorld* World = GEngine->GetWorldFromContextObject(WorldContextObject, false);
+	UWorld* World = GEngine->GetWorldFromContextObject(WorldContextObject, EGetWorldErrorMode::LogAndReturnNull);
 
 	if (TextureRenderTarget 
 		&& TextureRenderTarget->Resource 
@@ -271,8 +268,7 @@ void UKismetRenderingLibrary::ExportTexture2D(UObject* WorldContextObject, UText
 
 void UKismetRenderingLibrary::BeginDrawCanvasToRenderTarget(UObject* WorldContextObject, UTextureRenderTarget2D* TextureRenderTarget, UCanvas*& Canvas, FVector2D& Size, FDrawToRenderTargetContext& Context)
 {
-	check(WorldContextObject);
-	UWorld* World = GEngine->GetWorldFromContextObject(WorldContextObject, false);
+	UWorld* World = GEngine->GetWorldFromContextObject(WorldContextObject, EGetWorldErrorMode::LogAndReturnNull);
 
 	Canvas = NULL;
 	Size = FVector2D(0, 0);
@@ -325,8 +321,7 @@ void UKismetRenderingLibrary::BeginDrawCanvasToRenderTarget(UObject* WorldContex
 
 void UKismetRenderingLibrary::EndDrawCanvasToRenderTarget(UObject* WorldContextObject, const FDrawToRenderTargetContext& Context)
 {
-	check(WorldContextObject);
-	UWorld* World = GEngine->GetWorldFromContextObject(WorldContextObject, false);
+	UWorld* World = GEngine->GetWorldFromContextObject(WorldContextObject, EGetWorldErrorMode::LogAndReturnNull);
 
 	if (World)
 	{
@@ -356,15 +351,15 @@ void UKismetRenderingLibrary::EndDrawCanvasToRenderTarget(UObject* WorldContextO
 			// const cast required, as BP will treat Context as an output without the const
 			const_cast<FDrawToRenderTargetContext&>(Context) = FDrawToRenderTargetContext();
 		}
-		else if (!World)
-		{
-			FMessageLog("Blueprint").Warning(LOCTEXT("EndDrawCanvasToRenderTarget_InvalidWorldContextObject", "EndDrawCanvasToRenderTarget: WorldContextObject is not valid."));
-		}
 		else
 		{
 			FMessageLog("Blueprint").Warning(LOCTEXT("EndDrawCanvasToRenderTarget_InvalidContext", "EndDrawCanvasToRenderTarget: Context must be valid."));
 		}
 	}
+	else
+	{
+		FMessageLog("Blueprint").Warning(LOCTEXT("EndDrawCanvasToRenderTarget_InvalidWorldContextObject", "EndDrawCanvasToRenderTarget: WorldContextObject is not valid."));
+	}
 }
 
 FSkelMeshSkinWeightInfo UKismetRenderingLibrary::MakeSkinWeightInfo(int32 Bone0, uint8 Weight0, int32 Bone1, uint8 Weight1, int32 Bone2, uint8 Weight2, int32 Bone3, uint8 Weight3)
diff --git a/Engine/Source/Runtime/Engine/Private/KismetStringLibrary.cpp b/Engine/Source/Runtime/Engine/Private/KismetStringLibrary.cpp
--- a/Engine/Source/Runtime/Engine/Private/KismetStringLibrary.cpp
+++ b/Engine/Source/Runtime/Engine/Private/KismetStringLibrary.cpp
@@ -455,7 +455,7 @@ FString UKismetStringLibrary::RightChop(const FString& SourceString, int32 Count
 
 FString UKismetStringLibrary::Mid(const FString& SourceString, int32 Start, int32 Count = MAX_int32)
 {
-	return SourceString.Mid(Start, Count);
+	return (Count >= 0 ? SourceString.Mid(Start, Count) : FString());
 }
 
 FString UKismetStringLibrary::TimeSecondsToString(float InSeconds)
diff --git a/Engine/Source/Runtime/Engine/Private/KismetSystemLibrary.cpp b/Engine/Source/Runtime/Engine/Private/KismetSystemLibrary.cpp
--- a/Engine/Source/Runtime/Engine/Private/KismetSystemLibrary.cpp
+++ b/Engine/Source/Runtime/Engine/Private/KismetSystemLibrary.cpp
@@ -112,19 +112,19 @@ bool UKismetSystemLibrary::DoesImplementInterface(UObject* TestObject, TSubclass
 
 float UKismetSystemLibrary::GetGameTimeInSeconds(UObject* WorldContextObject)
 {
-	UWorld* World = GEngine->GetWorldFromContextObject( WorldContextObject );
+	UWorld* World = GEngine->GetWorldFromContextObject(WorldContextObject, EGetWorldErrorMode::LogAndReturnNull);
 	return World ? World->GetTimeSeconds() : 0.f;
 }
 
 bool UKismetSystemLibrary::IsServer(UObject* WorldContextObject)
 {
-	UWorld* World = GEngine->GetWorldFromContextObject( WorldContextObject );
+	UWorld* World = GEngine->GetWorldFromContextObject(WorldContextObject, EGetWorldErrorMode::LogAndReturnNull);
 	return World ? (World->GetNetMode() != NM_Client) : false;
 }
 
 bool UKismetSystemLibrary::IsDedicatedServer(UObject* WorldContextObject)
 {
-	UWorld* World = GEngine->GetWorldFromContextObject( WorldContextObject );
+	UWorld* World = GEngine->GetWorldFromContextObject(WorldContextObject, EGetWorldErrorMode::LogAndReturnNull);
 	if (World)
 	{
 		return (World->GetNetMode() == NM_DedicatedServer);
@@ -134,7 +134,7 @@ bool UKismetSystemLibrary::IsDedicatedServer(UObject* WorldContextObject)
 
 bool UKismetSystemLibrary::IsStandalone(UObject* WorldContextObject)
 {
-	UWorld* World = GEngine->GetWorldFromContextObject( WorldContextObject );
+	UWorld* World = GEngine->GetWorldFromContextObject(WorldContextObject, EGetWorldErrorMode::LogAndReturnNull);
 	return World ? (World->GetNetMode() == NM_Standalone) : false;
 }
 
@@ -164,7 +164,7 @@ void UKismetSystemLibrary::PrintString(UObject* WorldContextObject, const FStrin
 {
 #if !(UE_BUILD_SHIPPING || UE_BUILD_TEST) // Do not Print in Shipping or Test
 
-	UWorld* World = GEngine->GetWorldFromContextObject(WorldContextObject, false);
+	UWorld* World = GEngine->GetWorldFromContextObject(WorldContextObject, EGetWorldErrorMode::ReturnNull);
 	FString Prefix;
 	if (World)
 	{
@@ -321,7 +321,7 @@ FTimerHandle UKismetSystemLibrary::K2_SetTimerDelegate(FTimerDynamicDelegate Del
 	FTimerHandle Handle;
 	if (Delegate.IsBound())
 	{
-		const UWorld* const World = GEngine->GetWorldFromContextObject(Delegate.GetUObject());
+		const UWorld* const World = GEngine->GetWorldFromContextObject(Delegate.GetUObject(), EGetWorldErrorMode::LogAndReturnNull);
 		if(World)
 		{
 			FTimerManager& TimerManager = World->GetTimerManager();
@@ -351,7 +351,7 @@ void UKismetSystemLibrary::K2_ClearTimerDelegate(FTimerDynamicDelegate Delegate)
 {
 	if (Delegate.IsBound())
 	{
-		UWorld* World = GEngine->GetWorldFromContextObject(Delegate.GetUObject());
+		UWorld* World = GEngine->GetWorldFromContextObject(Delegate.GetUObject(), EGetWorldErrorMode::LogAndReturnNull);
 		if (World)
 		{
 			FTimerManager& TimerManager = World->GetTimerManager();
@@ -371,7 +371,7 @@ void UKismetSystemLibrary::K2_ClearTimerHandle(UObject* WorldContextObject, FTim
 {
 	if (Handle.IsValid())
 	{
-		UWorld* World = GEngine->GetWorldFromContextObject( WorldContextObject );
+		UWorld* World = GEngine->GetWorldFromContextObject(WorldContextObject, EGetWorldErrorMode::LogAndReturnNull);
 		if (World)
 		{
 			World->GetTimerManager().ClearTimer(Handle);
@@ -383,7 +383,7 @@ void UKismetSystemLibrary::K2_ClearAndInvalidateTimerHandle(UObject* WorldContex
 {
 	if (Handle.IsValid())
 	{
-		UWorld* World = GEngine->GetWorldFromContextObject(WorldContextObject);
+		UWorld* World = GEngine->GetWorldFromContextObject(WorldContextObject, EGetWorldErrorMode::LogAndReturnNull);
 		if (World)
 		{
 			World->GetTimerManager().ClearTimer(Handle);
@@ -403,7 +403,7 @@ void UKismetSystemLibrary::K2_PauseTimerDelegate(FTimerDynamicDelegate Delegate)
 {
 	if (Delegate.IsBound())
 	{
-		UWorld* World = GEngine->GetWorldFromContextObject(Delegate.GetUObject());
+		UWorld* World = GEngine->GetWorldFromContextObject(Delegate.GetUObject(), EGetWorldErrorMode::LogAndReturnNull);
 		if(World)
 		{
 			FTimerManager& TimerManager = World->GetTimerManager();
@@ -423,7 +423,7 @@ void UKismetSystemLibrary::K2_PauseTimerHandle(UObject* WorldContextObject, FTim
 {
 	if (Handle.IsValid())
 	{
-		UWorld* World = GEngine->GetWorldFromContextObject( WorldContextObject );
+		UWorld* World = GEngine->GetWorldFromContextObject(WorldContextObject, EGetWorldErrorMode::LogAndReturnNull);
 		if (World)
 		{
 			World->GetTimerManager().PauseTimer(Handle);
@@ -443,7 +443,7 @@ void UKismetSystemLibrary::K2_UnPauseTimerDelegate(FTimerDynamicDelegate Delegat
 {
 	if (Delegate.IsBound())
 	{
-		UWorld* World = GEngine->GetWorldFromContextObject(Delegate.GetUObject());
+		UWorld* World = GEngine->GetWorldFromContextObject(Delegate.GetUObject(), EGetWorldErrorMode::LogAndReturnNull);
 		if(World)
 		{
 			FTimerManager& TimerManager = World->GetTimerManager();
@@ -463,7 +463,7 @@ void UKismetSystemLibrary::K2_UnPauseTimerHandle(UObject* WorldContextObject, FT
 {
 	if (Handle.IsValid())
 	{
-		UWorld* World = GEngine->GetWorldFromContextObject( WorldContextObject );
+		UWorld* World = GEngine->GetWorldFromContextObject(WorldContextObject, EGetWorldErrorMode::LogAndReturnNull);
 		if (World)
 		{
 			World->GetTimerManager().UnPauseTimer(Handle);
@@ -484,7 +484,7 @@ bool UKismetSystemLibrary::K2_IsTimerActiveDelegate(FTimerDynamicDelegate Delega
 	bool bIsActive = false;
 	if (Delegate.IsBound())
 	{
-		UWorld* World = GEngine->GetWorldFromContextObject(Delegate.GetUObject());
+		UWorld* World = GEngine->GetWorldFromContextObject(Delegate.GetUObject(), EGetWorldErrorMode::LogAndReturnNull);
 		if(World)
 		{
 			FTimerManager& TimerManager = World->GetTimerManager();
@@ -507,7 +507,7 @@ bool UKismetSystemLibrary::K2_IsTimerActiveHandle(UObject* WorldContextObject, F
 	bool bIsActive = false;
 	if (Handle.IsValid())
 	{
-		UWorld* World = GEngine->GetWorldFromContextObject( WorldContextObject );
+		UWorld* World = GEngine->GetWorldFromContextObject(WorldContextObject, EGetWorldErrorMode::LogAndReturnNull);
 		if (World)
 		{
 			bIsActive = World->GetTimerManager().IsTimerActive(Handle);
@@ -530,7 +530,7 @@ bool UKismetSystemLibrary::K2_IsTimerPausedDelegate(FTimerDynamicDelegate Delega
 	bool bIsPaused = false;
 	if (Delegate.IsBound())
 	{
-		UWorld* World = GEngine->GetWorldFromContextObject(Delegate.GetUObject());
+		UWorld* World = GEngine->GetWorldFromContextObject(Delegate.GetUObject(), EGetWorldErrorMode::LogAndReturnNull);
 		if(World)
 		{
 			FTimerManager& TimerManager = World->GetTimerManager();
@@ -552,7 +552,7 @@ bool UKismetSystemLibrary::K2_IsTimerPausedHandle(UObject* WorldContextObject, F
 	bool bIsPaused = false;
 	if (Handle.IsValid())
 	{
-		UWorld* World = GEngine->GetWorldFromContextObject( WorldContextObject );
+		UWorld* World = GEngine->GetWorldFromContextObject(WorldContextObject, EGetWorldErrorMode::LogAndReturnNull);
 		if (World)
 		{
 			bIsPaused = World->GetTimerManager().IsTimerPaused(Handle);
@@ -575,7 +575,7 @@ bool UKismetSystemLibrary::K2_TimerExistsDelegate(FTimerDynamicDelegate Delegate
 	bool bTimerExists = false;
 	if (Delegate.IsBound())
 	{
-		UWorld* World = GEngine->GetWorldFromContextObject(Delegate.GetUObject());
+		UWorld* World = GEngine->GetWorldFromContextObject(Delegate.GetUObject(), EGetWorldErrorMode::LogAndReturnNull);
 		if(World)
 		{
 			FTimerManager& TimerManager = World->GetTimerManager();
@@ -597,7 +597,7 @@ bool UKismetSystemLibrary::K2_TimerExistsHandle(UObject* WorldContextObject, FTi
 	bool bTimerExists = false;
 	if (Handle.IsValid())
 	{
-		UWorld* World = GEngine->GetWorldFromContextObject( WorldContextObject );
+		UWorld* World = GEngine->GetWorldFromContextObject(WorldContextObject, EGetWorldErrorMode::LogAndReturnNull);
 		if (World)
 		{
 			bTimerExists = World->GetTimerManager().TimerExists(Handle);
@@ -620,7 +620,7 @@ float UKismetSystemLibrary::K2_GetTimerElapsedTimeDelegate(FTimerDynamicDelegate
 	float ElapsedTime = 0.f;
 	if (Delegate.IsBound())
 	{
-		UWorld* World = GEngine->GetWorldFromContextObject(Delegate.GetUObject());
+		UWorld* World = GEngine->GetWorldFromContextObject(Delegate.GetUObject(), EGetWorldErrorMode::LogAndReturnNull);
 		if(World)
 		{
 			FTimerManager& TimerManager = World->GetTimerManager();
@@ -642,7 +642,7 @@ float UKismetSystemLibrary::K2_GetTimerElapsedTimeHandle(UObject* WorldContextOb
 	float ElapsedTime = 0.f;
 	if (Handle.IsValid())
 	{
-		UWorld* World = GEngine->GetWorldFromContextObject( WorldContextObject );
+		UWorld* World = GEngine->GetWorldFromContextObject(WorldContextObject, EGetWorldErrorMode::LogAndReturnNull);
 		if (World)
 		{
 			ElapsedTime = World->GetTimerManager().GetTimerElapsed(Handle);
@@ -665,7 +665,7 @@ float UKismetSystemLibrary::K2_GetTimerRemainingTimeDelegate(FTimerDynamicDelega
 	float RemainingTime = 0.f;
 	if (Delegate.IsBound())
 	{
-		UWorld* World = GEngine->GetWorldFromContextObject(Delegate.GetUObject());
+		UWorld* World = GEngine->GetWorldFromContextObject(Delegate.GetUObject(), EGetWorldErrorMode::LogAndReturnNull);
 		if(World)
 		{
 			FTimerManager& TimerManager = World->GetTimerManager();
@@ -687,7 +687,7 @@ float UKismetSystemLibrary::K2_GetTimerRemainingTimeHandle(UObject* WorldContext
 	float RemainingTime = 0.f;
 	if (Handle.IsValid())
 	{
-		UWorld* World = GEngine->GetWorldFromContextObject( WorldContextObject );
+		UWorld* World = GEngine->GetWorldFromContextObject(WorldContextObject, EGetWorldErrorMode::LogAndReturnNull);
 		if (World)
 		{
 			RemainingTime = World->GetTimerManager().GetTimerRemaining(Handle);
@@ -1016,7 +1016,7 @@ bool UKismetSystemLibrary::SphereOverlapComponents(UObject* WorldContextObject,
 	}
 
 
-	UWorld* World = GEngine->GetWorldFromContextObject( WorldContextObject );
+	UWorld* World = GEngine->GetWorldFromContextObject(WorldContextObject, EGetWorldErrorMode::LogAndReturnNull);
 	if(World != nullptr)
 	{
 		World->OverlapMultiByObjectType(Overlaps, SpherePos, FQuat::Identity, ObjectParams, FCollisionShape::MakeSphere(SphereRadius), Params);
@@ -1070,7 +1070,7 @@ bool UKismetSystemLibrary::BoxOverlapComponents(UObject* WorldContextObject, con
 		ObjectParams.AddObjectTypesToQuery(Channel);
 	}
 
-	UWorld* World = GEngine->GetWorldFromContextObject( WorldContextObject );
+	UWorld* World = GEngine->GetWorldFromContextObject(WorldContextObject, EGetWorldErrorMode::LogAndReturnNull);
 	if (World != nullptr)
 	{
 		World->OverlapMultiByObjectType(Overlaps, BoxPos, FQuat::Identity, ObjectParams, FCollisionShape::MakeBox(BoxExtent), Params);
@@ -1123,7 +1123,7 @@ bool UKismetSystemLibrary::CapsuleOverlapComponents(UObject* WorldContextObject,
 		ObjectParams.AddObjectTypesToQuery(Channel);
 	}
 
-	UWorld* World = GEngine->GetWorldFromContextObject( WorldContextObject );	
+	UWorld* World = GEngine->GetWorldFromContextObject(WorldContextObject, EGetWorldErrorMode::LogAndReturnNull);
 	if (World != nullptr)
 	{
 		World->OverlapMultiByObjectType(Overlaps, CapsulePos, FQuat::Identity, ObjectParams, FCollisionShape::MakeCapsule(Radius, HalfHeight), Params);
@@ -1205,8 +1205,8 @@ bool UKismetSystemLibrary::LineTraceSingle(UObject* WorldContextObject, const FV
 	static const FName LineTraceSingleName(TEXT("LineTraceSingle"));
 	FCollisionQueryParams Params = ConfigureCollisionParams(LineTraceSingleName, bTraceComplex, ActorsToIgnore, bIgnoreSelf, WorldContextObject);
 
-	UWorld* World = GEngine->GetWorldFromContextObject(WorldContextObject);
-	bool const bHit = World->LineTraceSingleByChannel(OutHit, Start, End, CollisionChannel, Params);
+	UWorld* World = GEngine->GetWorldFromContextObject(WorldContextObject, EGetWorldErrorMode::LogAndReturnNull);
+	bool const bHit = World ? World->LineTraceSingleByChannel(OutHit, Start, End, CollisionChannel, Params) : false;
 
 #if ENABLE_DRAW_DEBUG
 	DrawDebugLineTraceSingle(World, Start, End, DrawDebugType, bHit, OutHit, TraceColor, TraceHitColor, DrawTime);
@@ -1222,8 +1222,8 @@ bool UKismetSystemLibrary::LineTraceMulti(UObject* WorldContextObject, const FVe
 	static const FName LineTraceMultiName(TEXT("LineTraceMulti"));
 	FCollisionQueryParams Params = ConfigureCollisionParams(LineTraceMultiName, bTraceComplex, ActorsToIgnore, bIgnoreSelf, WorldContextObject);
 
-	UWorld* World = GEngine->GetWorldFromContextObject(WorldContextObject);
-	bool const bHit = World->LineTraceMultiByChannel(OutHits, Start, End, CollisionChannel, Params);
+	UWorld* World = GEngine->GetWorldFromContextObject(WorldContextObject, EGetWorldErrorMode::LogAndReturnNull);
+	bool const bHit = World ? World->LineTraceMultiByChannel(OutHits, Start, End, CollisionChannel, Params) : false;
 
 #if ENABLE_DRAW_DEBUG
 	DrawDebugLineTraceMulti(World, Start, End, DrawDebugType, bHit, OutHits, TraceColor, TraceHitColor, DrawTime);
@@ -1237,7 +1237,7 @@ bool UKismetSystemLibrary::BoxTraceSingle(UObject* WorldContextObject, const FVe
 	static const FName BoxTraceSingleName(TEXT("BoxTraceSingle"));
 	FCollisionQueryParams Params = ConfigureCollisionParams(BoxTraceSingleName, bTraceComplex, ActorsToIgnore, bIgnoreSelf, WorldContextObject);
 
-	UWorld* World = GEngine->GetWorldFromContextObject(WorldContextObject);
+	UWorld* World = GEngine->GetWorldFromContextObject(WorldContextObject, EGetWorldErrorMode::LogAndReturnNull);
 	bool const bHit = World ? World->SweepSingleByChannel(OutHit, Start, End, Orientation.Quaternion(), UEngineTypes::ConvertToCollisionChannel(TraceChannel), FCollisionShape::MakeBox(HalfSize), Params) : false;
 
 #if ENABLE_DRAW_DEBUG
@@ -1252,7 +1252,7 @@ bool UKismetSystemLibrary::BoxTraceMulti(UObject* WorldContextObject, const FVec
 	static const FName BoxTraceMultiName(TEXT("BoxTraceMulti"));
 	FCollisionQueryParams Params = ConfigureCollisionParams(BoxTraceMultiName, bTraceComplex, ActorsToIgnore, bIgnoreSelf, WorldContextObject);
 
-	UWorld* World = GEngine->GetWorldFromContextObject(WorldContextObject);
+	UWorld* World = GEngine->GetWorldFromContextObject(WorldContextObject, EGetWorldErrorMode::LogAndReturnNull);
 	bool const bHit = World ? World->SweepMultiByChannel(OutHits, Start, End, Orientation.Quaternion(), UEngineTypes::ConvertToCollisionChannel(TraceChannel), FCollisionShape::MakeBox(HalfSize), Params) : false;
 
 #if ENABLE_DRAW_DEBUG
@@ -1269,8 +1269,8 @@ bool UKismetSystemLibrary::SphereTraceSingle(UObject* WorldContextObject, const
 	static const FName SphereTraceSingleName(TEXT("SphereTraceSingle"));
 	FCollisionQueryParams Params = ConfigureCollisionParams(SphereTraceSingleName, bTraceComplex, ActorsToIgnore, bIgnoreSelf, WorldContextObject);
 
-	UWorld* World = GEngine->GetWorldFromContextObject(WorldContextObject);
-	bool const bHit = World->SweepSingleByChannel(OutHit, Start, End, FQuat::Identity, CollisionChannel, FCollisionShape::MakeSphere(Radius), Params);
+	UWorld* World = GEngine->GetWorldFromContextObject(WorldContextObject, EGetWorldErrorMode::LogAndReturnNull);
+	bool const bHit = World ? World->SweepSingleByChannel(OutHit, Start, End, FQuat::Identity, CollisionChannel, FCollisionShape::MakeSphere(Radius), Params) : false;
 
 #if ENABLE_DRAW_DEBUG
 	DrawDebugSphereTraceSingle(World, Start, End, Radius, DrawDebugType, bHit, OutHit, TraceColor, TraceHitColor, DrawTime);
@@ -1286,8 +1286,8 @@ bool UKismetSystemLibrary::SphereTraceMulti(UObject* WorldContextObject, const F
 	static const FName SphereTraceMultiName(TEXT("SphereTraceMulti"));
 	FCollisionQueryParams Params = ConfigureCollisionParams(SphereTraceMultiName, bTraceComplex, ActorsToIgnore, bIgnoreSelf, WorldContextObject);
 
-	UWorld* World = GEngine->GetWorldFromContextObject(WorldContextObject);
-	bool const bHit = World->SweepMultiByChannel(OutHits, Start, End, FQuat::Identity, CollisionChannel, FCollisionShape::MakeSphere(Radius), Params);
+	UWorld* World = GEngine->GetWorldFromContextObject(WorldContextObject, EGetWorldErrorMode::LogAndReturnNull);
+	bool const bHit = World ? World->SweepMultiByChannel(OutHits, Start, End, FQuat::Identity, CollisionChannel, FCollisionShape::MakeSphere(Radius), Params) : false;
 
 #if ENABLE_DRAW_DEBUG
 	DrawDebugSphereTraceMulti(World, Start, End, Radius, DrawDebugType, bHit, OutHits, TraceColor, TraceHitColor, DrawTime);
@@ -1303,8 +1303,8 @@ bool UKismetSystemLibrary::CapsuleTraceSingle(UObject* WorldContextObject, const
 	static const FName CapsuleTraceSingleName(TEXT("CapsuleTraceSingle"));
 	FCollisionQueryParams Params = ConfigureCollisionParams(CapsuleTraceSingleName, bTraceComplex, ActorsToIgnore, bIgnoreSelf, WorldContextObject);
 
-	UWorld* World = GEngine->GetWorldFromContextObject(WorldContextObject);
-	bool const bHit = World->SweepSingleByChannel(OutHit, Start, End, FQuat::Identity, CollisionChannel, FCollisionShape::MakeCapsule(Radius, HalfHeight), Params);
+	UWorld* World = GEngine->GetWorldFromContextObject(WorldContextObject, EGetWorldErrorMode::LogAndReturnNull);
+	bool const bHit = World ? World->SweepSingleByChannel(OutHit, Start, End, FQuat::Identity, CollisionChannel, FCollisionShape::MakeCapsule(Radius, HalfHeight), Params) : false;
 
 #if ENABLE_DRAW_DEBUG
 	DrawDebugCapsuleTraceSingle(World, Start, End, Radius, HalfHeight, DrawDebugType, bHit, OutHit, TraceColor, TraceHitColor, DrawTime);
@@ -1321,8 +1321,8 @@ bool UKismetSystemLibrary::CapsuleTraceMulti(UObject* WorldContextObject, const
 	static const FName CapsuleTraceMultiName(TEXT("CapsuleTraceMulti"));
 	FCollisionQueryParams Params = ConfigureCollisionParams(CapsuleTraceMultiName, bTraceComplex, ActorsToIgnore, bIgnoreSelf, WorldContextObject);
 
-	UWorld* World = GEngine->GetWorldFromContextObject(WorldContextObject);
-	bool const bHit = World->SweepMultiByChannel(OutHits, Start, End, FQuat::Identity, CollisionChannel, FCollisionShape::MakeCapsule(Radius, HalfHeight), Params);
+	UWorld* World = GEngine->GetWorldFromContextObject(WorldContextObject, EGetWorldErrorMode::LogAndReturnNull);
+	bool const bHit = World ? World->SweepMultiByChannel(OutHits, Start, End, FQuat::Identity, CollisionChannel, FCollisionShape::MakeCapsule(Radius, HalfHeight), Params) : false;
 
 #if ENABLE_DRAW_DEBUG
 	DrawDebugCapsuleTraceMulti(World, Start, End, Radius, HalfHeight, DrawDebugType, bHit, OutHits, TraceColor, TraceHitColor, DrawTime);
@@ -1344,8 +1344,8 @@ bool UKismetSystemLibrary::LineTraceSingleForObjects(UObject* WorldContextObject
 		return false;
 	}
 
-	UWorld* World = GEngine->GetWorldFromContextObject(WorldContextObject);
-	bool const bHit = World->LineTraceSingleByObjectType(OutHit, Start, End, ObjectParams, Params);
+	UWorld* World = GEngine->GetWorldFromContextObject(WorldContextObject, EGetWorldErrorMode::LogAndReturnNull);
+	bool const bHit = World ? World->LineTraceSingleByObjectType(OutHit, Start, End, ObjectParams, Params) : false;
 
 #if ENABLE_DRAW_DEBUG
 	DrawDebugLineTraceSingle(World, Start, End, DrawDebugType, bHit, OutHit, TraceColor, TraceHitColor, DrawTime);
@@ -1366,8 +1366,8 @@ bool UKismetSystemLibrary::LineTraceMultiForObjects(UObject* WorldContextObject,
 		return false;
 	}
 
-	UWorld* World = GEngine->GetWorldFromContextObject(WorldContextObject);
-	bool const bHit = World->LineTraceMultiByObjectType(OutHits, Start, End, ObjectParams, Params);
+	UWorld* World = GEngine->GetWorldFromContextObject(WorldContextObject, EGetWorldErrorMode::LogAndReturnNull);
+	bool const bHit = World ? World->LineTraceMultiByObjectType(OutHits, Start, End, ObjectParams, Params) : false;
 
 #if ENABLE_DRAW_DEBUG
 	DrawDebugLineTraceMulti(World, Start, End, DrawDebugType, bHit, OutHits, TraceColor, TraceHitColor, DrawTime);
@@ -1388,8 +1388,8 @@ bool UKismetSystemLibrary::SphereTraceSingleForObjects(UObject* WorldContextObje
 		return false;
 	}
 
-	UWorld* World = GEngine->GetWorldFromContextObject(WorldContextObject);
-	bool const bHit = World->SweepSingleByObjectType(OutHit, Start, End, FQuat::Identity, ObjectParams, FCollisionShape::MakeSphere(Radius), Params);
+	UWorld* World = GEngine->GetWorldFromContextObject(WorldContextObject, EGetWorldErrorMode::LogAndReturnNull);
+	bool const bHit = World ? World->SweepSingleByObjectType(OutHit, Start, End, FQuat::Identity, ObjectParams, FCollisionShape::MakeSphere(Radius), Params) : false;
 
 #if ENABLE_DRAW_DEBUG
 	DrawDebugSphereTraceSingle(World, Start, End, Radius, DrawDebugType, bHit, OutHit, TraceColor, TraceHitColor, DrawTime);
@@ -1410,8 +1410,8 @@ bool UKismetSystemLibrary::SphereTraceMultiForObjects(UObject* WorldContextObjec
 		return false;
 	}
 
-	UWorld* World = GEngine->GetWorldFromContextObject(WorldContextObject);
-	bool const bHit = World->SweepMultiByObjectType(OutHits, Start, End, FQuat::Identity, ObjectParams, FCollisionShape::MakeSphere(Radius), Params);
+	UWorld* World = GEngine->GetWorldFromContextObject(WorldContextObject, EGetWorldErrorMode::LogAndReturnNull);
+	bool const bHit = World ? World->SweepMultiByObjectType(OutHits, Start, End, FQuat::Identity, ObjectParams, FCollisionShape::MakeSphere(Radius), Params) : false;
 
 #if ENABLE_DRAW_DEBUG
 	DrawDebugSphereTraceMulti(World, Start, End, Radius, DrawDebugType, bHit, OutHits, TraceColor, TraceHitColor, DrawTime);
@@ -1435,7 +1435,7 @@ bool UKismetSystemLibrary::BoxTraceSingleForObjects(UObject* WorldContextObject,
 		return false;
 	}
 
-	UWorld* World = GEngine->GetWorldFromContextObject(WorldContextObject);
+	UWorld* World = GEngine->GetWorldFromContextObject(WorldContextObject, EGetWorldErrorMode::LogAndReturnNull);
 	bool const bHit = World ? World->SweepSingleByObjectType(OutHit, Start, End, Orientation.Quaternion(), ObjectParams, FCollisionShape::MakeBox(HalfSize), Params) : false;
 
 #if ENABLE_DRAW_DEBUG
@@ -1457,7 +1457,7 @@ bool UKismetSystemLibrary::BoxTraceMultiForObjects(UObject* WorldContextObject,
 		return false;
 	}
 
-	UWorld* World = GEngine->GetWorldFromContextObject(WorldContextObject);
+	UWorld* World = GEngine->GetWorldFromContextObject(WorldContextObject, EGetWorldErrorMode::LogAndReturnNull);
 	bool const bHit = World ? World->SweepMultiByObjectType(OutHits, Start, End, Orientation.Quaternion(), ObjectParams, FCollisionShape::MakeBox(HalfSize), Params) : false;
 
 #if ENABLE_DRAW_DEBUG
@@ -1479,8 +1479,8 @@ bool UKismetSystemLibrary::CapsuleTraceSingleForObjects(UObject* WorldContextObj
 		return false;
 	}
 
-	UWorld* World = GEngine->GetWorldFromContextObject(WorldContextObject);
-	bool const bHit = World->SweepSingleByObjectType(OutHit, Start, End, FQuat::Identity, ObjectParams, FCollisionShape::MakeCapsule(Radius, HalfHeight), Params);
+	UWorld* World = GEngine->GetWorldFromContextObject(WorldContextObject, EGetWorldErrorMode::LogAndReturnNull);
+	bool const bHit = World ? World->SweepSingleByObjectType(OutHit, Start, End, FQuat::Identity, ObjectParams, FCollisionShape::MakeCapsule(Radius, HalfHeight), Params) : false;
 
 #if ENABLE_DRAW_DEBUG
 	DrawDebugCapsuleTraceSingle(World, Start, End, Radius, HalfHeight, DrawDebugType, bHit, OutHit, TraceColor, TraceHitColor, DrawTime);
@@ -1501,8 +1501,8 @@ bool UKismetSystemLibrary::CapsuleTraceMultiForObjects(UObject* WorldContextObje
 		return false;
 	}
 
-	UWorld* World = GEngine->GetWorldFromContextObject(WorldContextObject);
-	bool const bHit = World->SweepMultiByObjectType(OutHits, Start, End, FQuat::Identity, ObjectParams, FCollisionShape::MakeCapsule(Radius, HalfHeight), Params);
+	UWorld* World = GEngine->GetWorldFromContextObject(WorldContextObject, EGetWorldErrorMode::LogAndReturnNull);
+	bool const bHit = World ? World->SweepMultiByObjectType(OutHits, Start, End, FQuat::Identity, ObjectParams, FCollisionShape::MakeCapsule(Radius, HalfHeight), Params) : false;
 
 #if ENABLE_DRAW_DEBUG
 	DrawDebugCapsuleTraceMulti(World, Start, End, Radius, HalfHeight, DrawDebugType, bHit, OutHits, TraceColor, TraceHitColor, DrawTime);
@@ -1517,8 +1517,8 @@ bool UKismetSystemLibrary::LineTraceSingleByProfile(UObject* WorldContextObject,
 	static const FName LineTraceSingleName(TEXT("LineTraceSingleByProfile"));
 	FCollisionQueryParams Params = ConfigureCollisionParams(LineTraceSingleName, bTraceComplex, ActorsToIgnore, bIgnoreSelf, WorldContextObject);
 
-	UWorld* World = GEngine->GetWorldFromContextObject(WorldContextObject);
-	bool const bHit = World->LineTraceSingleByProfile(OutHit, Start, End, ProfileName, Params);
+	UWorld* World = GEngine->GetWorldFromContextObject(WorldContextObject, EGetWorldErrorMode::LogAndReturnNull);
+	bool const bHit = World ? World->LineTraceSingleByProfile(OutHit, Start, End, ProfileName, Params) : false;
 
 #if ENABLE_DRAW_DEBUG
 	DrawDebugLineTraceSingle(World, Start, End, DrawDebugType, bHit, OutHit, TraceColor, TraceHitColor, DrawTime);
@@ -1532,8 +1532,8 @@ bool UKismetSystemLibrary::LineTraceMultiByProfile(UObject* WorldContextObject,
 	static const FName LineTraceMultiName(TEXT("LineTraceMultiByProfile"));
 	FCollisionQueryParams Params = ConfigureCollisionParams(LineTraceMultiName, bTraceComplex, ActorsToIgnore, bIgnoreSelf, WorldContextObject);
 
-	UWorld* World = GEngine->GetWorldFromContextObject(WorldContextObject);
-	bool const bHit = World->LineTraceMultiByProfile(OutHits, Start, End, ProfileName, Params);
+	UWorld* World = GEngine->GetWorldFromContextObject(WorldContextObject, EGetWorldErrorMode::LogAndReturnNull);
+	bool const bHit = World ? World->LineTraceMultiByProfile(OutHits, Start, End, ProfileName, Params) : false;
 
 #if ENABLE_DRAW_DEBUG
 	DrawDebugLineTraceMulti(World, Start, End, DrawDebugType, bHit, OutHits, TraceColor, TraceHitColor, DrawTime);
@@ -1547,8 +1547,10 @@ bool UKismetSystemLibrary::BoxTraceSingleByProfile(UObject* WorldContextObject,
 	static const FName BoxTraceSingleName(TEXT("BoxTraceSingleByProfile"));
 	FCollisionQueryParams Params = ConfigureCollisionParams(BoxTraceSingleName, bTraceComplex, ActorsToIgnore, bIgnoreSelf, WorldContextObject);
 
-	UWorld* World = GEngine->GetWorldFromContextObject(WorldContextObject);
+	UWorld* World = GEngine->GetWorldFromContextObject(WorldContextObject, EGetWorldErrorMode::LogAndReturnNull);
+
 	bool const bHit = World ? World->SweepSingleByProfile(OutHit, Start, End, Orientation.Quaternion(), ProfileName, FCollisionShape::MakeBox(HalfSize), Params) : false;
+
 #if ENABLE_DRAW_DEBUG
 	DrawDebugBoxTraceSingle(World, Start, End, HalfSize, Orientation, DrawDebugType, bHit, OutHit, TraceColor, TraceHitColor, DrawTime);
 #endif
@@ -1561,7 +1563,7 @@ bool UKismetSystemLibrary::BoxTraceMultiByProfile(UObject* WorldContextObject, c
 	static const FName BoxTraceMultiName(TEXT("BoxTraceMultiByProfile"));
 	FCollisionQueryParams Params = ConfigureCollisionParams(BoxTraceMultiName, bTraceComplex, ActorsToIgnore, bIgnoreSelf, WorldContextObject);
 
-	UWorld* World = GEngine->GetWorldFromContextObject(WorldContextObject);
+	UWorld* World = GEngine->GetWorldFromContextObject(WorldContextObject, EGetWorldErrorMode::LogAndReturnNull);
 	bool const bHit = World ? World->SweepMultiByProfile(OutHits, Start, End, Orientation.Quaternion(), ProfileName, FCollisionShape::MakeBox(HalfSize), Params) : false;
 
 #if ENABLE_DRAW_DEBUG
@@ -1576,8 +1578,8 @@ bool UKismetSystemLibrary::SphereTraceSingleByProfile(UObject* WorldContextObjec
 	static const FName SphereTraceSingleName(TEXT("SphereTraceSingleByProfile"));
 	FCollisionQueryParams Params = ConfigureCollisionParams(SphereTraceSingleName, bTraceComplex, ActorsToIgnore, bIgnoreSelf, WorldContextObject);
 
-	UWorld* World = GEngine->GetWorldFromContextObject(WorldContextObject);
-	bool const bHit = World->SweepSingleByProfile(OutHit, Start, End, FQuat::Identity, ProfileName, FCollisionShape::MakeSphere(Radius), Params);
+	UWorld* World = GEngine->GetWorldFromContextObject(WorldContextObject, EGetWorldErrorMode::LogAndReturnNull);
+	bool const bHit = World ? World->SweepSingleByProfile(OutHit, Start, End, FQuat::Identity, ProfileName, FCollisionShape::MakeSphere(Radius), Params) : false;
 
 #if ENABLE_DRAW_DEBUG
 	DrawDebugSphereTraceSingle(World, Start, End, Radius, DrawDebugType, bHit, OutHit, TraceColor, TraceHitColor, DrawTime);
@@ -1591,8 +1593,9 @@ bool UKismetSystemLibrary::SphereTraceMultiByProfile(UObject* WorldContextObject
 	static const FName SphereTraceMultiName(TEXT("SphereTraceMultiByProfile"));
 	FCollisionQueryParams Params = ConfigureCollisionParams(SphereTraceMultiName, bTraceComplex, ActorsToIgnore, bIgnoreSelf, WorldContextObject);
 
-	UWorld* World = GEngine->GetWorldFromContextObject(WorldContextObject);
-	bool const bHit = World->SweepMultiByProfile(OutHits, Start, End, FQuat::Identity, ProfileName, FCollisionShape::MakeSphere(Radius), Params);
+	UWorld* World = GEngine->GetWorldFromContextObject(WorldContextObject, EGetWorldErrorMode::LogAndReturnNull);
+	bool const bHit = World ? World->SweepMultiByProfile(OutHits, Start, End, FQuat::Identity, ProfileName, FCollisionShape::MakeSphere(Radius), Params) : false;
+
 #if ENABLE_DRAW_DEBUG
 	DrawDebugSphereTraceMulti(World, Start, End, Radius, DrawDebugType, bHit, OutHits, TraceColor, TraceHitColor, DrawTime);
 #endif
@@ -1605,8 +1608,8 @@ bool UKismetSystemLibrary::CapsuleTraceSingleByProfile(UObject* WorldContextObje
 	static const FName CapsuleTraceSingleName(TEXT("CapsuleTraceSingleByProfile"));
 	FCollisionQueryParams Params = ConfigureCollisionParams(CapsuleTraceSingleName, bTraceComplex, ActorsToIgnore, bIgnoreSelf, WorldContextObject);
 
-	UWorld* World = GEngine->GetWorldFromContextObject(WorldContextObject);
-	bool const bHit = World->SweepSingleByProfile(OutHit, Start, End, FQuat::Identity, ProfileName, FCollisionShape::MakeCapsule(Radius, HalfHeight), Params);
+	UWorld* World = GEngine->GetWorldFromContextObject(WorldContextObject, EGetWorldErrorMode::LogAndReturnNull);
+	bool const bHit = World ? World->SweepSingleByProfile(OutHit, Start, End, FQuat::Identity, ProfileName, FCollisionShape::MakeCapsule(Radius, HalfHeight), Params) : false;
 
 #if ENABLE_DRAW_DEBUG
 	DrawDebugCapsuleTraceSingle(World, Start, End, Radius, HalfHeight, DrawDebugType, bHit, OutHit, TraceColor, TraceHitColor, DrawTime);
@@ -1621,8 +1624,8 @@ bool UKismetSystemLibrary::CapsuleTraceMultiByProfile(UObject* WorldContextObjec
 	static const FName CapsuleTraceMultiName(TEXT("CapsuleTraceMultiByProfile"));
 	FCollisionQueryParams Params = ConfigureCollisionParams(CapsuleTraceMultiName, bTraceComplex, ActorsToIgnore, bIgnoreSelf, WorldContextObject);
 
-	UWorld* World = GEngine->GetWorldFromContextObject(WorldContextObject);
-	bool const bHit = World->SweepMultiByProfile(OutHits, Start, End, FQuat::Identity, ProfileName, FCollisionShape::MakeCapsule(Radius, HalfHeight), Params);
+	UWorld* World = GEngine->GetWorldFromContextObject(WorldContextObject, EGetWorldErrorMode::LogAndReturnNull);
+	bool const bHit = World ? World->SweepMultiByProfile(OutHits, Start, End, FQuat::Identity, ProfileName, FCollisionShape::MakeCapsule(Radius, HalfHeight), Params) : false;
 
 #if ENABLE_DRAW_DEBUG
 	DrawDebugCapsuleTraceMulti(World, Start, End, Radius, HalfHeight, DrawDebugType, bHit, OutHits, TraceColor, TraceHitColor, DrawTime);
@@ -1636,8 +1639,7 @@ bool UKismetSystemLibrary::CapsuleTraceMultiByProfile(UObject* WorldContextObjec
 void UKismetSystemLibrary::DrawDebugLine(UObject* WorldContextObject, FVector const LineStart, FVector const LineEnd, FLinearColor Color, float LifeTime, float Thickness)
 {
 #if ENABLE_DRAW_DEBUG
-	UWorld* World = GEngine->GetWorldFromContextObject(WorldContextObject);
-	if(World != nullptr)
+	if (UWorld* World = GEngine->GetWorldFromContextObject(WorldContextObject, EGetWorldErrorMode::LogAndReturnNull))
 	{
 		::DrawDebugLine(World, LineStart, LineEnd, Color.ToFColor(true), false, LifeTime, SDPG_World, Thickness);
 	}
@@ -1648,7 +1650,10 @@ void UKismetSystemLibrary::DrawDebugLine(UObject* WorldContextObject, FVector co
 void UKismetSystemLibrary::DrawDebugCircle(UObject* WorldContextObject,FVector Center, float Radius, int32 NumSegments, FLinearColor LineColor, float LifeTime, float Thickness, FVector YAxis, FVector ZAxis, bool bDrawAxis)
 { 
 #if ENABLE_DRAW_DEBUG
-	::DrawDebugCircle(GEngine->GetWorldFromContextObject(WorldContextObject),Center, Radius, NumSegments, LineColor.ToFColor(true), false, LifeTime, SDPG_World, Thickness, YAxis, ZAxis, bDrawAxis);
+	if (UWorld* World = GEngine->GetWorldFromContextObject(WorldContextObject, EGetWorldErrorMode::LogAndReturnNull))
+	{
+		::DrawDebugCircle(World, Center, Radius, NumSegments, LineColor.ToFColor(true), false, LifeTime, SDPG_World, Thickness, YAxis, ZAxis, bDrawAxis);
+	}
 #endif
 }
 
@@ -1656,8 +1661,7 @@ void UKismetSystemLibrary::DrawDebugCircle(UObject* WorldContextObject,FVector C
 void UKismetSystemLibrary::DrawDebugPoint(UObject* WorldContextObject, FVector const Position, float Size, FLinearColor PointColor, float LifeTime)
 {
 #if ENABLE_DRAW_DEBUG
-	UWorld* World = GEngine->GetWorldFromContextObject(WorldContextObject);
-	if (World != nullptr)
+	if (UWorld* World = GEngine->GetWorldFromContextObject(WorldContextObject, EGetWorldErrorMode::LogAndReturnNull))
 	{
 		::DrawDebugPoint(World, Position, Size, PointColor.ToFColor(true), false, LifeTime, SDPG_World);
 	}
@@ -1668,7 +1672,7 @@ void UKismetSystemLibrary::DrawDebugPoint(UObject* WorldContextObject, FVector c
 void UKismetSystemLibrary::DrawDebugArrow(UObject* WorldContextObject, FVector const LineStart, FVector const LineEnd, float ArrowSize, FLinearColor Color, float LifeTime, float Thickness)
 {
 #if ENABLE_DRAW_DEBUG
-	UWorld* World = GEngine->GetWorldFromContextObject(WorldContextObject);
+	UWorld* World = GEngine->GetWorldFromContextObject(WorldContextObject, EGetWorldErrorMode::LogAndReturnNull);
 	if (World != nullptr)
 	{
 		::DrawDebugDirectionalArrow(World, LineStart, LineEnd, ArrowSize, Color.ToFColor(true), false, LifeTime, SDPG_World, Thickness);
@@ -1680,8 +1684,7 @@ void UKismetSystemLibrary::DrawDebugArrow(UObject* WorldContextObject, FVector c
 void UKismetSystemLibrary::DrawDebugBox(UObject* WorldContextObject, FVector const Center, FVector Extent, FLinearColor Color, const FRotator Rotation, float LifeTime, float Thickness)
 {
 #if ENABLE_DRAW_DEBUG
-	UWorld* World = GEngine->GetWorldFromContextObject(WorldContextObject);
-	if(World != nullptr)
+	if (UWorld* World = GEngine->GetWorldFromContextObject(WorldContextObject, EGetWorldErrorMode::LogAndReturnNull))
 	{
 		if (Rotation == FRotator::ZeroRotator)
 		{
@@ -1699,8 +1702,7 @@ void UKismetSystemLibrary::DrawDebugBox(UObject* WorldContextObject, FVector con
 void UKismetSystemLibrary::DrawDebugCoordinateSystem(UObject* WorldContextObject, FVector const AxisLoc, FRotator const AxisRot, float Scale, float LifeTime, float Thickness)
 {
 #if ENABLE_DRAW_DEBUG
-	UWorld* World = GEngine->GetWorldFromContextObject(WorldContextObject);
-	if (World != nullptr)
+	if (UWorld* World = GEngine->GetWorldFromContextObject(WorldContextObject, EGetWorldErrorMode::LogAndReturnNull))
 	{
 		::DrawDebugCoordinateSystem(World, AxisLoc, AxisRot, Scale, false, LifeTime, SDPG_World, Thickness);
 	}
@@ -1711,8 +1713,7 @@ void UKismetSystemLibrary::DrawDebugCoordinateSystem(UObject* WorldContextObject
 void UKismetSystemLibrary::DrawDebugSphere(UObject* WorldContextObject, FVector const Center, float Radius, int32 Segments, FLinearColor Color, float LifeTime, float Thickness)
 {
 #if ENABLE_DRAW_DEBUG
-	UWorld* World = GEngine->GetWorldFromContextObject(WorldContextObject);
-	if (World != nullptr)
+	if (UWorld* World = GEngine->GetWorldFromContextObject(WorldContextObject, EGetWorldErrorMode::LogAndReturnNull))
 	{
 		::DrawDebugSphere(World, Center, Radius, Segments, Color.ToFColor(true), false, LifeTime, SDPG_World, Thickness);
 	}
@@ -1723,8 +1724,7 @@ void UKismetSystemLibrary::DrawDebugSphere(UObject* WorldContextObject, FVector
 void UKismetSystemLibrary::DrawDebugCylinder(UObject* WorldContextObject, FVector const Start, FVector const End, float Radius, int32 Segments, FLinearColor Color, float LifeTime, float Thickness)
 {
 #if ENABLE_DRAW_DEBUG
-	UWorld* World = GEngine->GetWorldFromContextObject(WorldContextObject);
-	if (World != nullptr)
+	if (UWorld* World = GEngine->GetWorldFromContextObject(WorldContextObject, EGetWorldErrorMode::LogAndReturnNull))
 	{
 		::DrawDebugCylinder(World, Start, End, Radius, Segments, Color.ToFColor(true), false, LifeTime, SDPG_World, Thickness);
 	}
@@ -1735,8 +1735,7 @@ void UKismetSystemLibrary::DrawDebugCylinder(UObject* WorldContextObject, FVecto
 void UKismetSystemLibrary::DrawDebugCone(UObject* WorldContextObject, FVector const Origin, FVector const Direction, float Length, float AngleWidth, float AngleHeight, int32 NumSides, FLinearColor Color, float Duration, float Thickness)
 {
 #if ENABLE_DRAW_DEBUG
-	UWorld* World = GEngine->GetWorldFromContextObject(WorldContextObject);
-	if (World != nullptr)
+	if (UWorld* World = GEngine->GetWorldFromContextObject(WorldContextObject, EGetWorldErrorMode::LogAndReturnNull))
 	{
 		::DrawDebugCone(World, Origin, Direction, Length, AngleWidth, AngleHeight, NumSides, Color.ToFColor(true), false, Duration, SDPG_World, Thickness);
 	}
@@ -1746,8 +1745,7 @@ void UKismetSystemLibrary::DrawDebugCone(UObject* WorldContextObject, FVector co
 void UKismetSystemLibrary::DrawDebugConeInDegrees(UObject* WorldContextObject, FVector const Origin, FVector const Direction, float Length, float AngleWidth, float AngleHeight, int32 NumSides, FLinearColor Color, float LifeTime, float Thickness)
 {
 #if ENABLE_DRAW_DEBUG
-	UWorld* World = GEngine->GetWorldFromContextObject(WorldContextObject);
-	if (World != nullptr)
+	if (UWorld* World = GEngine->GetWorldFromContextObject(WorldContextObject, EGetWorldErrorMode::LogAndReturnNull))
 	{
 		::DrawDebugCone(World, Origin, Direction, Length, FMath::DegreesToRadians(AngleWidth), FMath::DegreesToRadians(AngleHeight), NumSides, Color.ToFColor(true), false, LifeTime, SDPG_World, Thickness);
 	}
@@ -1758,8 +1756,7 @@ void UKismetSystemLibrary::DrawDebugConeInDegrees(UObject* WorldContextObject, F
 void UKismetSystemLibrary::DrawDebugCapsule(UObject* WorldContextObject, FVector const Center, float HalfHeight, float Radius, const FRotator Rotation, FLinearColor Color, float LifeTime, float Thickness)
 {
 #if ENABLE_DRAW_DEBUG
-	UWorld* World = GEngine->GetWorldFromContextObject(WorldContextObject);
-	if (World != nullptr)
+	if (UWorld* World = GEngine->GetWorldFromContextObject(WorldContextObject, EGetWorldErrorMode::LogAndReturnNull))
 	{
 		::DrawDebugCapsule(World, Center, HalfHeight, Radius, Rotation.Quaternion(), Color.ToFColor(true), false, LifeTime, SDPG_World, Thickness);
 	}
@@ -1770,8 +1767,7 @@ void UKismetSystemLibrary::DrawDebugCapsule(UObject* WorldContextObject, FVector
 void UKismetSystemLibrary::DrawDebugString(UObject* WorldContextObject, FVector const TextLocation, const FString& Text, class AActor* TestBaseActor, FLinearColor TextColor, float Duration)
 {
 #if ENABLE_DRAW_DEBUG
-	UWorld* World = GEngine->GetWorldFromContextObject(WorldContextObject);
-	if (World != nullptr)
+	if (UWorld* World = GEngine->GetWorldFromContextObject(WorldContextObject, EGetWorldErrorMode::LogAndReturnNull))
 	{
 		::DrawDebugString(World, TextLocation, Text, TestBaseActor, TextColor.ToFColor(true), Duration);
 	}
@@ -1782,8 +1778,7 @@ void UKismetSystemLibrary::DrawDebugString(UObject* WorldContextObject, FVector
 void UKismetSystemLibrary::FlushDebugStrings( UObject* WorldContextObject )
 {
 #if ENABLE_DRAW_DEBUG
-	UWorld* World = GEngine->GetWorldFromContextObject( WorldContextObject );
-	if(World != nullptr)
+	if (UWorld* World = GEngine->GetWorldFromContextObject(WorldContextObject, EGetWorldErrorMode::LogAndReturnNull))
 	{
 		::FlushDebugStrings( World );
 	}
@@ -1794,8 +1789,7 @@ void UKismetSystemLibrary::FlushDebugStrings( UObject* WorldContextObject )
 void UKismetSystemLibrary::DrawDebugPlane(UObject* WorldContextObject, FPlane const& P, FVector const Loc, float Size, FLinearColor Color, float LifeTime)
 {
 #if ENABLE_DRAW_DEBUG
-	UWorld* World = GEngine->GetWorldFromContextObject(WorldContextObject);
-	if (World != nullptr)
+	if (UWorld* World = GEngine->GetWorldFromContextObject(WorldContextObject, EGetWorldErrorMode::LogAndReturnNull))
 	{
 		::DrawDebugSolidPlane(World, P, Loc, Size, Color.ToFColor(true), false, LifeTime, SDPG_World);
 	}
@@ -1806,8 +1800,7 @@ void UKismetSystemLibrary::DrawDebugPlane(UObject* WorldContextObject, FPlane co
 void UKismetSystemLibrary::FlushPersistentDebugLines(UObject* WorldContextObject)
 {
 #if ENABLE_DRAW_DEBUG
-	UWorld* World = GEngine->GetWorldFromContextObject( WorldContextObject );
-	if(World != nullptr)
+	if (UWorld* World = GEngine->GetWorldFromContextObject(WorldContextObject, EGetWorldErrorMode::LogAndReturnNull))
 	{
 		::FlushPersistentDebugLines( World );
 	}
@@ -1818,8 +1811,8 @@ void UKismetSystemLibrary::FlushPersistentDebugLines(UObject* WorldContextObject
 void UKismetSystemLibrary::DrawDebugFrustum(UObject* WorldContextObject, const FTransform& FrustumTransform, FLinearColor FrustumColor, float Duration, float Thickness)
 {
 #if ENABLE_DRAW_DEBUG
-	UWorld* World = GEngine->GetWorldFromContextObject(WorldContextObject);
-	if( World != nullptr && FrustumTransform.IsRotationNormalized() )
+	UWorld* World = GEngine->GetWorldFromContextObject(WorldContextObject, EGetWorldErrorMode::LogAndReturnNull);
+	if (World != nullptr && FrustumTransform.IsRotationNormalized())
 	{
 		FMatrix FrustumToWorld =  FrustumTransform.ToMatrixWithScale();
 		::DrawDebugFrustum(World, FrustumToWorld, FrustumColor.ToFColor(true), false, Duration, SDPG_World, Thickness);
@@ -1843,8 +1836,7 @@ void UKismetSystemLibrary::DrawDebugCamera(const ACameraActor* CameraActor, FLin
 void UKismetSystemLibrary::DrawDebugFloatHistoryTransform(UObject* WorldContextObject, const FDebugFloatHistory& FloatHistory, const FTransform& DrawTransform, FVector2D DrawSize, FLinearColor DrawColor, float LifeTime)
 {
 #if ENABLE_DRAW_DEBUG
-	UWorld * World = GEngine->GetWorldFromContextObject(WorldContextObject);
-	if (World)
+	if (UWorld* World = GEngine->GetWorldFromContextObject(WorldContextObject, EGetWorldErrorMode::LogAndReturnNull))
 	{
 		::DrawDebugFloatHistory(*World, FloatHistory, DrawTransform, DrawSize, DrawColor.ToFColor(true), false, LifeTime);
 	}
@@ -1854,8 +1846,7 @@ void UKismetSystemLibrary::DrawDebugFloatHistoryTransform(UObject* WorldContextO
 void UKismetSystemLibrary::DrawDebugFloatHistoryLocation(UObject* WorldContextObject, const FDebugFloatHistory& FloatHistory, FVector DrawLocation, FVector2D DrawSize, FLinearColor DrawColor, float LifeTime)
 {
 #if ENABLE_DRAW_DEBUG
-	UWorld * World = GEngine->GetWorldFromContextObject(WorldContextObject);
-	if (World)
+	if (UWorld* World = GEngine->GetWorldFromContextObject(WorldContextObject, EGetWorldErrorMode::LogAndReturnNull))
 	{
 		::DrawDebugFloatHistory(*World, FloatHistory, DrawLocation, DrawSize, DrawColor.ToFColor(true), false, LifeTime);
 	}
@@ -1915,7 +1906,7 @@ void UKismetSystemLibrary::GetActorBounds(const AActor* Actor, FVector& Origin,
 
 void UKismetSystemLibrary::Delay(UObject* WorldContextObject, float Duration, FLatentActionInfo LatentInfo )
 {
-	if (UWorld* World = GEngine->GetWorldFromContextObject(WorldContextObject))
+	if (UWorld* World = GEngine->GetWorldFromContextObject(WorldContextObject, EGetWorldErrorMode::LogAndReturnNull))
 	{
 		FLatentActionManager& LatentActionManager = World->GetLatentActionManager();
 		if (LatentActionManager.FindExistingAction<FDelayAction>(LatentInfo.CallbackTarget, LatentInfo.UUID) == NULL)
@@ -1928,7 +1919,7 @@ void UKismetSystemLibrary::Delay(UObject* WorldContextObject, float Duration, FL
 // Delay execution by Duration seconds; Calling again before the delay has expired will reset the countdown to Duration.
 void UKismetSystemLibrary::RetriggerableDelay(UObject* WorldContextObject, float Duration, FLatentActionInfo LatentInfo)
 {
-	if (UWorld* World = GEngine->GetWorldFromContextObject(WorldContextObject))
+	if (UWorld* World = GEngine->GetWorldFromContextObject(WorldContextObject, EGetWorldErrorMode::LogAndReturnNull))
 	{
 		FLatentActionManager& LatentActionManager = World->GetLatentActionManager();
 		FDelayAction* Action = LatentActionManager.FindExistingAction<FDelayAction>(LatentInfo.CallbackTarget, LatentInfo.UUID);
@@ -1946,7 +1937,7 @@ void UKismetSystemLibrary::RetriggerableDelay(UObject* WorldContextObject, float
 
 void UKismetSystemLibrary::MoveComponentTo(USceneComponent* Component, FVector TargetRelativeLocation, FRotator TargetRelativeRotation, bool bEaseOut, bool bEaseIn, float OverTime, bool bForceShortestRotationPath, TEnumAsByte<EMoveComponentAction::Type> MoveAction, FLatentActionInfo LatentInfo)
 {
-	if (UWorld* World = ((Component != NULL) ? Component->GetWorld() : NULL))
+	if (UWorld* World = GEngine->GetWorldFromContextObject(Component, EGetWorldErrorMode::LogAndReturnNull))
 	{
 		FLatentActionManager& LatentActionManager = World->GetLatentActionManager();
 		FInterpolateComponentToAction* Action = LatentActionManager.FindExistingAction<FInterpolateComponentToAction>(LatentInfo.CallbackTarget, LatentInfo.UUID);
@@ -2284,7 +2275,7 @@ bool UKismetSystemLibrary::IsControllerAssignedToGamepad(int32 ControllerId)
 
 void UKismetSystemLibrary::SetSuppressViewportTransitionMessage(UObject* WorldContextObject, bool bState)
 {
-	UWorld* World = GEngine->GetWorldFromContextObject(WorldContextObject);
+	UWorld* World = GEngine->GetWorldFromContextObject(WorldContextObject, EGetWorldErrorMode::LogAndReturnNull);
 	if (World && World->GetFirstLocalPlayerFromController() != nullptr && World->GetFirstLocalPlayerFromController()->ViewportClient != nullptr )
 	{
 		World->GetFirstLocalPlayerFromController()->ViewportClient->SetSuppressTransitionMessage(bState);
@@ -2365,7 +2356,7 @@ public:
 #endif
 };
 
-void UKismetSystemLibrary::LoadAsset(UObject* WorldContextObject, const TAssetPtr<UObject>& Asset, UKismetSystemLibrary::FOnAssetLoaded OnLoaded, FLatentActionInfo LatentInfo)
+void UKismetSystemLibrary::LoadAsset(UObject* WorldContextObject, TAssetPtr<UObject> Asset, UKismetSystemLibrary::FOnAssetLoaded OnLoaded, FLatentActionInfo LatentInfo)
 {
 	struct FLoadAssetAction : public FLoadAssetActionBase
 	{
@@ -2384,8 +2375,7 @@ void UKismetSystemLibrary::LoadAsset(UObject* WorldContextObject, const TAssetPt
 		}
 	};
 
-	UWorld* World = GEngine->GetWorldFromContextObject(WorldContextObject);
-	if (World != nullptr)
+	if (UWorld* World = GEngine->GetWorldFromContextObject(WorldContextObject, EGetWorldErrorMode::LogAndReturnNull))
 	{
 		FLatentActionManager& LatentManager = World->GetLatentActionManager();
 		if (LatentManager.FindExistingAction<FLoadAssetAction>(LatentInfo.CallbackTarget, LatentInfo.UUID) == nullptr)
@@ -2396,7 +2386,7 @@ void UKismetSystemLibrary::LoadAsset(UObject* WorldContextObject, const TAssetPt
 	}
 }
 
-void UKismetSystemLibrary::LoadAssetClass(UObject* WorldContextObject, const TAssetSubclassOf<UObject>& AssetClass, UKismetSystemLibrary::FOnAssetClassLoaded OnLoaded, FLatentActionInfo LatentInfo)
+void UKismetSystemLibrary::LoadAssetClass(UObject* WorldContextObject, TAssetSubclassOf<UObject> AssetClass, UKismetSystemLibrary::FOnAssetClassLoaded OnLoaded, FLatentActionInfo LatentInfo)
 {
 	struct FLoadAssetClassAction : public FLoadAssetActionBase
 	{
@@ -2415,8 +2405,7 @@ void UKismetSystemLibrary::LoadAssetClass(UObject* WorldContextObject, const TAs
 		}
 	};
 
-	UWorld* World = GEngine->GetWorldFromContextObject(WorldContextObject);
-	if (World != nullptr)
+	if (UWorld* World = GEngine->GetWorldFromContextObject(WorldContextObject, EGetWorldErrorMode::LogAndReturnNull))
 	{
 		FLatentActionManager& LatentManager = World->GetLatentActionManager();
 		if (LatentManager.FindExistingAction<FLoadAssetClassAction>(LatentInfo.CallbackTarget, LatentInfo.UUID) == nullptr)
diff --git a/Engine/Source/Runtime/Engine/Private/LevelBounds.cpp b/Engine/Source/Runtime/Engine/Private/LevelBounds.cpp
--- a/Engine/Source/Runtime/Engine/Private/LevelBounds.cpp
+++ b/Engine/Source/Runtime/Engine/Private/LevelBounds.cpp
@@ -36,7 +36,14 @@ ALevelBounds::ALevelBounds(const FObjectInitializer& ObjectInitializer)
 void ALevelBounds::PostLoad()
 {
 	Super::PostLoad();
-	GetLevel()->LevelBoundsActor = this;
+
+	if (!IsTemplate())
+	{
+		if (ULevel* Level = GetLevel())
+		{
+			Level->LevelBoundsActor = this;
+		}
+	}
 }
 
 FBox ALevelBounds::GetComponentsBoundingBox(bool bNonColliding) const
diff --git a/Engine/Source/Runtime/Engine/Private/LevelStreaming.cpp b/Engine/Source/Runtime/Engine/Private/LevelStreaming.cpp
--- a/Engine/Source/Runtime/Engine/Private/LevelStreaming.cpp
+++ b/Engine/Source/Runtime/Engine/Private/LevelStreaming.cpp
@@ -103,7 +103,12 @@ FString FStreamLevelAction::MakeSafeLevelName( const FName& InLevelName, UWorld*
 	// Special case for PIE, the PackageName gets mangled.
 	if (!InWorld->StreamingLevelsPrefix.IsEmpty())
 	{
-		FString PackageName = InWorld->StreamingLevelsPrefix + FPackageName::GetShortName(InLevelName);
+		FString PackageName = FPackageName::GetShortName(InLevelName);
+		if (!PackageName.StartsWith(InWorld->StreamingLevelsPrefix))
+		{
+			PackageName = InWorld->StreamingLevelsPrefix + PackageName;
+		}
+
 		if (!FPackageName::IsShortPackageName(InLevelName))
 		{
 			PackageName = FPackageName::GetLongPackagePath(InLevelName.ToString()) + TEXT("/") + PackageName;
@@ -1055,7 +1060,7 @@ ALevelScriptActor* ULevelStreaming::GetLevelScriptActor()
 ULevelStreamingKismet* ULevelStreamingKismet::LoadLevelInstance(UObject* WorldContextObject, const FString& LevelName, const FVector& Location, const FRotator& Rotation, bool& bOutSuccess)
 { 
 	bOutSuccess = false; 
-	UWorld* const World = GEngine->GetWorldFromContextObject(WorldContextObject, false);
+	UWorld* const World = GEngine->GetWorldFromContextObject(WorldContextObject, EGetWorldErrorMode::LogAndReturnNull);
 	if (!World) 
 	{
 		return nullptr;
