Fix non-unity
#rnx

diff --git a/Engine/Source/Runtime/Engine/Private/EdGraph/EdGraphPin.cpp b/Engine/Source/Runtime/Engine/Private/EdGraph/EdGraphPin.cpp
--- a/Engine/Source/Runtime/Engine/Private/EdGraph/EdGraphPin.cpp
+++ b/Engine/Source/Runtime/Engine/Private/EdGraph/EdGraphPin.cpp
@@ -9,6 +9,7 @@
 #include "EdGraph/EdGraphSchema.h"
 #include "Tickable.h"
 #include "EngineLogs.h"
+#include "HAL/IConsoleManager.h"
 #if WITH_EDITOR
 #include "Editor/EditorEngine.h"
 #include "Misc/ConfigCacheIni.h"
