Fix hotreload crash accessing orphan pin CVar
Improve naming and description of CVar
#rnx

diff --git a/Engine/Source/Editor/BlueprintGraph/Private/K2Node.cpp b/Engine/Source/Editor/BlueprintGraph/Private/K2Node.cpp
--- a/Engine/Source/Editor/BlueprintGraph/Private/K2Node.cpp
+++ b/Engine/Source/Editor/BlueprintGraph/Private/K2Node.cpp
@@ -887,7 +887,7 @@ void UK2Node::RewireOldPinsToNewPins(TArray<UEdGraphPin*>& InOldPins, TArray<UEd
 		// * The pin has been flagged not be saved if orphaned
 		// * The pin is hidden and not a split pin
 		const bool bVisibleOrSplitPin = (!OldPin->bHidden || (OldPin->SubPins.Num() > 0));
-		if (UEdGraphPin::EnableOrphanPins() && !bMatched && bVisibleOrSplitPin && OldPin->bSavePinIfOrphaned)
+		if (UEdGraphPin::AreOrphanPinsEnabled() && !bMatched && bVisibleOrSplitPin && OldPin->bSavePinIfOrphaned)
 		{
 			// The node can specify to save no pins, all pins, or all but exec pins. However, even if all is specified Execute and Then are never saved
 			 const bool bSaveOrphanedPin = ((OrphanedPinSaveMode == ESaveOrphanPinMode::SaveAll) ||
diff --git a/Engine/Source/Runtime/Engine/Classes/EdGraph/EdGraphPin.h b/Engine/Source/Runtime/Engine/Classes/EdGraph/EdGraphPin.h
--- a/Engine/Source/Runtime/Engine/Classes/EdGraph/EdGraphPin.h
+++ b/Engine/Source/Runtime/Engine/Classes/EdGraph/EdGraphPin.h
@@ -535,7 +535,7 @@ public:
 	/** This needs to be called if you want to use pin data within PostEditUndo */
 	ENGINE_API static void ResolveAllPinReferences();
 
-	ENGINE_API static bool EnableOrphanPins();
+	ENGINE_API static bool AreOrphanPinsEnabled();
 
 private:
 	/** Private Constructor. Create pins using CreatePin since all pin instances are managed by TSharedPtr. */
diff --git a/Engine/Source/Runtime/Engine/Private/EdGraph/EdGraphPin.cpp b/Engine/Source/Runtime/Engine/Private/EdGraph/EdGraphPin.cpp
--- a/Engine/Source/Runtime/Engine/Private/EdGraph/EdGraphPin.cpp
+++ b/Engine/Source/Runtime/Engine/Private/EdGraph/EdGraphPin.cpp
@@ -1052,7 +1052,7 @@ bool UEdGraphPin::ImportTextItem(const TCHAR*& Buffer, int32 PortFlags, class UO
 			bool LocalOrphanedPin = bOrphanedPin;
 			Buffer = BoolPropCDO->ImportText(Buffer, &LocalOrphanedPin, PortFlags, Parent, ErrorText);
 			bParseSuccess = (Buffer != nullptr);
-			if (EnableOrphanPins())
+			if (AreOrphanPinsEnabled())
 			{
 				bOrphanedPin = LocalOrphanedPin;
 			}
@@ -1449,7 +1449,7 @@ bool UEdGraphPin::Serialize(FArchive& Ar)
 		bDefaultValueIsReadOnly = !!(BitField & (1 << 2));
 		bDefaultValueIsIgnored = !!(BitField & (1 << 3));
 		bAdvancedView = !!(BitField & (1 << 4));
-		if (EnableOrphanPins())
+		if (AreOrphanPinsEnabled())
 		{
 			bOrphanedPin = !!(BitField & (1 << 5));
 		}
@@ -1861,11 +1861,12 @@ bool UEdGraphPin::ImportText_PinArray(const TCHAR*& Buffer, TArray<UEdGraphPin*>
 	return false;
 }
 
-TAutoConsoleVariable<int32> CVarDisableOrphanPins(TEXT("DisableOrphanPins"), 0, TEXT("0=Enable, 1=Disable"), ECVF_ReadOnly);
+int32 GCVarDisableOrphanPins = 0;
+FAutoConsoleVariableRef CVarDisableOrphanPins(TEXT("DisableOrphanPins"), GCVarDisableOrphanPins, TEXT("0=Orphan pins are enabled (default), 1=Orphan pins are disabled (note: this option will go away in the future)"), ECVF_ReadOnly);
 
-bool UEdGraphPin::EnableOrphanPins()
+bool UEdGraphPin::AreOrphanPinsEnabled()
 {
-	return (CVarDisableOrphanPins.GetValueOnGameThread() == 0);
+	return (GCVarDisableOrphanPins == 0);
 }
 
 
