UE-46651 Failure to build new C++ Project: Error C1083: Cannot open include file: 'Kismet/HeadMountedDisplayFunctionLibrary.h': No such file or directory
-Added compiler warning about moved header.

diff --git a/Engine/Source/Runtime/Engine/Classes/Kismet/HeadMountedDisplayFunctionLibrary.h b/Engine/Source/Runtime/Engine/Classes/Kismet/HeadMountedDisplayFunctionLibrary.h
new file mode 100644
--- /dev/null
+++ b/Engine/Source/Runtime/Engine/Classes/Kismet/HeadMountedDisplayFunctionLibrary.h
@@ -0,0 +1,8 @@
+// Copyright 1998-2017 Epic Games, Inc. All Rights Reserved.
+
+#include "Platform.h"
+
+COMPILE_WARNING("This file has been moved to HeadMountedDisplay/Public/HeadMountedDisplayFunctionLibrary.h. Some types were moved to HeadMountedDisplay/Public/HeadMountedDisplayTypes.h.  Replace all #include \"Kismet\\HeadMountedDisplayFunctionLibrary.h\" with: #include \"HeadMountedDisplayFunctionLibrary.h\" or #include \"HeadMountedDisplayTypes.h\".")
+
+// Needs highly explicit path, or we just end up including this header for an infinite loop.
+#include "../../../HeadMountedDisplay/Public/HeadMountedDisplayFunctionLibrary.h"
\ No newline at end of file
