UE-46651 Failure to build new C++ Project: Error C1083: Cannot open include file: 'Kismet/HeadMountedDisplayFunctionLibrary.h': No such file or directory
-Followup change
-Added#pragma once... the 'warn on ever include' idea doesn't really work as you get multiples for modules and it doesn't tell you where the include is anyway.

diff --git a/Engine/Source/Runtime/Engine/Classes/Kismet/HeadMountedDisplayFunctionLibrary.h b/Engine/Source/Runtime/Engine/Classes/Kismet/HeadMountedDisplayFunctionLibrary.h
--- a/Engine/Source/Runtime/Engine/Classes/Kismet/HeadMountedDisplayFunctionLibrary.h
+++ b/Engine/Source/Runtime/Engine/Classes/Kismet/HeadMountedDisplayFunctionLibrary.h
@@ -1,5 +1,7 @@
 // Copyright 1998-2017 Epic Games, Inc. All Rights Reserved.
 
+#pragma once
+
 #include "Platform.h"
 
 COMPILE_WARNING("This file has been moved to HeadMountedDisplay/Public/HeadMountedDisplayFunctionLibrary.h. Some types were moved to HeadMountedDisplay/Public/HeadMountedDisplayTypes.h.  Replace all #include \"Kismet\\HeadMountedDisplayFunctionLibrary.h\" with: #include \"HeadMountedDisplayFunctionLibrary.h\" or #include \"HeadMountedDisplayTypes.h\".")
