[Copy] Fixed crash when doing editor 'Force Delete' on a static mesh whose distance field is still being built.  Any UObject reference that is to an asset can be NULL'ed by the editor.

diff --git a/Engine/Documentation/Source/Programming/UnrealArchitecture/Objects/Optimizations/UnrealObjectHandling.INT.udn b/Engine/Documentation/Source/Programming/UnrealArchitecture/Objects/Optimizations/UnrealObjectHandling.INT.udn
--- a/Engine/Documentation/Source/Programming/UnrealArchitecture/Objects/Optimizations/UnrealObjectHandling.INT.udn
+++ b/Engine/Documentation/Source/Programming/UnrealArchitecture/Objects/Optimizations/UnrealObjectHandling.INT.udn
@@ -26,6 +26,8 @@ When an `AActor` or `UActorComponent` is destroyed or otherwise removed from pla
 
 It is important to realize that this feature applies only to `UActorComponent` or `AActor` references marked with `UPROPERTY` or stored in an Unreal Engine container class. An Object reference stored in a raw pointer will be unknown to the Unreal Engine, and will not be automatically nulled, nor will it prevent garbage collection. Note this does not mean that all `UObject*` variables must be `UProperties`. If you want an Object pointer that is not a `UProperty`, consider using `TWeakObjectPtr`. This is a "weak" pointer, meaning it will not prevent garbage collection, but it can be queried for validity before being accessed and will be set to null if the Object it points to is destroyed.
 
+Another case where a referenced UObject UProperty will be automatically null'ed is when using 'Force Delete' on an asset in the editor.  As a result, all code operating on UObjects which are assets must handle these pointers becoming null.
+
 ## Serialization
 
 When a `UObject` is serialized, all UProperty values are automatically written or read unless explicitly marked as "transient". For example, you could place an `AEnemy` instance in a level, set its Health to 500, save it and successfully reload it without writing a single line of code beyond the `UClass` definition.
diff --git a/Engine/Source/Runtime/Engine/Private/DistanceFieldAtlas.cpp b/Engine/Source/Runtime/Engine/Private/DistanceFieldAtlas.cpp
--- a/Engine/Source/Runtime/Engine/Private/DistanceFieldAtlas.cpp
+++ b/Engine/Source/Runtime/Engine/Private/DistanceFieldAtlas.cpp
@@ -549,6 +549,7 @@ void FDistanceFieldVolumeData::CacheDerivedData(const FString& InDDCKey, UStatic
 		COOK_STAT(Timer.TrackCyclesOnly());
 		FAsyncDistanceFieldTask* NewTask = new FAsyncDistanceFieldTask;
 		NewTask->DDCKey = InDDCKey;
+		check(Mesh && GenerateSource);
 		NewTask->StaticMesh = Mesh;
 		NewTask->GenerateSource = GenerateSource;
 		NewTask->DistanceFieldResolutionScale = DistanceFieldResolutionScale;
@@ -780,19 +781,25 @@ void FDistanceFieldAsyncQueue::BlockUntilAllBuildsComplete()
 void FDistanceFieldAsyncQueue::Build(FAsyncDistanceFieldTask* Task, FQueuedThreadPool& ThreadPool)
 {
 #if WITH_EDITOR
-	const FStaticMeshLODResources& LODModel = Task->GenerateSource->RenderData->LODResources[0];
-
-	MeshUtilities->GenerateSignedDistanceFieldVolumeData(
-		Task->StaticMesh->GetName(),
-		LODModel,
-		ThreadPool,
-		Task->MaterialBlendModes,
-		Task->GenerateSource->RenderData->Bounds,
-		Task->DistanceFieldResolutionScale,
-		Task->bGenerateDistanceFieldAsIfTwoSided,
-		*Task->GeneratedVolumeData);
-
-	CompletedTasks.Push(Task);
+
+	// Editor 'force delete' can null any UObject pointers which are seen by reference collecting (eg UProperty or serialized)
+	if (Task->StaticMesh && Task->GenerateSource)
+	{
+		const FStaticMeshLODResources& LODModel = Task->GenerateSource->RenderData->LODResources[0];
+
+		MeshUtilities->GenerateSignedDistanceFieldVolumeData(
+			Task->StaticMesh->GetName(),
+			LODModel,
+			ThreadPool,
+			Task->MaterialBlendModes,
+			Task->GenerateSource->RenderData->Bounds,
+			Task->DistanceFieldResolutionScale,
+			Task->bGenerateDistanceFieldAsIfTwoSided,
+			*Task->GeneratedVolumeData);
+	}
+
+    CompletedTasks.Push(Task);
+
 #endif
 }
 
@@ -820,29 +827,33 @@ void FDistanceFieldAsyncQueue::ProcessAsyncTasks()
 
 		ReferencedTasks.Remove(Task);
 
-		Task->GeneratedVolumeData->VolumeTexture.Initialize(Task->StaticMesh);
-		FDistanceFieldVolumeData* OldVolumeData = Task->StaticMesh->RenderData->LODResources[0].DistanceFieldData;
-
+		// Editor 'force delete' can null any UObject pointers which are seen by reference collecting (eg UProperty or serialized)
+		if (Task->StaticMesh)
 		{
-			// Cause all components using this static mesh to get re-registered, which will recreate their proxies and primitive uniform buffers
-			FStaticMeshComponentRecreateRenderStateContext RecreateRenderStateContext(Task->StaticMesh, false);
+			Task->GeneratedVolumeData->VolumeTexture.Initialize(Task->StaticMesh);
+			FDistanceFieldVolumeData* OldVolumeData = Task->StaticMesh->RenderData->LODResources[0].DistanceFieldData;
 
-			// Assign the new volume data
-			Task->StaticMesh->RenderData->LODResources[0].DistanceFieldData = Task->GeneratedVolumeData;
-		}
+			{
+				// Cause all components using this static mesh to get re-registered, which will recreate their proxies and primitive uniform buffers
+				FStaticMeshComponentRecreateRenderStateContext RecreateRenderStateContext(Task->StaticMesh, false);
 
-		OldVolumeData->VolumeTexture.Release();
+				// Assign the new volume data
+				Task->StaticMesh->RenderData->LODResources[0].DistanceFieldData = Task->GeneratedVolumeData;
+			}
 
-		// Rendering thread may still be referencing the old one, use the deferred cleanup interface to delete it next frame when it is safe
-		BeginCleanup(OldVolumeData);
+			OldVolumeData->VolumeTexture.Release();
 
-		{
-			TArray<uint8> DerivedData;
-			// Save built distance field volume to DDC
-			FMemoryWriter Ar(DerivedData, /*bIsPersistent=*/ true);
-			Ar << *(Task->StaticMesh->RenderData->LODResources[0].DistanceFieldData);
-			GetDerivedDataCacheRef().Put(*Task->DDCKey, DerivedData);
-			COOK_STAT(Timer.AddMiss(DerivedData.Num()));
+			// Rendering thread may still be referencing the old one, use the deferred cleanup interface to delete it next frame when it is safe
+			BeginCleanup(OldVolumeData);
+
+			{
+				TArray<uint8> DerivedData;
+				// Save built distance field volume to DDC
+				FMemoryWriter Ar(DerivedData, /*bIsPersistent=*/ true);
+				Ar << *(Task->StaticMesh->RenderData->LODResources[0].DistanceFieldData);
+				GetDerivedDataCacheRef().Put(*Task->DDCKey, DerivedData);
+				COOK_STAT(Timer.AddMiss(DerivedData.Num()));
+			}
 		}
 
 		delete Task;
