Back out changelist 3516977, 4.17 branch isn't open yet. Still trying to get a new build.

diff --git a/Engine/Documentation/Source/Programming/UnrealArchitecture/Objects/Optimizations/UnrealObjectHandling.INT.udn b/Engine/Documentation/Source/Programming/UnrealArchitecture/Objects/Optimizations/UnrealObjectHandling.INT.udn
--- a/Engine/Documentation/Source/Programming/UnrealArchitecture/Objects/Optimizations/UnrealObjectHandling.INT.udn
+++ b/Engine/Documentation/Source/Programming/UnrealArchitecture/Objects/Optimizations/UnrealObjectHandling.INT.udn
@@ -26,8 +26,6 @@ When an `AActor` or `UActorComponent` is destroyed or otherwise removed from pla
 
 It is important to realize that this feature applies only to `UActorComponent` or `AActor` references marked with `UPROPERTY` or stored in an Unreal Engine container class. An Object reference stored in a raw pointer will be unknown to the Unreal Engine, and will not be automatically nulled, nor will it prevent garbage collection. Note this does not mean that all `UObject*` variables must be `UProperties`. If you want an Object pointer that is not a `UProperty`, consider using `TWeakObjectPtr`. This is a "weak" pointer, meaning it will not prevent garbage collection, but it can be queried for validity before being accessed and will be set to null if the Object it points to is destroyed.
 
-Another case where a referenced UObject UProperty will be automatically null'ed is when using 'Force Delete' on an asset in the editor.  As a result, all code operating on UObjects which are assets must handle these pointers becoming null.
-
 ## Serialization
 
 When a `UObject` is serialized, all UProperty values are automatically written or read unless explicitly marked as "transient". For example, you could place an `AEnemy` instance in a level, set its Health to 500, save it and successfully reload it without writing a single line of code beyond the `UClass` definition.
diff --git a/Engine/Source/Runtime/Engine/Private/DistanceFieldAtlas.cpp b/Engine/Source/Runtime/Engine/Private/DistanceFieldAtlas.cpp
--- a/Engine/Source/Runtime/Engine/Private/DistanceFieldAtlas.cpp
+++ b/Engine/Source/Runtime/Engine/Private/DistanceFieldAtlas.cpp
@@ -549,7 +549,6 @@ void FDistanceFieldVolumeData::CacheDerivedData(const FString& InDDCKey, UStatic
 		COOK_STAT(Timer.TrackCyclesOnly());
 		FAsyncDistanceFieldTask* NewTask = new FAsyncDistanceFieldTask;
 		NewTask->DDCKey = InDDCKey;
-		check(Mesh && GenerateSource);
 		NewTask->StaticMesh = Mesh;
 		NewTask->GenerateSource = GenerateSource;
 		NewTask->DistanceFieldResolutionScale = DistanceFieldResolutionScale;
@@ -781,25 +780,19 @@ void FDistanceFieldAsyncQueue::BlockUntilAllBuildsComplete()
 void FDistanceFieldAsyncQueue::Build(FAsyncDistanceFieldTask* Task, FQueuedThreadPool& ThreadPool)
 {
 #if WITH_EDITOR
-
-	// Editor 'force delete' can null any UObject pointers which are seen by reference collecting (eg UProperty or serialized)
-	if (Task->StaticMesh && Task->GenerateSource)
-	{
-		const FStaticMeshLODResources& LODModel = Task->GenerateSource->RenderData->LODResources[0];
-
-		MeshUtilities->GenerateSignedDistanceFieldVolumeData(
-			Task->StaticMesh->GetName(),
-			LODModel,
-			ThreadPool,
-			Task->MaterialBlendModes,
-			Task->GenerateSource->RenderData->Bounds,
-			Task->DistanceFieldResolutionScale,
-			Task->bGenerateDistanceFieldAsIfTwoSided,
-			*Task->GeneratedVolumeData);
-	}
-
-    CompletedTasks.Push(Task);
-
+	const FStaticMeshLODResources& LODModel = Task->GenerateSource->RenderData->LODResources[0];
+
+	MeshUtilities->GenerateSignedDistanceFieldVolumeData(
+		Task->StaticMesh->GetName(),
+		LODModel,
+		ThreadPool,
+		Task->MaterialBlendModes,
+		Task->GenerateSource->RenderData->Bounds,
+		Task->DistanceFieldResolutionScale,
+		Task->bGenerateDistanceFieldAsIfTwoSided,
+		*Task->GeneratedVolumeData);
+
+	CompletedTasks.Push(Task);
 #endif
 }
 
@@ -827,33 +820,29 @@ void FDistanceFieldAsyncQueue::ProcessAsyncTasks()
 
 		ReferencedTasks.Remove(Task);
 
-		// Editor 'force delete' can null any UObject pointers which are seen by reference collecting (eg UProperty or serialized)
-		if (Task->StaticMesh)
-		{
-			Task->GeneratedVolumeData->VolumeTexture.Initialize(Task->StaticMesh);
-			FDistanceFieldVolumeData* OldVolumeData = Task->StaticMesh->RenderData->LODResources[0].DistanceFieldData;
+		Task->GeneratedVolumeData->VolumeTexture.Initialize(Task->StaticMesh);
+		FDistanceFieldVolumeData* OldVolumeData = Task->StaticMesh->RenderData->LODResources[0].DistanceFieldData;
 
-			{
-				// Cause all components using this static mesh to get re-registered, which will recreate their proxies and primitive uniform buffers
-				FStaticMeshComponentRecreateRenderStateContext RecreateRenderStateContext(Task->StaticMesh, false);
+		{
+			// Cause all components using this static mesh to get re-registered, which will recreate their proxies and primitive uniform buffers
+			FStaticMeshComponentRecreateRenderStateContext RecreateRenderStateContext(Task->StaticMesh, false);
 
-				// Assign the new volume data
-				Task->StaticMesh->RenderData->LODResources[0].DistanceFieldData = Task->GeneratedVolumeData;
-			}
+			// Assign the new volume data
+			Task->StaticMesh->RenderData->LODResources[0].DistanceFieldData = Task->GeneratedVolumeData;
+		}
 
-			OldVolumeData->VolumeTexture.Release();
+		OldVolumeData->VolumeTexture.Release();
 
-			// Rendering thread may still be referencing the old one, use the deferred cleanup interface to delete it next frame when it is safe
-			BeginCleanup(OldVolumeData);
+		// Rendering thread may still be referencing the old one, use the deferred cleanup interface to delete it next frame when it is safe
+		BeginCleanup(OldVolumeData);
 
-			{
-				TArray<uint8> DerivedData;
-				// Save built distance field volume to DDC
-				FMemoryWriter Ar(DerivedData, /*bIsPersistent=*/ true);
-				Ar << *(Task->StaticMesh->RenderData->LODResources[0].DistanceFieldData);
-				GetDerivedDataCacheRef().Put(*Task->DDCKey, DerivedData);
-				COOK_STAT(Timer.AddMiss(DerivedData.Num()));
-			}
+		{
+			TArray<uint8> DerivedData;
+			// Save built distance field volume to DDC
+			FMemoryWriter Ar(DerivedData, /*bIsPersistent=*/ true);
+			Ar << *(Task->StaticMesh->RenderData->LODResources[0].DistanceFieldData);
+			GetDerivedDataCacheRef().Put(*Task->DDCKey, DerivedData);
+			COOK_STAT(Timer.AddMiss(DerivedData.Num()));
 		}
 
 		delete Task;
