Copying //UE4/Dev-Framework to //UE4/Dev-Main (Source: //UE4/Dev-Framework @ 3544039)

diff --git a/Engine/Source/Runtime/Engine/Classes/Kismet/GameplayStatics.h b/Engine/Source/Runtime/Engine/Classes/Kismet/GameplayStatics.h
--- a/Engine/Source/Runtime/Engine/Classes/Kismet/GameplayStatics.h
+++ b/Engine/Source/Runtime/Engine/Classes/Kismet/GameplayStatics.h
@@ -476,12 +476,6 @@ class ENGINE_API UGameplayStatics : public UBlueprintFunctionLibrary
 	UFUNCTION(BlueprintCallable, Category="Audio", meta=(AdvancedDisplay = "2", UnsafeDuringActorConstruction = "true", Keywords = "play"))
 	static UAudioComponent* SpawnSoundAttached(USoundBase* Sound, USceneComponent* AttachToComponent, FName AttachPointName = NAME_None, FVector Location = FVector(ForceInit), FRotator Rotation = FRotator::ZeroRotator, EAttachLocation::Type LocationType = EAttachLocation::KeepRelativeOffset, bool bStopWhenAttachedToDestroyed = false, float VolumeMultiplier = 1.f, float PitchMultiplier = 1.f, float StartTime = 0.f, USoundAttenuation* AttenuationSettings = nullptr, USoundConcurrency* ConcurrencySettings = nullptr, bool bAutoDestroy = true);
 
-	DEPRECATED(4.9, "PlaySoundAttached has been renamed SpawnSoundAttached")
-	static UAudioComponent* PlaySoundAttached(USoundBase* Sound, USceneComponent* AttachToComponent, FName AttachPointName = NAME_None, FVector Location = FVector(ForceInit), EAttachLocation::Type LocationType = EAttachLocation::KeepRelativeOffset, bool bStopWhenAttachedToDestroyed = false, float VolumeMultiplier = 1.f, float PitchMultiplier = 1.f, float StartTime = 0.f, USoundAttenuation* AttenuationSettings = nullptr)
-	{
-		return SpawnSoundAttached(Sound, AttachToComponent, AttachPointName, Location, FRotator::ZeroRotator, LocationType, bStopWhenAttachedToDestroyed, VolumeMultiplier, PitchMultiplier, StartTime, AttenuationSettings);
-	}
-
 	static class UAudioComponent* SpawnSoundAttached(USoundBase* Sound, USceneComponent* AttachToComponent, FName AttachPointName, FVector Location, EAttachLocation::Type LocationType = EAttachLocation::KeepRelativeOffset, bool bStopWhenAttachedToDestroyed = false, float VolumeMultiplier = 1.f, float PitchMultiplier = 1.f, float StartTime = 0.f, USoundAttenuation* AttenuationSettings = nullptr, USoundConcurrency* ConcurrencySettings = nullptr, bool bAutoDestroy = true)
 	{
 		return SpawnSoundAttached(Sound, AttachToComponent, AttachPointName, Location, FRotator::ZeroRotator, LocationType, bStopWhenAttachedToDestroyed, VolumeMultiplier, PitchMultiplier, StartTime, AttenuationSettings, ConcurrencySettings, bAutoDestroy);
@@ -568,12 +562,6 @@ class ENGINE_API UGameplayStatics : public UBlueprintFunctionLibrary
 	UFUNCTION(BlueprintCallable, Category="Audio", meta=(AdvancedDisplay = "2", UnsafeDuringActorConstruction = "true", Keywords = "play"))
 	static UAudioComponent* SpawnDialogueAttached(UDialogueWave* Dialogue, const FDialogueContext& Context, USceneComponent* AttachToComponent, FName AttachPointName = NAME_None, FVector Location = FVector(ForceInit), FRotator Rotation = FRotator::ZeroRotator, EAttachLocation::Type LocationType = EAttachLocation::KeepRelativeOffset, bool bStopWhenAttachedToDestroyed = false, float VolumeMultiplier = 1.f, float PitchMultiplier = 1.f, float StartTime = 0.f, USoundAttenuation* AttenuationSettings = nullptr, bool bAutoDestroy = true);
 
-	DEPRECATED(4.9, "PlayDialogueAttached has been renamed SpawnDialogueAttached")
-	static UAudioComponent* PlayDialogueAttached(UDialogueWave* Dialogue, const FDialogueContext& Context, USceneComponent* AttachToComponent, FName AttachPointName = NAME_None, FVector Location = FVector(ForceInit), EAttachLocation::Type LocationType = EAttachLocation::KeepRelativeOffset, bool bStopWhenAttachedToDestroyed = false, float VolumeMultiplier = 1.f, float PitchMultiplier = 1.f, float StartTime = 0.f, USoundAttenuation* AttenuationSettings = nullptr)
-	{
-		return SpawnDialogueAttached(Dialogue, Context, AttachToComponent, AttachPointName, Location, FRotator::ZeroRotator, LocationType, bStopWhenAttachedToDestroyed, VolumeMultiplier, PitchMultiplier, StartTime, AttenuationSettings);
-	}
-
 	static UAudioComponent* SpawnDialogueAttached(UDialogueWave* Dialogue, const FDialogueContext& Context, USceneComponent* AttachToComponent, FName AttachPointName, FVector Location, EAttachLocation::Type LocationType = EAttachLocation::KeepRelativeOffset, bool bStopWhenAttachedToDestroyed = false, float VolumeMultiplier = 1.f, float PitchMultiplier = 1.f, float StartTime = 0.f, USoundAttenuation* AttenuationSettings = nullptr, bool bAutoDestroy = true)
 	{
 		return SpawnDialogueAttached(Dialogue, Context, AttachToComponent, AttachPointName, Location, FRotator::ZeroRotator, LocationType, bStopWhenAttachedToDestroyed, VolumeMultiplier, PitchMultiplier, StartTime, AttenuationSettings, bAutoDestroy);
diff --git a/Engine/Source/Runtime/Engine/Classes/Kismet/KismetMathLibrary.h b/Engine/Source/Runtime/Engine/Classes/Kismet/KismetMathLibrary.h
--- a/Engine/Source/Runtime/Engine/Classes/Kismet/KismetMathLibrary.h
+++ b/Engine/Source/Runtime/Engine/Classes/Kismet/KismetMathLibrary.h
@@ -1489,12 +1489,6 @@ class ENGINE_API UKismetMathLibrary : public UBlueprintFunctionLibrary
 		UPARAM(DisplayName="Y (Pitch)") float Pitch,
 		UPARAM(DisplayName="Z (Yaw)") float Yaw);
 
-	DEPRECATED(4.9, "Method renamed to MakeRotator as the parameter order changed")
-	static inline FRotator MakeRot(float Pitch, float Yaw, float Roll)
-	{
-		return MakeRotator(Roll, Pitch, Yaw);
-	}
-
 	/** Find a rotation for an object at Start location to point at Target location. */
 	UFUNCTION(BlueprintPure, Category="Math|Rotator", meta=(Keywords="rotation rotate"))
 	static FRotator FindLookAtRotation(const FVector& Start, const FVector& Target);
@@ -1543,12 +1537,6 @@ class ENGINE_API UKismetMathLibrary : public UBlueprintFunctionLibrary
 		UPARAM(DisplayName="Y (Pitch)") float& Pitch,
 		UPARAM(DisplayName="Z (Yaw)") float& Yaw);
 
-	DEPRECATED(4.9, "Method renamed to BreakRotator as the parameter order changed")
-	static inline void BreakRot(FRotator InRot, float& Pitch, float& Yaw, float& Roll)
-	{
-		BreakRotator(InRot, Roll, Pitch, Yaw);
-	}
-
 	/** Breaks apart a rotator into its component axes */
 	UFUNCTION(BlueprintPure, Category="Math|Rotator", meta=(Keywords="rotation rotate rotator breakrotator"))
 	static void BreakRotIntoAxes(const FRotator& InRot, FVector& X, FVector& Y, FVector& Z);
diff --git a/Engine/Source/Runtime/Engine/Classes/PhysicalMaterials/PhysicalMaterial.h b/Engine/Source/Runtime/Engine/Classes/PhysicalMaterials/PhysicalMaterial.h
--- a/Engine/Source/Runtime/Engine/Classes/PhysicalMaterials/PhysicalMaterial.h
+++ b/Engine/Source/Runtime/Engine/Classes/PhysicalMaterials/PhysicalMaterial.h
@@ -52,7 +52,7 @@ class ENGINE_API UPhysicalMaterial : public UObject
 	// Surface properties.
 	//
 	
-	/** Friction value of surface, controls how easily things can slide on this surface */
+	/** Friction value of surface, controls how easily things can slide on this surface (0 is frictionless, higher values increase the amount of friction) */
 	UPROPERTY(EditAnywhere, BlueprintReadOnly, Category=PhysicalMaterial, meta=(ClampMin=0))
 	float Friction;
 
diff --git a/Engine/Source/Runtime/Engine/Classes/PhysicsEngine/BodyInstance.h b/Engine/Source/Runtime/Engine/Classes/PhysicsEngine/BodyInstance.h
--- a/Engine/Source/Runtime/Engine/Classes/PhysicsEngine/BodyInstance.h
+++ b/Engine/Source/Runtime/Engine/Classes/PhysicsEngine/BodyInstance.h
@@ -556,30 +556,11 @@ public:
 	/** If scene type is not specified return the first scene used. */
 	int32 GetSceneIndex(int32 SceneType = -1) const;
 
-
-	/** Initialise dynamic properties for this instance when using PhysX - this must be done after scene addition.
-	 */
-	DEPRECATED(4.8, "Please call InitDynamicProperties_AssumesLocked and make sure you obtain the appropriate PhysX scene locks")
-	void InitDynamicProperties()
-	{
-		InitDynamicProperties_AssumesLocked();
-	}
-
 	/** Initialise dynamic properties for this instance when using PhysX - this must be done after scene addition.
 	 *  Note: This function is not thread safe. Make sure to obtain the appropriate PhysX scene locks before calling this function
 	 */
 	void InitDynamicProperties_AssumesLocked();
 
-	/** Populate the filter data within the provided FShapeData with the correct filters for this instance
-	 *	@param ShapeData ShapeData to populate
-	 *	@param bForceSimpleAsComplex Whether to force simple colision as complex
-	 */
-	DEPRECATED(4.8, "Please call GetFilterData_AssumesLocked and make sure you obtain the appropriate PhysX scene locks")
-	void GetFilterData(FShapeData& ShapeData, bool bForceSimpleAsComplex = false)
-	{
-		GetFilterData_AssumesLocked(ShapeData, bForceSimpleAsComplex);
-	}
-
 	/** Populate the filter data within the provided FShapeData with the correct filters for this instance
 	 *	@param ShapeData ShapeData to populate. ShapeData.CollisionEnabled and ShapeData.FilterData will be filled in.
 	 *	@param bForceSimpleAsComplex Whether to force simple colision as complex
@@ -587,19 +568,7 @@ public:
 	 */
 	void GetFilterData_AssumesLocked(FShapeData& ShapeData, bool bForceSimpleAsComplex = false);
 
-	/** Set initialisation flags on a provided shape
-	 *	@param UseCollisionEnabled Whether collision is enabled for the shape
-	 *	@param PShape The shape to set the flags on
-	 *	@param SceneType Which scene we are using
-	 *	@param bUseComplexAsSimple Whether to use complex collision as simple
-	 */
-	DEPRECATED(4.8, "Please call SetShapeFlags_AssumesLocked and make sure you obtain the appropriate PhysX scene locks")
-	void SetShapeFlags(TEnumAsByte<ECollisionEnabled::Type> UseCollisionEnabled, physx::PxShape* PShape, EPhysicsSceneType SceneType, const bool bUseComplexAsSimple = false)
-	{
-		SetShapeFlags_AssumesLocked(UseCollisionEnabled, PShape, SceneType, bUseComplexAsSimple);
-	}
-
-	/** Set initialisation flags on a provided shape
+	/** Set initialization flags on a provided shape
 	 *	@param UseCollisionEnabled Whether collision is enabled for the shape
 	 *	@param PShape The shape to set the flags on
 	 *	@param SceneType Which scene we are using
@@ -608,18 +577,7 @@ public:
 	 */
 	void SetShapeFlags_AssumesLocked(TEnumAsByte<ECollisionEnabled::Type> UseCollisionEnabled, physx::PxShape* PShape, EPhysicsSceneType SceneType, const bool bUseComplexAsSimple = false);
 
-	/** Populate the flag fields of the provided FShapeData with correct initialisation flags
-	 *	@param ShapeData ShapeData to populate
-	 *	@param UseCollisionEnabled Whether collision is enabled for this instance
-	 *	@param bUseComplexAsSimple Whether to use complex collision as simple
-	 */
-	DEPRECATED(4.8, "Please call GetShapeFlags_AssumesLocked and make sure you obtain the appropriate PhysX scene locks")
-	void GetShapeFlags(FShapeData& ShapeData, TEnumAsByte<ECollisionEnabled::Type> UseCollisionEnabled, const bool bUseComplexAsSimple = false)
-	{
-		GetShapeFlags_AssumesLocked(ShapeData, UseCollisionEnabled, bUseComplexAsSimple);
-	}
-
-	/** Populate the flag fields of the provided FShapeData with correct initialisation flags
+	/** Populate the flag fields of the provided FShapeData with correct initialization flags
 	 *	@param ShapeData ShapeData to populate. ShapeData.FilterData will not be modified.
 	 *	@param UseCollisionEnabled Whether collision is enabled for this instance
 	 *	@param bUseComplexAsSimple Whether to use complex collision as simple
@@ -627,17 +585,6 @@ public:
 	 */
 	void GetShapeFlags_AssumesLocked(FShapeData& ShapeData, TEnumAsByte<ECollisionEnabled::Type> UseCollisionEnabled, const bool bUseComplexAsSimple = false);
 
-	/**
-	 * Return the PxRigidActor from the given scene (see EPhysicsSceneType), if SceneType is in the range [0, PST_MAX).
-	 * If SceneType < 0, the PST_Sync actor is returned if it is not NULL, otherwise the PST_Async actor is returned.
-	 * Invalid scene types will cause NULL to be returned.
-	 */
-	DEPRECATED(4.8, "Please call GetPxRigidActor_AssumesLocked and make sure you obtain the appropriate PhysX scene locks")
-	physx::PxRigidActor* GetPxRigidActor(int32 SceneType = -1) const
-	{
-		return GetPxRigidActorFromScene_AssumesLocked(SceneType);
-	}
-
 	/**
 	 * Return the PxRigidActor from the given scene (see EPhysicsSceneType), if SceneType is in the range [0, PST_MAX).
 	 * If SceneType < 0, the PST_Sync actor is returned if it is not NULL, otherwise the PST_Async actor is returned.
@@ -651,47 +598,18 @@ public:
 
 	physx::PxRigidActor* GetPxRigidActorFromScene_AssumesLocked(int32 SceneType) const;
 
-
-	/** Return the PxRigidDynamic if it exists in one of the scenes (NULL otherwise).  Currently a PxRigidDynamic can exist in only one of the two scenes. */
-	DEPRECATED(4.8, "Please call GetPxRigidDynamic_AssumesLocked and make sure you obtain the appropriate PhysX scene locks")
-	physx::PxRigidDynamic* GetPxRigidDynamic() const
-	{
-		return GetPxRigidDynamic_AssumesLocked();
-	}
-
 	/** Return the PxRigidDynamic if it exists in one of the scenes (NULL otherwise).  Currently a PxRigidDynamic can exist in only one of the two scenes.
 	 *  Note: This function assumes the appropriate scene lock has been obtained.
 	 *  @see ExecuteOnPxRigidDynamicReadOnly, ExecuteOnPxRigidDynamicReadWrite
 	 */
 	physx::PxRigidDynamic* GetPxRigidDynamic_AssumesLocked() const;
 
-	/** Return the PxRigidBody if it exists in one of the scenes (NULL otherwise).  Currently a PxRigidBody can exist in only one of the two scenes. */
-	DEPRECATED(4.8, "Please call GetPxRigidBody_AssumesLocked and make sure you obtain the appropriate PhysX scene locks")
-	physx::PxRigidBody* GetPxRigidBody() const
-	{
-		return GetPxRigidBody_AssumesLocked();
-	}
-
 	/** Return the PxRigidBody if it exists in one of the scenes (NULL otherwise).  Currently a PxRigidBody can exist in only one of the two scenes.
 	 *  Note: This function assumes the appropriate scene lock has been obtained.
 	 *  @see ExecuteOnPxRigidBodyReadOnly, ExecuteOnPxRigidBodyReadWrite
 	 */
 	physx::PxRigidBody* GetPxRigidBody_AssumesLocked() const;
 
-	
-
-	/** 
-	 *	Utility to get all the shapes from a FBodyInstance 
-	 *	Shapes belonging to sync actor are first, then async. Number of shapes belonging to sync actor is returned.
-	 */
-	DEPRECATED(4.8, "Please call GetAllShapes_AssumesLocked and make sure you obtain the appropriate PhysX scene locks")
-	TArray<physx::PxShape*> GetAllShapes(int32& OutNumSyncShapes) const
-	{
-		TArray<physx::PxShape*> Shapes;
-		OutNumSyncShapes = GetAllShapes_AssumesLocked(Shapes);
-		return Shapes;
-	}
-
 	/** 
 	 *	Utility to get all the shapes from a FBodyInstance 
 	 *	Shapes belonging to sync actor are first, then async. Number of shapes belonging to sync actor is returned.
@@ -767,10 +685,6 @@ public:
 	/** Sets a custom slope override struct for this instance. Implicitly sets bOverrideWalkableSlopeOnInstance to true. */
 	void SetWalkableSlopeOverride(const FWalkableSlopeOverride& NewOverride);
 
-	/** Returns whether this body wants (and can) use the async scene. */
-	DEPRECATED(4.8, "Please call UseAsyncScene and pass in the PhysScene of the owning component (GetWorld()->GetPhysicsScene()).")
-	bool UseAsyncScene() const;
-
 	bool UseAsyncScene(const FPhysScene* PhysScene) const;
 
 	bool HasSharedShapes() const{ return bHasSharedShapes; }
@@ -900,9 +814,6 @@ public:
 	 */
 	void SetBodyTransform(const FTransform& NewTransform, ETeleportType Teleport);
 
-	DEPRECATED(4.9, "Please pass the teleport flag using the new ETeleportType enum ")
-	void SetBodyTransform(const FTransform& NewTransform, bool bTeleport);
-
 	/** Get current velocity in world space from physics body. */
 	FVector GetUnrealWorldVelocity() const;
 
@@ -1004,12 +915,6 @@ public:
 	void UpdatePhysicalMaterials();
 
 #if WITH_PHYSX
-	DEPRECATED(4.8, "Please call ApplyMaterialToShape_AssumesLocked and make sure you obtain the appropriate PhysX scene locks")
-	static void ApplyMaterialToShape(physx::PxShape* PShape, physx::PxMaterial* PSimpleMat, TArray<UPhysicalMaterial*>& ComplexPhysMats)
-	{
-		ApplyMaterialToShape_AssumesLocked(PShape, PSimpleMat, ComplexPhysMats, false);
-	}
-
 	/** 
 	 *  Apply a material directly to the passed in shape. Note this function is very advanced and requires knowledge of shape sharing as well as threading. Note: assumes the appropriate locks have been obtained
 	 *  @param  PShape					The shape we are applying the material to
@@ -1019,12 +924,6 @@ public:
 	 */
 	static void ApplyMaterialToShape_AssumesLocked(physx::PxShape* PShape, physx::PxMaterial* PSimpleMat, const TArray<UPhysicalMaterial*>& ComplexPhysMats, const bool bSharedShape);
 
-	DEPRECATED(4.8, "Please call ApplyMaterialToInstanceShapes_AssumesLocked and make sure you obtain the appropriate PhysX scene locks")
-	void ApplyMaterialToInstanceShapes(physx::PxMaterial* PSimpleMat, TArray<UPhysicalMaterial*>& ComplexPhysMats)
-	{
-		ApplyMaterialToInstanceShapes_AssumesLocked(PSimpleMat, ComplexPhysMats);
-	}
-
 	/** Note: This function is not thread safe. Make sure you obtain the appropriate PhysX scene lock before calling it*/
 	void ApplyMaterialToInstanceShapes_AssumesLocked(physx::PxMaterial* PSimpleMat, TArray<UPhysicalMaterial*>& ComplexPhysMats);
 #endif
@@ -1061,20 +960,6 @@ public:
 	bool Sweep(struct FHitResult& OutHit, const FVector& Start, const FVector& End, const FQuat& ShapeWorldRotation, const FCollisionShape& Shape, bool bTraceComplex) const;
 
 #if WITH_PHYSX
-	/**
-	 *  Test if the bodyinstance overlaps with the geometry in the Pos/Rot
-	 *
-	 *	@param	PGeom			Geometry it would like to test
-	 *  @param  ShapePose       Transform information in world. Use U2PTransform to convert from FTransform
-	 *  @param  OutMTD			The minimum translation direction needed to push the shape out of this BodyInstance. (Optional)
-	 *  @return true if PrimComp overlaps this component at the specified location/rotation
-	 */
-	DEPRECATED(4.8, "Please call OverlapPhysX_AssumesLocked and make sure you obtain the appropriate PhysX scene locks")
-	bool OverlapPhysX(const physx::PxGeometry& Geom, const physx::PxTransform&  ShapePose, FMTDResult* OutMTD = nullptr) const
-	{
-		return OverlapPhysX_AssumesLocked(Geom, ShapePose, OutMTD);
-	}
-
 	/**
 	 *  Test if the bodyinstance overlaps with the geometry in the Pos/Rot
 	 *
diff --git a/Engine/Source/Runtime/Engine/Classes/PhysicsEngine/BodySetup.h b/Engine/Source/Runtime/Engine/Classes/PhysicsEngine/BodySetup.h
--- a/Engine/Source/Runtime/Engine/Classes/PhysicsEngine/BodySetup.h
+++ b/Engine/Source/Runtime/Engine/Classes/PhysicsEngine/BodySetup.h
@@ -394,15 +394,6 @@ public:
 #endif // WITH_EDITOR
 
 #if WITH_PHYSX
-	/** 
-	 *   Add the shapes defined by this body setup to the supplied PxRigidBody. 
-	 */
-	DEPRECATED(4.8, "Please call AddShapesToRigidActor_AssumesLocked and make sure you obtain the appropriate PhysX scene locks")
-	ENGINE_API void AddShapesToRigidActor(FBodyInstance* OwningInstance, physx::PxRigidActor* PDestActor, EPhysicsSceneType SceneType, FVector& Scale3D, physx::PxMaterial* SimpleMaterial, TArray<UPhysicalMaterial*>& ComplexMaterials, FShapeData& ShapeData, const FTransform& RelativeTM = FTransform::Identity, TArray<physx::PxShape*>* NewShapes = NULL)
-	{
-		AddShapesToRigidActor_AssumesLocked(OwningInstance, PDestActor, SceneType, Scale3D, SimpleMaterial, ComplexMaterials, ShapeData, RelativeTM, NewShapes);
-	}
-
 	/** 
 	 *   Add the shapes defined by this body setup to the supplied PxRigidBody. 
 	 */
diff --git a/Engine/Source/Runtime/Engine/Classes/PhysicsEngine/ConvexElem.h b/Engine/Source/Runtime/Engine/Classes/PhysicsEngine/ConvexElem.h
--- a/Engine/Source/Runtime/Engine/Classes/PhysicsEngine/ConvexElem.h
+++ b/Engine/Source/Runtime/Engine/Classes/PhysicsEngine/ConvexElem.h
@@ -68,9 +68,6 @@ public:
 		return *this;
 	}
 
-	DEPRECATED(4.8, "Please call DrawElemWire which takes in a Scale parameter")
-	ENGINE_API void	DrawElemWire(class FPrimitiveDrawInterface* PDI, const FTransform& ElemTM, const FColor Color) const;
-	
 	ENGINE_API void	DrawElemWire(class FPrimitiveDrawInterface* PDI, const FTransform& ElemTM, const float Scale, const FColor Color) const;
 
 	ENGINE_API void AddCachedSolidConvexGeom(TArray<FDynamicMeshVertex>& VertexBuffer, TArray<int32>& IndexBuffer, const FColor VertexColor) const;
diff --git a/Engine/Source/Runtime/Engine/Classes/Sound/SoundWaveProcedural.h b/Engine/Source/Runtime/Engine/Classes/Sound/SoundWaveProcedural.h
--- a/Engine/Source/Runtime/Engine/Classes/Sound/SoundWaveProcedural.h
+++ b/Engine/Source/Runtime/Engine/Classes/Sound/SoundWaveProcedural.h
@@ -17,12 +17,6 @@
 
 DECLARE_DELEGATE_TwoParams( FOnSoundWaveProceduralUnderflow, class USoundWaveProcedural*, int32 );
 
-DEPRECATED(4.9, "FOnSoundWaveStreamingUnderflow has been renamed FOnSoundWaveProceduralUnderflow")
-typedef FOnSoundWaveProceduralUnderflow FOnSoundWaveStreamingUnderflow;
-
-DEPRECATED(4.9, "USoundWaveStreaming has been renamed USoundWaveProcedural.")
-typedef class USoundWaveProcedural USoundWaveStreaming;
-
 UCLASS()
 class ENGINE_API USoundWaveProcedural : public USoundWave
 {
diff --git a/Engine/Source/Runtime/Engine/Private/AI/Navigation/NavigationData.cpp b/Engine/Source/Runtime/Engine/Private/AI/Navigation/NavigationData.cpp
--- a/Engine/Source/Runtime/Engine/Private/AI/Navigation/NavigationData.cpp
+++ b/Engine/Source/Runtime/Engine/Private/AI/Navigation/NavigationData.cpp
@@ -79,17 +79,6 @@ FPathFindingQuery::FPathFindingQuery(FNavPathSharedRef PathToRecalculate, const
 //----------------------------------------------------------------------//
 uint32 FAsyncPathFindingQuery::LastPathFindingUniqueID = INVALID_NAVQUERYID;
 
-FAsyncPathFindingQuery::FAsyncPathFindingQuery(const UObject* InOwner, const ANavigationData* InNavData, const FVector& Start, const FVector& End, const FNavPathQueryDelegate& Delegate, FSharedConstNavQueryFilter SourceQueryFilter)
-: FPathFindingQuery(InOwner, *InNavData, Start, End, SourceQueryFilter)
-, QueryID(GetUniqueID())
-, OnDoneDelegate(Delegate)
-{
-	if (InNavData == nullptr)
-	{
-		UE_LOG(LogNavigation, Error, TEXT("Trying to instantiate FAsyncPathFindingQuery while InNavData == null"));
-	}
-}
-
 FAsyncPathFindingQuery::FAsyncPathFindingQuery(const UObject* InOwner, const ANavigationData& InNavData, const FVector& Start, const FVector& End, const FNavPathQueryDelegate& Delegate, FSharedConstNavQueryFilter SourceQueryFilter)
 : FPathFindingQuery(InOwner, InNavData, Start, End, SourceQueryFilter)
 , QueryID(GetUniqueID())
diff --git a/Engine/Source/Runtime/Engine/Private/AI/Navigation/PImplRecastNavMesh.cpp b/Engine/Source/Runtime/Engine/Private/AI/Navigation/PImplRecastNavMesh.cpp
--- a/Engine/Source/Runtime/Engine/Private/AI/Navigation/PImplRecastNavMesh.cpp
+++ b/Engine/Source/Runtime/Engine/Private/AI/Navigation/PImplRecastNavMesh.cpp
@@ -1321,49 +1321,6 @@ FNavLocation FPImplRecastNavMesh::GetRandomPoint(const FNavigationQueryFilter& F
 	return OutLocation;
 }
 
-bool FPImplRecastNavMesh::GetRandomPointInRadius(const FVector& Origin, float Radius, FNavLocation& OutLocation, const FNavigationQueryFilter& Filter, const UObject* Owner) const
-{
-	if (DetourNavMesh == NULL || Radius <= 0.f)
-	{
-		return false;
-	}
-
-	FRecastSpeciaLinkFilter LinkFilter(UNavigationSystem::GetCurrent(NavMeshOwner->GetWorld()), Owner);
-	INITIALIZE_NAVQUERY(NavQuery, Filter.GetMaxSearchNodes(), LinkFilter);
-
-	// inits to "pass all"
-	const dtQueryFilter* QueryFilter = ((const FRecastQueryFilter*)(Filter.GetImplementation()))->GetAsDetourQueryFilter();
-	ensure(QueryFilter);
-	if (QueryFilter)
-	{
-		// find starting poly
-		// convert start/end pos to Recast coords
-		const FVector NavExtent = NavMeshOwner->GetModifiedQueryExtent(FVector(Radius, Radius, Radius));
-		const float Extent[3] = { NavExtent.X, NavExtent.Z, NavExtent.Y };
-		float RecastOrigin[3];
-		Unr2RecastVector(Origin, RecastOrigin);
-		NavNodeRef OriginPolyID = INVALID_NAVNODEREF;
-		NavQuery.findNearestPoly(RecastOrigin, Extent, QueryFilter, &OriginPolyID, NULL);
-
-		dtPolyRef Poly;
-		float RandPt[3];
-		dtStatus Status = NavQuery.findRandomPointAroundCircle(OriginPolyID, RecastOrigin, Radius
-			, QueryFilter, FMath::FRand, &Poly, RandPt);
-
-		if (dtStatusSucceed(Status))
-		{
-			OutLocation = FNavLocation(Recast2UnrVector(RandPt), Poly);
-			return true;
-		}
-		else
-		{
-			OutLocation = FNavLocation(Origin, OriginPolyID);
-		}
-	}
-
-	return false;
-}
-
 bool FPImplRecastNavMesh::GetRandomPointInCluster(NavNodeRef ClusterRef, FNavLocation& OutLocation) const
 {
 	if (DetourNavMesh == NULL || ClusterRef == 0)
diff --git a/Engine/Source/Runtime/Engine/Private/AI/Navigation/RecastNavMeshDataChunk.cpp b/Engine/Source/Runtime/Engine/Private/AI/Navigation/RecastNavMeshDataChunk.cpp
--- a/Engine/Source/Runtime/Engine/Private/AI/Navigation/RecastNavMeshDataChunk.cpp
+++ b/Engine/Source/Runtime/Engine/Private/AI/Navigation/RecastNavMeshDataChunk.cpp
@@ -352,26 +352,3 @@ void URecastNavMeshDataChunk::GatherTiles(const FPImplRecastNavMesh* NavMeshImpl
 	}
 }
 #endif// WITH_EDITOR
-
-//----------------------------------------------------------------------//
-// DEPRECATED
-//----------------------------------------------------------------------//
-TArray<uint32> URecastNavMeshDataChunk::AttachTiles(FPImplRecastNavMesh* NavMeshImpl)
-{
-	if (NavMeshImpl)
-	{
-		return AttachTiles(*NavMeshImpl);
-	}
-
-	return TArray<uint32>();
-}
-
-TArray<uint32> URecastNavMeshDataChunk::DetachTiles(FPImplRecastNavMesh* NavMeshImpl)
-{
-	if (NavMeshImpl)
-	{
-		return DetachTiles(*NavMeshImpl);
-	}
-
-	return TArray<uint32>();
-}
diff --git a/Engine/Source/Runtime/Engine/Private/Actor.cpp b/Engine/Source/Runtime/Engine/Private/Actor.cpp
--- a/Engine/Source/Runtime/Engine/Private/Actor.cpp
+++ b/Engine/Source/Runtime/Engine/Private/Actor.cpp
@@ -420,11 +420,6 @@ bool AActor::K2_TeleportTo( FVector DestLocation, FRotator DestRotation )
 	return TeleportTo(DestLocation, DestRotation, false, false);
 }
 
-void AActor::SetTickPrerequisite(AActor* PrerequisiteActor)
-{
-	AddTickPrerequisiteActor(PrerequisiteActor);
-}
-
 void AActor::AddTickPrerequisiteActor(AActor* PrerequisiteActor)
 {
 	if (PrimaryActorTick.bCanEverTick && PrerequisiteActor && PrerequisiteActor->PrimaryActorTick.bCanEverTick)
@@ -2019,11 +2014,6 @@ FVector AActor::GetPlacementExtent() const
 	return Extent;
 }
 
-class UClass* AActor::GetActorClass() const
-{
-	return GetClass();
-}
-
 void AActor::Destroyed()
 {
 	RouteEndPlay(EEndPlayReason::Destroyed);
@@ -3607,11 +3597,6 @@ void AActor::K2_DestroyComponent(UActorComponent* Component)
 	}
 }
 
-UPrimitiveComponent* AActor::GetRootPrimitiveComponent() const
-{ 
-	return Cast<class UPrimitiveComponent>(RootComponent); 
-}
-
 bool AActor::SetRootComponent(class USceneComponent* NewRootComponent)
 {
 	/** Only components owned by this actor can be used as a its root component. */
diff --git a/Engine/Source/Runtime/Engine/Private/ActorReplication.cpp b/Engine/Source/Runtime/Engine/Private/ActorReplication.cpp
--- a/Engine/Source/Runtime/Engine/Private/ActorReplication.cpp
+++ b/Engine/Source/Runtime/Engine/Private/ActorReplication.cpp
@@ -27,24 +27,8 @@ static bool		SavedbHidden;
 static AActor*	SavedOwner;
 static bool		SavedbRepPhysics;
 
-#define DEPRECATED_NET_PRIORITY -17.0f // Pick an arbitrary invalid priority, check for that
-
-float AActor::GetNetPriority(const FVector& ViewPos, const FVector& ViewDir, APlayerController* Viewer, UActorChannel* InChannel, float Time, bool bLowBandwidth)
-{
-	return DEPRECATED_NET_PRIORITY;
-}
-
 float AActor::GetNetPriority(const FVector& ViewPos, const FVector& ViewDir, AActor* Viewer, AActor* ViewTarget, UActorChannel* InChannel, float Time, bool bLowBandwidth)
 {
-	PRAGMA_DISABLE_DEPRECATION_WARNINGS
-	// Call the deprecated version so subclasses will still work for a version
-	float DeprecatedValue = AActor::GetNetPriority(ViewPos, ViewDir, Cast<APlayerController>(Viewer), InChannel, Time, bLowBandwidth);
-	PRAGMA_ENABLE_DEPRECATION_WARNINGS
-	if (DeprecatedValue != DEPRECATED_NET_PRIORITY)
-	{
-		return DeprecatedValue;
-	}
-
 	if (bNetUseOwnerRelevancy && Owner)
 	{
 		// If we should use our owner's priority, pass it through
@@ -135,11 +119,6 @@ bool AActor::GetNetDormancy(const FVector& ViewPos, const FVector& ViewDir, AAct
 	return false;
 }
 
-bool AActor::GetNetDormancy(const FVector& ViewPos, const FVector& ViewDir, APlayerController* Viewer, AActor* ViewTarget, UActorChannel* InChannel, float Time, bool bLowBandwidth)
-{
-	return GetNetDormancy(ViewPos, ViewDir, Cast<AActor>(Viewer), ViewTarget, InChannel, Time, bLowBandwidth);
-}
-
 void AActor::PreNetReceive()
 {
 	SavedbHidden = bHidden;
@@ -256,11 +235,6 @@ void AActor::SyncReplicatedPhysicsSimulation()
 	}
 }
 
-bool AActor::IsNetRelevantFor(const APlayerController* RealViewer, const AActor* ViewTarget, const FVector& SrcLocation) const
-{
-	return IsNetRelevantFor(Cast<AActor>(RealViewer), ViewTarget, SrcLocation);
-}
-
 bool AActor::IsNetRelevantFor(const AActor* RealViewer, const AActor* ViewTarget, const FVector& SrcLocation) const
 {
 	if (bAlwaysRelevant || IsOwnedBy(ViewTarget) || IsOwnedBy(RealViewer) || this == ViewTarget || ViewTarget == Instigator)
diff --git a/Engine/Source/Runtime/Engine/Private/AssetManager.cpp b/Engine/Source/Runtime/Engine/Private/AssetManager.cpp
--- a/Engine/Source/Runtime/Engine/Private/AssetManager.cpp
+++ b/Engine/Source/Runtime/Engine/Private/AssetManager.cpp
@@ -13,6 +13,7 @@
 #include "Misc/ScopedSlowTask.h"
 #include "Misc/Paths.h"
 #include "AssetRegistryState.h"
+#include "HAL/PlatformFilemanager.h"
 
 #if WITH_EDITOR
 #include "Editor.h"
@@ -102,6 +103,8 @@ UAssetManager::UAssetManager()
 	bIsGlobalAsyncScanEnvironment = false;
 	bShouldGuessTypeAndName = false;
 	bShouldUseSynchronousLoad = false;
+	bIsLoadingFromPakFiles = false;
+	bShouldAcquireMissingChunksOnLoad = false;
 	bIsBulkScanning = false;
 	bIsManagementDatabaseCurrent = false;
 	bUpdateManagementDatabaseAfterScan = false;
@@ -133,11 +136,17 @@ void UAssetManager::PostInitProperties()
 		FEditorDelegates::EndPIE.AddUObject(this, &UAssetManager::EndPIE);
 
 		// In editor builds guess the type/name if allowed
-		bShouldGuessTypeAndName = GetSettings().bShouldGuessTypeAndNameInEditor;
+		bShouldGuessTypeAndName = Settings.bShouldGuessTypeAndNameInEditor;
+		bOnlyCookProductionAssets = Settings.bOnlyCookProductionAssets;
 #else 
 		// Never guess type in cooked builds
 		bShouldGuessTypeAndName = false;
+
+		// Only cooked builds supoprt pak files and chunk download
+		bIsLoadingFromPakFiles = FPlatformFileManager::Get().FindPlatformFile(TEXT("PakFile")) != nullptr;
+		bShouldAcquireMissingChunksOnLoad = Settings.bShouldAcquireMissingChunksOnLoad;
 #endif
+		
 		bShouldUseSynchronousLoad = IsRunningCommandlet();
 
 		LoadRedirectorMaps();
@@ -931,7 +940,7 @@ void UAssetManager::GetPrimaryAssetTypeInfoList(TArray<FPrimaryAssetTypeInfo>& A
 
 TSharedPtr<FStreamableHandle> UAssetManager::ChangeBundleStateForPrimaryAssets(const TArray<FPrimaryAssetId>& AssetsToChange, const TArray<FName>& AddBundles, const TArray<FName>& RemoveBundles, bool bRemoveAllBundles, FStreamableDelegate DelegateToCall, TAsyncLoadPriority Priority)
 {
-	TArray<TSharedPtr<FStreamableHandle> > NewHandles;
+	TArray<TSharedPtr<FStreamableHandle> > NewHandles, ExistingHandles;
 	TArray<FPrimaryAssetId> NewAssets;
 	TSharedPtr<FStreamableHandle> ReturnHandle;
 
@@ -970,6 +979,8 @@ TSharedPtr<FStreamableHandle> UAssetManager::ChangeBundleStateForPrimaryAssets(c
 			{
 				if (NameData->PendingState.BundleNames == NewBundleState)
 				{
+					// This will wait on any existing handles to finish
+					ExistingHandles.Add(NameData->PendingState.Handle);
 					continue;
 				}
 				// Clear pending state
@@ -1026,15 +1037,7 @@ TSharedPtr<FStreamableHandle> UAssetManager::ChangeBundleStateForPrimaryAssets(c
 				DebugName += TEXT(")");
 			}
 
-			// Create the per asset handle
-			if (bShouldUseSynchronousLoad)
-			{
-				NewHandle = StreamableManager.RequestSyncLoad(PathsToLoad.Array(), false, DebugName);
-			}
-			else
-			{
-				NewHandle = StreamableManager.RequestAsyncLoad(PathsToLoad.Array(), FStreamableDelegate(), Priority, false, DebugName);
-			}
+			NewHandle = LoadAssetList(PathsToLoad.Array(), FStreamableDelegate(), Priority, DebugName);
 
 			if (!ensureMsgf(NewHandle.IsValid(), TEXT("Requested load of Primary Asset with no referenced assets!")))
 			{
@@ -1061,9 +1064,11 @@ TSharedPtr<FStreamableHandle> UAssetManager::ChangeBundleStateForPrimaryAssets(c
 		}
 	}
 
-	if (NewHandles.Num() > 1)
+	if (NewHandles.Num() > 1 || ExistingHandles.Num() > 0)
 	{
-		// If multiple handles, need to make wrapper handle
+		// If multiple handles or we have an old handle, need to make wrapper handle
+		NewHandles.Append(ExistingHandles);
+
 		ReturnHandle = StreamableManager.CreateCombinedHandle(NewHandles, FString::Printf(TEXT("%s CreateCombinedHandle"), *GetName()));
 
 		// Call delegate or bind to meta handle
@@ -1172,15 +1177,7 @@ TSharedPtr<FStreamableHandle> UAssetManager::PreloadPrimaryAssets(const TArray<F
 		}
 	}
 
-	if (bShouldUseSynchronousLoad)
-	{
-		ReturnHandle = StreamableManager.RequestSyncLoad(PathsToLoad.Array(), false, DebugName);
-		FStreamableHandle::ExecuteDelegate(DelegateToCall);
-	}
-	else
-	{
-		ReturnHandle = StreamableManager.RequestAsyncLoad(PathsToLoad.Array(), DelegateToCall, Priority, false, DebugName);
-	}
+	ReturnHandle = LoadAssetList(PathsToLoad.Array(), DelegateToCall, Priority, DebugName);
 
 	if (!ensureMsgf(ReturnHandle.IsValid(), TEXT("Requested preload of Primary Asset with no referenced assets!")))
 	{
@@ -1371,9 +1368,39 @@ int32 UAssetManager::UnloadPrimaryAssetsWithType(FPrimaryAssetType PrimaryAssetT
 	return UnloadPrimaryAssets(Assets);
 }
 
-TSharedPtr<FStreamableHandle> UAssetManager::LoadAssetList(const TArray<FStringAssetReference>& AssetList)
+TSharedPtr<FStreamableHandle> UAssetManager::LoadAssetList(const TArray<FStringAssetReference>& AssetList, FStreamableDelegate DelegateToCall, TAsyncLoadPriority Priority, const FString& DebugName)
 {
-	return GetStreamableManager().RequestAsyncLoad(AssetList, FStreamableDelegate(), FStreamableManager::DefaultAsyncLoadPriority, false);
+	TSharedPtr<FStreamableHandle> NewHandle;
+	TArray<int32> MissingChunks, ErrorChunks;
+
+	if (bShouldAcquireMissingChunksOnLoad)
+	{
+		FindMissingChunkList(AssetList, MissingChunks, ErrorChunks);
+
+		if (ErrorChunks.Num() > 0)
+		{
+			// At least one chunk doesn't exist, fail
+			UE_LOG(LogAssetManager, Error, TEXT("Failure loading %s, Required chunk %d does not exist!"), *DebugName, ErrorChunks[0]);
+			return nullptr;
+		}
+	}
+
+	// SynchronousLoad doesn't make sense if chunks are missing
+	if (bShouldUseSynchronousLoad && MissingChunks.Num() > 0)
+	{
+		NewHandle = StreamableManager.RequestSyncLoad(AssetList, false, DebugName);
+	}
+	else
+	{
+		NewHandle = StreamableManager.RequestAsyncLoad(AssetList, DelegateToCall, Priority, false, MissingChunks.Num() > 0, DebugName);
+
+		if (MissingChunks.Num() > 0 && NewHandle.IsValid())
+		{
+			AcquireChunkList(MissingChunks, FAssetManagerAcquireResourceDelegate(), EChunkPriority::Immediate, NewHandle);
+		}
+	}
+
+	return NewHandle;
 }
 
 FAssetBundleEntry UAssetManager::GetAssetBundleEntry(const FPrimaryAssetId& BundleScope, FName BundleName) const
@@ -1411,6 +1438,242 @@ bool UAssetManager::GetAssetBundleEntries(const FPrimaryAssetId& BundleScope, TA
 	return bFoundAny;
 }
 
+bool UAssetManager::FindMissingChunkList(const TArray<FStringAssetReference>& AssetList, TArray<int32>& OutMissingChunkList, TArray<int32>& OutErrorChunkList) const
+{
+	IPlatformChunkInstall* ChunkInstall = FPlatformMisc::GetPlatformChunkInstall();
+
+	if (!ChunkInstall || !bIsLoadingFromPakFiles)
+	{
+		return false;
+	}
+
+	for (const FStringAssetReference& Asset : AssetList)
+	{
+		FAssetData FoundData;
+		GetAssetDataForPath(Asset, FoundData);
+		
+		TSet<int32> FoundChunks, MissingChunks, ErrorChunks;
+
+		for (int32 ChunkId : FoundData.ChunkIDs)
+		{
+			EChunkLocation::Type ChunkLocation = ChunkInstall->GetChunkLocation(ChunkId);
+
+			switch (ChunkLocation)
+			{			
+			case EChunkLocation::DoesNotExist:
+				ErrorChunks.Add(ChunkId);
+				break;
+			case EChunkLocation::NotAvailable:
+				MissingChunks.Add(ChunkId);
+				break;
+			case EChunkLocation::LocalSlow:
+			case EChunkLocation::LocalFast:
+				FoundChunks.Add(ChunkId);
+				break;
+			}
+		}
+
+		// Assets may be redundantly in multiple chunks, if we have any of the chunks then we have the asset
+		if (FoundChunks.Num() == 0)
+		{
+			if (MissingChunks.Num() > 0)
+			{
+				int32 MissingChunkToAdd = -1;
+
+				for (int32 MissingChunkId : MissingChunks)
+				{
+					if (OutMissingChunkList.Contains(MissingChunkId))
+					{
+						// This chunk is already scheduled, don't add a new one
+						MissingChunkToAdd = -1;
+						break;
+					}
+					else if (MissingChunkToAdd == -1)
+					{
+						// Add the first mentioned missing chunk
+						MissingChunkToAdd = MissingChunkId;
+					}
+				}
+
+				if (MissingChunkToAdd != -1)
+				{
+					OutMissingChunkList.Add(MissingChunkToAdd);
+				}
+			}
+			else if (ErrorChunks.Num() > 0)
+			{
+				// Only have error chunks, report the errors
+				for (int32 ErrorChunkId : ErrorChunks)
+				{
+					OutErrorChunkList.Add(ErrorChunkId);
+				}
+			}
+		}
+	}
+
+	return OutMissingChunkList.Num() > 0 || OutErrorChunkList.Num() > 0;
+}
+
+void UAssetManager::AcquireChunkList(const TArray<int32>& ChunkList, FAssetManagerAcquireResourceDelegate CompleteDelegate, EChunkPriority::Type Priority, TSharedPtr<FStreamableHandle> StalledHandle)
+{
+	FPendingChunkInstall* PendingChunkInstall = new(PendingChunkInstalls) FPendingChunkInstall;
+	PendingChunkInstall->ManualCallback = CompleteDelegate;
+	PendingChunkInstall->RequestedChunks = ChunkList;
+	PendingChunkInstall->PendingChunks = ChunkList;
+	PendingChunkInstall->StalledStreamableHandle = StalledHandle;
+
+	IPlatformChunkInstall* ChunkInstall = FPlatformMisc::GetPlatformChunkInstall();
+
+	if (!ChunkInstallDelegateHandle.IsValid())
+	{
+		ChunkInstallDelegateHandle = ChunkInstall->AddChunkInstallDelegate(FPlatformChunkInstallDelegate::CreateUObject(this, &ThisClass::OnChunkDownloaded));
+	}
+
+	for (int32 MissingChunk : PendingChunkInstall->PendingChunks)
+	{
+		ChunkInstall->PrioritizeChunk(MissingChunk, Priority);
+	}
+}
+
+void UAssetManager::AcquireResourcesForAssetList(const TArray<FStringAssetReference>& AssetList, FAssetManagerAcquireResourceDelegate CompleteDelegate, EChunkPriority::Type Priority)
+{
+	TArray<int32> MissingChunks, ErrorChunks;
+
+	FindMissingChunkList(AssetList, MissingChunks, ErrorChunks);
+
+	if (ErrorChunks.Num() > 0)
+	{
+		// At least one chunk doesn't exist, fail
+		FStreamableDelegate TempDelegate = FStreamableDelegate::CreateLambda([CompleteDelegate]() {CompleteDelegate.ExecuteIfBound(false); });
+		FStreamableHandle::ExecuteDelegate(TempDelegate);
+
+		return;
+	}
+
+	if (MissingChunks.Num() == 0)
+	{
+		// All here, schedule the callback
+		FStreamableDelegate TempDelegate = FStreamableDelegate::CreateLambda([CompleteDelegate]() {CompleteDelegate.ExecuteIfBound(true); });
+		FStreamableHandle::ExecuteDelegate(TempDelegate);
+
+		return;
+	}
+	
+	AcquireChunkList(MissingChunks, CompleteDelegate, Priority, nullptr);
+}
+
+void UAssetManager::AcquireResourcesForPrimaryAssetList(const TArray<FPrimaryAssetId>& PrimaryAssetList, FAssetManagerAcquireResourceDelegate CompleteDelegate, EChunkPriority::Type Priority)
+{
+	TSet<FStringAssetReference> PathsToLoad;
+	TSharedPtr<FStreamableHandle> ReturnHandle;
+
+	for (const FPrimaryAssetId& PrimaryAssetId : PrimaryAssetList)
+	{
+		FPrimaryAssetData* NameData = GetNameData(PrimaryAssetId);
+
+		if (NameData)
+		{
+			// Gather asset refs
+			const FStringAssetReference& AssetPath = NameData->AssetPtr.ToStringReference();
+
+			if (!AssetPath.IsNull())
+			{
+				// Dynamic types can have no base asset path
+				PathsToLoad.Add(AssetPath);
+			}
+
+			TArray<FAssetBundleEntry> BundleEntries;
+			GetAssetBundleEntries(PrimaryAssetId, BundleEntries);
+			for (const FAssetBundleEntry& Entry : BundleEntries)
+			{
+				if (Entry.IsValid())
+				{
+					PathsToLoad.Append(Entry.BundleAssets);
+				}
+			}
+		}
+	}
+
+	AcquireResourcesForAssetList(PathsToLoad.Array(), CompleteDelegate, Priority);
+}
+
+bool UAssetManager::GetResourceAcquireProgress(int32& OutAcquiredCount, int32& OutRequestedCount) const
+{
+	OutAcquiredCount = OutRequestedCount = 0;
+	// Iterate pending callbacks, in order they were added
+	for (const FPendingChunkInstall& PendingChunkInstall : PendingChunkInstalls)
+	{
+		OutRequestedCount += PendingChunkInstall.RequestedChunks.Num();
+		OutAcquiredCount += (PendingChunkInstall.RequestedChunks.Num() - PendingChunkInstall.PendingChunks.Num());
+	}
+
+	return PendingChunkInstalls.Num() > 0;
+}
+
+void UAssetManager::OnChunkDownloaded(uint32 ChunkId, bool bSuccess)
+{
+	IPlatformChunkInstall* ChunkInstall = FPlatformMisc::GetPlatformChunkInstall();
+
+	// Iterate pending callbacks, in order they were added
+	for (int32 i = 0; i < PendingChunkInstalls.Num(); i++)
+	{
+		// Make a copy so if we resize the array it's safe
+		FPendingChunkInstall PendingChunkInstall = PendingChunkInstalls[i];
+		if (PendingChunkInstall.PendingChunks.Contains(ChunkId))
+		{
+			bool bFailed = !bSuccess;
+			TArray<int32> NewPendingList;
+			
+			// Check all chunks if they are done or failed
+			for (int32 PendingChunkId : PendingChunkInstall.PendingChunks)
+			{
+				EChunkLocation::Type ChunkLocation = ChunkInstall->GetChunkLocation(PendingChunkId);
+
+				switch (ChunkLocation)
+				{
+				case EChunkLocation::DoesNotExist:
+					bFailed = true;
+					break;
+				case EChunkLocation::NotAvailable:
+					NewPendingList.Add(PendingChunkId);
+					break;
+				}
+			}
+
+			if (bFailed)
+			{
+				// Resize array first
+				PendingChunkInstalls.RemoveAt(i);
+				i--;
+
+				if (PendingChunkInstall.StalledStreamableHandle.IsValid())
+				{
+					PendingChunkInstall.StalledStreamableHandle->CancelHandle();
+				}
+
+				PendingChunkInstall.ManualCallback.ExecuteIfBound(false);
+			}
+			else if (NewPendingList.Num() == 0)
+			{
+				// Resize array first
+				PendingChunkInstalls.RemoveAt(i);
+				i--;
+
+				if (PendingChunkInstall.StalledStreamableHandle.IsValid())
+				{
+					PendingChunkInstall.StalledStreamableHandle->StartStalledHandle();
+				}
+
+				PendingChunkInstall.ManualCallback.ExecuteIfBound(true);
+			}
+			else
+			{
+				PendingChunkInstalls[i].PendingChunks = NewPendingList;
+			}
+		}
+	}
+}
+
 FPrimaryAssetData* UAssetManager::GetNameData(const FPrimaryAssetId& PrimaryAssetId, bool bCheckRedirector)
 {
 	return const_cast<FPrimaryAssetData*>(const_cast<const UAssetManager*>(this)->GetNameData(PrimaryAssetId));
@@ -1682,7 +1945,7 @@ bool UAssetManager::GetAssetDataForPath(const FStringAssetReference& ObjectPath,
 
 FStringAssetReference UAssetManager::GetAssetPathForData(const FAssetData& AssetData) const
 {
-	FString AssetPath = AssetData.ObjectPath.ToString();
+	FString AssetPath = AssetData.IsValid() ? AssetData.ObjectPath.ToString() : FString();
 
 	// All blueprint types end with blueprint as the class, there is no better test unfortunately
 	if (AssetData.AssetClass.ToString().EndsWith(TEXT("Blueprint")))
@@ -1698,12 +1961,14 @@ void UAssetManager::GetAssetDataForPathInternal(IAssetRegistry& AssetRegistry, c
 	// We're a class if our path is foo.foo_C
 	bool bIsClass = AssetPath.EndsWith(TEXT("_C"), ESearchCase::CaseSensitive) && !AssetPath.Contains(TEXT("_C."), ESearchCase::CaseSensitive);
 
+	const bool bIncludeOnlyOnDiskAssets = !GIsEditor;
+
 	// If we're a class, first look for the asset data without the trailing _C
 	// We do this first because in cooked builds you have to search the asset registry for the Blueprint, not the class itself
 	if (bIsClass)
 	{
 		// We need to strip the class suffix because the asset registry has it listed by blueprint name
-		const bool bIncludeOnlyOnDiskAssets = !GIsEditor;
+		
 		OutAssetData = AssetRegistry.GetAssetByObjectPath(FName(*AssetPath.LeftChop(2)), bIncludeOnlyOnDiskAssets);
 
 		if (OutAssetData.IsValid())
@@ -1712,7 +1977,7 @@ void UAssetManager::GetAssetDataForPathInternal(IAssetRegistry& AssetRegistry, c
 		}
 	}
 
-	OutAssetData = AssetRegistry.GetAssetByObjectPath(FName(*AssetPath));
+	OutAssetData = AssetRegistry.GetAssetByObjectPath(FName(*AssetPath), bIncludeOnlyOnDiskAssets);
 }
 
 bool UAssetManager::WriteCustomReport(FString FileName, TArray<FString>& FileLines) const
@@ -2465,8 +2730,6 @@ EPrimaryAssetCookRule UAssetManager::GetPackageCookRule(FName PackageName) const
 
 bool UAssetManager::VerifyCanCookPackage(FName PackageName, bool bLogError) const
 {
-	const UAssetManagerSettings& Settings = GetSettings();
-
 	EPrimaryAssetCookRule CookRule = UAssetManager::Get().GetPackageCookRule(PackageName);
 	if (CookRule == EPrimaryAssetCookRule::NeverCook)
 	{
@@ -2477,7 +2740,7 @@ bool UAssetManager::VerifyCanCookPackage(FName PackageName, bool bLogError) cons
 		
 		return false;
 	}
-	else if (CookRule == EPrimaryAssetCookRule::DevelopmentCook && Settings.bOnlyCookProductionAssets)
+	else if (CookRule == EPrimaryAssetCookRule::DevelopmentCook && bOnlyCookProductionAssets)
 	{
 		if (bLogError)
 		{
@@ -2610,13 +2873,19 @@ void UAssetManager::RefreshPrimaryAssetDirectory()
 
 void UAssetManager::ReinitializeFromConfig()
 {
+	// We specifically do not reset AssetRuleOverrides as those can be set by something other than inis
 	AssetPathMap.Reset();
-	AssetRuleOverrides.Reset();
 	ManagementParentMap.Reset();
 	CachedAssetBundles.Reset();
 	AlreadyScannedDirectories.Reset();
 	AssetTypeMap.Reset();
 
+	// This code is editor only, so reinitialize globals
+	const UAssetManagerSettings& Settings = GetSettings();
+	bShouldGuessTypeAndName = Settings.bShouldGuessTypeAndNameInEditor;
+	bShouldAcquireMissingChunksOnLoad = Settings.bShouldAcquireMissingChunksOnLoad;
+	bOnlyCookProductionAssets = Settings.bOnlyCookProductionAssets;
+
 	LoadRedirectorMaps();
 	ScanPrimaryAssetTypesFromConfig();
 }
@@ -2645,7 +2914,7 @@ void UAssetManager::OnInMemoryAssetCreated(UObject *Object)
 
 			GetAssetDataForPathInternal(AssetRegistry, Object->GetPathName(), NewAssetData);
 
-			if (ensure(NewAssetData.IsValid()))
+			if (NewAssetData.IsValid())
 			{
 				// Make sure it's in a valid path
 				bool bFoundPath = false;
diff --git a/Engine/Source/Runtime/Engine/Private/BlueprintMapLibrary.cpp b/Engine/Source/Runtime/Engine/Private/BlueprintMapLibrary.cpp
--- a/Engine/Source/Runtime/Engine/Private/BlueprintMapLibrary.cpp
+++ b/Engine/Source/Runtime/Engine/Private/BlueprintMapLibrary.cpp
@@ -22,17 +22,24 @@ bool UBlueprintMapLibrary::GenericMap_Remove(const void* TargetMap, const UMapPr
 	return false;
 }
 
-bool UBlueprintMapLibrary::GenericMap_Find(const void* TargetMap, const UMapProperty* MapProperty, const void* KeyPtr, void* ValuePtr)
+bool UBlueprintMapLibrary::GenericMap_Find(const void* TargetMap, const UMapProperty* MapProperty, const void* KeyPtr, void* OutValuePtr)
 {
 	if(TargetMap)
 	{
 		FScriptMapHelper MapHelper(MapProperty, TargetMap);
-		uint8* Result = MapHelper.FindValueFromHash(KeyPtr);
-		if(Result && ValuePtr)
+		uint8* FoundValuePtr = MapHelper.FindValueFromHash(KeyPtr);
+		if (OutValuePtr)
 		{
-			MapProperty->ValueProp->CopyCompleteValueFromScriptVM(ValuePtr, Result);
+			if (FoundValuePtr)
+			{
+				MapProperty->ValueProp->CopyCompleteValueFromScriptVM(OutValuePtr, FoundValuePtr);
+			}
+			else
+			{
+				MapProperty->ValueProp->InitializeValue(OutValuePtr);
+			}
 		}
-		return Result != nullptr;
+		return FoundValuePtr != nullptr;
 	}
 	return false;
 }
@@ -102,3 +109,14 @@ void UBlueprintMapLibrary::GenericMap_Clear(const void* TargetMap, const UMapPro
 	}
 }
 
+void UBlueprintMapLibrary::GenericMap_SetMapPropertyByName(UObject* OwnerObject, FName MapPropertyName, const void* SrcMapAddr)
+{
+	if (OwnerObject)
+	{
+		if (UMapProperty* MapProp = FindField<UMapProperty>(OwnerObject->GetClass(), MapPropertyName))
+		{
+			void* Dest = MapProp->ContainerPtrToValuePtr<void>(OwnerObject);
+			MapProp->CopyValuesInternal(Dest, SrcMapAddr, 1);
+		}
+	}
+}
diff --git a/Engine/Source/Runtime/Engine/Private/BlueprintSetLibrary.cpp b/Engine/Source/Runtime/Engine/Private/BlueprintSetLibrary.cpp
--- a/Engine/Source/Runtime/Engine/Private/BlueprintSetLibrary.cpp
+++ b/Engine/Source/Runtime/Engine/Private/BlueprintSetLibrary.cpp
@@ -186,3 +186,14 @@ void UBlueprintSetLibrary::GenericSet_Difference(const void* SetA, const USetPro
 	}
 }
 
+void UBlueprintSetLibrary::GenericSet_SetSetPropertyByName(UObject* OwnerObject, FName SetPropertyName, const void* SrcSetAddr)
+{
+	if (OwnerObject)
+	{
+		if (USetProperty* SetProp = FindField<USetProperty>(OwnerObject->GetClass(), SetPropertyName))
+		{
+			void* Dest = SetProp->ContainerPtrToValuePtr<void>(OwnerObject);
+			SetProp->CopyValuesInternal(Dest, SrcSetAddr, 1);
+		}
+	}
+}
diff --git a/Engine/Source/Runtime/Engine/Private/Brush.cpp b/Engine/Source/Runtime/Engine/Private/Brush.cpp
--- a/Engine/Source/Runtime/Engine/Private/Brush.cpp
+++ b/Engine/Source/Runtime/Engine/Private/Brush.cpp
@@ -168,16 +168,6 @@ void ABrush::SetIsTemporarilyHiddenInEditor( bool bIsHidden )
 	}
 }
 
-FVector ABrush::GetPrePivot() const
-{
-	return GetPivotOffset();
-}
-
-void ABrush::SetPrePivot( const FVector& InPrePivot )
-{
-	SetPivotOffset(InPrePivot);
-}
-
 void ABrush::PostLoad()
 {
 	Super::PostLoad();
diff --git a/Engine/Source/Runtime/Engine/Private/Camera/CameraComponent.cpp b/Engine/Source/Runtime/Engine/Private/Camera/CameraComponent.cpp
--- a/Engine/Source/Runtime/Engine/Private/Camera/CameraComponent.cpp
+++ b/Engine/Source/Runtime/Engine/Private/Camera/CameraComponent.cpp
@@ -19,9 +19,6 @@
 
 #define LOCTEXT_NAMESPACE "CameraComponent"
 
-static void SetDeprecatedControllerViewRotation(UCameraComponent& Component, bool bValue);
-
-
 //////////////////////////////////////////////////////////////////////////
 // UCameraComponent
 
@@ -48,9 +45,6 @@ UCameraComponent::UCameraComponent(const FObjectInitializer& ObjectInitializer)
 	bUsePawnControlRotation = false;
 	bAutoActivate = true;
 	bLockToHmd = true;
-
-	// Init deprecated var, for old code that may refer to it.
-	SetDeprecatedControllerViewRotation(*this, bUsePawnControlRotation);
 }
 
 void UCameraComponent::OnUpdateTransform(EUpdateTransformFlags UpdateTransformFlags, ETeleportType Teleport)
@@ -121,9 +115,6 @@ void UCameraComponent::OnRegister()
 #endif
 
 	Super::OnRegister();
-
-	// Init deprecated var, for old code that may refer to it.
-	SetDeprecatedControllerViewRotation(*this, bUsePawnControlRotation);
 }
 
 void UCameraComponent::PostLoad()
@@ -136,9 +127,6 @@ void UCameraComponent::PostLoad()
 	{
 		bUsePawnControlRotation = bUseControllerViewRotation_DEPRECATED;
 	}
-
-	// Init deprecated var, for old code that may refer to it.
-	SetDeprecatedControllerViewRotation(*this, bUsePawnControlRotation);
 }
 
 #if WITH_EDITORONLY_DATA
@@ -333,14 +321,6 @@ bool UCameraComponent::GetEditorPreviewInfo(float DeltaTime, FMinimalViewInfo& V
 }
 #endif	// WITH_EDITOR
 
-
-void SetDeprecatedControllerViewRotation(UCameraComponent& Component, bool bValue)
-{
-	PRAGMA_DISABLE_DEPRECATION_WARNINGS
-	Component.bUseControllerViewRotation = bValue;
-	PRAGMA_ENABLE_DEPRECATION_WARNINGS
-}
-
 void UCameraComponent::NotifyCameraCut()
 {
 	// if we are owned by a camera actor, notify it too
diff --git a/Engine/Source/Runtime/Engine/Private/Character.cpp b/Engine/Source/Runtime/Engine/Private/Character.cpp
--- a/Engine/Source/Runtime/Engine/Private/Character.cpp
+++ b/Engine/Source/Runtime/Engine/Private/Character.cpp
@@ -319,13 +319,6 @@ bool ACharacter::IsJumpProvidingForce() const
 	return (bPressedJump && JumpKeyHoldTime < GetJumpMaxHoldTime());
 }
 
-// Deprecated
-bool ACharacter::DoJump( bool bReplayingMoves )
-{
-	return CanJump() && CharacterMovement->DoJump(bReplayingMoves);
-}
-
-
 void ACharacter::RecalculateBaseEyeHeight()
 {
 	if (!bIsCrouched)
diff --git a/Engine/Source/Runtime/Engine/Private/Components/ActorComponent.cpp b/Engine/Source/Runtime/Engine/Private/Components/ActorComponent.cpp
--- a/Engine/Source/Runtime/Engine/Private/Components/ActorComponent.cpp
+++ b/Engine/Source/Runtime/Engine/Private/Components/ActorComponent.cpp
@@ -723,6 +723,8 @@ void UActorComponent::OnUnregister()
 {
 	check(bRegistered);
 	bRegistered = false;
+
+	ClearNeedEndOfFrameUpdate();
 }
 
 void UActorComponent::InitializeComponent()
@@ -1427,6 +1429,18 @@ void UActorComponent::MarkForNeededEndOfFrameUpdate()
 	}
 }
 
+void UActorComponent::ClearNeedEndOfFrameUpdate_Internal()
+{
+	// If this is being garbage collected we don't really need to worry about clearing this
+	if (!HasAnyFlags(RF_BeginDestroyed) && !IsUnreachable())
+	{
+		if (UWorld* World = GetWorld())
+		{
+			World->ClearActorComponentEndOfFrameUpdate(this);
+		}
+	}
+}
+
 void UActorComponent::MarkForNeededEndOfFrameRecreate()
 {
 	if (bNeverNeedsRenderUpdate)
diff --git a/Engine/Source/Runtime/Engine/Private/Components/CharacterMovementComponent.cpp b/Engine/Source/Runtime/Engine/Private/Components/CharacterMovementComponent.cpp
--- a/Engine/Source/Runtime/Engine/Private/Components/CharacterMovementComponent.cpp
+++ b/Engine/Source/Runtime/Engine/Private/Components/CharacterMovementComponent.cpp
@@ -1000,13 +1000,6 @@ void UCharacterMovementComponent::ApplyNetworkMovementMode(const uint8 ReceivedM
 	SetMovementMode(NetMovementMode, NetCustomMode);
 }
 
-// TODO: deprecated, remove
-void UCharacterMovementComponent::PerformAirControl(FVector Direction, float ZDiff)
-{
-	PerformAirControlForPathFollowing(Direction, ZDiff);
-}
-
-
 void UCharacterMovementComponent::PerformAirControlForPathFollowing(FVector Direction, float ZDiff)
 {
 	// use air control if low grav or above destination and falling towards it
@@ -2808,15 +2801,6 @@ FVector UCharacterMovementComponent::HandleSlopeBoosting(const FVector& SlideRes
 	return Result;
 }
 
-
-// TODO: deprecated, remove.
-FVector UCharacterMovementComponent::AdjustUpperHemisphereImpact(const FVector& Delta, const FHitResult& Hit) const
-{
-	const float ZScale = FMath::Clamp(1.f - (FMath::Abs(Hit.Normal.Z) * UpperImpactNormalScale_DEPRECATED), 0.f, 1.f);
-	return FVector(Delta.X, Delta.Y, Delta.Z * ZScale);
-}
-
-
 FVector UCharacterMovementComponent::NewFallVelocity(const FVector& InitialVelocity, const FVector& Gravity, float DeltaTime) const
 {
 	FVector Result = InitialVelocity;
@@ -4110,36 +4094,6 @@ void UCharacterMovementComponent::PhysFalling(float deltaTime, int32 Iterations)
 	}
 }
 
-
-// Note: Deprecated for 4.9
-bool UCharacterMovementComponent::FindAirControlImpact(float DeltaTime, float AdditionalTime, const FVector& FallVelocity, const FVector& FallAcceleration, const FVector& Gravity, FHitResult& OutHitResult)
-{
-	// Test for slope to avoid using air control to climb walls.
-	FVector TestWalk = Velocity * DeltaTime;
-	if (AdditionalTime > 0.f)
-	{
-		const FVector PostGravityVelocity = NewFallVelocity(FallVelocity, Gravity, AdditionalTime);
-		TestWalk += ((FallAcceleration * AdditionalTime) + PostGravityVelocity) * AdditionalTime;
-	}
-	
-	if (!TestWalk.IsZero())
-	{
-		FCollisionQueryParams CapsuleQuery(SCENE_QUERY_STAT(FallingTraceParam), false, CharacterOwner);
-		FCollisionResponseParams ResponseParam;
-		InitCollisionParams(CapsuleQuery, ResponseParam);
-		const FVector CapsuleLocation = UpdatedComponent->GetComponentLocation();
-		const FCollisionShape CapsuleShape = GetPawnCapsuleCollisionShape(SHRINK_None);
-		
-		if (GetWorld()->SweepSingleByChannel(OutHitResult, CapsuleLocation, CapsuleLocation + TestWalk, FQuat::Identity, UpdatedComponent->GetCollisionObjectType(), CapsuleShape, CapsuleQuery, ResponseParam))
-		{
-			return true;
-		}
-	}
-
-	return false;
-}
-
-
 FVector UCharacterMovementComponent::LimitAirControl(float DeltaTime, const FVector& FallAcceleration, const FHitResult& HitResult, bool bCheckForValidLandingSpot)
 {
 	FVector Result(FallAcceleration);
diff --git a/Engine/Source/Runtime/Engine/Private/Components/ChildActorComponent.cpp b/Engine/Source/Runtime/Engine/Private/Components/ChildActorComponent.cpp
--- a/Engine/Source/Runtime/Engine/Private/Components/ChildActorComponent.cpp
+++ b/Engine/Source/Runtime/Engine/Private/Components/ChildActorComponent.cpp
@@ -142,6 +142,13 @@ void UChildActorComponent::PostEditImport()
 	{
 		ChildActorTemplate = CastChecked<UChildActorComponent>(GetArchetype())->ChildActorTemplate;
 	}
+
+	// Any cached instance data is invalid if we've had data imported in to us
+	if (CachedInstanceData)
+	{
+		delete CachedInstanceData;
+		CachedInstanceData = nullptr;
+	}
 }
 
 void UChildActorComponent::PostEditChangeProperty(FPropertyChangedEvent& PropertyChangedEvent)
diff --git a/Engine/Source/Runtime/Engine/Private/Components/NavMovementComponent.cpp b/Engine/Source/Runtime/Engine/Private/Components/NavMovementComponent.cpp
--- a/Engine/Source/Runtime/Engine/Private/Components/NavMovementComponent.cpp
+++ b/Engine/Source/Runtime/Engine/Private/Components/NavMovementComponent.cpp
@@ -124,29 +124,3 @@ void UNavMovementComponent::SetUpdateNavAgentWithOwnersCollisions(bool bUpdateWi
 {
 	bUpdateNavAgentWithOwnersCollision = bUpdateWithOwner;
 }
-
-//----------------------------------------------------------------------//
-// DEPRECATED
-//----------------------------------------------------------------------//
-const FNavAgentProperties* UNavMovementComponent::GetNavAgentProperties() const
-{
-	return &GetNavAgentPropertiesRef();
-}
-FNavAgentProperties* UNavMovementComponent::GetNavAgentProperties()
-{
-	return &GetNavAgentPropertiesRef();
-}
-void UNavMovementComponent::UpdateNavAgent(AActor* Owner)
-{
-	if (Owner)
-	{
-		UpdateNavAgent(*Owner);
-	}
-}
-void UNavMovementComponent::UpdateNavAgent(UCapsuleComponent* CapsuleComponent)
-{
-	if (CapsuleComponent)
-	{
-		UpdateNavAgent(*CapsuleComponent);
-	}
-}
diff --git a/Engine/Source/Runtime/Engine/Private/Components/PrimitiveComponent.cpp b/Engine/Source/Runtime/Engine/Private/Components/PrimitiveComponent.cpp
--- a/Engine/Source/Runtime/Engine/Private/Components/PrimitiveComponent.cpp
+++ b/Engine/Source/Runtime/Engine/Private/Components/PrimitiveComponent.cpp
@@ -483,13 +483,9 @@ void UPrimitiveComponent::OnUnregister()
 	if (!HasAnyFlags(RF_BeginDestroyed) && !IsUnreachable())
 	{
 		UWorld* World = GetWorld();
-		if (World)
+		if (World && World->Scene)
 		{
-			if (World->Scene)
-			{
-				World->Scene->ReleasePrimitive(this);
-			}
-			World->ClearActorComponentEndOfFrameUpdate(this);
+			World->Scene->ReleasePrimitive(this);
 		}
 	}
 
diff --git a/Engine/Source/Runtime/Engine/Private/Components/ProjectileMovementComponent.cpp b/Engine/Source/Runtime/Engine/Private/Components/ProjectileMovementComponent.cpp
--- a/Engine/Source/Runtime/Engine/Private/Components/ProjectileMovementComponent.cpp
+++ b/Engine/Source/Runtime/Engine/Private/Components/ProjectileMovementComponent.cpp
@@ -340,12 +340,6 @@ void UProjectileMovementComponent::SetVelocityInLocalSpace(FVector NewVelocity)
 }
 
 
-// Deprecated
-FVector UProjectileMovementComponent::CalculateVelocity(FVector OldVelocity, float DeltaTime, bool bGravityEnabled_UNUSED) const
-{
-	return ComputeVelocity(OldVelocity, DeltaTime);
-}
-
 FVector UProjectileMovementComponent::ComputeVelocity(FVector InitialVelocity, float DeltaTime) const
 {
 	// v = v0 + a*t
@@ -403,11 +397,6 @@ FVector UProjectileMovementComponent::ComputeHomingAcceleration(const FVector& I
 	return HomingAcceleration;
 }
 
-// Deprecated
-float UProjectileMovementComponent::GetEffectiveGravityZ() const
-{
-	return GetGravityZ();
-}
 
 float UProjectileMovementComponent::GetGravityZ() const
 {
diff --git a/Engine/Source/Runtime/Engine/Private/DataAsset.cpp b/Engine/Source/Runtime/Engine/Private/DataAsset.cpp
--- a/Engine/Source/Runtime/Engine/Private/DataAsset.cpp
+++ b/Engine/Source/Runtime/Engine/Private/DataAsset.cpp
@@ -76,11 +76,11 @@ void UPrimaryDataAsset::PostLoad()
 	Super::PostLoad();
 
 #if WITH_EDITORONLY_DATA
-	int32 OldAssetDataNum = AssetBundleData.Bundles.Num();
+	FAssetBundleData OldData = AssetBundleData;
 	
 	UpdateAssetBundleData();
 
-	if (AssetBundleData.Bundles.Num() != OldAssetDataNum && UAssetManager::IsValid())
+	if (UAssetManager::IsValid() && OldData != AssetBundleData)
 	{
 		// Bundles changed, refresh
 		UAssetManager::Get().RefreshAssetData(this);
diff --git a/Engine/Source/Runtime/Engine/Private/EdGraph/EdGraphNode.cpp b/Engine/Source/Runtime/Engine/Private/EdGraph/EdGraphNode.cpp
--- a/Engine/Source/Runtime/Engine/Private/EdGraph/EdGraphNode.cpp
+++ b/Engine/Source/Runtime/Engine/Private/EdGraph/EdGraphNode.cpp
@@ -17,6 +17,17 @@
 
 #define LOCTEXT_NAMESPACE "EdGraph"
 
+FEdGraphTerminalType FEdGraphTerminalType::FromPinType(const FEdGraphPinType& PinType)
+{
+	FEdGraphTerminalType TerminalType;
+	TerminalType.TerminalCategory = PinType.PinCategory;
+	TerminalType.TerminalSubCategory = PinType.PinSubCategory;
+	TerminalType.TerminalSubCategoryObject = PinType.PinSubCategoryObject;
+	TerminalType.bTerminalIsConst = PinType.bIsConst;
+	TerminalType.bTerminalIsWeakPointer = PinType.bIsWeakPointer;
+	return TerminalType;
+}
+
 FArchive& operator<<(FArchive& Ar, FEdGraphTerminalType& T)
 {
 	Ar << T.TerminalCategory;
diff --git a/Engine/Source/Runtime/Engine/Private/EdGraph/EdGraphPin.cpp b/Engine/Source/Runtime/Engine/Private/EdGraph/EdGraphPin.cpp
--- a/Engine/Source/Runtime/Engine/Private/EdGraph/EdGraphPin.cpp
+++ b/Engine/Source/Runtime/Engine/Private/EdGraph/EdGraphPin.cpp
@@ -409,15 +409,18 @@ void TransferPersistentDataFromOldPin(UEdGraphPin& DestPin, T& SourcePin, const
 	DestPin.PinId = SourcePin.PinId;
 
 	// Only move the default value if it was modified; inherit the new default value otherwise
-	if (!SourcePin.DoesDefaultValueMatchAutogenerated())
+	if (DestPin.Direction == EGPD_Input)
 	{
-		DestPin.DefaultObject = SourcePin.DefaultObject;
-		DestPin.DefaultValue = MoveTemp(SourcePin.DefaultValue);
-		DestPin.DefaultTextValue = MoveTemp(SourcePin.DefaultTextValue);
-	}
-	else
-	{
-		DestPin.GetSchema()->ResetPinToAutogeneratedDefaultValue(&DestPin, false);
+		if (!SourcePin.DoesDefaultValueMatchAutogenerated())
+		{
+			DestPin.DefaultObject = SourcePin.DefaultObject;
+			DestPin.DefaultValue = MoveTemp(SourcePin.DefaultValue);
+			DestPin.DefaultTextValue = MoveTemp(SourcePin.DefaultTextValue);
+		}
+		else
+		{
+			DestPin.GetSchema()->ResetPinToAutogeneratedDefaultValue(&DestPin, false);
+		}
 	}
 
 	// Copy the links
diff --git a/Engine/Source/Runtime/Engine/Private/GameEngine.cpp b/Engine/Source/Runtime/Engine/Private/GameEngine.cpp
--- a/Engine/Source/Runtime/Engine/Private/GameEngine.cpp
+++ b/Engine/Source/Runtime/Engine/Private/GameEngine.cpp
@@ -545,6 +545,8 @@ UEngine::UEngine(const FObjectInitializer& ObjectInitializer)
 	bIsVanillaProduct = false;
 
 	GameScreenshotSaveDirectory.Path = FPaths::ScreenShotDir();
+
+	LastGCFrame = TNumericLimits<uint64>::Max();
 }
 
 void UGameEngine::Init(IEngineLoop* InEngineLoop)
diff --git a/Engine/Source/Runtime/Engine/Private/GameFramework/SpringArmComponent.cpp b/Engine/Source/Runtime/Engine/Private/GameFramework/SpringArmComponent.cpp
--- a/Engine/Source/Runtime/Engine/Private/GameFramework/SpringArmComponent.cpp
+++ b/Engine/Source/Runtime/Engine/Private/GameFramework/SpringArmComponent.cpp
@@ -10,7 +10,6 @@
 //////////////////////////////////////////////////////////////////////////
 // USpringArmComponent
 
-static void SetDeprecatedControllerViewRotation(USpringArmComponent& Component, bool bValue);
 const FName USpringArmComponent::SocketName(TEXT("SpringEndpoint"));
 
 USpringArmComponent::USpringArmComponent(const FObjectInitializer& ObjectInitializer)
@@ -39,9 +38,6 @@ USpringArmComponent::USpringArmComponent(const FObjectInitializer& ObjectInitial
 	CameraRotationLagSpeed = 10.f;
 	CameraLagMaxTimeStep = 1.f / 60.f;
 	CameraLagMaxDistance = 0.f;
-	
-	// Init deprecated var, for old code that may refer to it.
-	SetDeprecatedControllerViewRotation(*this, bUsePawnControlRotation);
 }
 
 void USpringArmComponent::UpdateDesiredArmLocation(bool bDoTrace, bool bDoLocationLag, bool bDoRotationLag, float DeltaTime)
@@ -215,17 +211,11 @@ void USpringArmComponent::OnRegister()
 
 	// Set initial location (without lag).
 	UpdateDesiredArmLocation(false, false, false, 0.f);
-
-	// Init deprecated var, for old code that may refer to it.
-	SetDeprecatedControllerViewRotation(*this, bUsePawnControlRotation);
 }
 
 void USpringArmComponent::PostLoad()
 {
 	Super::PostLoad();
-
-	// Init deprecated var, for old code that may refer to it.
-	SetDeprecatedControllerViewRotation(*this, bUsePawnControlRotation);
 }
 
 void USpringArmComponent::TickComponent(float DeltaTime, enum ELevelTick TickType, FActorComponentTickFunction* ThisTickFunction)
@@ -271,11 +261,3 @@ void USpringArmComponent::QuerySupportedSockets(TArray<FComponentSocketDescripti
 {
 	new (OutSockets) FComponentSocketDescription(SocketName, EComponentSocketType::Socket);
 }
-
-
-void SetDeprecatedControllerViewRotation(USpringArmComponent& Component, bool bValue)
-{
-	PRAGMA_DISABLE_DEPRECATION_WARNINGS
-	Component.bUseControllerViewRotation = bValue;
-	PRAGMA_ENABLE_DEPRECATION_WARNINGS
-}
diff --git a/Engine/Source/Runtime/Engine/Private/GameViewportClient.cpp b/Engine/Source/Runtime/Engine/Private/GameViewportClient.cpp
--- a/Engine/Source/Runtime/Engine/Private/GameViewportClient.cpp
+++ b/Engine/Source/Runtime/Engine/Private/GameViewportClient.cpp
@@ -700,18 +700,6 @@ bool UGameViewportClient::GetMousePosition(FVector2D& MousePosition) const
 	return bGotMousePosition;
 }
 
-FVector2D UGameViewportClient::GetMousePosition() const
-{
-	FVector2D MousePosition;
-	if (!GetMousePosition(MousePosition))
-	{
-		MousePosition = FVector2D::ZeroVector;
-	}
-
-	return MousePosition;
-}
-
-
 bool UGameViewportClient::RequiresUncapturedAxisInput() const
 {
 	bool bRequired = false;
@@ -1739,18 +1727,6 @@ ULocalPlayer* UGameViewportClient::SetupInitialLocalPlayer(FString& OutError)
 	return ViewportGameInstance->CreateInitialPlayer(OutError);
 }
 
-ULocalPlayer* UGameViewportClient::CreatePlayer(int32 ControllerId, FString& OutError, bool bSpawnActor)
-{
-	UGameInstance * ViewportGameInstance = GEngine->GetWorldContextFromGameViewportChecked(this).OwningGameInstance;
-	return ( ViewportGameInstance != NULL ) ? ViewportGameInstance->CreateLocalPlayer(ControllerId, OutError, bSpawnActor) : NULL;
-}
-
-bool UGameViewportClient::RemovePlayer(class ULocalPlayer* ExPlayer)
-{
-	UGameInstance * ViewportGameInstance = GEngine->GetWorldContextFromGameViewportChecked(this).OwningGameInstance;
-	return (ViewportGameInstance != NULL) ? ViewportGameInstance->RemoveLocalPlayer(ExPlayer) : false;
-}
-
 void UGameViewportClient::UpdateActiveSplitscreenType()
 {
 	ESplitScreenType::Type SplitType = ESplitScreenType::None;
diff --git a/Engine/Source/Runtime/Engine/Private/KismetSystemLibrary.cpp b/Engine/Source/Runtime/Engine/Private/KismetSystemLibrary.cpp
--- a/Engine/Source/Runtime/Engine/Private/KismetSystemLibrary.cpp
+++ b/Engine/Source/Runtime/Engine/Private/KismetSystemLibrary.cpp
@@ -2122,9 +2122,10 @@ bool UKismetSystemLibrary::CanLaunchURL(const FString& URL)
 
 	return false;
 }
+
 void UKismetSystemLibrary::CollectGarbage()
 {
-	GEngine->DeferredCommands.Add(TEXT("obj gc"));
+	GEngine->ForceGarbageCollection(true);
 }
 
 void UKismetSystemLibrary::ShowAdBanner(int32 AdIdIndex, bool bShowOnBottomOfScreen)
diff --git a/Engine/Source/Runtime/Engine/Private/Level.cpp b/Engine/Source/Runtime/Engine/Private/Level.cpp
--- a/Engine/Source/Runtime/Engine/Private/Level.cpp
+++ b/Engine/Source/Runtime/Engine/Private/Level.cpp
@@ -1558,11 +1558,6 @@ void ULevel::BuildStreamingData(UWorld* World, ULevel* TargetLevel/*=NULL*/, UTe
 #endif
 }
 
-ABrush* ULevel::GetBrush() const
-{
-	return GetDefaultBrush();
-}
-
 ABrush* ULevel::GetDefaultBrush() const
 {
 	ABrush* DefaultBrush = nullptr;
diff --git a/Engine/Source/Runtime/Engine/Private/LevelTick.cpp b/Engine/Source/Runtime/Engine/Private/LevelTick.cpp
--- a/Engine/Source/Runtime/Engine/Private/LevelTick.cpp
+++ b/Engine/Source/Runtime/Engine/Private/LevelTick.cpp
@@ -773,11 +773,6 @@ static TAutoConsoleVariable<int32> CVarAllowAsyncRenderThreadUpdatesEditor(
 	0,
 	TEXT("Used to control async renderthread updates in the editor."));
 
-static TAutoConsoleVariable<int32> CVarCollectGarbageEveryFrame(
-	TEXT("gc.CollectGarbageEveryFrame"),
-	0,
-	TEXT("Used to debug garbage collection...Collects garbage every frame if the value is > 0."));
-
 namespace EComponentMarkedForEndOfFrameUpdateState
 {
 	enum Type
@@ -1100,13 +1095,6 @@ public:
 extern bool GCollisionAnalyzerIsRecording;
 #endif // ENABLE_COLLISION_ANALYZER
 
-#if !(UE_BUILD_SHIPPING || UE_BUILD_TEST)
-static TAutoConsoleVariable<int32> CVarStressTestGCWhileStreaming(
-	TEXT("gc.StressTestGC"),
-	0,
-	TEXT("If set to 1, the engine will attempt to trigger GC each frame while async loading."));
-#endif
-
 DECLARE_CYCLE_STAT(TEXT("TG_PrePhysics"), STAT_TG_PrePhysics, STATGROUP_TickGroups);
 DECLARE_CYCLE_STAT(TEXT("TG_StartPhysics"), STAT_TG_StartPhysics, STATGROUP_TickGroups);
 DECLARE_CYCLE_STAT(TEXT("Start TG_DuringPhysics"), STAT_TG_DuringPhysics, STATGROUP_TickGroups);
@@ -1115,14 +1103,6 @@ DECLARE_CYCLE_STAT(TEXT("TG_PostPhysics"), STAT_TG_PostPhysics, STATGROUP_TickGr
 DECLARE_CYCLE_STAT(TEXT("TG_PostUpdateWork"), STAT_TG_PostUpdateWork, STATGROUP_TickGroups);
 DECLARE_CYCLE_STAT(TEXT("TG_LastDemotable"), STAT_TG_LastDemotable, STATGROUP_TickGroups);
 
-static float GTimeBetweenPurgingPendingKillObjects = 60.0f;
-static FAutoConsoleVariableRef CVarTimeBetweenPurgingPendingKillObjects(
-	TEXT("gc.TimeBetweenPurgingPendingKillObjects"),
-	GTimeBetweenPurgingPendingKillObjects,
-	TEXT("Time in seconds (game time) we should wait between purging object references to objects that are pending kill."),
-	ECVF_Default
-	);
-
 #include "GameFramework/SpawnActorTimer.h"
 
 TDrawEvent<FRHICommandList>* BeginTickDrawEvent()
@@ -1556,55 +1536,7 @@ void UWorld::Tick( ELevelTick TickType, float DeltaSeconds )
 	bInTick = false;
 	Mark.Pop();
 
-	
-#if !(UE_BUILD_SHIPPING || UE_BUILD_TEST)
-	if (CVarStressTestGCWhileStreaming.GetValueOnGameThread() && IsAsyncLoading())
-	{
-		TryCollectGarbage(GARBAGE_COLLECTION_KEEPFLAGS, true);
-	}
-	else 
-#endif
-	if (FullPurgeTriggered)
-	{
-		if (TryCollectGarbage(GARBAGE_COLLECTION_KEEPFLAGS, true))
-		{
-			CleanupActors();
-			FullPurgeTriggered = false;
-			TimeSinceLastPendingKillPurge = 0.0f;
-		}
-	}
-	else if( HasBegunPlay() )
-	{
-		TimeSinceLastPendingKillPurge += DeltaSeconds;
-
-		const float TimeBetweenPurgingPendingKillObjects = GetTimeBetweenGarbageCollectionPasses();
-
-		
-
-		// See if we should delay garbage collect for this frame
-		if (bShouldDelayGarbageCollect)
-		{
-			bShouldDelayGarbageCollect = false;
-		}
-		// Perform incremental purge update if it's pending or in progress.
-		else if( !IsIncrementalPurgePending() 
-		// Purge reference to pending kill objects every now and so often.
-		&&	(TimeSinceLastPendingKillPurge > TimeBetweenPurgingPendingKillObjects) && TimeBetweenPurgingPendingKillObjects > 0 )
-		{
-			SCOPE_CYCLE_COUNTER(STAT_GCMarkTime);
-			PerformGarbageCollectionAndCleanupActors();
-		}
-		else
-		{
-			SCOPE_CYCLE_COUNTER(STAT_GCSweepTime);
-			IncrementalPurgeGarbage( true );
-		}
-	}
-
-	if (CVarCollectGarbageEveryFrame.GetValueOnGameThread() > 0)
-	{
-		ForceGarbageCollection(true);
-	}
+	GEngine->ConditionalCollectGarbage();
 
 	// players only request from last frame
 	if (bPlayersOnlyPending)
@@ -1685,23 +1617,22 @@ void UWorld::Tick( ELevelTick TickType, float DeltaSeconds )
  */
 void UWorld::DelayGarbageCollection()
 {
-	bShouldDelayGarbageCollect = true;
+	GEngine->DelayGarbageCollection();
 }
 
-void UWorld::SetTimeUntilNextGarbageCollection(float MinTimeUntilNextPass)
+void UWorld::ForceGarbageCollection( bool bFullPurge)
 {
-	const float TimeBetweenPurgingPendingKillObjects = GetTimeBetweenGarbageCollectionPasses();
+	GEngine->ForceGarbageCollection(bFullPurge);
+}
 
-	// This can make it go negative if the desired interval is longer than the typical interval, but it's only ever compared against TimeBetweenPurgingPendingKillObjects
-	TimeSinceLastPendingKillPurge = TimeBetweenPurgingPendingKillObjects - MinTimeUntilNextPass;
+void UWorld::SetTimeUntilNextGarbageCollection(const float MinTimeUntilNextPass)
+{
+	GEngine->SetTimeUntilNextGarbageCollection(MinTimeUntilNextPass);
 }
 
 float UWorld::GetTimeBetweenGarbageCollectionPasses() const
 {
-	const bool bAtLeastOnePlayerConnected = NetDriver && NetDriver->ClientConnections.Num() > 0;
-	const bool bShouldUseLowFrequencyGC = IsRunningDedicatedServer() && !bAtLeastOnePlayerConnected;
-
-	return bShouldUseLowFrequencyGC ? (GTimeBetweenPurgingPendingKillObjects * 10) : GTimeBetweenPurgingPendingKillObjects;
+	return GEngine->GetTimeBetweenGarbageCollectionPasses();
 }
 
 /**
@@ -1709,50 +1640,42 @@ float UWorld::GetTimeBetweenGarbageCollectionPasses() const
  */
 void UWorld::PerformGarbageCollectionAndCleanupActors()
 {
-	// We don't collect garbage while there are outstanding async load requests as we would need
-	// to block on loading the remaining data.
-	if( !IsAsyncLoading() )
-	{
-		// Perform housekeeping.
-		if (TryCollectGarbage(GARBAGE_COLLECTION_KEEPFLAGS, false))
-		{
-			CleanupActors();
-
-			// Reset counter.
-			TimeSinceLastPendingKillPurge = 0.0f;
-		}
-	}
+	GEngine->PerformGarbageCollectionAndCleanupActors();
 }
 
-
 void UWorld::CleanupActors()
 {
 	// Remove NULL entries from actor list. Only does so for dynamic actors to avoid resorting; in theory static 
 	// actors shouldn't be deleted during gameplay.
-	for( int32 LevelIndex=0; LevelIndex<Levels.Num(); LevelIndex++ )
+	for (ULevel* Level : Levels)
 	{
-		ULevel* Level = Levels[LevelIndex];
 		// Don't compact actors array for levels that are currently in the process of being made visible as the
 		// code that spreads this work across several frames relies on the actor count not changing as it keeps
 		// an index into the array.
 		if( CurrentLevelPendingVisibility != Level )
 		{
 			// Actor 0 (world info) and 1 (default brush) are special and should never be removed from the actor array even if NULL
-			int32 FirstDynamicIndex = 2;
+			const int32 FirstDynamicIndex = 2;
+			int32 NumActorsToRemove = 0;
 			// Remove NULL entries from array, we're iterating backwards to avoid unnecessary memcpys during removal.
 			for( int32 ActorIndex=Level->Actors.Num()-1; ActorIndex>=FirstDynamicIndex; ActorIndex-- )
 			{
-				if( Level->Actors[ActorIndex] == NULL )
+				// To avoid shuffling things down repeatedly when not necessary count nulls and then remove in bunches
+				if (Level->Actors[ActorIndex] == nullptr)
+				{
+					++NumActorsToRemove;
+				}
+				else if (NumActorsToRemove > 0)
 				{
-					Level->Actors.RemoveAt( ActorIndex );
+					Level->Actors.RemoveAt(ActorIndex+1, NumActorsToRemove, false);
+					NumActorsToRemove = 0;
 				}
 			}
+			if (NumActorsToRemove > 0)
+			{
+				// If our FirstDynamicIndex (and any immediately following it) were null it won't get caught in the loop, so do a cleanup pass here
+				Level->Actors.RemoveAt(FirstDynamicIndex, NumActorsToRemove, false);
+			}
 		}
 	}
-}
-
-void UWorld::ForceGarbageCollection( bool bForcePurge/*=false*/ )
-{
-	TimeSinceLastPendingKillPurge = 1.0f + GetTimeBetweenGarbageCollectionPasses();
-	FullPurgeTriggered = FullPurgeTriggered || bForcePurge;
-}
+}
\ No newline at end of file
diff --git a/Engine/Source/Runtime/Engine/Private/Materials/MaterialShared.cpp b/Engine/Source/Runtime/Engine/Private/Materials/MaterialShared.cpp
--- a/Engine/Source/Runtime/Engine/Private/Materials/MaterialShared.cpp
+++ b/Engine/Source/Runtime/Engine/Private/Materials/MaterialShared.cpp
@@ -888,11 +888,6 @@ bool FMaterialResource::IsUsedWithAPEXCloth() const
 	return Material->bUsedWithClothing;
 }
 
-bool FMaterialResource::IsUsedWithUI() const
-{
-	return IsUIMaterial();
-}
-
 EMaterialTessellationMode FMaterialResource::GetTessellationMode() const 
 { 
 	return (EMaterialTessellationMode)Material->D3D11TessellationMode; 
diff --git a/Engine/Source/Runtime/Engine/Private/MemberReference.cpp b/Engine/Source/Runtime/Engine/Private/MemberReference.cpp
--- a/Engine/Source/Runtime/Engine/Private/MemberReference.cpp
+++ b/Engine/Source/Runtime/Engine/Private/MemberReference.cpp
@@ -6,8 +6,11 @@
 
 #if WITH_EDITOR
 #include "Kismet2/BlueprintEditorUtils.h"
+#include "BlueprintCompilationManager.h"
 #endif
 
+extern COREUOBJECT_API bool GBlueprintUseCompilationManager;
+
 //////////////////////////////////////////////////////////////////////////
 // FMemberReference
 
@@ -248,6 +251,18 @@ void FMemberReference::InitFieldRedirectMap()
 	}
 }
 
+UClass* FMemberReference::GetClassToUse(UClass* InClass, bool bUseUpToDateClass)
+{
+	if(GBlueprintUseCompilationManager && bUseUpToDateClass)
+	{
+		return FBlueprintEditorUtils::GetMostUpToDateClass(InClass);
+	}
+	else
+	{
+		return InClass;
+	}
+}
+
 UField* FMemberReference::FindRemappedField(UClass *FieldClass, UClass* InitialScope, FName InitialName, bool bInitialScopeMustBeOwnerOfField)
 {
 	DECLARE_SCOPE_CYCLE_COUNTER(TEXT("FMemberReference::FindRemappedField"), STAT_LinkerLoad_FindRemappedField, STATGROUP_LoadTimeVerbose);
diff --git a/Engine/Source/Runtime/Engine/Private/PhysicsEngine/BodyInstance.cpp b/Engine/Source/Runtime/Engine/Private/PhysicsEngine/BodyInstance.cpp
--- a/Engine/Source/Runtime/Engine/Private/PhysicsEngine/BodyInstance.cpp
+++ b/Engine/Source/Runtime/Engine/Private/PhysicsEngine/BodyInstance.cpp
@@ -3032,11 +3032,6 @@ FTransform FBodyInstance::GetUnrealWorldTransform_AssumesLocked(bool bWithProjec
 	return GetUnrealWorldTransformImp<false>(this, bWithProjection);
 }
 
-void FBodyInstance::SetBodyTransform(const FTransform& NewTransform, bool bTeleport)
-{
-	SetBodyTransform(NewTransform, TeleportFlagToEnum(bTeleport));
-}
-
 void FBodyInstance::SetBodyTransform(const FTransform& NewTransform, ETeleportType Teleport)
 {
 	SCOPE_CYCLE_COUNTER(STAT_SetBodyTransform);
@@ -5089,13 +5084,6 @@ void FBodyInstance::FixupData(class UObject* Loader)
 	}
 }
 
-bool FBodyInstance::UseAsyncScene() const
-{
-	UPrimitiveComponent* OwnerComponentInst = OwnerComponent.Get();
-	UWorld* World = OwnerComponentInst ? OwnerComponentInst->GetWorld() : nullptr;
-	return UseAsyncScene(World ? World->GetPhysicsScene() : nullptr);
-}
-
 bool FBodyInstance::UseAsyncScene(const FPhysScene* PhysScene) const
 {
 	bool bHasAsyncScene = true;
diff --git a/Engine/Source/Runtime/Engine/Private/PhysicsEngine/PhysDrawing.cpp b/Engine/Source/Runtime/Engine/Private/PhysicsEngine/PhysDrawing.cpp
--- a/Engine/Source/Runtime/Engine/Private/PhysicsEngine/PhysDrawing.cpp
+++ b/Engine/Source/Runtime/Engine/Private/PhysicsEngine/PhysDrawing.cpp
@@ -380,12 +380,6 @@ void FKSphylElem::DrawElemSolid(FPrimitiveDrawInterface* PDI, const FTransform&
 // FKConvexElem
 /////////////////////////////////////////////////////////////////////////////////////
 
-
-void FKConvexElem::DrawElemWire(FPrimitiveDrawInterface* PDI, const FTransform& ElemTM, const FColor Color) const
-{
-	DrawElemWire(PDI, ElemTM, 1.f, Color);
-}
-
 void FKConvexElem::DrawElemWire(FPrimitiveDrawInterface* PDI, const FTransform& ElemTM, const float Scale, const FColor Color) const
 {
 #if WITH_PHYSX
diff --git a/Engine/Source/Runtime/Engine/Private/PhysicsEngine/PhysXSupport.h b/Engine/Source/Runtime/Engine/Private/PhysicsEngine/PhysXSupport.h
--- a/Engine/Source/Runtime/Engine/Private/PhysicsEngine/PhysXSupport.h
+++ b/Engine/Source/Runtime/Engine/Private/PhysicsEngine/PhysXSupport.h
@@ -76,39 +76,13 @@ ENGINE_API void AddRadialImpulseToPxRigidBody_AssumesLocked(PxRigidBody& PRigidB
 
 
 /** Calculates correct impulse at the body's center of mass and adds the impulse to the body. */
-DEPRECATED(4.8, "Please call AddRadialImpulseToPxRigidBody_AssumesLocked and make sure you obtain the appropriate PhysX scene locks")
-inline void AddRadialImpulseToPxRigidBody(PxRigidBody& PRigidBody, const FVector& Origin, float Radius, float Strength, uint8 Falloff, bool bVelChange)
-{
-	AddRadialImpulseToPxRigidBody_AssumesLocked(PRigidBody, Origin, Radius, Strength, Falloff, bVelChange);
-}
-
 ENGINE_API void AddRadialForceToPxRigidBody_AssumesLocked(PxRigidBody& PRigidBody, const FVector& Origin, float Radius, float Strength, uint8 Falloff, bool bAccelChange);
 
-/** Calculates correct force at the body's center of mass and adds force to the body. */
-DEPRECATED(4.8, "Please call AddRadialImpulseToPxRigidBody_AssumesLocked and make sure you obtain the appropriate PhysX scene locks")
-inline void AddRadialForceToPxRigidBody(PxRigidBody& PRigidBody, const FVector& Origin, float Radius, float Strength, uint8 Falloff, bool bAccelChange)
-{
-	AddRadialForceToPxRigidBody_AssumesLocked(PRigidBody, Origin, Radius, Strength, Falloff, bAccelChange);
-}
-
+/** Util to see if a PxRigidBody is non-kinematic */
 bool IsRigidBodyKinematic_AssumesLocked(const PxRigidBody* PRigidBody);
 
 bool IsRigidBodyKinematicAndInSimulationScene_AssumesLocked(const PxRigidBody* PRigidBody);
 
-DEPRECATED(4.12, "Please call IsRigidBodyKinematic_AssumesLocked")
-inline bool IsRigidBodyNonKinematic_AssumesLocked(const PxRigidBody* PRigidBody)
-{
-	return !IsRigidBodyKinematic_AssumesLocked(PRigidBody);
-}
-
-/** Util to see if a PxRigidActor is non-kinematic */
-DEPRECATED(4.8, "Please call IsRigidBodyKinematic_AssumesLocked and make sure you obtain the appropriate PhysX scene locks")
-inline bool IsRigidBodyNonKinematic(PxRigidBody* PRigidBody)
-{
-	return !IsRigidBodyKinematic_AssumesLocked(PRigidBody);
-}
-
-
 /////// GLOBAL POINTERS
 
 /** Pointer to PhysX Foundation singleton */
diff --git a/Engine/Source/Runtime/Engine/Private/PlayerController.cpp b/Engine/Source/Runtime/Engine/Private/PlayerController.cpp
--- a/Engine/Source/Runtime/Engine/Private/PlayerController.cpp
+++ b/Engine/Source/Runtime/Engine/Private/PlayerController.cpp
@@ -251,7 +251,7 @@ void APlayerController::ClientUpdateLevelStreamingStatus_Implementation(FName Pa
 						// If we're unloading any levels, we need to request a one frame delay of garbage collection to make sure it happens after the level is actually unloaded
 						if (LevelStreamingObject->bShouldBeLoaded && !bNewShouldBeLoaded)
 						{
-							GetWorld()->DelayGarbageCollection();
+							GEngine->DelayGarbageCollection();
 						}
 
 						LevelStreamingObject->bShouldBeLoaded		= bNewShouldBeLoaded;
@@ -286,7 +286,7 @@ void APlayerController::ClientFlushLevelStreaming_Implementation()
 		GetWorld()->UpdateLevelStreaming();
 		GetWorld()->bRequestedBlockOnAsyncLoading = true;
 		// request GC as soon as possible to remove any unloaded levels from memory
-		GetWorld()->ForceGarbageCollection();
+		GEngine->ForceGarbageCollection();
 	}
 }
 
@@ -3093,7 +3093,7 @@ void APlayerController::ClientSetCinematicMode_Implementation(bool bInCinematicM
 
 void APlayerController::ClientForceGarbageCollection_Implementation()
 {
-	GetWorld()->ForceGarbageCollection();
+	GEngine->ForceGarbageCollection();
 }
 
 /// @endcond
diff --git a/Engine/Source/Runtime/Engine/Private/PrimaryAssetLabel.cpp b/Engine/Source/Runtime/Engine/Private/PrimaryAssetLabel.cpp
--- a/Engine/Source/Runtime/Engine/Private/PrimaryAssetLabel.cpp
+++ b/Engine/Source/Runtime/Engine/Private/PrimaryAssetLabel.cpp
@@ -7,7 +7,14 @@
 #include "Engine/AssetManager.h"
 #include "AssetRegistryModule.h"
 
+#if WITH_EDITOR
+#include "CollectionManagerTypes.h"
+#include "ICollectionManager.h"
+#include "CollectionManagerModule.h"
+#endif
+
 const FName UPrimaryAssetLabel::DirectoryBundle = FName("Directory");
+const FName UPrimaryAssetLabel::CollectionBundle = FName("Collection");
 
 UPrimaryAssetLabel::UPrimaryAssetLabel()
 {
@@ -54,6 +61,28 @@ void UPrimaryAssetLabel::UpdateAssetBundleData()
 		// Fast set, destroys NewPaths
 		AssetBundleData.SetBundleAssets(DirectoryBundle, MoveTemp(NewPaths));
 	}
+
+	if (AssetCollection.CollectionName != NAME_None)
+	{
+		TArray<FStringAssetReference> NewPaths;
+		TArray<FName> CollectionAssets;
+		ICollectionManager& CollectionManager = FCollectionManagerModule::GetModule().Get();
+		CollectionManager.GetAssetsInCollection(AssetCollection.CollectionName, ECollectionShareType::CST_All, CollectionAssets);
+		for (int32 Index = 0; Index < CollectionAssets.Num(); ++Index)
+		{
+			FAssetData FoundAsset = Manager.GetAssetRegistry().GetAssetByObjectPath(CollectionAssets[Index]);
+
+			FStringAssetReference AssetRef = Manager.GetAssetPathForData(FoundAsset);
+
+			if (!AssetRef.IsNull())
+			{
+				NewPaths.Add(AssetRef);
+			}
+		}
+
+		// Fast set, destroys NewPaths
+		AssetBundleData.SetBundleAssets(CollectionBundle, MoveTemp(NewPaths));
+	}
 	
 	// Update rules
 	FPrimaryAssetId PrimaryAssetId = GetPrimaryAssetId();
diff --git a/Engine/Source/Runtime/Engine/Private/StreamableManager.cpp b/Engine/Source/Runtime/Engine/Private/StreamableManager.cpp
--- a/Engine/Source/Runtime/Engine/Private/StreamableManager.cpp
+++ b/Engine/Source/Runtime/Engine/Private/StreamableManager.cpp
@@ -167,6 +167,12 @@ EAsyncPackageState::Type FStreamableHandle::WaitUntilComplete(float Timeout)
 		return EAsyncPackageState::Complete;
 	}
 
+	if (bStalled)
+	{
+		// If we were stalled, start us now to avoid deadlocks
+		StartStalledHandle();
+	}
+
 	EAsyncPackageState::Type State = ProcessAsyncLoadingUntilComplete([this]() { return HasLoadCompleted(); }, Timeout);
 	
 	return State;
@@ -376,6 +382,20 @@ void FStreamableHandle::ReleaseHandle()
 	}
 }
 
+void FStreamableHandle::StartStalledHandle()
+{
+	if (!bStalled || !IsActive())
+	{
+		// Cannot start
+		return;
+	}
+
+	check(OwningManager);
+
+	bStalled = false;
+	OwningManager->StartHandleRequests(AsShared());
+}
+
 FStreamableHandle::~FStreamableHandle()
 {
 	if (IsActive())
@@ -663,12 +683,6 @@ FStreamable* FStreamableManager::StreamInternal(const FStringAssetReference& InT
 	check(IsInGameThread());
 	UE_LOG(LogStreamableManager, Verbose, TEXT("Asynchronous load %s"), *InTargetName.ToString());
 
-	if (FPackageName::IsShortPackageName(InTargetName.ToString()))
-	{
-		UE_LOG(LogStreamableManager, Error, TEXT("     Can't load invalid package name %s"), *InTargetName.ToString());
-		return nullptr;
-	}
-
 	FStringAssetReference TargetName = ResolveRedirects(InTargetName);
 	FStreamable* Existing = StreamableItems.FindRef(TargetName);
 	if (Existing)
@@ -768,7 +782,7 @@ UObject* FStreamableManager::SynchronousLoad(FStringAssetReference const& Target
 	return LoadSynchronous(Target, true);
 }
 
-TSharedPtr<FStreamableHandle> FStreamableManager::RequestAsyncLoad(const TArray<FStringAssetReference>& TargetsToStream, FStreamableDelegate DelegateToCall, TAsyncLoadPriority Priority, bool bManageActiveHandle, const FString& DebugName)
+TSharedPtr<FStreamableHandle> FStreamableManager::RequestAsyncLoad(const TArray<FStringAssetReference>& TargetsToStream, FStreamableDelegate DelegateToCall, TAsyncLoadPriority Priority, bool bManageActiveHandle, bool bStartStalled, const FString& DebugName)
 {
 	// Schedule a new callback, this will get called when all related async loads are completed
 	TSharedRef<FStreamableHandle> NewRequest = MakeShareable(new FStreamableHandle());
@@ -776,9 +790,26 @@ TSharedPtr<FStreamableHandle> FStreamableManager::RequestAsyncLoad(const TArray<
 	NewRequest->OwningManager = this;
 	NewRequest->RequestedAssets = TargetsToStream;
 	NewRequest->DebugName = DebugName;
+	NewRequest->Priority = Priority;
 
 	// Remove null requests
-	NewRequest->RequestedAssets.Remove(FStringAssetReference());
+
+	for (int32 i = NewRequest->RequestedAssets.Num() - 1; i >= 0 ; i--)
+	{
+		FStringAssetReference& TargetName = NewRequest->RequestedAssets[i];
+		if (TargetName.IsNull())
+		{
+			// Remove null entries
+			NewRequest->RequestedAssets.RemoveAt(i);
+			continue;
+		}
+		else if (FPackageName::IsShortPackageName(TargetName.ToString()))
+		{
+			UE_LOG(LogStreamableManager, Error, TEXT("RequestAsyncLoad called with invalid package name %s"), *TargetName.ToString());
+			NewRequest->CancelHandle();
+			return nullptr;
+		}
+	}
 
 	if (NewRequest->RequestedAssets.Num() == 0)
 	{
@@ -827,59 +858,37 @@ TSharedPtr<FStreamableHandle> FStreamableManager::RequestAsyncLoad(const TArray<
 		NewRequest->RequestedAssets = TargetSet.Array();
 	}
 
-	TArray<FStreamable *> ExistingStreamables;
-	ExistingStreamables.Reserve(NewRequest->RequestedAssets.Num());
-
-	for (int32 i = 0; i < NewRequest->RequestedAssets.Num(); i++)
+	if (bManageActiveHandle)
 	{
-		FStreamable* Existing = StreamInternal(NewRequest->RequestedAssets[i], Priority, NewRequest);
-
-		if (!Existing)
-		{
-			// Requested an invalid asset path
-			NewRequest->CancelHandle();
-			return nullptr;
-		}
-
-		ExistingStreamables.Add(Existing);
-		Existing->AddLoadingRequest(NewRequest);
+		// This keeps a reference around until explicitly released
+		ManagedActiveHandles.Add(NewRequest);
 	}
 
-	// Go through and complete loading anything that's already in memory, this may call the callback right away
-	for (int32 i = 0; i < NewRequest->RequestedAssets.Num(); i++)
+	if (bStartStalled)
 	{
-		FStreamable* Existing = ExistingStreamables[i];
-
-		if (Existing && (Existing->Target || Existing->bLoadFailed))
-		{
-			Existing->bAsyncLoadRequestOutstanding = false;
-
-			CheckCompletedRequests(NewRequest->RequestedAssets[i], Existing);
-		}
+		NewRequest->bStalled = true;
 	}
-
-	if (bManageActiveHandle)
+	else
 	{
-		// This keeps a reference around until explicitly released
-		ManagedActiveHandles.Add(NewRequest);
+		StartHandleRequests(NewRequest);
 	}
 
 	return NewRequest;
 }
 
-TSharedPtr<FStreamableHandle> FStreamableManager::RequestAsyncLoad(const FStringAssetReference& TargetToStream, FStreamableDelegate DelegateToCall, TAsyncLoadPriority Priority, bool bManageActiveHandle, const FString& DebugName)
+TSharedPtr<FStreamableHandle> FStreamableManager::RequestAsyncLoad(const FStringAssetReference& TargetToStream, FStreamableDelegate DelegateToCall, TAsyncLoadPriority Priority, bool bManageActiveHandle, bool bStartStalled, const FString& DebugName)
 {
-	return RequestAsyncLoad(TArray<FStringAssetReference>{TargetToStream}, DelegateToCall, Priority, bManageActiveHandle, DebugName);
+	return RequestAsyncLoad(TArray<FStringAssetReference>{TargetToStream}, DelegateToCall, Priority, bManageActiveHandle, bStartStalled, DebugName);
 }
 
-TSharedPtr<FStreamableHandle>  FStreamableManager::RequestAsyncLoad(const TArray<FStringAssetReference>& TargetsToStream, TFunction<void()>&& Callback, TAsyncLoadPriority Priority, bool bManageActiveHandle, const FString& DebugName)
+TSharedPtr<FStreamableHandle>  FStreamableManager::RequestAsyncLoad(const TArray<FStringAssetReference>& TargetsToStream, TFunction<void()>&& Callback, TAsyncLoadPriority Priority, bool bManageActiveHandle, bool bStartStalled, const FString& DebugName)
 {
-	return RequestAsyncLoad(TargetsToStream, FStreamableDelegate::CreateLambda( MoveTemp( Callback ) ), Priority, bManageActiveHandle, DebugName);
+	return RequestAsyncLoad(TargetsToStream, FStreamableDelegate::CreateLambda( MoveTemp( Callback ) ), Priority, bManageActiveHandle, bStartStalled, DebugName);
 }
 
-TSharedPtr<FStreamableHandle>  FStreamableManager::RequestAsyncLoad(const FStringAssetReference& TargetToStream, TFunction<void()>&& Callback, TAsyncLoadPriority Priority, bool bManageActiveHandle, const FString& DebugName)
+TSharedPtr<FStreamableHandle>  FStreamableManager::RequestAsyncLoad(const FStringAssetReference& TargetToStream, TFunction<void()>&& Callback, TAsyncLoadPriority Priority, bool bManageActiveHandle, bool bStartStalled, const FString& DebugName)
 {
-	return RequestAsyncLoad(TargetToStream, FStreamableDelegate::CreateLambda( MoveTemp( Callback ) ), Priority, bManageActiveHandle, DebugName);
+	return RequestAsyncLoad(TargetToStream, FStreamableDelegate::CreateLambda( MoveTemp( Callback ) ), Priority, bManageActiveHandle, bStartStalled, DebugName);
 }
 
 TSharedPtr<FStreamableHandle> FStreamableManager::RequestSyncLoad(const TArray<FStringAssetReference>& TargetsToStream, bool bManageActiveHandle, const FString& DebugName)
@@ -890,7 +899,7 @@ TSharedPtr<FStreamableHandle> FStreamableManager::RequestSyncLoad(const TArray<F
 	bForceSynchronousLoads = IsInAsyncLoadingThread() || IsEventDrivenLoaderEnabled() || !IsAsyncLoading();
 
 	// Do an async load and wait to complete. In some cases this will do a sync load due to safety issues
-	TSharedPtr<FStreamableHandle> Request = RequestAsyncLoad(TargetsToStream, FStreamableDelegate(), AsyncLoadHighPriority, bManageActiveHandle, DebugName);
+	TSharedPtr<FStreamableHandle> Request = RequestAsyncLoad(TargetsToStream, FStreamableDelegate(), AsyncLoadHighPriority, bManageActiveHandle, false, DebugName);
 
 	bForceSynchronousLoads = false;
 
@@ -910,6 +919,34 @@ TSharedPtr<FStreamableHandle> FStreamableManager::RequestSyncLoad(const FStringA
 	return RequestSyncLoad(TArray<FStringAssetReference>{TargetToStream}, bManageActiveHandle, DebugName);
 }
 
+void FStreamableManager::StartHandleRequests(TSharedRef<FStreamableHandle> Handle)
+{
+	TArray<FStreamable *> ExistingStreamables;
+	ExistingStreamables.Reserve(Handle->RequestedAssets.Num());
+
+	for (int32 i = 0; i < Handle->RequestedAssets.Num(); i++)
+	{
+		FStreamable* Existing = StreamInternal(Handle->RequestedAssets[i], Handle->Priority, Handle);
+		check(Existing);
+
+		ExistingStreamables.Add(Existing);
+		Existing->AddLoadingRequest(Handle);
+	}
+
+	// Go through and complete loading anything that's already in memory, this may call the callback right away
+	for (int32 i = 0; i < Handle->RequestedAssets.Num(); i++)
+	{
+		FStreamable* Existing = ExistingStreamables[i];
+
+		if (Existing && (Existing->Target || Existing->bLoadFailed))
+		{
+			Existing->bAsyncLoadRequestOutstanding = false;
+
+			CheckCompletedRequests(Handle->RequestedAssets[i], Existing);
+		}
+	}
+}
+
 UObject* FStreamableManager::LoadSynchronous(const FStringAssetReference& Target, bool bManageActiveHandle, TSharedPtr<FStreamableHandle>* RequestHandlePointer)
 {
 	TSharedPtr<FStreamableHandle> Request = RequestSyncLoad(Target, bManageActiveHandle, FString::Printf(TEXT("LoadSynchronous of %s"), *Target.ToString()));
diff --git a/Engine/Source/Runtime/Engine/Private/TextureRenderTarget2D.cpp b/Engine/Source/Runtime/Engine/Private/TextureRenderTarget2D.cpp
--- a/Engine/Source/Runtime/Engine/Private/TextureRenderTarget2D.cpp
+++ b/Engine/Source/Runtime/Engine/Private/TextureRenderTarget2D.cpp
@@ -469,7 +469,7 @@ void FTextureRenderTarget2DResource::UpdateDeferredResource( FRHICommandListImme
 	if (bClearRenderTarget)
 	{
 		RHICmdList.SetViewport(0, 0, 0.0f, TargetSizeX, TargetSizeY, 1.0f);
-		ensure(RenderTargetTextureRHI->GetClearColor() == ClearColor);
+		ensure(RenderTargetTextureRHI.IsValid() && (RenderTargetTextureRHI->GetClearColor() == ClearColor));
 		SetRenderTarget(RHICmdList, RenderTargetTextureRHI, FTextureRHIRef(), ESimpleRenderTargetMode::EClearColorExistingDepth, FExclusiveDepthStencil::DepthWrite_StencilWrite, true);
 	}
  
