Copying //UE4/Dev-Framework to //UE4/Dev-Main (Source: //UE4/Dev-Framework @ 3544039)

diff --git a/Engine/Source/Runtime/Engine/Private/UnrealEngine.cpp b/Engine/Source/Runtime/Engine/Private/UnrealEngine.cpp
--- a/Engine/Source/Runtime/Engine/Private/UnrealEngine.cpp
+++ b/Engine/Source/Runtime/Engine/Private/UnrealEngine.cpp
@@ -836,18 +836,178 @@ void EngineMemoryWarningHandler(const FGenericMemoryWarningContext& GenericConte
 
 UEngine::FOnNewStatRegistered UEngine::NewStatDelegate;
 
+#if !(UE_BUILD_SHIPPING || UE_BUILD_TEST)
+static TAutoConsoleVariable<int32> CVarStressTestGCWhileStreaming(
+	TEXT("gc.StressTestGC"),
+	0,
+	TEXT("If set to 1, the engine will attempt to trigger GC each frame while async loading."));
+#endif
+
+static TAutoConsoleVariable<int32> CVarCollectGarbageEveryFrame(
+	TEXT("gc.CollectGarbageEveryFrame"),
+	0,
+	TEXT("Used to debug garbage collection...Collects garbage every frame if the value is > 0."));
+
+static float GTimeBetweenPurgingPendingKillObjects = 60.0f;
+static FAutoConsoleVariableRef CVarTimeBetweenPurgingPendingKillObjects(
+	TEXT("gc.TimeBetweenPurgingPendingKillObjects"),
+	GTimeBetweenPurgingPendingKillObjects,
+	TEXT("Time in seconds (game time) we should wait between purging object references to objects that are pending kill."),
+	ECVF_Default
+);
+
+static float GTimeBetweenPurgingPendingKillObjectsOnIdleServerMultiplier = 10.0f;
+static FAutoConsoleVariableRef CVarTimeBetweenPurgingPendingKillObjectsOnIdleServerMultiplier(
+	TEXT("gc.TimeBetweenPurgingPendingKillObjectsOnIdleServerMultiplier"),
+	GTimeBetweenPurgingPendingKillObjectsOnIdleServerMultiplier,
+	TEXT("Multiplier to apply to time between purging pending kill objects when on an idle server."),
+	ECVF_Default
+);
+
 void UEngine::PreGarbageCollect()
 {
-	for (TObjectIterator<UWorld> WorldIt; WorldIt; ++WorldIt)
+	ForEachObjectOfClass(UWorld::StaticClass(), [](UObject* WorldObj)
 	{
-		UWorld* World = *WorldIt;
+		UWorld* World = CastChecked<UWorld>(WorldObj);
 
-		if (World && World->HasEndOfFrameUpdates())
+		if (World->HasEndOfFrameUpdates())
 		{
 			// Make sure deferred component updates have been sent to the rendering thread before deleting any UObjects which the rendering thread may be referencing
 			// This fixes rendering thread crashes in the following order of operations 1) UMeshComponent::SetMaterial 2) GC 3) Rendering command that dereferences the UMaterial
 			World->SendAllEndOfFrameUpdates();
 		}
+	});
+}
+
+float UEngine::GetTimeBetweenGarbageCollectionPasses() const
+{
+	float TimeBetweenGC = GTimeBetweenPurgingPendingKillObjects;
+
+	if (IsRunningDedicatedServer())
+	{
+		bool bAtLeastOnePlayerConnected = false;
+
+		ForEachObjectOfClass(UWorld::StaticClass(),[&bAtLeastOnePlayerConnected](UObject* WorldObj)
+		{
+			UWorld* World = CastChecked<UWorld>(WorldObj);
+			bAtLeastOnePlayerConnected = bAtLeastOnePlayerConnected || ((World->NetDriver != nullptr) && World->NetDriver->ClientConnections.Num() > 0);
+		});
+
+		if (!bAtLeastOnePlayerConnected)
+		{
+			TimeBetweenGC *= GTimeBetweenPurgingPendingKillObjectsOnIdleServerMultiplier;
+		}
+	}
+
+	return TimeBetweenGC;
+}
+
+void UEngine::ForceGarbageCollection(bool bForcePurge/*=false*/)
+{
+	TimeSinceLastPendingKillPurge = 1.0f + GetTimeBetweenGarbageCollectionPasses();
+	bFullPurgeTriggered = bFullPurgeTriggered || bForcePurge;
+}
+
+void UEngine::DelayGarbageCollection()
+{
+	bShouldDelayGarbageCollect = true;
+}
+
+void UEngine::SetTimeUntilNextGarbageCollection(const float MinTimeUntilNextPass)
+{
+	const float TimeBetweenPurgingPendingKillObjects = GetTimeBetweenGarbageCollectionPasses();
+
+	// This can make it go negative if the desired interval is longer than the typical interval, but it's only ever compared against TimeBetweenPurgingPendingKillObjects
+	TimeSinceLastPendingKillPurge = TimeBetweenPurgingPendingKillObjects - MinTimeUntilNextPass;
+}
+
+void UEngine::ConditionalCollectGarbage()
+{
+	if (GFrameCounter != LastGCFrame)
+	{
+	#if !(UE_BUILD_SHIPPING || UE_BUILD_TEST)
+		if (CVarStressTestGCWhileStreaming.GetValueOnGameThread() && IsAsyncLoading())
+		{
+			TryCollectGarbage(GARBAGE_COLLECTION_KEEPFLAGS, true);
+		}
+		else
+	#endif
+		if (bFullPurgeTriggered)
+		{
+			if (TryCollectGarbage(GARBAGE_COLLECTION_KEEPFLAGS, true))
+			{
+				ForEachObjectOfClass(UWorld::StaticClass(),[](UObject* World)
+				{
+					CastChecked<UWorld>(World)->CleanupActors();
+				});
+				bFullPurgeTriggered = false;
+				bShouldDelayGarbageCollect = false;
+				TimeSinceLastPendingKillPurge = 0.0f;
+			}
+		}
+		else
+		{
+			bool bHasAWorldBegunPlay = false;
+			ForEachObjectOfClass(UWorld::StaticClass(), [&bHasAWorldBegunPlay](UObject* World)
+			{
+				bHasAWorldBegunPlay = bHasAWorldBegunPlay || CastChecked<UWorld>(World)->HasBegunPlay();
+			});
+
+			if (bHasAWorldBegunPlay)
+			{
+				TimeSinceLastPendingKillPurge += FApp::GetDeltaTime();
+
+				const float TimeBetweenPurgingPendingKillObjects = GetTimeBetweenGarbageCollectionPasses();
+
+				// See if we should delay garbage collect for this frame
+				if (bShouldDelayGarbageCollect)
+				{
+					bShouldDelayGarbageCollect = false;
+				}
+				// Perform incremental purge update if it's pending or in progress.
+				else if (!IsIncrementalPurgePending()
+					// Purge reference to pending kill objects every now and so often.
+					&& (TimeSinceLastPendingKillPurge > TimeBetweenPurgingPendingKillObjects) && TimeBetweenPurgingPendingKillObjects > 0.f)
+				{
+					SCOPE_CYCLE_COUNTER(STAT_GCMarkTime);
+					PerformGarbageCollectionAndCleanupActors();
+				}
+				else
+				{
+					SCOPE_CYCLE_COUNTER(STAT_GCSweepTime);
+					IncrementalPurgeGarbage(true);
+				}
+			}
+		}
+
+		if (CVarCollectGarbageEveryFrame.GetValueOnGameThread() > 0)
+		{
+			ForceGarbageCollection(true);
+		}
+
+		LastGCFrame = GFrameCounter;
+	}
+}
+
+void UEngine::PerformGarbageCollectionAndCleanupActors()
+{
+	// We don't collect garbage while there are outstanding async load requests as we would need
+	// to block on loading the remaining data.
+	if (!IsAsyncLoading())
+	{
+		// Perform housekeeping.
+		if (TryCollectGarbage(GARBAGE_COLLECTION_KEEPFLAGS, false))
+		{
+			ForEachObjectOfClass(UWorld::StaticClass(), [](UObject* World)
+			{
+				CastChecked<UWorld>(World)->CleanupActors();
+			});
+
+			// Reset counter.
+			TimeSinceLastPendingKillPurge = 0.0f;
+			bFullPurgeTriggered = false;
+			LastGCFrame = GFrameCounter;
+		}
 	}
 }
 
@@ -2019,11 +2179,6 @@ FAudioDevice* UEngine::GetActiveAudioDevice()
 	return nullptr;
 }
 
-FAudioDevice* UEngine::GetAudioDevice()
-{
-	return GetMainAudioDevice();
-}
-
 /**
  *	Initialize the audio device
  *
@@ -4888,7 +5043,7 @@ bool UEngine::HandleMergeMeshCommand( const TCHAR* Cmd, FOutputDevice& Ar, UWorl
 	if (PlayerPawn && SourceMeshList.Num() >= 2)
 	{
 		// create the composite mesh
-		auto CompositeMesh = NewObject<USkeletalMesh>( GetTransientPackage(), NAME_None, RF_Transient );
+		USkeletalMesh* CompositeMesh = NewObject<USkeletalMesh>( GetTransientPackage(), NAME_None, RF_Transient );
 
 		TArray<FSkelMeshMergeSectionMapping> InForceSectionMapping;
 		// create an instance of the FSkeletalMeshMerge utility
@@ -5196,13 +5351,9 @@ bool UEngine::HandleObjCommand( const TCHAR* Cmd, FOutputDevice& Ar )
 	if( FParse::Command(&Cmd,TEXT("GARBAGE")) || FParse::Command(&Cmd,TEXT("GC")) )
 	{
 		// Purge unclaimed objects.
-		Ar.Logf(TEXT("Collecting garbage and resetting GC timers on all worlds."));
+		Ar.Logf(TEXT("Collecting garbage and resetting GC timer."));
 		CollectGarbage( GARBAGE_COLLECTION_KEEPFLAGS );
-		for (TObjectIterator<UWorld> It; It; ++It)
-		{
-			UWorld* CurrentWorld = *It;
-			CurrentWorld->TimeSinceLastPendingKillPurge = 0;
-		}
+		TimeSinceLastPendingKillPurge = 0.f;
 		return true;
 	}
 	else if (FParse::Command(&Cmd, TEXT("TRYGC")))
@@ -5210,12 +5361,8 @@ bool UEngine::HandleObjCommand( const TCHAR* Cmd, FOutputDevice& Ar )
 		// Purge unclaimed objects.		
 		if (TryCollectGarbage(GARBAGE_COLLECTION_KEEPFLAGS))
 		{
-			Ar.Logf(TEXT("Collecting garbage and resetting GC timers on all worlds."));
-			for (TObjectIterator<UWorld> It; It; ++It)
-			{
-				UWorld* CurrentWorld = *It;
-				CurrentWorld->TimeSinceLastPendingKillPurge = 0;
-			}
+			Ar.Logf(TEXT("Collecting garbage and resetting GC timer."));
+			TimeSinceLastPendingKillPurge = 0.f;
 		}
 		else
 		{
diff --git a/Engine/Source/Runtime/Engine/Private/UserInterface/PlayerInput.cpp b/Engine/Source/Runtime/Engine/Private/UserInterface/PlayerInput.cpp
--- a/Engine/Source/Runtime/Engine/Private/UserInterface/PlayerInput.cpp
+++ b/Engine/Source/Runtime/Engine/Private/UserInterface/PlayerInput.cpp
@@ -27,6 +27,7 @@ static float DebugSmoothedMouseY = 0.f;
 static float DebugSmoothedMouseSensitivity = 8.f;
 
 const TArray<FInputActionKeyMapping> UPlayerInput::NoKeyMappings;
+const TArray<FInputAxisKeyMapping> UPlayerInput::NoAxisMappings;
 TArray<FInputActionKeyMapping> UPlayerInput::EngineDefinedActionMappings;
 TArray<FInputAxisKeyMapping> UPlayerInput::EngineDefinedAxisMappings;
 
@@ -452,20 +453,6 @@ void UPlayerInput::SetMouseSensitivity(const float Sensitivity)
 	}
 }
 
-void UPlayerInput::SetMouseSensitivityToDefault()
-{
-	// find default sensitivity restore to that
-	for (const FInputAxisConfigEntry& AxisConfigEntry : GetDefault<UInputSettings>()->AxisConfig)
-	{
-		const FKey AxisKey = AxisConfigEntry.AxisKeyName;
-		if (AxisKey == EKeys::MouseX)
-		{
-			SetMouseSensitivity(AxisConfigEntry.AxisProperties.Sensitivity);
-			break;
-		}
-	}
-}
-
 bool UPlayerInput::GetInvertAxis(const FName AxisName)
 {
 	ConditionalBuildKeyMappings();
@@ -546,11 +533,6 @@ void UPlayerInput::InvertAxisKey(const FKey AxisKey)
 	}
 }
 
-void UPlayerInput::InvertMouse()
-{
-	InvertAxisKey(EKeys::MouseY);
-}
-
 struct FAxisDelegate
 {
 	UObject* Obj;
@@ -1768,6 +1750,19 @@ const TArray<FInputActionKeyMapping>& UPlayerInput::GetKeysForAction(const FName
 	return UPlayerInput::NoKeyMappings;
 }
 
+const TArray<FInputAxisKeyMapping>& UPlayerInput::GetKeysForAxis(const FName AxisName)
+{
+	ConditionalBuildKeyMappings();
+
+	const FAxisKeyDetails* KeyDetails = AxisKeyMap.Find(AxisName);
+	if (KeyDetails)
+	{
+		return KeyDetails->KeyMappings;
+	}
+
+	return UPlayerInput::NoAxisMappings;
+}
+
 #if !UE_BUILD_SHIPPING
 bool UPlayerInput::ExecInputCommands( UWorld* InWorld, const TCHAR* Cmd, FOutputDevice& Ar )
 {
diff --git a/Engine/Source/Runtime/Engine/Private/World.cpp b/Engine/Source/Runtime/Engine/Private/World.cpp
--- a/Engine/Source/Runtime/Engine/Private/World.cpp
+++ b/Engine/Source/Runtime/Engine/Private/World.cpp
@@ -2750,7 +2750,7 @@ void UWorld::UpdateLevelStreamingInner(ULevelStreaming* StreamingLevel)
 	// kicking off the async load.
 	if (bAreLevelsPendingPurge)
 	{
-		ForceGarbageCollection( false );
+		GEngine->ForceGarbageCollection( false );
 	}
 
 	// See whether level is already loaded
@@ -2857,7 +2857,7 @@ void UWorld::UpdateLevelStreaming()
 	// In case more levels has been requested to unload, force GC on next tick 
 	if (NumLevelsPendingPurge < FLevelStreamingGCHelper::GetNumLevelsPendingPurge())
 	{
-		ForceGarbageCollection(true); 
+		GEngine->ForceGarbageCollection(true); 
 	}
 }
 
@@ -2909,11 +2909,6 @@ void UWorld::FlushLevelStreaming(EFlushLevelStreamingType FlushType)
 	}
 }
 
-void UWorld::FlushLevelStreaming(EFlushLevelStreamingType FlushType, FName ExcludeType)
-{
-	FlushLevelStreaming(FlushType);
-}
-
 /**
  * Forces streaming data to be rebuilt for the current world.
  */
@@ -3730,11 +3725,6 @@ void UWorld::RemoveOnActorSpawnedHandler( FDelegateHandle InHandle )
 	OnActorSpawned.Remove(InHandle);
 }
 
-ABrush* UWorld::GetBrush() const
-{
-	return GetDefaultBrush();
-}
-
 ABrush* UWorld::GetDefaultBrush() const
 {
 	check(PersistentLevel);
diff --git a/Engine/Source/Runtime/Engine/Public/AI/Navigation/PImplRecastNavMesh.h b/Engine/Source/Runtime/Engine/Public/AI/Navigation/PImplRecastNavMesh.h
--- a/Engine/Source/Runtime/Engine/Public/AI/Navigation/PImplRecastNavMesh.h
+++ b/Engine/Source/Runtime/Engine/Public/AI/Navigation/PImplRecastNavMesh.h
@@ -91,11 +91,6 @@ public:
 	/** Returns a random location on the navmesh. */
 	FNavLocation GetRandomPoint(const FNavigationQueryFilter& Filter, const UObject* Owner) const;
 
-	/** Returns a random location on the navmesh within Radius from Origin. 
-	 *	@return false if no valid navigable location available in specified area */
-	DEPRECATED(4.8, "GetRandomPointInRadius is deprecated, please use ANavigationData::GetRandomReachablePointInRadius")
-	bool GetRandomPointInRadius(const FVector& Origin, float Radius, FNavLocation& OutLocation, const FNavigationQueryFilter& Filter, const UObject* Owner) const;
-
 	/** Returns a random location on the navmesh within cluster */
 	bool GetRandomPointInCluster(NavNodeRef ClusterRef, FNavLocation& OutLocation) const;
 
diff --git a/Engine/Source/Runtime/Engine/Public/MaterialShared.h b/Engine/Source/Runtime/Engine/Public/MaterialShared.h
--- a/Engine/Source/Runtime/Engine/Public/MaterialShared.h
+++ b/Engine/Source/Runtime/Engine/Public/MaterialShared.h
@@ -1737,8 +1737,6 @@ public:
 	ENGINE_API virtual bool IsUsedWithSplineMeshes() const override;
 	ENGINE_API virtual bool IsUsedWithInstancedStaticMeshes() const override;
 	ENGINE_API virtual bool IsUsedWithAPEXCloth() const override;
-	DEPRECATED(4.9, "IsUsedWithUI is now replaced by IsUIMaterial")
-	ENGINE_API virtual bool IsUsedWithUI() const override;
 	ENGINE_API virtual enum EMaterialTessellationMode GetTessellationMode() const override;
 	ENGINE_API virtual bool IsCrackFreeDisplacementEnabled() const override;
 	ENGINE_API virtual bool IsAdaptiveTessellationEnabled() const override;
diff --git a/Engine/Source/Runtime/Engine/Public/PhysicsPublic.h b/Engine/Source/Runtime/Engine/Public/PhysicsPublic.h
--- a/Engine/Source/Runtime/Engine/Public/PhysicsPublic.h
+++ b/Engine/Source/Runtime/Engine/Public/PhysicsPublic.h
@@ -536,83 +536,27 @@ public:
 	/** Returns whether an async scene is setup and can be used. This depends on the console variable "p.EnableAsyncScene". */
 	ENGINE_API bool HasAsyncScene() const { return bAsyncSceneEnabled; }
 
-	/** Lets the scene update anything related to this BodyInstance as it's now being terminated */
-	DEPRECATED(4.8, "Please call AddCustomPhysics_AssumesLocked and make sure you obtain the appropriate PhysX scene locks")
-	void TermBody(FBodyInstance* BodyInstance)
-	{
-		TermBody_AssumesLocked(BodyInstance);
-	}
-
 	/** Lets the scene update anything related to this BodyInstance as it's now being terminated */
 	void TermBody_AssumesLocked(FBodyInstance* BodyInstance);
 
-	/** Add a custom callback for next step that will be called on every substep */
-	DEPRECATED(4.8, "Please call AddCustomPhysics_AssumesLocked and make sure you obtain the appropriate PhysX scene locks")
-	void AddCustomPhysics(FBodyInstance* BodyInstance, FCalculateCustomPhysics& CalculateCustomPhysics)
-	{
-		AddCustomPhysics_AssumesLocked(BodyInstance, CalculateCustomPhysics);
-	}
-
 	/** Add a custom callback for next step that will be called on every substep */
 	void AddCustomPhysics_AssumesLocked(FBodyInstance* BodyInstance, FCalculateCustomPhysics& CalculateCustomPhysics);
 
 	/** Adds a force to a body - We need to go through scene to support substepping */
-	DEPRECATED(4.8, "Please call AddForce_AssumesLocked and make sure you obtain the appropriate PhysX scene locks")
-	void AddForce(FBodyInstance* BodyInstance, const FVector& Force, bool bAllowSubstepping, bool bAccelChange)
-	{
-		AddForce_AssumesLocked(BodyInstance, Force, bAllowSubstepping, bAccelChange);
-	}
-
 	void AddForce_AssumesLocked(FBodyInstance* BodyInstance, const FVector& Force, bool bAllowSubstepping, bool bAccelChange);
 
-	/** Adds a force to a body at a specific position - We need to go through scene to support substepping */
-	DEPRECATED(4.8, "Please call AddForceAtPosition_AssumesLocked and make sure you obtain the appropriate PhysX scene locks")
-	void AddForceAtPosition(FBodyInstance* BodyInstance, const FVector& Force, const FVector& Position, bool bAllowSubstepping)
-	{
-		AddForceAtPosition_AssumesLocked(BodyInstance, Force, Position, bAllowSubstepping);
-	}
-
 	/** Adds a force to a body at a specific position - We need to go through scene to support substepping */
 	void AddForceAtPosition_AssumesLocked(FBodyInstance* BodyInstance, const FVector& Force, const FVector& Position, bool bAllowSubstepping, bool bIsLocalForce=false);
 
-	/** Adds a radial force to a body - We need to go through scene to support substepping */
-	DEPRECATED(4.8, "Please call AddRadialForceToBody_AssumesLocked and make sure you obtain the appropriate PhysX scene locks")
-	void AddRadialForceToBody(FBodyInstance* BodyInstance, const FVector& Origin, const float Radius, const float Strength, const uint8 Falloff, bool bAccelChange, bool bAllowSubstepping)
-	{
-		AddRadialForceToBody_AssumesLocked(BodyInstance, Origin, Radius, Strength, Falloff, bAccelChange, bAllowSubstepping);
-	}
-
 	/** Adds a radial force to a body - We need to go through scene to support substepping */
 	void AddRadialForceToBody_AssumesLocked(FBodyInstance* BodyInstance, const FVector& Origin, const float Radius, const float Strength, const uint8 Falloff, bool bAccelChange, bool bAllowSubstepping);
 
-	/** Adds torque to a body - We need to go through scene to support substepping */
-	DEPRECATED(4.8, "Please call AddTorque_AssumesLocked and make sure you obtain the appropriate PhysX scene locks")
-	void AddTorque(FBodyInstance* BodyInstance, const FVector& Torque, bool bAllowSubstepping, bool bAccelChange)
-	{
-		AddTorque_AssumesLocked(BodyInstance, Torque, bAllowSubstepping, bAccelChange);
-	}
-
 	/** Adds torque to a body - We need to go through scene to support substepping */
 	void AddTorque_AssumesLocked(FBodyInstance* BodyInstance, const FVector& Torque, bool bAllowSubstepping, bool bAccelChange);
 
-	/** Sets a Kinematic actor's target position - We need to do this here to support substepping*/
-	DEPRECATED(4.8, "Please call SetKinematicTarget_AssumesLocked and make sure you obtain the appropriate PhysX scene locks")
-	void SetKinematicTarget(FBodyInstance* BodyInstance, const FTransform& TargetTM, bool bAllowSubstepping)
-	{
-		SetKinematicTarget_AssumesLocked(BodyInstance, TargetTM, bAllowSubstepping);
-	}
-	
 	/** Sets a Kinematic actor's target position - We need to do this here to support substepping*/
 	void SetKinematicTarget_AssumesLocked(FBodyInstance* BodyInstance, const FTransform& TargetTM, bool bAllowSubstepping);
 
-	/** Gets a Kinematic actor's target position - We need to do this here to support substepping
-	  * Returns true if kinematic target has been set. If false the OutTM is invalid */
-	DEPRECATED(4.8, "Please call GetKinematicTarget_AssumesLocked and make sure you obtain the appropriate PhysX scene locks")
-	bool GetKinematicTarget(const FBodyInstance* BodyInstance, FTransform& OutTM) const
-	{
-		return GetKinematicTarget_AssumesLocked(BodyInstance, OutTM);
-	}
-
 	/** Gets a Kinematic actor's target position - We need to do this here to support substepping
 	  * Returns true if kinematic target has been set. If false the OutTM is invalid */
 	bool GetKinematicTarget_AssumesLocked(const FBodyInstance* BodyInstance, FTransform& OutTM) const;
diff --git a/Engine/Source/Runtime/Engine/Public/SkeletalMeshMerge.h b/Engine/Source/Runtime/Engine/Public/SkeletalMeshMerge.h
--- a/Engine/Source/Runtime/Engine/Public/SkeletalMeshMerge.h
+++ b/Engine/Source/Runtime/Engine/Public/SkeletalMeshMerge.h
@@ -42,22 +42,6 @@ struct FRefPoseOverride
 		this->SkeletalMesh = ReferenceMesh;
 	}
 
-	/**
-	 * Adds a bone to the list of poses to override.
-	 */
-	DEPRECATED(4.8, "Use AddOverride(FName, EBoneOverrideMode) instead.")
-	void AddOverride(FName BoneName, bool bOverrideChildren)
-	{
-		EBoneOverrideMode OverrideMode = BoneAndChildren;
-
-		if (!bOverrideChildren)
-		{
-			OverrideMode = BoneOnly;
-		}
-
-		AddOverride(BoneName, OverrideMode);
-	}
-
 	/**
 	 * Adds a bone to the list of poses to override.
 	 */
diff --git a/Engine/Source/Runtime/GameplayTags/Private/GameplayTagContainer.cpp b/Engine/Source/Runtime/GameplayTags/Private/GameplayTagContainer.cpp
--- a/Engine/Source/Runtime/GameplayTags/Private/GameplayTagContainer.cpp
+++ b/Engine/Source/Runtime/GameplayTags/Private/GameplayTagContainer.cpp
@@ -1331,7 +1331,7 @@ bool FGameplayTag::ImportTextItem(const TCHAR*& Buffer, int32 PortFlags, UObject
 	return false;
 }
 
-void FGameplayTag::FromExportString(FString ExportString)
+void FGameplayTag::FromExportString(const FString& ExportString)
 {
 	TagName = NAME_None;
 
diff --git a/Engine/Source/Runtime/Landscape/Classes/LandscapeComponent.h b/Engine/Source/Runtime/Landscape/Classes/LandscapeComponent.h
--- a/Engine/Source/Runtime/Landscape/Classes/LandscapeComponent.h
+++ b/Engine/Source/Runtime/Landscape/Classes/LandscapeComponent.h
@@ -221,7 +221,7 @@ struct FLandscapeComponentGrassData
 	friend FArchive& operator<<(FArchive& Ar, FLandscapeComponentGrassData& Data);
 };
 
-UCLASS(hidecategories=(Display, Attachment, Physics, Debug, Collision, Movement, Rendering, PrimitiveComponent, Object, Transform), showcategories=("Rendering|Material"), MinimalAPI)
+UCLASS(hidecategories=(Display, Attachment, Physics, Debug, Collision, Movement, Rendering, PrimitiveComponent, Object, Transform), showcategories=("Rendering|Material"), MinimalAPI, Within=LandscapeProxy)
 class ULandscapeComponent : public UPrimitiveComponent
 {
 	GENERATED_UCLASS_BODY()
diff --git a/Engine/Source/Runtime/Landscape/Classes/LandscapeHeightfieldCollisionComponent.h b/Engine/Source/Runtime/Landscape/Classes/LandscapeHeightfieldCollisionComponent.h
--- a/Engine/Source/Runtime/Landscape/Classes/LandscapeHeightfieldCollisionComponent.h
+++ b/Engine/Source/Runtime/Landscape/Classes/LandscapeHeightfieldCollisionComponent.h
@@ -31,7 +31,7 @@ namespace physx
 }
 #endif // WITH_PHYSX
 
-UCLASS(MinimalAPI)
+UCLASS(MinimalAPI, Within=LandscapeProxy)
 class ULandscapeHeightfieldCollisionComponent : public UPrimitiveComponent
 {
 	GENERATED_UCLASS_BODY()
diff --git a/Engine/Source/Runtime/MRMesh/Private/MRMeshComponent.cpp b/Engine/Source/Runtime/MRMesh/Private/MRMeshComponent.cpp
--- a/Engine/Source/Runtime/MRMesh/Private/MRMeshComponent.cpp
+++ b/Engine/Source/Runtime/MRMesh/Private/MRMeshComponent.cpp
@@ -16,6 +16,7 @@
 #include "PhysicsPublic.h"
 #include "IPhysXCooking.h"
 #include "PhysXCookHelper.h"
+#include "Misc/RuntimeErrors.h"
 
 /**
  * A vertex buffer with some dummy data to send down for MRMesh vertex components that we aren't feeding at the moment.
@@ -384,10 +385,13 @@ UMRMeshComponent::UMRMeshComponent(const FObjectInitializer& ObjectInitializer)
 
 void UMRMeshComponent::ConnectReconstructor(UMeshReconstructorBase* Reconstructor)
 {
-	if (ensure(MeshReconstructor == nullptr))
+	if (ensureAsRuntimeWarning(Reconstructor != nullptr))
 	{
-		MeshReconstructor = Reconstructor;
-		MeshReconstructor->ConnectMRMesh(this);
+		if (ensure(MeshReconstructor == nullptr))
+		{
+			MeshReconstructor = Reconstructor;
+			MeshReconstructor->ConnectMRMesh(this);
+		}
 	}
 }
 
diff --git a/Engine/Source/Runtime/MediaAssets/Public/MediaPlaylist.h b/Engine/Source/Runtime/MediaAssets/Public/MediaPlaylist.h
--- a/Engine/Source/Runtime/MediaAssets/Public/MediaPlaylist.h
+++ b/Engine/Source/Runtime/MediaAssets/Public/MediaPlaylist.h
@@ -83,6 +83,7 @@ public:
 	UFUNCTION(BlueprintCallable, Category="Media|MediaPlaylist")
 	void Insert(UMediaSource* MediaSource, int32 Index)
 	{
+		Index = FMath::Clamp(Index, 0, Items.Num());
 		Items.Insert(MediaSource, Index);
 	}
 
@@ -118,7 +119,10 @@ public:
 	UFUNCTION(BlueprintCallable, Category="Media|MediaPlaylist")
 	void RemoveAt(int32 Index)
 	{
-		Items.RemoveAt(Index);
+		if (Items.IsValidIndex(Index))
+		{
+			Items.RemoveAt(Index);
+		}
 	}
 
 protected:
diff --git a/Engine/Source/Runtime/MovieScene/Private/Evaluation/MovieSceneEvaluationTemplateInstance.cpp b/Engine/Source/Runtime/MovieScene/Private/Evaluation/MovieSceneEvaluationTemplateInstance.cpp
--- a/Engine/Source/Runtime/MovieScene/Private/Evaluation/MovieSceneEvaluationTemplateInstance.cpp
+++ b/Engine/Source/Runtime/MovieScene/Private/Evaluation/MovieSceneEvaluationTemplateInstance.cpp
@@ -202,7 +202,7 @@ void FMovieSceneRootEvaluationTemplateInstance::Evaluate(FMovieSceneContext Cont
 
 	ensureMsgf(Instance, TEXT("Could not find instance for supplied sequence ID."));
 
-	const int32 FieldIndex = Instance ? Instance->Template->EvaluationField.GetSegmentFromTime(Context.GetTime() * Instance->RootToSequenceTransform) : INDEX_NONE;
+	const int32 FieldIndex = (Instance && Instance->Template) ? Instance->Template->EvaluationField.GetSegmentFromTime(Context.GetTime() * Instance->RootToSequenceTransform) : INDEX_NONE;
 	if (FieldIndex == INDEX_NONE)
 	{
 		CallSetupTearDown(Player);
diff --git a/Engine/Source/Runtime/MovieScene/Private/MovieSceneSequencePlayer.cpp b/Engine/Source/Runtime/MovieScene/Private/MovieSceneSequencePlayer.cpp
--- a/Engine/Source/Runtime/MovieScene/Private/MovieSceneSequencePlayer.cpp
+++ b/Engine/Source/Runtime/MovieScene/Private/MovieSceneSequencePlayer.cpp
@@ -4,7 +4,7 @@
 #include "MovieScene.h"
 #include "MovieSceneSequence.h"
 #include "Engine/Engine.h"
-
+#include "Misc/RuntimeErrors.h"
 
 bool FMovieSceneSequencePlaybackSettings::SerializeFromMismatchedTag( const FPropertyTag& Tag, FArchive& Ar )
 {
@@ -98,9 +98,10 @@ void UMovieSceneSequencePlayer::PlayInternal()
 		}
 
 		UMovieSceneSequence* MovieSceneSequence = RootTemplateInstance.GetSequence(MovieSceneSequenceID::Root);
-		TOptional<float> FixedFrameInterval = MovieSceneSequence->GetMovieScene() ? MovieSceneSequence->GetMovieScene()->GetOptionalFixedFrameInterval() : TOptional<float>();
+		const bool bHasValidMovieScene = (MovieSceneSequence && MovieSceneSequence->GetMovieScene());
+		TOptional<float> FixedFrameInterval = bHasValidMovieScene ? MovieSceneSequence->GetMovieScene()->GetOptionalFixedFrameInterval() : TOptional<float>();
 
-		if (FixedFrameInterval.IsSet() && MovieSceneSequence->GetMovieScene()->GetForceFixedFrameIntervalPlayback())
+		if (FixedFrameInterval.IsSet() && bHasValidMovieScene && MovieSceneSequence->GetMovieScene()->GetForceFixedFrameIntervalPlayback())
 		{
 			OldMaxTickRate = GEngine->GetMaxFPS();
 			GEngine->SetMaxFPS(1.f / FixedFrameInterval.GetValue());
@@ -181,9 +182,12 @@ void UMovieSceneSequencePlayer::Scrub()
 	// @todo Sequencer playback: Should we recreate the instance every time?
 	// We must not recreate the instance since it holds stateful information (such as which objects it has spawned). Recreating the instance would break any 
 	// @todo: Is this still the case now that eval state is stored (correctly) in the player?
-	if (!RootTemplateInstance.IsValid())
+	if (ensureAsRuntimeWarning(Sequence != nullptr))
 	{
-		RootTemplateInstance.Initialize(*Sequence, *this);
+		if (!RootTemplateInstance.IsValid())
+		{
+			RootTemplateInstance.Initialize(*Sequence, *this);
+		}
 	}
 
 	Status = EMovieScenePlayerStatus::Scrubbing;
@@ -327,7 +331,7 @@ void UMovieSceneSequencePlayer::UpdateTimeCursorPosition(float NewPosition, TOpt
 	float Length = GetLength();
 
 	UMovieSceneSequence* MovieSceneSequence = RootTemplateInstance.GetSequence(MovieSceneSequenceID::Root);
-	TOptional<float> FixedFrameInterval = MovieSceneSequence->GetMovieScene() ? MovieSceneSequence->GetMovieScene()->GetOptionalFixedFrameInterval() : TOptional<float>();
+	TOptional<float> FixedFrameInterval = (MovieSceneSequence && MovieSceneSequence->GetMovieScene()) ? MovieSceneSequence->GetMovieScene()->GetOptionalFixedFrameInterval() : TOptional<float>();
 
 	if (bPendingFirstUpdate)
 	{
diff --git a/Engine/Source/Runtime/Niagara/Private/NiagaraComponent.cpp b/Engine/Source/Runtime/Niagara/Private/NiagaraComponent.cpp
--- a/Engine/Source/Runtime/Niagara/Private/NiagaraComponent.cpp
+++ b/Engine/Source/Runtime/Niagara/Private/NiagaraComponent.cpp
@@ -616,26 +616,34 @@ void UNiagaraComponent::SetNiagaraEmitterSpawnRate(const FString& InEmitterName,
 
 const FNiagaraScriptDataInterfaceInfo* UNiagaraComponent::GetEffectDataInterface(FGuid Id) 
 {
-	UNiagaraEffect* Effect = GetAsset();
-	const TArray<FNiagaraScriptDataInterfaceInfo>& Params = Effect->GetEffectScript()->DataInterfaceInfo;
-	const FNiagaraScriptDataInterfaceInfo* SrcVariable = Params.FindByPredicate([&](const FNiagaraScriptDataInterfaceInfo& Var)
+	if (UNiagaraEffect* Effect = GetAsset())
 	{
-		return Var.Id == Id;
-	});
+		const TArray<FNiagaraScriptDataInterfaceInfo>& Params = Effect->GetEffectScript()->DataInterfaceInfo;
+		const FNiagaraScriptDataInterfaceInfo* SrcVariable = Params.FindByPredicate([&](const FNiagaraScriptDataInterfaceInfo& Var)
+		{
+			return Var.Id == Id;
+		});
 
-	return SrcVariable;
+		return SrcVariable;
+	}
+
+	return nullptr;
 }
 
 const FNiagaraScriptDataInterfaceInfo* UNiagaraComponent::GetEffectDataInterface(FName InName)
 {
-	UNiagaraEffect* Effect = GetAsset();
-	const TArray<FNiagaraScriptDataInterfaceInfo>& Params = Effect->GetEffectScript()->DataInterfaceInfo;
-	const FNiagaraScriptDataInterfaceInfo* SrcVariable = Params.FindByPredicate([&](const FNiagaraScriptDataInterfaceInfo& Var)
+	if (UNiagaraEffect* Effect = GetAsset())
 	{
-		return Var.Name == InName;
-	});
+		const TArray<FNiagaraScriptDataInterfaceInfo>& Params = Effect->GetEffectScript()->DataInterfaceInfo;
+		const FNiagaraScriptDataInterfaceInfo* SrcVariable = Params.FindByPredicate([&](const FNiagaraScriptDataInterfaceInfo& Var)
+		{
+			return Var.Name == InName;
+		});
 
-	return SrcVariable;
+		return SrcVariable;
+	}
+
+	return nullptr;
 }
 
 FNiagaraScriptDataInterfaceInfo* UNiagaraComponent::GetLocalOverrideDataInterface(FGuid Id)
diff --git a/Engine/Source/Runtime/RuntimeAssetCache/Private/RuntimeAssetCache.cpp b/Engine/Source/Runtime/RuntimeAssetCache/Private/RuntimeAssetCache.cpp
--- a/Engine/Source/Runtime/RuntimeAssetCache/Private/RuntimeAssetCache.cpp
+++ b/Engine/Source/Runtime/RuntimeAssetCache/Private/RuntimeAssetCache.cpp
@@ -7,6 +7,7 @@
 #include "RuntimeAssetCachePluginInterface.h"
 #include "RuntimeAssetCacheBackend.h"
 #include "RuntimeAssetCacheBucket.h"
+#include "Misc/RuntimeErrors.h"
 
 DEFINE_STAT(STAT_RAC_ASyncWaitTime);
 
@@ -42,12 +43,12 @@ FRuntimeAssetCache::~FRuntimeAssetCache()
 
 int32 FRuntimeAssetCache::GetCacheSize(FName Bucket) const
 {
-	return ensure(Buckets.Contains(Bucket)) ? Buckets[Bucket]->GetSize() : 0;
+	return ensureAsRuntimeWarning(Buckets.Contains(Bucket)) ? Buckets[Bucket]->GetSize() : 0;
 }
 
 int32 FRuntimeAssetCache::GetAsynchronous(IRuntimeAssetCacheBuilder* CacheBuilder, const FOnRuntimeAssetCacheAsyncComplete& OnComplete)
 {
-	if (ensure(CacheBuilder != nullptr))
+	if (ensureAsRuntimeWarning(CacheBuilder != nullptr))
 	{
 		int32 Handle = GetNextHandle();
 
@@ -85,7 +86,7 @@ int32 FRuntimeAssetCache::GetAsynchronous(IRuntimeAssetCacheBuilder* CacheBuilde
 
 FVoidPtrParam FRuntimeAssetCache::GetSynchronous(IRuntimeAssetCacheBuilder* CacheBuilder)
 {
-	if (ensure(CacheBuilder))
+	if (ensureAsRuntimeWarning(CacheBuilder))
 	{
 		if (ensureMsgf(!CacheBuilder->ShouldBuildAsynchronously(), TEXT("CacheBuilder %s can be only called asynchronously."), CacheBuilder->GetBuilderName()))
 		{
@@ -153,7 +154,7 @@ void FRuntimeAssetCache::WaitAsynchronousCompletion(int32 Handle)
 			AsyncTask = PendingTasks.FindRef(Handle);
 		}
 
-		if (ensure(AsyncTask))
+		if (ensureAsRuntimeWarning(AsyncTask))
 		{
 			AsyncTask->EnsureCompletion();
 		}
@@ -167,13 +168,13 @@ FVoidPtrParam FRuntimeAssetCache::GetAsynchronousResults(int32 Handle)
 	{
 		FScopeLock ScopeLock(&SynchronizationObject);
 		PendingTasks.RemoveAndCopyValue(Handle, AsyncTask);
-		if (ensure(AsyncTask))
+		if (ensureAsRuntimeWarning(AsyncTask))
 		{
 			AsyncTask->GetTask().FireCompletionDelegate();
 		}
 	}
 
-	if (ensure(AsyncTask))
+	if (ensureAsRuntimeWarning(AsyncTask))
 	{
 		FVoidPtrParam Data = AsyncTask->GetTask().GetDataAndSize();
 		delete AsyncTask;
@@ -192,7 +193,7 @@ bool FRuntimeAssetCache::PollAsynchronousCompletion(int32 Handle)
 		FScopeLock ScopeLock(&SynchronizationObject);
 		AsyncTask = PendingTasks.FindRef(Handle);
 	}
-	ensure(AsyncTask);
+	ensureAsRuntimeWarning(AsyncTask);
 	return AsyncTask ? AsyncTask->IsDone() : true;
 }
 
diff --git a/Engine/Source/Runtime/Slate/Private/Framework/Application/MenuStack.cpp b/Engine/Source/Runtime/Slate/Private/Framework/Application/MenuStack.cpp
--- a/Engine/Source/Runtime/Slate/Private/Framework/Application/MenuStack.cpp
+++ b/Engine/Source/Runtime/Slate/Private/Framework/Application/MenuStack.cpp
@@ -157,43 +157,6 @@ namespace
 
 }	// anon namespace
 
-
-TSharedRef<SWindow> FMenuStack::PushMenu( const TSharedRef<SWindow>& ParentWindow, const TSharedRef<SWidget>& InContent, const FVector2D& SummonLocation, const FPopupTransitionEffect& TransitionEffect, const bool bFocusImmediately, const bool bShouldAutoSize, const FVector2D& WindowSize, const FVector2D& SummonLocationSize)
-{
-	// Deprecated method that is no longer called by its deprecated counterpart in FSlateApplication so it will only create a dummy message warning the caller not to use this method
-	// Backwards compatibility of the API is taken care of by FSlateApplication.
-	ensureMsgf(false, TEXT("PushMenu() returning an SWindow is deprecated. Use Push() that returns an IMenu instead."));
-
-	FSlateRect Anchor(SummonLocation, SummonLocation + SummonLocationSize);
-	FVector2D ExpectedSize(300, 200);
-	FVector2D AdjustedSummonLocation = FSlateApplication::Get().CalculatePopupWindowPosition(Anchor, ExpectedSize, FVector2D::ZeroVector, Orient_Vertical);
-
-	TSharedRef<SWindow> NewMenuWindow =
-		SNew(SWindow)
-		.Type(EWindowType::Menu)
-		.IsPopupWindow(true)
-		.SizingRule(bShouldAutoSize ? ESizingRule::Autosized : ESizingRule::UserSized)
-		.ScreenPosition(AdjustedSummonLocation)
-		.AutoCenter(EAutoCenter::None)
-		.ClientSize(ExpectedSize)
-		.FocusWhenFirstShown(true)
-		.ActivationPolicy(EWindowActivationPolicy::Always)
-		[
-			SNew(SVerticalBox)
-			+ SVerticalBox::Slot()
-			.VAlign(VAlign_Center)
-			.HAlign(HAlign_Center)
-			[
-				SNew(STextBlock)
-				.Text(LOCTEXT("PushMenuDeprecatedWarning", "WARNING: PushMenu() returning an SWindow is deprecated. Use Push() that returns an IMenu instead."))
-			]
-		];
-
-	FSlateApplication::Get().AddWindowAsNativeChild(NewMenuWindow, ParentWindow, true);
-
-	return NewMenuWindow;
-}
-
 TSharedRef<IMenu> FMenuStack::Push(const FWidgetPath& InOwnerPath, const TSharedRef<SWidget>& InContent, const FVector2D& SummonLocation, const FPopupTransitionEffect& TransitionEffect, const bool bFocusImmediately, const FVector2D& SummonLocationSize, TOptional<EPopupMethod> InMethod, const bool bIsCollapsedByParent, const bool bEnablePerPixelTransparency)
 {
 	// We want to ensure that when the window is restored to restore the current keyboard focus
@@ -559,12 +522,6 @@ FPopupMethodReply FMenuStack::QueryPopupMethod(const FWidgetPath& PathToQuery)
 	return FPopupMethodReply::UseMethod(EPopupMethod::CreateNewWindow);
 }
 
-void FMenuStack::Dismiss(int32 FirstStackIndexToRemove)
-{
-	// deprecated
-	DismissInternal(FirstStackIndexToRemove);
-}
-
 void FMenuStack::DismissFrom(const TSharedPtr<IMenu>& InFromMenu)
 {
 	int32 Index = Stack.IndexOfByKey(InFromMenu);
@@ -721,12 +678,6 @@ TSharedPtr<IMenu> FMenuStack::FindMenuInWidgetPath(const FWidgetPath& PathToQuer
 	return TSharedPtr<IMenu>();
 }
 
-void FMenuStack::RemoveWindow( TSharedRef<SWindow> WindowToRemove )
-{
-	// deprecated
-	OnWindowDestroyed(WindowToRemove);
-}
-
 void FMenuStack::OnWindowDestroyed(TSharedRef<SWindow> InWindow)
 {
 	if (HostWindow == InWindow)
@@ -779,21 +730,6 @@ void FMenuStack::OnWindowActivated( TSharedRef<SWindow> ActivatedWindow )
 	}
 }
 
-
-int32 FMenuStack::FindLocationInStack( TSharedPtr<SWindow> WindowToFind ) const
-{
-	// deprecated
-	if (WindowToFind.IsValid())
-	{
-		TWeakPtr<IMenu> Menu = FindMenuFromWindow(WindowToFind.ToSharedRef());
-
-		return Stack.IndexOfByKey(Menu);
-	}
-
-	// The window was not found
-	return INDEX_NONE;
-}
-
 TSharedPtr<IMenu> FMenuStack::FindMenuFromWindow(TSharedRef<SWindow> InWindow) const
 {
 	const TSharedPtr<FMenuBase>* FoundMenu = Stack.FindByPredicate([InWindow](TSharedPtr<FMenuBase> Menu) { return Menu->GetOwnedWindow() == InWindow; });
@@ -839,41 +775,15 @@ bool FMenuStack::HasMenus() const
 	return Stack.Num() > 0;
 }
 
-bool FMenuStack::HasOpenSubMenus( const TSharedRef<SWindow>& Window ) const
-{
-	// deprecated
-	TWeakPtr<IMenu> Menu = FindMenuFromWindow(Window);
-	if (Menu.IsValid())
-	{
-		return Menu != Stack.Last();
-	}
-	return false;
-}
-
 bool FMenuStack::HasOpenSubMenus(TSharedPtr<IMenu> InMenu) const
 {
 	int32 StackIndex = Stack.IndexOfByKey(InMenu);
 	return StackIndex != INDEX_NONE && StackIndex < Stack.Num() - 1;
 }
 
-int32 FMenuStack::GetNumStackLevels() const
-{
-	return Stack.Num();
-}
-
 TSharedPtr<SWindow> FMenuStack::GetHostWindow() const
 {
 	return HostWindow;
 }
 
-
-FMenuWindowList& FMenuStack::GetWindowsAtStackLevel( const int32 StackLevelIndex )
-{
-	// deprecated
-	// can't support this so return an empty list
-	ensure(false);
-	static FMenuWindowList EmptyList;
-	return EmptyList;
-}
-
 #undef LOCTEXT_NAMESPACE
diff --git a/Engine/Source/Runtime/Slate/Public/Framework/Application/MenuStack.h b/Engine/Source/Runtime/Slate/Public/Framework/Application/MenuStack.h
--- a/Engine/Source/Runtime/Slate/Public/Framework/Application/MenuStack.h
+++ b/Engine/Source/Runtime/Slate/Public/Framework/Application/MenuStack.h
@@ -70,23 +70,6 @@ public:
 		: bHostWindowGuard(false)
 	{}
 
-	/**
-	 * [Deprecated] Pushes a new menu onto the stack.  Automatically removes windows from the stack that are no longer valid
-	 * Invalid windows are those that are not parents of the window menu being pushed and any windows on the same stack level
-	 * This is the old method for creating menus.
-	 *
-	 * @param InParentWindow		The parent window of the menu
-	 * @param InContent				The menu's content
-	 * @param SummonLocation		Location where the menu should appear
-	 * @param TransitionEffect		Animation to use when the popup appears
-	 * @param bFocusImmediately		Should the popup steal focus when shown?
-	 * @param bShouldAutoSize		True if the new window should automatically size itself to its content
-	 * @param WindowSize			When bShouldAutoSize=false, this must be set to the size of the window to be created
-	 * @param SummonLocationSize	An optional size around the summon location which describes an area in which the menu may not appear
-	 */
-	DEPRECATED(4.9, "PushMenu() returning an SWindow is deprecated. Use Push() that returns an IMenu instead.")
-	TSharedRef<SWindow> PushMenu( const TSharedRef<SWindow>& InParentWindow, const TSharedRef<SWidget>& InContent, const FVector2D& SummonLocation, const FPopupTransitionEffect& TransitionEffect, const bool bFocusImmediately = true, const bool bShouldAutoSize = true, const FVector2D& WindowSize = FVector2D::ZeroVector, const FVector2D& SummonLocationSize = FVector2D::ZeroVector );
-
 	/**
 	 * Pushes a new menu onto the stack. The widget path will be searched for existing menus and the new menu will be parented appropriately.
 	 * Menus are always auto-sized. Use fixed-size content if a fixed size is required.
@@ -144,15 +127,6 @@ public:
 	 * @param bIsCollapsedByParent	Is this menu collapsed when a parent menu receives focus/activation? If false, only focus/activation outside the entire stack will auto collapse it.
 	 */
 	TSharedRef<IMenu> PushHosted(const TSharedPtr<IMenu>& InParentMenu, const TSharedRef<IMenuHost>& InMenuHost, const TSharedRef<SWidget>& InContent, TSharedPtr<SWidget>& OutWrappedContent, const FPopupTransitionEffect& TransitionEffect, EShouldThrottle ShouldThrottle, const bool bIsCollapsedByParent = true);
-
-	/**
-	 * [Deprecated] Dismisses the menu stack up to a certain level (by default removes all levels in the stack)
-	 * Dismisses in reverse order so to keep just the first level you would pass in 1
-	 *
-	 * @param FirstStackIndexToRemove	The first level of the stack to remove
-	 */
-	DEPRECATED(4.9, "Dismiss() taking a stack index is deprecated. Use DismissFrom() or DismissAll().")
-	void Dismiss( int32 FirstStackIndexToRemove = 0 );
 	
 	/**
 	 * Dismisses the menu stack including InFromMenu and all its child menus
@@ -167,14 +141,6 @@ public:
 	 */
 	void DismissAll();
 
-	/**
-	 * Removes a window from the stack and all of its children.  Assumes the window is cleaned up elsewhere
-	 *
-	 * @param WindowToRemove	The window to remove (will remove all children of this window as well)	
-	 */
-	DEPRECATED(4.9, "RemoveWindow() is deprecated. Do not use.")
-	void RemoveWindow( TSharedRef<SWindow> WindowToRemove );
-
 	/**
 	 * Called by the application when any window is destroyed.
 	 *
@@ -189,15 +155,6 @@ public:
 	 */
 	void OnWindowActivated( TSharedRef<SWindow> ActivatedWindow );
 
-	/**
-	 * [Deprecated] Finds a window in the stack.  (Deprecated but not using the deprecation macro. This is to avoid warnings as this is still called from a deprecated FSlateApplication method)
-	 * This is deprecated because menus are no longer necessarily windows. Callers should also not need to know the stack levels by index. Menus are identified using the IMenu interface.
-	 *
-	 * @param	WindowToFind	The window to look for
-	 * @return	The level of the stack where the window is located or INDEX_NONE if the window could not be found
-	 */
-	int32 FindLocationInStack( TSharedPtr<SWindow> WindowToFind ) const;
-
 	/**
 	 * Finds a menu in the stack that owns InWindow.
 	 * Although this method isn't deprecated, it is intended for private use inside FMenuStack
@@ -226,13 +183,6 @@ public:
 	 */
 	FSlateRect GetToolTipForceFieldRect(TSharedRef<IMenu> InMenu, const FWidgetPath& PathContainingMenu) const;
 
-	/**
-	 * [Deprecated]
-	 * @return	Returns the size of the menu stack
-	 */
-	DEPRECATED(4.9, "GetNumStackLevels() is deprecated. Consider HasMenus().")
-	int32 GetNumStackLevels( ) const;
-
 	/**
 	 * @return	Returns the window that is the parent of everything in the stack, if any. If the stack is empty, returns an invalid ptr.
 	 */
@@ -243,29 +193,11 @@ public:
 	*/
 	bool HasMenus() const;
 
-	/**
-	 * [Deprecated] This is deprecated because menus are no longer necessarily windows. Menus are identified using the IMenu interface.
-	 *
-	 * @return Returns whether the window has child menus. If the window is not a menu, returns false.
-	 */
-	DEPRECATED(4.9, "HasOpenSubMenus() taking a window is deprecated. Use HasOpenSubMenus() taking an IMenu as a parameter.")
-	bool HasOpenSubMenus( const TSharedRef<SWindow>& Window ) const;
-
 	/**
 	 * @return Returns whether the menu has child menus. If the menu isn't in the stack, returns false.
 	 */
 	bool HasOpenSubMenus(TSharedPtr<IMenu> InMenu) const;
 
-	/**
-	 * [Deprecated] Returned all of the windows at the specified stack level index. Always returns an empty list because it's deprecated and the return type is not supported.
-	 *
-	 * @param	StackLevelIndex	Index into the menu stack
-	 *
-	 * @return	List of windows at this stack level (returns an empty list because it's deprecated)
-	 */
-	DEPRECATED(4.9, "GetWindowsAtStackLevel() is deprecated. Do not use.")
-	FMenuWindowList& GetWindowsAtStackLevel( const int32 StackLevelIndex );
-
 private:
 	/**
 	 * Queries the widgets in the path for a menu popup method.
diff --git a/Engine/Source/Runtime/Slate/Public/Framework/Commands/InputBindingManager.h b/Engine/Source/Runtime/Slate/Public/Framework/Commands/InputBindingManager.h
--- a/Engine/Source/Runtime/Slate/Public/Framework/Commands/InputBindingManager.h
+++ b/Engine/Source/Runtime/Slate/Public/Framework/Commands/InputBindingManager.h
@@ -10,16 +10,11 @@ class FUserDefinedChords;
 
 typedef TMap<FName, TSharedPtr<FUICommandInfo>> FCommandInfoMap;
 typedef TMap<FInputChord, FName> FChordMap;
-DEPRECATED(4.8, "Use FChordMap instead of FGestureMap")
-typedef FChordMap FGestureMap;
 
 
 /** Delegate for alerting subscribers the input manager records a user-defined chord */
 DECLARE_MULTICAST_DELEGATE_OneParam(FOnUserDefinedChordChanged, const FUICommandInfo& );
 
-DEPRECATED(4.8, "Use FOnUserDefinedChordChanged instead of FOnUserDefinedGestureChanged")
-typedef FOnUserDefinedChordChanged FOnUserDefinedGestureChanged;
-
 /**
  * Manager responsible for creating and processing input bindings.                  
  */
@@ -81,8 +76,6 @@ public:
 	 * @param	bCheckDefault			Whether or not to check the default chord.  Will check the active chord if false
 	 * @return	A pointer to the command info which has the InChord as its active chord or null if one cannot be found
 	 */
-	DEPRECATED(4.8, "Use GetCommandInfoFromInputChord instead of GetCommandInfoFromInputGesture")
-	const TSharedPtr<FUICommandInfo> GetCommandInfoFromInputGesture( const FName InBindingContext, const FInputChord& InGesture, bool bCheckDefault ) const;
 	const TSharedPtr<FUICommandInfo> GetCommandInfoFromInputChord( const FName InBindingContext, const FInputChord& InChord, bool bCheckDefault ) const;
 
 	/** 
@@ -108,8 +101,6 @@ public:
 	 *
 	 * @param CommandInfo	The command that had the active chord changed. 
 	 */
-	DEPRECATED(4.8, "Use NotifyActiveChordChanged instead of NotifyActiveGestureChanged")
-	void NotifyActiveGestureChanged( const FUICommandInfo& CommandInfo );
 	virtual void NotifyActiveChordChanged( const FUICommandInfo& CommandInfo );
 	
 	/**
@@ -120,8 +111,6 @@ public:
 	/**
 	 * Removes any user defined chords
 	 */
-	DEPRECATED(4.8, "Use RemoveUserDefinedChords instead of RemoveUserDefinedGestures")
-	void RemoveUserDefinedGestures() { RemoveUserDefinedChords(); }
 	void RemoveUserDefinedChords();
 
 	/**
@@ -133,22 +122,12 @@ public:
 	void GetCommandInfosFromContext( const FName InBindingContext, TArray< TSharedPtr<FUICommandInfo> >& OutCommandInfos ) const;
 
 	/** Registers a delegate to be called when a user-defined chord is edited */
-	DEPRECATED(4.8, "Use RegisterUserDefinedChordChanged instead of RegisterUserDefinedGestureChanged")
-	FDelegateHandle RegisterUserDefinedGestureChanged(const FOnUserDefinedChordChanged::FDelegate& Delegate)
-	{
-		return RegisterUserDefinedChordChanged(Delegate);
-	}
 	FDelegateHandle RegisterUserDefinedChordChanged(const FOnUserDefinedChordChanged::FDelegate& Delegate)
 	{
 		return OnUserDefinedChordChanged.Add(Delegate);
 	}
 
 	/** Unregisters a delegate to be called when a user-defined chord is edited */
-	DEPRECATED(4.8, "Use RegisterUserDefinedChordChanged instead of RegisterUserDefinedGestureChanged")
-	void UnregisterUserDefinedGestureChanged(FDelegateHandle DelegateHandle)
-	{
-		UnregisterUserDefinedChordChanged(DelegateHandle);
-	}
 	void UnregisterUserDefinedChordChanged(FDelegateHandle DelegateHandle)
 	{
 		OnUserDefinedChordChanged.Remove(DelegateHandle);
diff --git a/Engine/Source/Runtime/Slate/Public/Framework/Commands/InputChord.h b/Engine/Source/Runtime/Slate/Public/Framework/Commands/InputChord.h
--- a/Engine/Source/Runtime/Slate/Public/Framework/Commands/InputChord.h
+++ b/Engine/Source/Runtime/Slate/Public/Framework/Commands/InputChord.h
@@ -182,12 +182,6 @@ struct SLATE_API FInputChord
 		return (bAlt || bCtrl || bCmd || bShift);
 	}
 
-	DEPRECATED(4.8, "Use IsValidChord instead of IsValidGesture")
-	bool IsValidGesture() const
-	{
-		return IsValidChord();
-	}
-
 	/**
 	 * Determines if this chord is valid.  A chord is valid if it has a non modifier key that must be pressed
 	 * and zero or more modifier keys that must be pressed
diff --git a/Engine/Source/Runtime/Slate/Public/Widgets/Docking/SDockTab.h b/Engine/Source/Runtime/Slate/Public/Widgets/Docking/SDockTab.h
--- a/Engine/Source/Runtime/Slate/Public/Widgets/Docking/SDockTab.h
+++ b/Engine/Source/Runtime/Slate/Public/Widgets/Docking/SDockTab.h
@@ -132,10 +132,6 @@ public:
 	/** Similar to GetTabRole() but returns the correct role for UI style and user input purposes */
 	ETabRole GetVisualTabRole() const;
 
-	/** Returns true if the tab role is NomadTab, but is being visualized as a major tab */
-	DEPRECATED(4.8, "Use GetVisualTabRole(), which will return ETabRole::MajorTab for nomads with major tab style.")
-	bool IsNomadTabWithMajorTabStyle() const { return GetTabRole() == ETabRole::NomadTab && GetVisualTabRole() == ETabRole::MajorTab; }
-
 	/**
 	 * What should the content area look like for this type of tab?
 	 * Documents, Apps, and Tool Panels have different backgrounds.
diff --git a/Engine/Source/Runtime/Slate/Public/Widgets/Text/STextBlock.h b/Engine/Source/Runtime/Slate/Public/Widgets/Text/STextBlock.h
--- a/Engine/Source/Runtime/Slate/Public/Widgets/Text/STextBlock.h
+++ b/Engine/Source/Runtime/Slate/Public/Widgets/Text/STextBlock.h
@@ -176,9 +176,6 @@ public:
 	 */
 	void SetFont(const TAttribute< FSlateFontInfo >& InFont);
 
-	DEPRECATED(4.3, "SetForegroundColor is deprecated use SetColorAndOpacity instead.")
-	void SetForegroundColor(const TAttribute<FSlateColor>& InColorAndOpacity) { SetColorAndOpacity(InColorAndOpacity); }
-
 	/** See ColorAndOpacity attribute */
 	void SetColorAndOpacity(const TAttribute<FSlateColor>& InColorAndOpacity);
 
diff --git a/Engine/Source/Runtime/SlateCore/Private/Animation/CurveHandle.cpp b/Engine/Source/Runtime/SlateCore/Private/Animation/CurveHandle.cpp
--- a/Engine/Source/Runtime/SlateCore/Private/Animation/CurveHandle.cpp
+++ b/Engine/Source/Runtime/SlateCore/Private/Animation/CurveHandle.cpp
@@ -38,29 +38,6 @@ float FCurveHandle::GetLerp( ) const
 	return ApplyEasing(Time, TheCurve.EaseFunction);
 }
 
-float FCurveHandle::GetLerpLooping() const
-{
-	return DEPRECATED_GetLerpLooping();
-}
-float FCurveHandle::DEPRECATED_GetLerpLooping() const
-{
-	if (OwnerSequence == nullptr)
-	{
-		return 0.0f;
-	}
-
-	// How far we've played through the curve sequence so far.
-	const float CurSequenceTime = OwnerSequence->DEPRECATED_GetSequenceTimeLooping();
-
-	const FCurveSequence::FSlateCurve& TheCurve = OwnerSequence->GetCurve(CurveIndex);
-	const float TimeSinceStarted = CurSequenceTime - TheCurve.StartTime;
-
-	// How far we passed through the current curve scaled between 0 and 1.
-	const float Time = FMath::Clamp(TimeSinceStarted / TheCurve.DurationSeconds, 0.0f, 1.0f);
-
-	return ApplyEasing(Time, TheCurve.EaseFunction);
-}
-
 
 /* FCurveHandle static functions
  *****************************************************************************/
diff --git a/Engine/Source/Runtime/SlateCore/Private/Animation/CurveSequence.cpp b/Engine/Source/Runtime/SlateCore/Private/Animation/CurveSequence.cpp
--- a/Engine/Source/Runtime/SlateCore/Private/Animation/CurveSequence.cpp
+++ b/Engine/Source/Runtime/SlateCore/Private/Animation/CurveSequence.cpp
@@ -58,16 +58,6 @@ FCurveHandle FCurveSequence::AddCurveRelative( const float InOffset, const float
 	return AddCurve(CurveStartTime, InDurationSecond, InEaseFunction);
 }
 
-
-void FCurveSequence::Play( const float StartAtTime )
-{
-	// Playing forward
-	bInReverse = false;
-
-	// We start playing NOW.
-	SetStartTime(FSlateApplicationBase::Get().GetCurrentTime() - StartAtTime);
-}
-
 void FCurveSequence::Play( const TSharedRef<SWidget>& InOwnerWidget, bool bPlayLooped, const float StartAtTime )
 {
 	RegisterActiveTimerIfNeeded(InOwnerWidget);
@@ -99,15 +89,6 @@ void FCurveSequence::Reverse( )
 	SetStartTime(NewStartTime);
 }
 
-
-void FCurveSequence::PlayReverse( const float StartAtTime )
-{
-	bInReverse = true;
-
-	// We start reversing NOW.
-	SetStartTime(FSlateApplicationBase::Get().GetCurrentTime() - StartAtTime);
-}
-
 void FCurveSequence::PlayReverse( const TSharedRef<SWidget>& InOwnerWidget, bool bPlayLooped, const float StartAtTime )
 {
 	RegisterActiveTimerIfNeeded(InOwnerWidget);
@@ -166,26 +147,6 @@ float FCurveSequence::GetSequenceTime( ) const
 	return bIsLooping ? FMath::Fmod( SequenceTime, TotalDuration ) : SequenceTime;
 }
 
-float FCurveSequence::GetSequenceTimeLooping( ) const
-{
-	return DEPRECATED_GetSequenceTimeLooping();
-}
-float FCurveSequence::DEPRECATED_GetSequenceTimeLooping() const
-{
-	if (!bIsLooping)
-	{
-		auto MutableThis = const_cast<FCurveSequence*>( this );
-
-		// Fake that we're looping to get the lerp
-		MutableThis->bIsLooping = true;
-		float SequenceTime = GetSequenceTime();
-		MutableThis->bIsLooping = false;
-
-		return SequenceTime;
-	}
-
-	return GetSequenceTime();
-}
 
 bool FCurveSequence::IsInReverse( ) const
 {
@@ -242,19 +203,6 @@ float FCurveSequence::GetLerp( ) const
 	return FCurveHandle( this, 0 ).GetLerp();
 }
 
-float FCurveSequence::GetLerpLooping() const
-{
-	return DEPRECATED_GetLerpLooping();
-}
-float FCurveSequence::DEPRECATED_GetLerpLooping() const
-{
-	// Only supported for sequences with a single curve.  If you have multiple curves, use your FCurveHandle to compute
-	// interpolation alpha values.
-	checkSlow(Curves.Num() == 1);
-
-	return bIsLooping ? GetLerp() : FCurveHandle(this, 0).DEPRECATED_GetLerpLooping();
-}
-
 const FCurveSequence::FSlateCurve& FCurveSequence::GetCurve( int32 CurveIndex ) const
 {
 	return Curves[CurveIndex];
diff --git a/Engine/Source/Runtime/SlateCore/Private/Widgets/SWindow.cpp b/Engine/Source/Runtime/SlateCore/Private/Widgets/SWindow.cpp
--- a/Engine/Source/Runtime/SlateCore/Private/Widgets/SWindow.cpp
+++ b/Engine/Source/Runtime/SlateCore/Private/Widgets/SWindow.cpp
@@ -1155,18 +1155,6 @@ bool SWindow::AppearsInTaskbar() const
 	return !bIsPopupWindow && Type != EWindowType::ToolTip && Type != EWindowType::CursorDecorator;
 }
 
-void SWindow::SetOnWindowActivated( const FOnWindowActivated& InDelegate )
-{
-	// deprecated
-	OnWindowActivated = InDelegate;
-}
-
-void SWindow::SetOnWindowDeactivated( const FOnWindowDeactivated& InDelegate )
-{
-	// deprecated
-	OnWindowDeactivated = InDelegate;
-}
-
 /** Sets the delegate to execute right before the window is closed */
 void SWindow::SetOnWindowClosed( const FOnWindowClosed& InDelegate )
 {
diff --git a/Engine/Source/Runtime/SlateCore/Public/Animation/CurveHandle.h b/Engine/Source/Runtime/SlateCore/Public/Animation/CurveHandle.h
--- a/Engine/Source/Runtime/SlateCore/Public/Animation/CurveHandle.h
+++ b/Engine/Source/Runtime/SlateCore/Public/Animation/CurveHandle.h
@@ -58,17 +58,6 @@ public:
 	 */
 	float GetLerp( ) const;
 
-	/**
-	 * Just like GetLerp, but loops forever.
-	 *
-	 * @return Looped lerp value.
-	 * @see GetLerp
-	 */
-	DEPRECATED(4.8, "FCurveHandle::GetLerpLooping() is deprecated. Use GetLerp() instead. To play a sequence on a loop, pass \"true\" as the second parameter to FCurveSequence::Play / ::PlayReverse.")
-	float GetLerpLooping( ) const;
-	/** Shell to avoid generating internal deprecated warnings. Do not call. Will be removed when GetLerpLooping() is removed. */
-	float DEPRECATED_GetLerpLooping() const;
-
 	/**
 	 * Checks whether this handle is initialized.
 	 *
diff --git a/Engine/Source/Runtime/SlateCore/Public/Animation/CurveSequence.h b/Engine/Source/Runtime/SlateCore/Public/Animation/CurveSequence.h
--- a/Engine/Source/Runtime/SlateCore/Public/Animation/CurveSequence.h
+++ b/Engine/Source/Runtime/SlateCore/Public/Animation/CurveSequence.h
@@ -89,22 +89,6 @@ public:
 	 * @param InEaseFunction       Easing function to use for this curve.  Defaults to Linear.  Use this to smooth out your animation transitions.
 	 */
 	FCurveHandle AddCurveRelative( const float InOffset, const float InDurationSecond, const ECurveEaseFunction::Type InEaseFunction = ECurveEaseFunction::Linear );
-	
-	/**
-	 * Start playing this curve sequence
-	 *
-	 * @param	StartAtTime		Specifies a time offset relative to the animation to start at.  Defaults to zero (the actual start of the sequence.)
-	 */
-	DEPRECATED(4.8, "FCurveSequence::Play(const float StartAtTime) is deprecated. Curve sequences must be associated with a specific widget when played. Use Play(const TSharedRef<SWidget>& InOwnerWidget, bool bPlayLooped, const float StartAtTime) instead.")
-	void Play( const float StartAtTime = 0.0f );
-
-	/**
-	 * Start playing this curve sequence in reverse
-	 *
-	 * @param	StartAtTime		Specifies a time offset relative to the animation to start at.  Defaults to zero (the actual start of the sequence.)
-	 */
-	DEPRECATED(4.8, "FCurveSequence::PlayReverse(const float StartAtTime) is deprecated. Curve sequences must be associated with a specific widget when played. Use PlayReverse(const TSharedRef<SWidget>& InOwnerWidget, bool bPlayLooped, const float StartAtTime) instead.")
-	void PlayReverse( const float StartAtTime = 0.0f );
 
 	/**
 	 * Start playing this curve sequence. Registers an active timer with the widget being animated.
@@ -143,12 +127,6 @@ public:
 	/** @return the current time relative to the beginning of the sequence. */
 	float GetSequenceTime( ) const;
 
-	/** @return the current time relative to the beginning of the sequence as if the animation were a looping one. */
-	DEPRECATED(4.8, "FCurveSequence::GetSequenceTimeLooping() is deprecated. Use GetSequenceTime() instead. To play a sequence on a loop, pass \"true\" as the second parameter to Play or PlayReverse.")
-	float GetSequenceTimeLooping( ) const;
-	/** Shell to avoid generating internal deprecated warnings. Do not call. Will be removed when GetSequenceTimeLooping() is removed. */
-	float DEPRECATED_GetSequenceTimeLooping() const;
-
 	/** @return true if the animation is in reverse */
 	bool IsInReverse( ) const;
 
@@ -178,17 +156,6 @@ public:
 	 */
 	float GetLerp() const;
 
-	/**
-	 * For single-curve animations, returns the looping interpolation alpha for the animation.  If you call this
-	 * function on a sequence with multiple curves, an assertion will trigger.
-	 *
-	 * @return A linearly interpolated value between 0 and 1 for this curve.
-	 */
-	DEPRECATED(4.8, "FCurveSequence::GetLerpLooping() is deprecated. Use GetLerp() instead. To play a sequence on a loop, pass \"true\" as the second parameter to Play or PlayReverse.")
-	float GetLerpLooping( ) const;
-	/** Shell to avoid generating internal deprecated warnings. Do not call. Will be removed when GetLerpLooping() is removed. */
-	float DEPRECATED_GetLerpLooping() const;
-
 	/**
 	 * @param CurveIndex  Index of a curve in the curves array.
 	 *
diff --git a/Engine/Source/Runtime/SlateCore/Public/Application/SlateApplicationBase.h b/Engine/Source/Runtime/SlateCore/Public/Application/SlateApplicationBase.h
--- a/Engine/Source/Runtime/SlateCore/Public/Application/SlateApplicationBase.h
+++ b/Engine/Source/Runtime/SlateCore/Public/Application/SlateApplicationBase.h
@@ -212,17 +212,6 @@ public:
 
 	virtual EUINavigation GetNavigationDirectionFromKey( const FKeyEvent& InKeyEvent ) const = 0;
 
-	/**
-	* Gets the Widget that currently captures the mouse.
-	*
-	* @return The captor widget, or nullptr if no widget captured the mouse.
-	*/
-	DEPRECATED(4.5, "This API is no longer supported.  You can check if a specific widget has capture using the SWidget::HasMouseCapture() API")
-	inline TSharedPtr< SWidget > GetMouseCaptor() const
-	{
-		return GetMouseCaptorImpl();
-	}
-
 protected:
 	/**
 	 * Implementation of GetMouseCaptor which can be overridden without warnings.
diff --git a/Engine/Source/Runtime/SlateCore/Public/Input/Events.h b/Engine/Source/Runtime/SlateCore/Public/Input/Events.h
--- a/Engine/Source/Runtime/SlateCore/Public/Input/Events.h
+++ b/Engine/Source/Runtime/SlateCore/Public/Input/Events.h
@@ -40,20 +40,8 @@ enum class EFocusCause : uint8
 	WindowActivate,
 };
 
-/** DEPRECATED 4.6 - Do not use */
-//@Todo slate: Remove this as soon as the 4.6 deprecated API is Removed.
-namespace EKeyboardFocusCause
-{
-	DEPRECATED(4.6, "EFocusCause::Type is deprecated and was renamed to EFocusCause. Please use that type instead.")
-	typedef ::EFocusCause Type;
-
-	const EFocusCause Mouse = EFocusCause::Mouse;
-	const EFocusCause Keyboard = EFocusCause::Navigation;
-	const EFocusCause SetDirectly = EFocusCause::SetDirectly;
-	const EFocusCause Cleared = EFocusCause::Cleared;
-	const EFocusCause OtherWidgetLostFocus = EFocusCause::OtherWidgetLostFocus;
-	const EFocusCause WindowActivate = EFocusCause::WindowActivate;
-}
+DEPRECATED(4.18, "EKeyboardFocusCause is deprecated and was renamed to EFocusCause. Please use that type instead.")
+typedef EFocusCause EKeyboardFocusCause;
 
 /**
  * FFocusEvent is used when notifying widgets about keyboard focus changes
diff --git a/Engine/Source/Runtime/SlateCore/Public/Input/Reply.h b/Engine/Source/Runtime/SlateCore/Public/Input/Reply.h
--- a/Engine/Source/Runtime/SlateCore/Public/Input/Reply.h
+++ b/Engine/Source/Runtime/SlateCore/Public/Input/Reply.h
@@ -49,18 +49,6 @@ public:
 
 	/** An event should return FReply::Handled().SetUserFocus( SomeWidget ) as a means of asking the system to set users focus to the provided widget*/
 	FReply& SetUserFocus(TSharedRef<SWidget> GiveMeFocus, EFocusCause ReasonFocusIsChanging = EFocusCause::SetDirectly, bool bInAllUsers = false);
-	
-	DEPRECATED(4.6, "FReply::CaptureJoystick() is deprecated, use FReply::SetUserFocus() instead.")
-	FReply& CaptureJoystick(TSharedRef<SWidget> InJoystickCaptor, bool bInAllJoysticks = false)
-	{
-		return SetUserFocus(InJoystickCaptor, EFocusCause::SetDirectly, bInAllJoysticks);
-	}
-
-	DEPRECATED(4.6, "FReply::SetKeyboardFocus() is deprecated, use FReply::SetUserFocus() instead.")
-	FReply& SetKeyboardFocus(TSharedRef<SWidget> GiveMeFocus, EFocusCause ReasonFocusIsChanging)
-	{
-		return SetUserFocus(GiveMeFocus, ReasonFocusIsChanging, false);
-	}
 
 	/** An event should return a FReply::Handled().ClearUserFocus() to ask the system to clear user focus*/
 	FReply& ClearUserFocus(bool bInAllUsers = false)
@@ -71,12 +59,6 @@ public:
 	/** An event should return a FReply::Handled().ClearUserFocus() to ask the system to clear user focus*/
 	FReply& ClearUserFocus(EFocusCause ReasonFocusIsChanging, bool bInAllUsers = false);
 
-	DEPRECATED(4.6, "FReply::ReleaseJoystickCapture() is deprecated, use FReply::ClearUserFocus() instead.")
-	FReply& ReleaseJoystickCapture(bool bInAllJoysticks = false)
-	{
-		return ClearUserFocus(bInAllJoysticks);
-	}
-
 	/** An event should return FReply::Handled().SetNavigation( NavigationType ) as a means of asking the system to attempt a navigation*/
 	FReply& SetNavigation(EUINavigation InNavigationType, const ENavigationGenesis InNavigationGenesis, const ENavigationSource InNavigationSource = ENavigationSource::FocusedWidget)
 	{
@@ -182,9 +164,6 @@ public:
 	/** true if this reply indicated that we should release focus as a result of the event being handled */
 	bool ShouldReleaseUserFocus() const { return bReleaseUserFocus; }
 
-	DEPRECATED(4.6, "FReply::ShouldReleaseJoystick() is deprecated, use FReply::ShouldReleaseUserFocus() instead.")
-	bool ShouldReleaseJoystick() const { return ShouldReleaseUserFocus(); }
-
 	/** If the event replied with a request to change the user focus whether it should do it for all users or just the current UserIndex */
 	bool AffectsAllUsers() const { return bAllUsers; }
 
@@ -203,12 +182,6 @@ public:
 	/** When not nullptr, user focus has been requested to be set on the FocusRecipient. */
 	const TSharedPtr<SWidget>& GetUserFocusRecepient() const { return FocusRecipient; }
 
-	DEPRECATED(4.6, "FReply::GetFocusRecepient() is deprecated, use FReply::GetUserFocusRecepient() instead.")
-	const TSharedPtr<SWidget>& GetFocusRecepient() const { return GetUserFocusRecepient(); }
-
-	DEPRECATED(4.6, "FReply::GetJoystickCaptor() is deprecated, use FReply::GetUserFocusRecepient() instead.")
-	const TSharedPtr<SWidget>& GetJoystickCaptor() const { return GetUserFocusRecepient(); }
-
 	/** Get the reason that a focus change is being requested. */
 	EFocusCause GetFocusCause() const { return FocusChangeReason; }
 
diff --git a/Engine/Source/Runtime/SlateCore/Public/Widgets/SWidget.h b/Engine/Source/Runtime/SlateCore/Public/Widgets/SWidget.h
--- a/Engine/Source/Runtime/SlateCore/Public/Widgets/SWidget.h
+++ b/Engine/Source/Runtime/SlateCore/Public/Widgets/SWidget.h
@@ -167,7 +167,7 @@ class IToolTip;
  * 
  * SLATE_ATTRIBUTE(ECheckBoxState, IsChecked)
  *
- * DEPRECATED(4.3, "Please use IsChecked(TAttribute<ECheckBoxState>)")
+ * DEPRECATED(4.xx, "Please use IsChecked(TAttribute<ECheckBoxState>)")
  * FArguments& IsChecked(bool InIsChecked)
  * {
  * 		_IsChecked = InIsChecked ? ECheckBoxState::Checked : ECheckBoxState::Unchecked;
diff --git a/Engine/Source/Runtime/SlateCore/Public/Widgets/SWindow.h b/Engine/Source/Runtime/SlateCore/Public/Widgets/SWindow.h
--- a/Engine/Source/Runtime/SlateCore/Public/Widgets/SWindow.h
+++ b/Engine/Source/Runtime/SlateCore/Public/Widgets/SWindow.h
@@ -526,17 +526,9 @@ public:
 	/** @return should this window show up in the taskbar */
 	bool AppearsInTaskbar() const;
 
-	/** Sets the delegate to execute when the window is activated */
-	DEPRECATED(4.9, "SetOnWindowActivated() is deprecated. Use GetOnWindowActivatedEvent() and subscribe to the multicast event.")
-	void SetOnWindowActivated( const FOnWindowActivated& InDelegate );
-
 	/** Gets the multicast delegate executed when the window is deactivated */
 	FOnWindowActivatedEvent& GetOnWindowActivatedEvent() { return WindowActivatedEvent; }
 
-	/** Sets the delegate to execute when the window is deactivated */
-	DEPRECATED(4.9, "SetOnWindowDeactivated() is deprecated. Use GetOnWindowDeactivatedEvent() and subscribe to the multicast event.")
-	void SetOnWindowDeactivated( const FOnWindowDeactivated& InDelegate );
-
 	/** Gets the multicast delegate executed when the window is deactivated */
 	FOnWindowDeactivatedEvent& GetOnWindowDeactivatedEvent() { return WindowDeactivatedEvent; }
 
diff --git a/Engine/Source/Runtime/UMG/Private/WidgetBlueprintGeneratedClass.cpp b/Engine/Source/Runtime/UMG/Private/WidgetBlueprintGeneratedClass.cpp
--- a/Engine/Source/Runtime/UMG/Private/WidgetBlueprintGeneratedClass.cpp
+++ b/Engine/Source/Runtime/UMG/Private/WidgetBlueprintGeneratedClass.cpp
@@ -479,7 +479,7 @@ void UWidgetBlueprintGeneratedClass::InitializeTemplate(const ITargetPlatform* T
 
 		// Remove the old archetype.
 		{
-			UUserWidget* OldArchetype = LoadObject<UUserWidget>(WidgetTemplatePackage, TEXT("WidgetArchetype"), nullptr, LOAD_NoWarn);
+			UUserWidget* OldArchetype = FindObject<UUserWidget>(WidgetTemplatePackage, TEXT("WidgetArchetype"));
 			if ( OldArchetype )
 			{
 				const ERenameFlags RenFlags = REN_DontCreateRedirectors | REN_NonTransactional | REN_DoNotDirty;
