Copying //UE4/Orion-Staging to //UE4/Main (Source: //Orion/Dev-General @ 3564337)

diff --git a/Engine/Source/Runtime/Engine/Private/Animation/AnimationUtils.cpp b/Engine/Source/Runtime/Engine/Private/Animation/AnimationUtils.cpp
--- a/Engine/Source/Runtime/Engine/Private/Animation/AnimationUtils.cpp
+++ b/Engine/Source/Runtime/Engine/Private/Animation/AnimationUtils.cpp
@@ -13,11 +13,13 @@
 #include "Animation/AnimCompress_PerTrackCompression.h"
 #include "Animation/AnimCompress_LeastDestructive.h"
 #include "Animation/AnimCompress_RemoveEverySecondKey.h"
+#include "Animation/AnimCompress_Automatic.h"
 #include "Animation/AnimSet.h"
 #include "Animation/AnimationSettings.h"
 #include "AnimationCompression.h"
 #include "Engine/SkeletalMeshSocket.h"
 #include "AnimEncoding.h"
+#include "UObject/LinkerLoad.h"
 
 /** Array to keep track of SkeletalMeshes we have built metadata for, and log out the results just once. */
 //static TArray<USkeleton*> UniqueSkeletonsMetadataArray;
@@ -50,6 +52,9 @@ void FAnimationUtils::BuildSkeletonMetaData(USkeleton* Skeleton, TArray<FBoneDat
 		// Copy over data from the skeleton.
 		const FTransform& SrcTransform = SkeletonRefPose[BoneIndex];
 
+		ensure(!SrcTransform.ContainsNaN());
+		ensure(SrcTransform.IsRotationNormalized());
+
 		BoneData.Orientation = SrcTransform.GetRotation();
 		BoneData.Position = SrcTransform.GetTranslation();
 		BoneData.Name = RefSkeleton.GetBoneName(BoneIndex);
@@ -97,7 +102,7 @@ void FAnimationUtils::BuildSkeletonMetaData(USkeleton* Skeleton, TArray<FBoneDat
 	// Enumerate end effectors.  For each end effector, propagate its index up to all ancestors.
 	if( bEnableLogging )
 	{
-		UE_LOG(LogAnimation, Log, TEXT("Enumerate End Effectors for %s"), *Skeleton->GetFName().ToString());
+		UE_LOG(LogAnimationCompression, Log, TEXT("Enumerate End Effectors for %s"), *Skeleton->GetFName().ToString());
 	}
 	for ( int32 BoneIndex = 0 ; BoneIndex < OutBoneData.Num() ; ++BoneIndex )
 	{
@@ -125,12 +130,12 @@ void FAnimationUtils::BuildSkeletonMetaData(USkeleton* Skeleton, TArray<FBoneDat
 			}
 			if( bEnableLogging )
 			{
-				UE_LOG(LogAnimation, Log, TEXT("\t %s bKeyEndEffector: %d"), *BoneData.Name.ToString(), BoneData.bKeyEndEffector);
+				UE_LOG(LogAnimationCompression, Log, TEXT("\t %s bKeyEndEffector: %d"), *BoneData.Name.ToString(), BoneData.bKeyEndEffector);
 			}
 		}
 	}
 #if 0
-	UE_LOG(LogAnimation, Log,  TEXT("====END EFFECTORS:") );
+	UE_LOG(LogAnimationCompression, Log,  TEXT("====END EFFECTORS:") );
 	int32 NumEndEffectors = 0;
 	for ( int32 BoneIndex = 0 ; BoneIndex < OutBoneData.Num() ; ++BoneIndex )
 	{
@@ -147,8 +152,8 @@ void FAnimationUtils::BuildSkeletonMetaData(USkeleton* Skeleton, TArray<FBoneDat
 			++NumEndEffectors;
 		}
 	}
-	UE_LOG(LogAnimation, Log,  TEXT("====NUM EFFECTORS %i(%i)"), NumEndEffectors, OutBoneData(0).Children.Num() );
-	UE_LOG(LogAnimation, Log,  TEXT("====NON END EFFECTORS:") );
+	UE_LOG(LogAnimationCompression, Log,  TEXT("====NUM EFFECTORS %i(%i)"), NumEndEffectors, OutBoneData(0).Children.Num() );
+	UE_LOG(LogAnimationCompression, Log,  TEXT("====NON END EFFECTORS:") );
 	for ( int32 BoneIndex = 0 ; BoneIndex < OutBoneData.Num() ; ++BoneIndex )
 	{
 		const FBoneData& BoneData = OutBoneData(BoneIndex);
@@ -168,10 +173,10 @@ void FAnimationUtils::BuildSkeletonMetaData(USkeleton* Skeleton, TArray<FBoneDat
 				Message += FString::Printf( TEXT("%s "), *OutBoneData(EndEffectorIndex).Name );
 				check( OutBoneData(EndEffectorIndex).IsEndEffector() );
 			}
-			UE_LOG(LogAnimation, Log,  *Message );
+			UE_LOG(LogAnimationCompression, Log,  *Message );
 		}
 	}
-	UE_LOG(LogAnimation, Log,  TEXT("===================") );
+	UE_LOG(LogAnimationCompression, Log,  TEXT("===================") );
 #endif
 }
 
@@ -186,21 +191,24 @@ void FAnimationUtils::BuildComponentSpaceTransform(FTransform& OutTransform,
 	// Put root-to-component in OutTransform.
 	OutTransform = BoneSpaceTransforms[0];
 
-	if ( BoneIndex > 0 )
+	if (BoneIndex > 0)
 	{
 		const FBoneData& Bone = BoneData[BoneIndex];
 
-		checkSlow( Bone.BonesToRoot.Num()-1 == 0 );
+		checkSlow((Bone.BonesToRoot.Num() - 1) == 0);
 
 		// Compose BoneData.BonesToRoot down.
-		for ( int32 i = Bone.BonesToRoot.Num()-2 ; i >=0 ; --i )
+		for (int32 i = Bone.BonesToRoot.Num()-2; i >=0; --i)
 		{
 			const int32 AncestorIndex = Bone.BonesToRoot[i];
-			OutTransform = BoneSpaceTransforms[AncestorIndex]*OutTransform;
+			ensure(AncestorIndex != INDEX_NONE);
+			OutTransform = BoneSpaceTransforms[AncestorIndex] * OutTransform;
+			OutTransform.NormalizeRotation();
 		}
 
 		// Finally, include the bone's local-to-parent.
-		OutTransform = BoneSpaceTransforms[BoneIndex]*OutTransform;
+		OutTransform = BoneSpaceTransforms[BoneIndex] * OutTransform;
+		OutTransform.NormalizeRotation();
 	}
 }
 
@@ -223,7 +231,13 @@ void FAnimationUtils::ComputeCompressionError(const UAnimSequence* AnimSeq, cons
 
 	if (AnimSeq->NumFrames > 0)
 	{
-		const float TimeStep = (float)AnimSeq->SequenceLength/(float)AnimSeq->NumFrames;
+		const bool bCanUseCompressedData = (AnimSeq->CompressedTrackOffsets.Num() > 0);
+		if (!bCanUseCompressedData)
+		{
+			// If we can't use CompressedData, there's not much point in being here.
+			return;
+		}
+
 		const int32 NumBones = BoneData.Num();
 		
 		float ErrorCount = 0.0f;
@@ -232,102 +246,129 @@ void FAnimationUtils::ComputeCompressionError(const UAnimSequence* AnimSeq, cons
 		USkeleton* Skeleton = AnimSeq->GetSkeleton();
 		check ( Skeleton );
 
-		const TArray<FTransform>& RefPose	= Skeleton->GetRefLocalPoses();
+		const TArray<FTransform>& RefPose = Skeleton->GetRefLocalPoses();
 
-		TArray<FTransform> RawAtoms;
-		TArray<FTransform> NewAtoms;
 		TArray<FTransform> RawTransforms;
 		TArray<FTransform> NewTransforms;
-
-		RawAtoms.AddZeroed(NumBones);
-		NewAtoms.AddZeroed(NumBones);
 		RawTransforms.AddZeroed(NumBones);
 		NewTransforms.AddZeroed(NumBones);
 
-		FTransform const DummyBone(FQuat::Identity, FVector(END_EFFECTOR_SOCKET_DUMMY_BONE_SIZE, END_EFFECTOR_SOCKET_DUMMY_BONE_SIZE, END_EFFECTOR_SOCKET_DUMMY_BONE_SIZE));
+		// Cache these to speed up animations with a lot of frames.
+		// We do this only once, instead of every frame.
+		struct FCachedBoneIndexData
+		{
+			int32 RawTrackIndex;
+			int32 CompressedTrackIndex;
+			int32 ParentIndex;
+		};
+		TArray<FCachedBoneIndexData> CachedBoneIndexData;
+		CachedBoneIndexData.AddZeroed(NumBones);
+		for (int32 BoneIndex = 0; BoneIndex < NumBones; ++BoneIndex)
+		{
+			FCachedBoneIndexData& BoneIndexData = CachedBoneIndexData[BoneIndex];
+			BoneIndexData.RawTrackIndex = Skeleton->GetAnimationTrackIndex(BoneIndex, AnimSeq, true);
+			BoneIndexData.CompressedTrackIndex = Skeleton->GetAnimationTrackIndex(BoneIndex, AnimSeq, false);
+			BoneIndexData.ParentIndex = Skeleton->GetReferenceSkeleton().GetParentIndex(BoneIndex);
+		}
+
+		// Check the precondition that parents occur before children in the RequiredBones array.
+		for (int32 BoneIndex = 1; BoneIndex < NumBones; ++BoneIndex)
+		{
+			const FCachedBoneIndexData& BoneIndexData = CachedBoneIndexData[BoneIndex];
+			check(BoneIndexData.ParentIndex != INDEX_NONE);
+			check(BoneIndexData.ParentIndex < BoneIndex);
+		}
 
-		const bool bCanUseCompressedData = AnimSeq->CompressedTrackOffsets.Num() > 0;
+		const FTransform EndEffectorDummyBoneSocket(FQuat::Identity, FVector(END_EFFECTOR_DUMMY_BONE_LENGTH_SOCKET));
+		const FTransform EndEffectorDummyBone(FQuat::Identity, FVector(END_EFFECTOR_DUMMY_BONE_LENGTH));
+		const FAnimKeyHelper Helper(AnimSeq->SequenceLength, AnimSeq->NumFrames);
+		const float KeyLength = Helper.TimePerKey() + SMALL_NUMBER;
 
-		// for each whole increment of time (frame stepping)
-		for( float Time = 0.0f; Time < AnimSeq->SequenceLength; Time+= TimeStep )
+		for (int32 FrameIndex = 0; FrameIndex<AnimSeq->NumFrames; FrameIndex++)
 		{
+			const float Time = (float)FrameIndex * KeyLength;
+
 			// get the raw and compressed atom for each bone
-			for( int32 BoneIndex = 0; BoneIndex < NumBones; ++BoneIndex )
+			for (int32 BoneIndex = 0; BoneIndex < NumBones; ++BoneIndex)
 			{
-				int32 const RawTrackIndex = Skeleton->GetAnimationTrackIndex(BoneIndex, AnimSeq, true);
-				int32 const CompressedTrackIndex = bCanUseCompressedData ? Skeleton->GetAnimationTrackIndex(BoneIndex, AnimSeq, false) : RawTrackIndex;
-
-				if (RawTrackIndex == INDEX_NONE)
+				const FCachedBoneIndexData& BoneIndexData = CachedBoneIndexData[BoneIndex];
+				if (BoneIndexData.RawTrackIndex == INDEX_NONE)
 				{
-					// No track for the bone was found, so use the reference pose.
-					RawAtoms[BoneIndex]	= RefPose[BoneIndex];
-					NewAtoms[BoneIndex] = RawAtoms[BoneIndex];
+					// No track for the bone was found, use default transform
+					const FTransform& RefPoseTransform = RefPose[BoneIndex];
+					RawTransforms[BoneIndex] = RefPoseTransform;
+					NewTransforms[BoneIndex] = RefPoseTransform;
 				}
 				else
 				{
-					AnimSeq->GetBoneTransform(RawAtoms[BoneIndex], RawTrackIndex, Time, true);
-					AnimSeq->GetBoneTransform(NewAtoms[BoneIndex], CompressedTrackIndex, Time, false);
-
-
-					bool bSkipTranslationTrack = false;
-					// If we forcibly reduced the translation track to one key, make sure we don't introduce error if it was animated previously.
-					// So short-circuit RAW data for error measuring past that first key.
-					bool bReducedTranslationTrack = false;
-
-					// If we don't care about this translation track, because it's going to get skipped, then use RefSkel translation for error measurement.
-#if( SKIP_FORCEMESHTRANSLATION_TRACKS || SKIP_ANIMROTATIONONLY_TRACKS || REDUCE_ANIMROTATIONONLY_TRACKS )		
-					const bool bUseRefPoseTranslation = (Skeleton->GetBoneTranslationRetargetingMode(BoneIndex) == EBoneTranslationRetargetingMode::Skeleton);
-	#if( SKIP_FORCEMESHTRANSLATION_TRACKS || SKIP_ANIMROTATIONONLY_TRACKS )	
-					bSkipTranslationTrack = bUseRefPoseTranslation;
-	#endif
-	#if( REDUCE_ANIMROTATIONONLY_TRACKS )
-					bReducedTranslationTrack = bUseRefPoseTranslation && (Time > 0.f);
-	#endif
-
-#endif
-					// bAnimRotationOnly tracks - ignore translation data always use Ref Skeleton.
-					if( bSkipTranslationTrack || bReducedTranslationTrack )
+					// If we have transforms, but they're additive, apply them to the ref pose.
+					// This is because additive animations are mostly rotation.
+					// And for the error metric we measure distance between end effectors.
+					// So that means additive animations by default will all be balled up at the origin and not show any error.
+					if (AnimSeq->IsValidAdditive())
+					{
+						const FTransform& RefPoseTransform = RefPose[BoneIndex];
+						RawTransforms[BoneIndex] = RefPoseTransform;
+						NewTransforms[BoneIndex] = RefPoseTransform;
+
+						FTransform AdditiveRawTransform;
+						FTransform AdditiveNewTransform;
+						AnimSeq->GetBoneTransform(AdditiveRawTransform, BoneIndexData.RawTrackIndex, Time, true);
+						AnimSeq->GetBoneTransform(AdditiveNewTransform, BoneIndexData.CompressedTrackIndex, Time, false);
+
+						const ScalarRegister VBlendWeight(1.f);
+						RawTransforms[BoneIndex].AccumulateWithAdditiveScale(AdditiveRawTransform, VBlendWeight);
+						NewTransforms[BoneIndex].AccumulateWithAdditiveScale(AdditiveNewTransform, VBlendWeight);
+					}
+					else
 					{
-						RawAtoms[BoneIndex].SetTranslation( RefPose[BoneIndex].GetTranslation() );
-						NewAtoms[BoneIndex].SetTranslation( RefPose[BoneIndex].GetTranslation() );
+						AnimSeq->GetBoneTransform(RawTransforms[BoneIndex], BoneIndexData.RawTrackIndex, Time, true);
+						AnimSeq->GetBoneTransform(NewTransforms[BoneIndex], BoneIndexData.CompressedTrackIndex, Time, false);
 					}
 				}
 
-				RawTransforms[BoneIndex] = RawAtoms[BoneIndex];
-				NewTransforms[BoneIndex] = NewAtoms[BoneIndex];
+				ensure(!RawTransforms[BoneIndex].ContainsNaN());
+				ensure(!NewTransforms[BoneIndex].ContainsNaN());
 
 				// For all bones below the root, final component-space transform is relative transform * component-space transform of parent.
-				if( BoneIndex > 0 )
+				if (BoneIndex > 0)
 				{
-					const int32 ParentIndex = Skeleton->GetReferenceSkeleton().GetParentIndex(BoneIndex);
-
-					// Check the precondition that parents occur before children in the RequiredBones array.
-					check( ParentIndex != INDEX_NONE );
-					check( ParentIndex < BoneIndex );
-
-					RawTransforms[BoneIndex] *= RawTransforms[ParentIndex];
-					NewTransforms[BoneIndex] *= NewTransforms[ParentIndex];
+					RawTransforms[BoneIndex] *= RawTransforms[BoneIndexData.ParentIndex];
+					NewTransforms[BoneIndex] *= NewTransforms[BoneIndexData.ParentIndex];
 				}
 				
-				if( BoneData[BoneIndex].IsEndEffector() )
+				// If this is an EndEffector, add a dummy bone to measure the effect of compressing the rotation.
+				if (BoneData[BoneIndex].IsEndEffector())
 				{
-					// If this is an EndEffector with a Socket attached to it, add an extra bone, to measure error introduced by effector rotation compression.
-					if( BoneData[BoneIndex].bHasSocket || BoneData[BoneIndex].bKeyEndEffector )
+					// Sockets and Key EndEffectors have a longer dummy bone to maintain higher precision.
+					if (BoneData[BoneIndex].bHasSocket || BoneData[BoneIndex].bKeyEndEffector)
 					{
-						RawTransforms[BoneIndex] = DummyBone * RawTransforms[BoneIndex];
-						NewTransforms[BoneIndex] = DummyBone * NewTransforms[BoneIndex];
+						RawTransforms[BoneIndex] = EndEffectorDummyBoneSocket * RawTransforms[BoneIndex];
+						NewTransforms[BoneIndex] = EndEffectorDummyBoneSocket * NewTransforms[BoneIndex];
 					}
+					else
+					{
+						RawTransforms[BoneIndex] = EndEffectorDummyBone * RawTransforms[BoneIndex];
+						NewTransforms[BoneIndex] = EndEffectorDummyBone * NewTransforms[BoneIndex];
+					}
+				}
 
+				// Normalize rotations
+				RawTransforms[BoneIndex].NormalizeRotation();
+				NewTransforms[BoneIndex].NormalizeRotation();
+
+				if (BoneData[BoneIndex].IsEndEffector())
+				{
 					float Error = (RawTransforms[BoneIndex].GetLocation() - NewTransforms[BoneIndex].GetLocation()).Size();
 
 					ErrorTotal += Error;
 					ErrorCount += 1.0f;
 
-					if( Error > ErrorStats.MaxError )
+					if (Error > ErrorStats.MaxError)
 					{
-						ErrorStats.MaxError		= Error;
+						ErrorStats.MaxError	= Error;
 						ErrorStats.MaxErrorBone = BoneIndex;
-						MaxErrorTrack = RawTrackIndex;
+						MaxErrorTrack = BoneIndexData.RawTrackIndex;
 						ErrorStats.MaxErrorTime = Time;
 					}
 				}
@@ -444,11 +485,11 @@ bool FAnimationUtils::GetForcedRecompressionSetting()
 	AnimSeq->CompressionScheme = (compressionalgorithm);																									\
 	(compressionalgorithm)->Reduce(AnimSeq, CompressContext);																								\
 	AnimSeq->SetUseRawDataOnly(false);																														\
-	const SIZE_T NewSize = AnimSeq->GetResourceSizeBytes(EResourceSizeMode::Inclusive);																		\
+	const SIZE_T NewSize = AnimSeq->GetApproxCompressedSize();																		\
 																																							\
 	/* compute the savings and compression error*/																											\
-	const int32 MemorySavingsFromOriginal = OriginalSize - NewSize;																						\
-	const int32 MemorySavingsFromPrevious = CurrentSize - NewSize;																							\
+	const int64 MemorySavingsFromOriginal = OriginalSize - NewSize;																						\
+	const int64 MemorySavingsFromPrevious = CurrentSize - NewSize;																							\
 	PctSaving = 0.f;																																		\
 	/* figure out our new compression error*/																												\
 	FAnimationUtils::ComputeCompressionError(AnimSeq, BoneData, NewErrorStats);																				\
@@ -468,13 +509,13 @@ bool FAnimationUtils::GetForcedRecompressionSetting()
 	bKeepNewCompressionMethod |= bLowersErrorAndSavesSameOrBetter;																							\
 																																							\
 	PctSaving = (OriginalSize > 0) ? (100.f - (100.f * float(NewSize) / float(OriginalSize))) : 0.f;														\
-	UE_LOG(LogAnimation, Verbose, TEXT("- %s - bytes saved(%i) (%.1f%%) from previous(%i) MaxError(%.2f) bLowersError(%d) %s"),								\
+	UE_LOG(LogAnimationCompression, Log, TEXT("- %s - bytes saved(%i) (%.1f%%) from previous(%i) MaxError(%.2f) bLowersError(%d) %s"),								\
 		compressionname, MemorySavingsFromOriginal, PctSaving, MemorySavingsFromPrevious, NewErrorStats.MaxError, bLowersError, bKeepNewCompressionMethod ? TEXT("(**Best so far**)") : TEXT(""));									\
 																																							\
-	UE_LOG(LogAnimation, Verbose, TEXT("    bReducesErrorBelowThreshold(%d) bHasAcceptableErrorAndSavesSpace(%d) bLowersErrorAndSavesSameOrBetter(%d)"),	\
+	UE_LOG(LogAnimationCompression, Log, TEXT("    bReducesErrorBelowThreshold(%d) bHasAcceptableErrorAndSavesSpace(%d) bLowersErrorAndSavesSameOrBetter(%d)"),	\
 		bReducesErrorBelowThreshold, bHasAcceptableErrorAndSavesSpace, bLowersErrorAndSavesSameOrBetter);													\
 																																							\
-	UE_LOG(LogAnimation, Verbose, TEXT("    WinningCompressorError(%f) MasterTolerance(%f) bForceBelowThreshold(%d) bErrorUnderThreshold(%d)"),				\
+	UE_LOG(LogAnimationCompression, Log, TEXT("    WinningCompressorError(%f) MasterTolerance(%f) bForceBelowThreshold(%d) bErrorUnderThreshold(%d)"),				\
 		WinningCompressorError, MasterTolerance, bForceBelowThreshold, bErrorUnderThreshold);																\
 																																							\
 	if (bKeepNewCompressionMethod)																															\
@@ -517,7 +558,7 @@ bool FAnimationUtils::GetForcedRecompressionSetting()
 		AnimSeq->SetUseRawDataOnly(bSavedUseRawDataOnly);																									\
 		AnimationFormat_SetInterfaceLinks(*AnimSeq);																										\
 																																							\
-		const SIZE_T RestoredSize = AnimSeq->GetResourceSizeBytes(EResourceSizeMode::Inclusive);															\
+		const SIZE_T RestoredSize = AnimSeq->GetApproxCompressedSize();															\
 		check(RestoredSize == CurrentSize);																													\
 	}																																						\
 }
@@ -531,7 +572,7 @@ bool FAnimationUtils::GetForcedRecompressionSetting()
 #define DECLARE_ANIM_COMP_ALGORITHM(Algorithm) \
 	static int32 Algorithm ## CompressorWins = 0; \
 	static float Algorithm ## CompressorSumError = 0.0f; \
-	static int32 Algorithm ## CompressorWinMargin = 0
+	static int64 Algorithm ## CompressorWinMargin = 0
 
 /**
  * Utility function to compress an animation. If the animation is currently associated with a codec, it will be used to 
@@ -665,11 +706,19 @@ void FAnimationUtils::CompressAnimSequenceExplicit(
 	DECLARE_ANIM_COMP_ALGORITHM(Linear_PerTrackExp1);
 	DECLARE_ANIM_COMP_ALGORITHM(Linear_PerTrackExp2);
 
-	check(AnimSeq != NULL);
+	check(AnimSeq != nullptr);
+	if (AnimSeq->HasAnyFlags(RF_NeedLoad))
+	{
+		AnimSeq->GetLinker()->Preload(AnimSeq);
+	}
 
 	// attempt to find the default skeletal mesh associated with this sequence
 	USkeleton* Skeleton = AnimSeq->GetSkeleton();
 	check (Skeleton);
+	if (Skeleton->HasAnyFlags(RF_NeedLoad))
+	{
+		Skeleton->GetLinker()->Preload(Skeleton);
+	}
 
 	static int32 TotalRecompressions = 0;
 
@@ -687,26 +736,19 @@ void FAnimationUtils::CompressAnimSequenceExplicit(
 	// we must have raw data to continue
 	if(NumRawDataTracks > 0 )
 	{
-		// See if we're trying alternate compressors
-		bool const bTryAlternateCompressor = MasterTolerance > 0.0f;
-
-		// Get the current size
-		SIZE_T OriginalSize = AnimSeq->GetResourceSizeBytes(EResourceSizeMode::Inclusive);
-		TotalSizeBefore += OriginalSize;
+		// If compression Scheme is automatic, then we definitely want to 'TryAlternateCompressor'.
+		if (AnimSeq->CompressionScheme && AnimSeq->CompressionScheme->IsA(UAnimCompress_Automatic::StaticClass()))
+		{
+			MasterTolerance = GetAlternativeCompressionThreshold();
+		}
 
-		// Estimate total uncompressed
-		TotalUncompressed += ((sizeof(FVector) + sizeof(FQuat) + sizeof(FVector)) * NumRawDataTracks * AnimSeq->NumFrames);
+		// See if we're trying alternate compressors
+		const bool bTryAlternateCompressor = MasterTolerance > 0.0f;
 
 		// Filter RAW data to get rid of mismatched tracks (translation/rotation data with a different number of keys than there are frames)
 		// No trivial key removal is done at this point (impossible error metrics of -1), since all of the techniques will perform it themselves
 		AnimSeq->CompressRawAnimData(-1.0f, -1.0f);
 
-		// start with the current technique, or the default if none exists.
-		// this will serve as our fallback if no better technique can be found
-		int32 OriginalKeyEncodingFormat = AnimSeq->KeyEncodingFormat;
-		int32 OriginalTranslationFormat = AnimSeq->TranslationCompressionFormat;
-		int32 OriginalRotationFormat = AnimSeq->RotationCompressionFormat;
-
 		AnimationErrorStats OriginalErrorStats;
 		AnimationErrorStats TrueOriginalErrorStats;
 		TArray<FBoneData> BoneData;
@@ -716,29 +758,58 @@ void FAnimationUtils::CompressAnimSequenceExplicit(
 		FAnimationUtils::ComputeCompressionError(AnimSeq, BoneData, TrueOriginalErrorStats);
 
 		int32 AfterOriginalRecompression = 0;
-		if( (bFirstRecompressUsingCurrentOrDefault && !bTryAlternateCompressor) )
+		if( (bFirstRecompressUsingCurrentOrDefault && !bTryAlternateCompressor) || (AnimSeq->CompressedByteStream.Num() == 0))
 		{
 			UAnimCompress* OriginalCompressionAlgorithm = AnimSeq->CompressionScheme ? AnimSeq->CompressionScheme : FAnimationUtils::GetDefaultAnimationCompressionAlgorithm();
 
+			// Automatic compression brings us back here, so don't create an infinite loop and pick bitwise compress instead.
+			if (!OriginalCompressionAlgorithm || OriginalCompressionAlgorithm->IsA(UAnimCompress_Automatic::StaticClass()))
+			{
+				UAnimCompress* CompressionAlgorithm = NewObject<UAnimCompress_BitwiseCompressOnly>();
+				OriginalCompressionAlgorithm = static_cast<UAnimCompress*>(StaticDuplicateObject(CompressionAlgorithm, AnimSeq));
+			}
+
+			UE_LOG(LogAnimationCompression, Log, TEXT("Recompressing (%s) using current/default (%s) bFirstRecompressUsingCurrentOrDefault(%d) bTryAlternateCompressor(%d) IsCompressedDataValid(%d)"),
+				*AnimSeq->GetFullName(),
+				*OriginalCompressionAlgorithm->GetName(),
+				bFirstRecompressUsingCurrentOrDefault,
+				bTryAlternateCompressor,
+				AnimSeq->IsCompressedDataValid());
+
 			AnimSeq->CompressionScheme = OriginalCompressionAlgorithm;
 			OriginalCompressionAlgorithm->Reduce(AnimSeq, CompressContext);
 			AnimSeq->SetUseRawDataOnly(false);
-			AfterOriginalRecompression = AnimSeq->GetResourceSizeBytes(EResourceSizeMode::Inclusive);
+			AfterOriginalRecompression = AnimSeq->GetApproxCompressedSize();
 
 			// figure out our current compression error
 			FAnimationUtils::ComputeCompressionError(AnimSeq, BoneData, OriginalErrorStats);
 		}
 		else
 		{
-			AfterOriginalRecompression = AnimSeq->GetResourceSizeBytes(EResourceSizeMode::Inclusive);
+			AfterOriginalRecompression = AnimSeq->GetApproxCompressedSize();
 			OriginalErrorStats = TrueOriginalErrorStats;
 		}
- 
+
+		// Get the current size
+		SIZE_T OriginalSize = AnimSeq->GetApproxCompressedSize();
+		TotalSizeBefore += OriginalSize;
+
+		// Estimate total uncompressed
+		TotalUncompressed += AnimSeq->GetUncompressedRawSize();
+
+		// start with the current technique, or the default if none exists.
+		// this will serve as our fallback if no better technique can be found
+		const int32 OriginalKeyEncodingFormat = AnimSeq->KeyEncodingFormat;
+		const int32 OriginalTranslationFormat = AnimSeq->TranslationCompressionFormat;
+		const int32 OriginalRotationFormat = AnimSeq->RotationCompressionFormat;
+
 		// Check for global permission to try an alternative compressor
 		// we don't check for bDoNotOverrideCompression here, as that is now used as part of the UAnimCompress_Automatic compressor
 		// And it's valid to manually recompress animations
 		if( bTryAlternateCompressor /*&& (!AnimSeq->bDoNotOverrideCompression */)
 		{
+			ensure(AnimSeq->CompressedByteStream.Num() > 0);
+
 			AnimationErrorStats NewErrorStats = OriginalErrorStats;
 			if (bRaiseMaxErrorToExisting)
 			{
@@ -767,11 +838,11 @@ void FAnimationUtils::CompressAnimSequenceExplicit(
 				++TotalRecompressions;
 
 				// Prepare to compress
-				SIZE_T CurrentSize = AnimSeq->GetResourceSizeBytes(EResourceSizeMode::Inclusive);
+				SIZE_T CurrentSize = AnimSeq->GetApproxCompressedSize();
 				int32* WinningCompressorCounter = NULL;
 				float* WinningCompressorErrorSum = NULL;
-				int32* WinningCompressorMarginalSavingsSum = NULL;
-				int32 WinningCompressorMarginalSavings = 0;
+				int64* WinningCompressorMarginalSavingsSum = NULL;
+				int64 WinningCompressorMarginalSavings = 0;
 				FString WinningCompressorName;
 				int32 WinningCompressorSavings = 0;
 				float PctSaving = 0.f;
@@ -807,9 +878,9 @@ void FAnimationUtils::CompressAnimSequenceExplicit(
 					UAnimCompress_PerTrackCompression* PerTrackCompressor = NewObject<UAnimCompress_PerTrackCompression>();
 
 					// Start not too aggressive
-					PerTrackCompressor->MaxPosDiffBitwise /= 10.f;
-					PerTrackCompressor->MaxAngleDiffBitwise /= 10.f;
-					PerTrackCompressor->MaxScaleDiffBitwise /= 10.f;
+// 					PerTrackCompressor->MaxPosDiffBitwise /= 10.f;
+// 					PerTrackCompressor->MaxAngleDiffBitwise /= 10.f;
+// 					PerTrackCompressor->MaxScaleDiffBitwise /= 10.f;
 					PerTrackCompressor->bUseAdaptiveError2 = true;
 
 					// Try default compressor first
@@ -818,6 +889,9 @@ void FAnimationUtils::CompressAnimSequenceExplicit(
 					if( NewErrorStats.MaxError >= MasterTolerance )
 					{
 						UE_LOG(LogAnimationCompression, Log, TEXT("\tStandard bitwise compressor too aggressive, lower default settings."));
+
+						AnimationErrorStats TestErrorStats;
+						FAnimationUtils::ComputeCompressionError(AnimSeq, BoneData, TestErrorStats);
 					}
 					else
 					{
@@ -1211,7 +1285,6 @@ void FAnimationUtils::CompressAnimSequenceExplicit(
 				}
 
 				// Increase winning compressor.
-				if( CurrentSize != OriginalSize )
 				{
 					int32 SizeDecrease = OriginalSize - CurrentSize;
 					if (WinningCompressorCounter)
@@ -1222,7 +1295,7 @@ void FAnimationUtils::CompressAnimSequenceExplicit(
 						*WinningCompressorMarginalSavingsSum += WinningCompressorMarginalSavings;
 						check(WinningCompressorSavings == SizeDecrease);
 
-					UE_LOG(LogAnimationCompression, Log, TEXT("  Recompressing '%s' with compressor '%s' saved %i bytes (%i -> %i -> %i) (max diff=%f)\n"),
+					UE_LOG(LogAnimationCompression, Log, TEXT("  Recompressing(%s) with compressor('%s') saved %i bytes (%i -> %i -> %i) (max diff=%f)\n"),
 							*AnimSeq->GetName(),
 							*WinningCompressorName,
 							SizeDecrease,
@@ -1231,14 +1304,20 @@ void FAnimationUtils::CompressAnimSequenceExplicit(
 					}
 					else
 					{
-						TotalNoWinnerRounds++;
-						UE_LOG(LogAnimationCompression, Log, TEXT("  Recompressing '%s' with original/default compressor saved %i bytes (%i -> %i -> %i) (max diff=%f)\n"),
+						UE_LOG(LogAnimationCompression, Log, TEXT("  No compressor suitable! Recompressing(%s) with original/default compressor(%s) saved %i bytes (%i -> %i -> %i) (max diff=%f)\n"),
 							*AnimSeq->GetName(),
+							*AnimSeq->CompressionScheme->GetName(),
 							SizeDecrease,
 							OriginalSize, AfterOriginalRecompression, CurrentSize,
 							WinningCompressorError);
-					}
 
+						UE_LOG(LogAnimationCompression, Log, TEXT("  CompressedTrackOffsets(%d) CompressedByteStream(%d) CompressedScaleOffsets(%d)"),
+							AnimSeq->CompressedTrackOffsets.Num(),
+							AnimSeq->CompressedByteStream.Num(),
+							AnimSeq->CompressedScaleOffsets.GetMemorySize());
+
+						TotalNoWinnerRounds++;
+					}
 
 					// Update the memory stats
 #if STATS
@@ -1257,11 +1336,11 @@ void FAnimationUtils::CompressAnimSequenceExplicit(
 				}
 
 				// Make sure we got that right.
-				check(CurrentSize == AnimSeq->GetResourceSizeBytes(EResourceSizeMode::Inclusive));
+				check(CurrentSize == AnimSeq->GetApproxCompressedSize());
 				TotalSizeNow += CurrentSize;
 
 				PctSaving = TotalSizeBefore > 0 ? 100.f - (100.f * float(TotalSizeNow) / float(TotalSizeBefore)) : 0.f;
-				UE_LOG(LogAnimationCompression, Log, TEXT("Compression Stats Summary [%i total, %i Bytes saved, %i before, %i now, %3.1f%% savings. Uncompressed: %i TotalRatio: %i:1]"),
+				UE_LOG(LogAnimationCompression, Log, TEXT("Compression Stats Summary [Recompressions(%i) Bytes saved(%i) before(%i) now(%i) savings(%3.1f%%) Uncompressed(%i) TotalRatio(%i:1)]"),
 				TotalRecompressions,
 				AlternativeCompressorSavings,
 				TotalSizeBefore, 
@@ -1339,7 +1418,7 @@ void FAnimationUtils::CompressAnimSequenceExplicit(
 		// Do not recompress - Still take into account size for stats.
 		else
 		{
-			TotalSizeNow += AnimSeq->GetResourceSizeBytes(EResourceSizeMode::Inclusive);
+			TotalSizeNow += AnimSeq->GetApproxCompressedSize();
 		}
 	}
 	else
@@ -1441,7 +1520,7 @@ FString FAnimationUtils::GetAnimationCompressionFormatString(AnimationCompressio
 	case ACF_Identity:
 		return FString(TEXT("ACF_Identity"));
 	default:
-		UE_LOG(LogAnimation, Warning, TEXT("AnimationCompressionFormat was not found:  %i"), static_cast<int32>(InFormat) );
+		UE_LOG(LogAnimationCompression, Warning, TEXT("AnimationCompressionFormat was not found:  %i"), static_cast<int32>(InFormat) );
 	}
 
 	return FString(TEXT("Unknown"));
@@ -1464,7 +1543,7 @@ FString FAnimationUtils::GetAnimationKeyFormatString(AnimationKeyFormat InFormat
 	case AKF_PerTrackCompression:
 		return FString(TEXT("AKF_PerTrackCompression"));
 	default:
-		UE_LOG(LogAnimation, Warning, TEXT("AnimationKeyFormat was not found:  %i"), static_cast<int32>(InFormat) );
+		UE_LOG(LogAnimationCompression, Warning, TEXT("AnimationKeyFormat was not found:  %i"), static_cast<int32>(InFormat) );
 	}
 
 	return FString(TEXT("Unknown"));
diff --git a/Engine/Source/Runtime/Engine/Private/Animation/CachedAnimData.cpp b/Engine/Source/Runtime/Engine/Private/Animation/CachedAnimData.cpp
--- a/Engine/Source/Runtime/Engine/Private/Animation/CachedAnimData.cpp
+++ b/Engine/Source/Runtime/Engine/Private/Animation/CachedAnimData.cpp
@@ -58,6 +58,11 @@ bool FCachedAnimStateData::IsRelevant(UAnimInstance& InAnimInstance) const
 	return FAnimWeight::IsRelevant(GetWeight(InAnimInstance));
 }
 
+bool FCachedAnimStateData::IsActiveState(class UAnimInstance& InAnimInstance) const
+{
+	return IsValid(InAnimInstance) ? (InAnimInstance.GetCurrentStateName(MachineIndex) == StateName) : false;
+}
+
 bool FCachedAnimStateArray::IsValid(UAnimInstance& InAnimInstance) const
 {
 	// Make sure the setup validates our assumptions.
diff --git a/Engine/Source/Runtime/Engine/Private/AudioStreaming.cpp b/Engine/Source/Runtime/Engine/Private/AudioStreaming.cpp
--- a/Engine/Source/Runtime/Engine/Private/AudioStreaming.cpp
+++ b/Engine/Source/Runtime/Engine/Private/AudioStreaming.cpp
@@ -300,7 +300,7 @@ void FStreamingWaveData::BeginPendingRequests(const TArray<uint32>& IndicesToLoa
 			// returns.
 			PendingChunkChangeRequestStatus.Increment();
 
-			EAsyncIOPriority AsyncIOPriority = CurrentRequest.bPrioritiseRequest ? AIOP_BelowNormal : AIOP_Low;
+			EAsyncIOPriority AsyncIOPriority = AIOP_High;
 
 			// Load and decompress async.
 #if WITH_EDITORONLY_DATA
diff --git a/Engine/Source/Runtime/Engine/Private/Collision/Collision.cpp b/Engine/Source/Runtime/Engine/Private/Collision/Collision.cpp
--- a/Engine/Source/Runtime/Engine/Private/Collision/Collision.cpp
+++ b/Engine/Source/Runtime/Engine/Private/Collision/Collision.cpp
@@ -12,7 +12,15 @@
 #include "CollisionQueryParams.h"
 #include "Engine/World.h"
 #include "Components/PrimitiveComponent.h"
+#include "Components/ShapeComponent.h"
+#include "Components/StaticMeshComponent.h"
+#include "Components/SkinnedMeshComponent.h"
+#include "Engine/StaticMesh.h"
+#include "Engine/SkeletalMesh.h"
 #include "Engine/CollisionProfile.h"
+#include "PhysicsEngine/BodySetupEnums.h"
+#include "PhysicsEngine/BodySetup.h"
+#include "PhysicsEngine/PhysicsSettings.h"
 
 // TEMP until crash is fixed in IsCollisionEnabled().
 #include "LandscapeHeightfieldCollisionComponent.h"
@@ -524,7 +532,24 @@ DEFINE_LOG_CATEGORY(LogCollisionCommands);
 
 namespace CollisionResponseConsoleCommands
 {
-	static const FString ResponseStrings[] = {TEXT("Ignore"), TEXT("Overlap"), TEXT("Block")};
+	static const TArray<FString> ResponseStrings = {TEXT("Ignore"), TEXT("Overlap"), TEXT("Block")};
+	static const TArray<FString> ComplexityStrings = {TEXT("Default"), TEXT("SimpleAndComplex"), TEXT("UseSimpleAsComplex"), TEXT("UseComplexAsSimple")};
+
+	FString GetCommaSeparatedList(const TArray<FString>& List)
+	{
+		FString Result;
+		if (List.Num() > 0)
+		{
+			Result.Reserve(128);
+			for (int32 i=0; i < List.Num()-1; i++)
+			{
+				Result.Append(List[i]);
+				Result.Append(TEXT(", "));
+			}
+			Result.Append(List[List.Num()-1]);
+		}
+		return Result;
+	}
 
 	FString FillString(TCHAR Char, int32 Count)
 	{
@@ -560,16 +585,54 @@ namespace CollisionResponseConsoleCommands
 		{
 			if (Obj->GetOuter())
 			{
-				Result = TEXT("'") + Obj->GetPathName(Obj->GetOuter()->GetOuter()) + TEXT("'");
+				Result = FString::Printf(TEXT("'%s'"), *Obj->GetPathName(Obj->GetOuter()->GetOuter()));
 			}
 			else
 			{
-				Result = TEXT("'") + Obj->GetPathName(Obj->GetOuter()) + TEXT("'");
+				Result = FString::Printf(TEXT("'%s'"), *Obj->GetPathName(Obj->GetOuter()));
 			}
 		}
 		return Result;
 	}
 
+	FString GetAssetName(const UPrimitiveComponent* Comp)
+	{
+		FString AssetName;
+		if (const UStaticMeshComponent* StaticMeshComp = Cast<UStaticMeshComponent>(Comp))
+		{
+			if (StaticMeshComp->GetStaticMesh())
+			{
+				AssetName = StaticMeshComp->GetStaticMesh()->GetPathName();
+			}
+		}
+		else if (const USkinnedMeshComponent* SkinnedMeshComp = Cast<USkinnedMeshComponent>(Comp))
+		{
+			if (SkinnedMeshComp->SkeletalMesh)
+			{
+				AssetName = SkinnedMeshComp->SkeletalMesh->GetPathName();
+			}
+		}
+
+		return AssetName;
+	}
+
+	//////////////////////////////////////////////////////////////////////////
+	// Helper to map enum to display name through a map (to avoid repeated string building).
+	template<typename EnumType>
+	const FString& MapEnumToDisplayName(const UEnum *ComplexityEnum, EnumType EnumFlag, TMap<EnumType, FString>& EnumToDisplayNameMap)
+	{
+		if (!EnumToDisplayNameMap.Contains(EnumFlag))
+		{
+			const FString ComplexityName = ComplexityEnum ? ComplexityEnum->GetNameByValue((int64)EnumFlag).ToString() : TEXT("<unknown>");
+			const FString ComplexityDisplayName = GetDisplayNameText(ComplexityEnum, (int64)EnumFlag, ComplexityName);
+			return EnumToDisplayNameMap.Add(EnumFlag, ComplexityDisplayName);
+		}
+		else
+		{
+			return EnumToDisplayNameMap.FindChecked(EnumFlag);
+		}
+	}
+
 	void ListCollisionProfileNames()
 	{
 		TArray<TSharedPtr<FName>> ProfileNameList;
@@ -631,6 +694,41 @@ namespace CollisionResponseConsoleCommands
 		ECollisionResponse RequiredResponse;
 	};
 
+	struct FSortComponentsForComplexity
+	{
+		FSortComponentsForComplexity(TMap<UPrimitiveComponent*, FString>& NameMap)
+		: InternalNameMap(NameMap)
+		{
+		}
+
+		bool operator()(const UPrimitiveComponent& A, const UPrimitiveComponent& B) const
+		{
+			// Sort by asset name;
+			const FString NameA = InternalNameMap.FindRef(&A);
+			const FString NameB = InternalNameMap.FindRef(&B);
+			if (NameA != NameB)
+			{
+				return NameA < NameB;
+			}
+
+			// Sort by object name.
+			UObject* AOwner = A.GetOuter();
+			UObject* BOwner = B.GetOuter();
+			if (AOwner && BOwner)
+			{
+				return AOwner->GetName() < BOwner->GetName();
+			}
+			else
+			{
+				return false;
+			}
+		}
+
+	private: 
+
+		TMap<UPrimitiveComponent*, FString>& InternalNameMap;
+	};
+
 	ECollisionResponse StringToCollisionResponse(const FString& InString)
 	{
 		int32 TestIndex = 0;
@@ -708,6 +806,52 @@ namespace CollisionResponseConsoleCommands
 		return NAME_None;
 	}
 
+	ECollisionTraceFlag StringToCollisionComplexity(const FString& InString)
+	{
+		const UEnum *ComplexityEnum = FindObject<UEnum>(ANY_PACKAGE, TEXT("ECollisionTraceFlag"), true);
+		if (ComplexityEnum)
+		{
+			int32 ComplexityInt = INDEX_NONE;
+
+			// Try the name they gave us.
+			ComplexityInt = ComplexityEnum->GetValueByName(FName(*InString));
+			if (ComplexityInt != INDEX_NONE)
+			{
+				return ECollisionTraceFlag(ComplexityInt);
+			}
+
+			// Try with adding the prefix, ie "CTF_".
+			const FString WithPrefixName = ComplexityEnum->GenerateEnumPrefix() + TEXT("_") + InString;
+			ComplexityInt = ComplexityEnum->GetValueByName(FName(*WithPrefixName));
+			if (ComplexityInt != INDEX_NONE)
+			{
+				return ECollisionTraceFlag(ComplexityInt);
+			}
+
+			// Try with adding the prefix and "Use", ie "CTF_Use" (for example 'Default'->'CTF_UseDefault').
+			const FString WithPrefixUseName = ComplexityEnum->GenerateEnumPrefix() + TEXT("_Use") + InString;
+			ComplexityInt = ComplexityEnum->GetValueByName(FName(*WithPrefixUseName));
+			if (ComplexityInt != INDEX_NONE)
+			{
+				return ECollisionTraceFlag(ComplexityInt);
+			}
+
+			// Try matching the display name
+			const FString NullString(TEXT(""));
+			for (int32 ComplexityIndex = 0; ComplexityIndex < ComplexityEnum->NumEnums(); ComplexityIndex++)
+			{
+				int64 ComplexityValue = ComplexityEnum->GetValueByIndex(ComplexityIndex);
+				const FString ComplexityDisplayName = GetDisplayNameText(ComplexityEnum, ComplexityValue, NullString);
+				if (ComplexityDisplayName == InString)
+				{
+					return ECollisionTraceFlag(ComplexityValue);
+				}
+			}
+		}
+
+		return CTF_MAX;
+	}
+
 	//////////////////////////////////////////////////////////////////////////
 	void ListComponentsWithResponseToProfile(const ECollisionResponse RequiredResponse, const FName& ProfileToCheck)
 	{
@@ -749,6 +893,7 @@ namespace CollisionResponseConsoleCommands
 		{
 			Results.Sort(FSortComponentsWithResponseToProfile(RequiredResponse));
 			const UEnum *ChannelEnum = FindObject<UEnum>(ANY_PACKAGE, TEXT("ECollisionChannel"), true);
+			TMap<ECollisionChannel, FString> EnumToDisplayNameMap;
 
 			// Get max column widths for some data
 			int32 MaxNameWidth = 0;
@@ -765,8 +910,7 @@ namespace CollisionResponseConsoleCommands
 
 				if (ChannelEnum)
 				{
-					const FString ChannelName = ChannelEnum->GetNameByValue((int64)Comp->GetCollisionObjectType()).ToString();
-					const FString ChannelDisplayName = GetDisplayNameText(ChannelEnum, (int64)Comp->GetCollisionObjectType(), ChannelName);
+					const FString& ChannelDisplayName = MapEnumToDisplayName<ECollisionChannel>(ChannelEnum, Comp->GetCollisionObjectType(), EnumToDisplayNameMap);
 					MaxChannelWidth = FMath::Max<int32>(MaxChannelWidth, ChannelDisplayName.Len());
 				}
 
@@ -792,8 +936,7 @@ namespace CollisionResponseConsoleCommands
 			int32 Index=0;
 			for (UPrimitiveComponent* Comp : Results)
 			{
-				const FString ChannelName = (ChannelEnum ? ChannelEnum->GetNameByValue((int64)Comp->GetCollisionObjectType()).ToString() : TEXT("<unknown>"));
-				const FString ChannelDisplayName = GetDisplayNameText(ChannelEnum, (int64)Comp->GetCollisionObjectType(), ChannelName);
+				const FString& ChannelDisplayName = MapEnumToDisplayName<ECollisionChannel>(ChannelEnum, Comp->GetCollisionObjectType(), EnumToDisplayNameMap);
 				UObject* Outer = Comp->GetOuter();
 				if (Outer)
 				{
@@ -814,6 +957,7 @@ namespace CollisionResponseConsoleCommands
 			}
 			UE_LOG(LogCollisionCommands, Log, TEXT("%s"), *LineMarker);
 		}
+		// Summary
 		check(RequiredResponse < ECollisionResponse::ECR_MAX);
 		UE_LOG(LogCollisionCommands, Log, TEXT("Found %d components with '%s' response to profile '%s'."), Results.Num(), *ResponseStrings[(int32)RequiredResponse], *ProfileToCheck.ToString());
 	}
@@ -825,7 +969,7 @@ namespace CollisionResponseConsoleCommands
 		if (Args.Num() < 2)
 		{
 			UE_LOG(LogCollisionCommands, Warning, TEXT("Usage: 'Collision.ListComponentsWithResponseToProfile <Response> <Profile>'."));
-			UE_LOG(LogCollisionCommands, Warning, TEXT("  Response: Ignore, Overlap, Block"));
+			UE_LOG(LogCollisionCommands, Warning, TEXT("  Response: %s"), *GetCommaSeparatedList(ResponseStrings));
 			UE_LOG(LogCollisionCommands, Warning, TEXT("  Profile:  Collision profile name or index. Use 'Collision.ListProfiles' to see a full list."));
 			return;
 		}
@@ -835,7 +979,7 @@ namespace CollisionResponseConsoleCommands
 		ECollisionResponse RequiredResponse = StringToCollisionResponse(ResponseString);
 		if (RequiredResponse == ECollisionResponse::ECR_MAX)
 		{
-			UE_LOG(LogCollisionCommands, Warning, TEXT("Unknown response '%s'. Must be Ignore, Overlap, or Block."), *ResponseString);
+			UE_LOG(LogCollisionCommands, Warning, TEXT("Unknown response '%s'. Must be one of %s."), *ResponseString, *GetCommaSeparatedList(ResponseStrings));
 			return;
 		}
 
@@ -854,6 +998,7 @@ namespace CollisionResponseConsoleCommands
 	//////////////////////////////////////////////////////////////////////////
 	void ListProfilesWithResponseToChannel(const ECollisionResponse RequiredResponse, const ECollisionChannel TestChannel)
 	{
+		// Match
 		TSet<FName> Results;
 		TArray<TSharedPtr<FName>> ProfileNameList;
 		UCollisionProfile::Get()->GetProfileNames(ProfileNameList);
@@ -873,6 +1018,7 @@ namespace CollisionResponseConsoleCommands
 			++Index;
 		}
 
+		// Display Data
 		if (Results.Num() > 0)
 		{
 			Results.Sort([](const FName& A, const FName& B) { return A < B; });
@@ -881,6 +1027,7 @@ namespace CollisionResponseConsoleCommands
 				UE_LOG(LogCollisionCommands, Log, TEXT("%s"), *ResultName.ToString());
 			}
 		}
+		// Display Summary
 		check(RequiredResponse < ECollisionResponse::ECR_MAX);
 		const UEnum *ChannelEnum = FindObject<UEnum>(ANY_PACKAGE, TEXT("ECollisionChannel"), true);
 		const FString ChannelName = (ChannelEnum ? ChannelEnum->GetNameStringByValue(TestChannel) : TEXT("<unknown>"));
@@ -896,7 +1043,7 @@ namespace CollisionResponseConsoleCommands
 		if (Args.Num() < 2)
 		{
 			UE_LOG(LogCollisionCommands, Warning, TEXT("Usage: 'Collision.ListProfilesWithResponseToChannel <Response> <Channel>'."));
-			UE_LOG(LogCollisionCommands, Warning, TEXT("  Response: Ignore, Overlap, Block"));
+			UE_LOG(LogCollisionCommands, Warning, TEXT("  Response: %s"), *GetCommaSeparatedList(ResponseStrings));
 			UE_LOG(LogCollisionCommands, Warning, TEXT("  Profile:  Collision channel name or index. Use 'Collision.ListChannels' to see a full list."));
 			return;
 		}
@@ -906,7 +1053,7 @@ namespace CollisionResponseConsoleCommands
 		const ECollisionResponse RequiredResponse = StringToCollisionResponse(ResponseString);
 		if (RequiredResponse == ECollisionResponse::ECR_MAX)
 		{
-			UE_LOG(LogCollisionCommands, Warning, TEXT("Unknown response '%s'. Must be Ignore, Overlap, or Block."), *ResponseString);
+			UE_LOG(LogCollisionCommands, Warning, TEXT("Unknown response '%s'. Must be one of %s."), *ResponseString, *GetCommaSeparatedList(ResponseStrings));
 			return;
 		}
 
@@ -923,6 +1070,168 @@ namespace CollisionResponseConsoleCommands
 	}
 
 
+	//////////////////////////////////////////////////////////////////////////
+	void ListObjectsWithCollisionComplexity(ECollisionTraceFlag Complexity)
+	{
+		if (Complexity == CTF_MAX)
+		{
+			return;
+		}
+
+		// See what "Default" means, so we can include those
+		const ECollisionTraceFlag DefaultFlag = UPhysicsSettings::Get()->DefaultShapeComplexity;
+		const ECollisionTraceFlag ActualComplexity = (Complexity == CTF_UseDefault) ? DefaultFlag : Complexity;
+
+		// Look at all components and check collision complexity
+		TSet<UPrimitiveComponent*> Results;
+		for (TObjectIterator<UPrimitiveComponent> Iter(RF_NoFlags, /*bIncludeDerivedClasses*/ true, /*ExclusionFlags*/ EInternalObjectFlags::None); Iter; ++Iter)
+		{
+			UPrimitiveComponent* Comp = *Iter;
+			if (Comp)
+			{
+				// Special case for UShapeComponent CDOs, GetBodySetup() asserts.
+				if (Comp->GetClass() == UShapeComponent::StaticClass())
+				{
+					continue;
+				}
+
+				// Get collision complexity from body setup
+				if (UBodySetup* BodySetup = Comp->GetBodySetup())
+				{
+					// If matching "Default", only list those explicitly set to default.
+					if (Complexity == CTF_UseDefault && BodySetup->CollisionTraceFlag == CTF_UseDefault)
+					{
+						Results.Add(Comp);
+					}
+					else
+					{
+						// Using GetCollisionTraceFlag includes both the same requested complexity and those with Default if that is the same as the requested mode.
+						const ECollisionTraceFlag BodyFlag = BodySetup->GetCollisionTraceFlag();
+						if (Complexity == BodyFlag)
+						{
+							Results.Add(Comp);
+						}
+					}
+				}
+			}
+		}
+
+		// Log results.
+		if (Results.Num() > 0)
+		{
+			// Fill component name Maps.
+			TMap<UPrimitiveComponent*, FString> ComponentAssetNameMap;
+			TMap<UPrimitiveComponent*, FString> ComponentPathNameMap;
+			for (UPrimitiveComponent* Comp : Results)
+			{
+				const FString AssetName = GetAssetName(Comp);
+				if (AssetName.Len() > 0)
+				{
+					ComponentAssetNameMap.Add(Comp, AssetName);
+				}
+
+				UObject* Outer = Comp->GetOuter();
+				if (Outer)
+				{
+					const FString PathName = FormatObjectName(Comp);
+					ComponentPathNameMap.Add(Comp, PathName);
+				}
+			}
+
+			Results.Sort(FSortComponentsForComplexity(ComponentAssetNameMap));
+			const UEnum *ComplexityEnum = FindObject<UEnum>(ANY_PACKAGE, TEXT("ECollisionTraceFlag"), true);
+			if (!ensure(ComplexityEnum != nullptr))
+			{
+				return;
+			}
+
+			// Generate map of enum->name mappings
+			TMap<ECollisionTraceFlag, FString> EnumToDisplayNameMap;
+
+			// Get max column widths for some data
+			const FString ComplexityHeading = TEXT("Collision Complexity");
+			int32 MaxComplexityWidth = ComplexityHeading.Len();
+			int32 MaxNameWidth = 0;
+			int32 MaxAssetNameWidth = 0;
+			for (UPrimitiveComponent* Comp : Results)
+			{
+				// Object name width
+				UObject* Outer = Comp->GetOuter();
+				if (Outer)
+				{
+					const FString PathName = ComponentPathNameMap.FindRef(Comp);
+					MaxNameWidth = FMath::Max<int32>(MaxNameWidth, PathName.Len());
+				}
+
+				// Complexity display name width
+				const UBodySetup* BodySetup = Comp->GetBodySetup();
+				const ECollisionTraceFlag BodyFlag = BodySetup->CollisionTraceFlag;
+
+				const FString& ComplexityDisplayName = MapEnumToDisplayName<ECollisionTraceFlag>(ComplexityEnum, BodyFlag, EnumToDisplayNameMap);				
+				MaxComplexityWidth = FMath::Max(MaxComplexityWidth, ComplexityDisplayName.Len());
+
+				// Asset name width
+				const FString& AssetName = ComponentAssetNameMap.FindRef(Comp);
+				MaxAssetNameWidth = FMath::Max(MaxAssetNameWidth, AssetName.Len());
+			}
+
+			// Display Column headings
+			const FString Output = FString::Printf(TEXT("  #, %-*s, %-*s, %-*s, Path"), MaxNameWidth, TEXT("Component"), MaxComplexityWidth, *ComplexityHeading, MaxAssetNameWidth, TEXT("Asset"));
+			UE_LOG(LogCollisionCommands, Log, TEXT("%s"), *Output);
+			const int32 TotalLen = Output.Len() + 16;
+			const FString LineMarker = FillString(TCHAR('-'), TotalLen);
+			UE_LOG(LogCollisionCommands, Log, TEXT("%s"), *LineMarker);
+
+			// Display Data
+			int32 Index=0;
+			for (UPrimitiveComponent* Comp : Results)
+			{
+				UObject* Outer = Comp->GetOuter();
+				if (Outer)
+				{
+					const FString& PathName = ComponentPathNameMap.FindRef(Comp);
+					const UBodySetup* BodySetup = Comp->GetBodySetup();
+					const ECollisionTraceFlag BodyFlag = BodySetup->CollisionTraceFlag;
+					const FString& ComplexityDisplayName = MapEnumToDisplayName<ECollisionTraceFlag>(ComplexityEnum, BodyFlag, EnumToDisplayNameMap);
+					const FString& AssetName = ComponentAssetNameMap.FindRef(Comp);
+
+					UE_LOG(LogCollisionCommands, Log, TEXT("%3d, %-*s, %-*s, %-*s, %s"),
+						   Index, MaxNameWidth, *PathName, MaxComplexityWidth, *ComplexityDisplayName, MaxAssetNameWidth, *AssetName, Outer->GetOuter() ? *GetPathNameSafe(Outer->GetOuter()) : *GetPathNameSafe(Outer));
+					Index++;
+				}
+			}
+			UE_LOG(LogCollisionCommands, Log, TEXT("%s"), *LineMarker);
+		}
+		// Display Summary
+		UE_LOG(LogCollisionCommands, Log, TEXT("Found %d components with '%s' collision complexity."), Results.Num(), *ComplexityStrings[int32(Complexity)]);
+	}
+
+
+
+	//////////////////////////////////////////////////////////////////////////
+	// Args: <Complexity>
+	void Parse_ListObjectsWithCollisionComplexity(const TArray<FString>& Args, UWorld* World)
+	{
+		if (Args.Num() < 1)
+		{
+			UE_LOG(LogCollisionCommands, Warning, TEXT("Usage: 'Collision.ListObjectsWithCollisionComplexity <Complexity>'."));
+			UE_LOG(LogCollisionCommands, Warning, TEXT("  Complexity: %s"), *GetCommaSeparatedList(ComplexityStrings));
+			return;
+		}
+
+		// Arg0 : Complexity
+		const FString& ComplexityString = Args[0];
+		const ECollisionTraceFlag Complexity = StringToCollisionComplexity(ComplexityString);
+		if (Complexity == CTF_MAX)
+		{
+			UE_LOG(LogCollisionCommands, Warning, TEXT("Unknown complexity '%s'. Must be one of %s."), *ComplexityString, *GetCommaSeparatedList(ComplexityStrings));
+			return;
+		}
+
+		ListObjectsWithCollisionComplexity(Complexity);
+	}
+
+
 	//////////////////////////////////////////////////////////////////////////
 	// Console commands
 
@@ -950,6 +1259,12 @@ namespace CollisionResponseConsoleCommands
 		FConsoleCommandWithWorldAndArgsDelegate::CreateStatic(CollisionResponseConsoleCommands::Parse_ListProfilesWithResponseToChannel)
 	);
 
+	static FAutoConsoleCommandWithWorldAndArgs ListProfilesWithCollisionComplexityCommand(
+		TEXT("Collision.ListObjectsWithCollisionComplexity"),
+		TEXT(""),
+		FConsoleCommandWithWorldAndArgsDelegate::CreateStatic(CollisionResponseConsoleCommands::Parse_ListObjectsWithCollisionComplexity)
+	);
+
 }
 
 #endif //!UE_BUILD_SHIPPING
diff --git a/Engine/Source/Runtime/Engine/Private/Collision/CollisionConversions.cpp b/Engine/Source/Runtime/Engine/Private/Collision/CollisionConversions.cpp
--- a/Engine/Source/Runtime/Engine/Private/Collision/CollisionConversions.cpp
+++ b/Engine/Source/Runtime/Engine/Private/Collision/CollisionConversions.cpp
@@ -686,7 +686,8 @@ static bool ComputeInflatedMTD_Internal(const float MtdInflation, const PxLocati
 
 	PxVec3 PxMtdNormal(0.f);
 	PxF32 PxMtdDepth = 0.f;
-	const PxGeometry& POtherGeom = PHit.shape->getGeometry().any();
+	PxGeometryHolder Holder = PHit.shape->getGeometry();
+	const PxGeometry& POtherGeom = Holder.any();
 	const bool bMtdResult = PxGeometryQuery::computePenetration(PxMtdNormal, PxMtdDepth, Geom, QueryTM, POtherGeom, PShapeWorldPose);
 	if (bMtdResult)
 	{
@@ -966,7 +967,8 @@ static bool ConvertOverlappedShapeToImpactHit(const UWorld* World, const PxLocat
 				{
 					// MTD failed, use point distance. This is not ideal.
 					// Note: faceIndex seems to be unreliable for convex meshes in these cases, so not using FindGeomOpposingNormal() for them here.
-					PxGeometry& PGeom = PShape->getGeometry().any();
+					PxGeometryHolder Holder = PShape->getGeometry();
+					PxGeometry& PGeom = Holder.any();
 					PxVec3 PClosestPoint;
 					const float Distance = PxGeometryQuery::pointDistance(QueryTM.p, PGeom, PShapeWorldPose, &PClosestPoint);
 
diff --git a/Engine/Source/Runtime/Engine/Private/Components/PrimitiveComponent.cpp b/Engine/Source/Runtime/Engine/Private/Components/PrimitiveComponent.cpp
--- a/Engine/Source/Runtime/Engine/Private/Components/PrimitiveComponent.cpp
+++ b/Engine/Source/Runtime/Engine/Private/Components/PrimitiveComponent.cpp
@@ -205,6 +205,10 @@ UPrimitiveComponent::UPrimitiveComponent(const FObjectInitializer& ObjectInitial
 	LastCheckedAllCollideableDescendantsTime = 0.f;
 	
 	bApplyImpulseOnDamage = true;
+
+#if WITH_EDITORONLY_DATA
+	bEnableAutoLODGeneration = true;
+#endif // WITH_EDITORONLY_DATA
 }
 
 bool UPrimitiveComponent::UsesOnlyUnlitMaterials() const
@@ -3199,9 +3203,14 @@ void UPrimitiveComponent::SetRenderInMono(bool bValue)
 
 void UPrimitiveComponent::SetLODParentPrimitive(UPrimitiveComponent * InLODParentPrimitive)
 {
-	// @todo, what do we do with old parent. We can't just reset undo parent because the parent migh be used by other primitive
-	LODParentPrimitive = InLODParentPrimitive;
-	MarkRenderStateDirty();
+#if WITH_EDITOR
+	if (ShouldGenerateAutoLOD())
+#endif
+	{
+		// @todo, what do we do with old parent. We can't just reset undo parent because the parent migh be used by other primitive
+		LODParentPrimitive = InLODParentPrimitive;
+		MarkRenderStateDirty();
+	}
 }
 
 UPrimitiveComponent* UPrimitiveComponent::GetLODParentPrimitive() const
@@ -3251,7 +3260,7 @@ void UPrimitiveComponent::SetCustomNavigableGeometry(const EHasCustomNavigableGe
 #if WITH_EDITOR
 const bool UPrimitiveComponent::ShouldGenerateAutoLOD() const
 {
-	return (Mobility != EComponentMobility::Movable);
+	return (Mobility != EComponentMobility::Movable) && bEnableAutoLODGeneration;
 }
 #endif 
 
diff --git a/Engine/Source/Runtime/Engine/Private/Components/SkinnedMeshComponent.cpp b/Engine/Source/Runtime/Engine/Private/Components/SkinnedMeshComponent.cpp
--- a/Engine/Source/Runtime/Engine/Private/Components/SkinnedMeshComponent.cpp
+++ b/Engine/Source/Runtime/Engine/Private/Components/SkinnedMeshComponent.cpp
@@ -166,6 +166,11 @@ namespace FAnimUpdateRateManager
 		0,
 		TEXT("Set to 1 to force interpolation"));
 
+	static TAutoConsoleVariable<int32> CVarURODisableInterpolation(
+		TEXT("a.URO.DisableInterpolation"),
+		0,
+		TEXT("Set to 1 to disable interpolation"));
+
 	void AnimUpdateRateSetParams(FAnimUpdateRateParametersTracker* Tracker, float DeltaTime, bool bRecentlyRendered, float MaxDistanceFactor, int32 MinLod, bool bNeedsValidRootMotion, bool bUsingRootMotionFromEverything)
 	{
 		// default rules for setting update rates
@@ -2860,7 +2865,8 @@ void FAnimUpdateRateParameters::SetTrailMode(float DeltaTime, uint8 UpdateRateSh
 	UpdateRate = FMath::Max(NewUpdateRate, 1);
 	// Make sure EvaluationRate is a multiple of UpdateRate.
 	EvaluationRate = FMath::Max((NewEvaluationRate / UpdateRate) * UpdateRate, 1);
-	bInterpolateSkippedFrames = (bNewInterpSkippedFrames && (EvaluationRate < MaxEvalRateForInterpolation)) || (FAnimUpdateRateManager::CVarForceInterpolation.GetValueOnAnyThread() == 1);
+	bInterpolateSkippedFrames = (FAnimUpdateRateManager::CVarURODisableInterpolation.GetValueOnAnyThread() == 0) &&
+		((bNewInterpSkippedFrames && (EvaluationRate < MaxEvalRateForInterpolation)) || (FAnimUpdateRateManager::CVarForceInterpolation.GetValueOnAnyThread() == 1));
 
 	// Make sure we don't overflow. we don't need very large numbers.
 	const uint32 Counter = (GFrameCounter + UpdateRateShift)% MAX_uint32;
diff --git a/Engine/Source/Runtime/Engine/Private/Components/StaticMeshComponent.cpp b/Engine/Source/Runtime/Engine/Private/Components/StaticMeshComponent.cpp
--- a/Engine/Source/Runtime/Engine/Private/Components/StaticMeshComponent.cpp
+++ b/Engine/Source/Runtime/Engine/Private/Components/StaticMeshComponent.cpp
@@ -2221,7 +2221,8 @@ void UStaticMeshComponent::ApplyComponentInstanceData(FStaticMeshComponentInstan
 		}
 		else
 		{
-			UE_LOG(LogStaticMesh, Warning, TEXT("Cached component instance data transform did not match!  Discarding cached lighting data which will cause lighting to be unbuilt.\n%s\nCurrent: %s Cached: %s"), 
+			UE_ASSET_LOG(LogStaticMesh, Warning, this,
+				TEXT("Cached component instance data transform did not match!  Discarding cached lighting data which will cause lighting to be unbuilt.\n%s\nCurrent: %s Cached: %s"), 
 				*GetPathName(),
 				*GetComponentTransform().ToString(),
 				*StaticMeshInstanceData->CachedStaticLighting.Transform.ToString());
diff --git a/Engine/Source/Runtime/Engine/Private/ImageUtils.cpp b/Engine/Source/Runtime/Engine/Private/ImageUtils.cpp
--- a/Engine/Source/Runtime/Engine/Private/ImageUtils.cpp
+++ b/Engine/Source/Runtime/Engine/Private/ImageUtils.cpp
@@ -12,6 +12,7 @@ ImageUtils.cpp: Image utility functions.
 #include "Logging/MessageLog.h"
 #include "IImageWrapper.h"
 #include "IImageWrapperModule.h"
+#include "ModuleManager.h"
 
 DEFINE_LOG_CATEGORY_STATIC(LogImageUtils, Log, All);
 
diff --git a/Engine/Source/Runtime/Engine/Private/KismetMathLibrary.cpp b/Engine/Source/Runtime/Engine/Private/KismetMathLibrary.cpp
--- a/Engine/Source/Runtime/Engine/Private/KismetMathLibrary.cpp
+++ b/Engine/Source/Runtime/Engine/Private/KismetMathLibrary.cpp
@@ -758,10 +758,24 @@ void UKismetMathLibrary::GetYawPitchFromVector(FVector InVec, float& Yaw, float&
 {
 	FVector NormalizedVector = InVec.GetSafeNormal();
 	// Find yaw.
-	Yaw = FMath::Atan2(NormalizedVector.Y, NormalizedVector.X) * 180.f / PI;
+	Yaw = FMath::RadiansToDegrees(FMath::Atan2(NormalizedVector.Y, NormalizedVector.X));
 
 	// Find pitch.
-	Pitch = FMath::Atan2(NormalizedVector.Z, FMath::Sqrt(NormalizedVector.X*NormalizedVector.X + NormalizedVector.Y*NormalizedVector.Y)) * 180.f / PI;
+	Pitch = FMath::RadiansToDegrees(FMath::Atan2(NormalizedVector.Z, FMath::Sqrt(NormalizedVector.X*NormalizedVector.X + NormalizedVector.Y*NormalizedVector.Y)));
+}
+
+void UKismetMathLibrary::GetAzimuthAndElevation(FVector InDirection, const FTransform& ReferenceFrame, float& Azimuth, float& Elevation)
+{
+	FVector2D Result = FMath::GetAzimuthAndElevation
+	(
+		InDirection.GetSafeNormal(),
+		ReferenceFrame.GetUnitAxis(EAxis::X),
+		ReferenceFrame.GetUnitAxis(EAxis::Y),
+		ReferenceFrame.GetUnitAxis(EAxis::Z)
+	);
+
+	Azimuth = FMath::RadiansToDegrees(Result.X);
+	Elevation = FMath::RadiansToDegrees(Result.Y);
 }
 
 void UKismetMathLibrary::BreakRotIntoAxes(const FRotator& InRot, FVector& X, FVector& Y, FVector& Z)
diff --git a/Engine/Source/Runtime/Engine/Private/LevelActor.cpp b/Engine/Source/Runtime/Engine/Private/LevelActor.cpp
--- a/Engine/Source/Runtime/Engine/Private/LevelActor.cpp
+++ b/Engine/Source/Runtime/Engine/Private/LevelActor.cpp
@@ -39,6 +39,8 @@
 #include "Components/BoxComponent.h"
 #include "GameFramework/MovementComponent.h"
 
+#include "Misc/TimeGuard.h"
+
 #define LOCTEXT_NAMESPACE "LevelActor"
 
 // CVars
@@ -295,6 +297,9 @@ AActor* UWorld::SpawnActor( UClass* Class, FVector const* Location, FRotator con
 AActor* UWorld::SpawnActor( UClass* Class, FTransform const* UserTransformPtr, const FActorSpawnParameters& SpawnParameters )
 {
 	SCOPE_CYCLE_COUNTER(STAT_SpawnActorTime);
+	SCOPE_TIME_GUARD_NAMED_MS(TEXT("SpawnActor Of Type"), Class->GetFName(), 2);
+	
+
 	check( CurrentLevel ); 	
 	check(GIsEditor || (CurrentLevel == PersistentLevel));
 
diff --git a/Engine/Source/Runtime/Engine/Private/Materials/Material.cpp b/Engine/Source/Runtime/Engine/Private/Materials/Material.cpp
--- a/Engine/Source/Runtime/Engine/Private/Materials/Material.cpp
+++ b/Engine/Source/Runtime/Engine/Private/Materials/Material.cpp
@@ -2293,20 +2293,19 @@ void UMaterial::CacheShadersForResources(EShaderPlatform ShaderPlatform, const T
 		{
 			if (IsDefaultMaterial())
 			{
-				UE_LOG(LogMaterial, Fatal, TEXT("Failed to compile Default Material %s for platform %s!"), 
-					*GetPathName(), 
+				UE_ASSET_LOG(LogMaterial, Fatal, this,
+					TEXT("Failed to compile Default Material for platform %s!"),
 					*LegacyShaderPlatformToShaderFormat(ShaderPlatform).ToString());
 			}
 
-			UE_LOG(LogMaterial, Warning, TEXT("Failed to compile Material %s for platform %s, Default Material will be used in game."), 
-				*GetPathName(), 
+			UE_ASSET_LOG(LogMaterial, Warning, this, TEXT("Failed to compile Material for platform %s, Default Material will be used in game."), 
 				*LegacyShaderPlatformToShaderFormat(ShaderPlatform).ToString());
 
 			const TArray<FString>& CompileErrors = CurrentResource->GetCompileErrors();
 			for (int32 ErrorIndex = 0; ErrorIndex < CompileErrors.Num(); ErrorIndex++)
 			{
 				// Always log material errors in an unsuppressed category
-				UE_LOG(LogMaterial, Warning, TEXT("	%s"), *CompileErrors[ErrorIndex]);
+				UE_LOG(LogMaterial, Log, TEXT("	%s"), *CompileErrors[ErrorIndex]);
 			}
 		}
 	}
diff --git a/Engine/Source/Runtime/Engine/Private/Materials/MaterialInstance.cpp b/Engine/Source/Runtime/Engine/Private/Materials/MaterialInstance.cpp
--- a/Engine/Source/Runtime/Engine/Private/Materials/MaterialInstance.cpp
+++ b/Engine/Source/Runtime/Engine/Private/Materials/MaterialInstance.cpp
@@ -1741,9 +1741,8 @@ void UMaterialInstance::CacheShadersForResources(EShaderPlatform ShaderPlatform,
 
 		if (!bSuccess)
 		{
-			UE_LOG(LogMaterial, Warning,
-				TEXT("Failed to compile Material Instance %s with Base %s for platform %s, Default Material will be used in game."), 
-				*GetPathName(), 
+			UE_ASSET_LOG(LogMaterial, Warning, this,
+				TEXT("Failed to compile Material Instance with Base %s for platform %s, Default Material will be used in game."), 
 				BaseMaterial ? *BaseMaterial->GetName() : TEXT("Null"), 
 				*LegacyShaderPlatformToShaderFormat(ShaderPlatform).ToString()
 				);
@@ -1751,7 +1750,7 @@ void UMaterialInstance::CacheShadersForResources(EShaderPlatform ShaderPlatform,
 			const TArray<FString>& CompileErrors = CurrentResource->GetCompileErrors();
 			for (int32 ErrorIndex = 0; ErrorIndex < CompileErrors.Num(); ErrorIndex++)
 			{
-				UE_LOG(LogMaterial, Warning, TEXT("	%s"), *CompileErrors[ErrorIndex]);
+				UE_LOG(LogMaterial, Log, TEXT("	%s"), *CompileErrors[ErrorIndex]);
 			}
 		}
 	}
@@ -2341,12 +2340,20 @@ void UMaterialInstance::SetTextureParameterValueInternal(FName ParameterName, UT
 	// Don't enqueue an update if it isn't needed
 	if (ParameterValue->ParameterValue != Value)
 	{
-		checkf(!Value || Value->IsA(UTexture::StaticClass()), TEXT("Expecting a UTexture! Value='%s' class='%s'"), *Value->GetName(), *Value->GetClass()->GetName());
-
-		ParameterValue->ParameterValue = Value;
-		// Update the material instance data in the rendering thread.
-		GameThread_UpdateMIParameter(this, *ParameterValue);
-		CacheMaterialInstanceUniformExpressions(this);
+		// set as an ensure, because it is somehow possible to accidentally pass non-textures into here via blueprints...
+		if (Value && ensureMsgf(Value->IsA(UTexture::StaticClass()), TEXT("Expecting a UTexture! Value='%s' class='%s'"), *Value->GetName(), *Value->GetClass()->GetName()))
+		{
+			if ( FPlatformProperties::RequiresCookedData() )
+			{
+				// Daniel L: Temporary code to track down OR-40458 REMOVE ME
+				check(IsInGameThread());
+				check(!Value->IsPendingKill());
+			}
+			ParameterValue->ParameterValue = Value;
+			// Update the material instance data in the rendering thread.
+			GameThread_UpdateMIParameter(this, *ParameterValue);
+			CacheMaterialInstanceUniformExpressions(this);
+		}		
 	}
 }
 
diff --git a/Engine/Source/Runtime/Engine/Private/Particles/ParticleComponents.cpp b/Engine/Source/Runtime/Engine/Private/Particles/ParticleComponents.cpp
--- a/Engine/Source/Runtime/Engine/Private/Particles/ParticleComponents.cpp
+++ b/Engine/Source/Runtime/Engine/Private/Particles/ParticleComponents.cpp
@@ -119,6 +119,17 @@ FAutoConsoleVariableRef CVarParticleLODBias(
 	ECVF_Scalability
 	);
 
+static TAutoConsoleVariable<float> CVarQLSpawnRateReferenceLevel(
+	TEXT("fx.QualityLevelSpawnRateScaleReferenceLevel"),
+	2,
+	TEXT("Controls the reference level for quality level based spawn rate scaling. This is the FX quality level\n")
+	TEXT("at which spawn rate is not scaled down; Spawn rate scaling will happen by each emitter's\n")
+	TEXT("QualityLevelSpawnRateScale value for each reduction in level below the reference level.\n")
+	TEXT("\n")
+	TEXT("Default = 2. Value should range from 0 to the maximum FX quality level."),
+	ECVF_Scalability);
+
+
 /** Whether to allow particle systems to perform work. */
 bool GIsAllowingParticles = true;
 
@@ -1739,13 +1750,10 @@ void UParticleEmitter::CacheEmitterModuleInfo()
 float UParticleEmitter::GetQualityLevelSpawnRateMult()
 {
 	int32 EffectsQuality = Scalability::GetEffectsQualityDirect(true);
-	float Q = 1;
-	float Level = (1 - EffectsQuality);
-	for (int i = 0; i < Level + 1; i++)
-	{
-		Q = Q*QualityLevelSpawnRateScale;
-	}
-	return Q;
+	int32 ReferenceLevel = CVarQLSpawnRateReferenceLevel.GetValueOnAnyThread();
+	float Level = (ReferenceLevel - EffectsQuality);
+	float Q = FMath::Pow(QualityLevelSpawnRateScale, Level);
+	return FMath::Min(1.0f, Q);
 }
 
 bool UParticleEmitter::HasAnyEnabledLODs()const
diff --git a/Engine/Source/Runtime/Engine/Private/Particles/ParticleModules_Location.cpp b/Engine/Source/Runtime/Engine/Private/Particles/ParticleModules_Location.cpp
--- a/Engine/Source/Runtime/Engine/Private/Particles/ParticleModules_Location.cpp
+++ b/Engine/Source/Runtime/Engine/Private/Particles/ParticleModules_Location.cpp
@@ -1970,7 +1970,18 @@ int32 UParticleModuleLocationBoneSocket::GetMaxSourceIndex(FModuleLocationBoneSo
 		break;
 		case EBoneSocketSourceIndexMode::Direct:
 		{
-			return SourceType == BONESOCKETSOURCE_Sockets ? SourceComponent->SkeletalMesh->NumSockets() : SourceComponent->GetNumBones();
+			if (SourceType == BONESOCKETSOURCE_Sockets)
+			{
+				if (ensure(SourceComponent->SkeletalMesh))
+				{
+					return SourceComponent->SkeletalMesh->NumSockets();
+				}
+				return 0;
+			}
+			else
+			{
+				return SourceComponent->GetNumBones();
+			}
 		}
 		break;
 	}
diff --git a/Engine/Source/Runtime/Engine/Private/PhysicsEngine/BodyInstance.cpp b/Engine/Source/Runtime/Engine/Private/PhysicsEngine/BodyInstance.cpp
--- a/Engine/Source/Runtime/Engine/Private/PhysicsEngine/BodyInstance.cpp
+++ b/Engine/Source/Runtime/Engine/Private/PhysicsEngine/BodyInstance.cpp
@@ -4642,7 +4642,9 @@ bool FBodyInstance::GetSquaredDistanceToBody(const FVector& Point, float& OutDis
 				continue;
 			}
 
-			PxGeometry& PGeom = PShape->getGeometry().any();
+			PxGeometryHolder Holder = PShape->getGeometry();	// getGeometry() result is stored on the stack, if we don't hold on to it it may be gone next statement.
+			PxGeometry& PGeom = Holder.any();
+
 			PxTransform PGlobalPose = GetPxTransform_AssumesLocked(PShape, RigidActor);
 			PxGeometryType::Enum GeomType = PShape->getGeometryType();
 
diff --git a/Engine/Source/Runtime/Engine/Private/PhysicsEngine/BodySetup.cpp b/Engine/Source/Runtime/Engine/Private/PhysicsEngine/BodySetup.cpp
--- a/Engine/Source/Runtime/Engine/Private/PhysicsEngine/BodySetup.cpp
+++ b/Engine/Source/Runtime/Engine/Private/PhysicsEngine/BodySetup.cpp
@@ -367,6 +367,13 @@ void UBodySetup::CreatePhysicsMeshes()
 	// Find or create cooked physics data
 	static FName PhysicsFormatName(FPlatformProperties::GetPhysicsFormat());
 	FByteBulkData* FormatData = GetCookedData(PhysicsFormatName);
+
+	// On dedicated servers we may be cooking generic data and sharing it
+	if (FormatData == nullptr && IsRunningDedicatedServer())
+	{
+		FormatData = GetCookedData(FGenericPlatformProperties::GetPhysicsFormat());
+	}
+
 	if (FormatData)
 	{
 		if (FormatData->IsLocked())
@@ -791,7 +798,6 @@ template <> FString FBodySetupShapeIterator::GetDebugName<FKSphylElem>() const
 	return TEXT("Capsule");
 }
 
-
 ////////////////////////////// Convex elements ////////////////////////////
 template <> bool FBodySetupShapeIterator::PopulatePhysXGeometryAndTransform(const FKConvexElem& ConvexElem, PxConvexMeshGeometry& OutGeometry, PxTransform& OutTM) const
 {
diff --git a/Engine/Source/Runtime/Engine/Private/PlayerController.cpp b/Engine/Source/Runtime/Engine/Private/PlayerController.cpp
--- a/Engine/Source/Runtime/Engine/Private/PlayerController.cpp
+++ b/Engine/Source/Runtime/Engine/Private/PlayerController.cpp
@@ -1755,6 +1755,10 @@ bool APlayerController::SetPause( bool bPause, FCanUnpause CanUnpauseDelegate)
 			{
 				// Pause gamepad rumbling too if needed
 				bResult = GameMode->SetPause(this, CanUnpauseDelegate);
+
+				// Force an update, otherwise since the game time is not updating, the net driver
+				// might not see that it is time for the world settings actor to replicate
+				ForceSingleNetUpdateFor(GetWorldSettings());
 			}
 			else if (!bPause && bCurrentPauseState)
 			{
@@ -1765,7 +1769,7 @@ bool APlayerController::SetPause( bool bPause, FCanUnpause CanUnpauseDelegate)
 	return bResult;
 }
 
-bool APlayerController::IsPaused()
+bool APlayerController::IsPaused() const
 {
 	return GetWorldSettings()->Pauser != NULL;
 }
@@ -4267,6 +4271,7 @@ void APlayerController::TickActor( float DeltaSeconds, ELevelTick TickType, FAct
 								// Also null the pointer to make sure no one accidentally starts using it below the call to ForcePositionUpdate
 								ServerData->ServerTimeStamp = GetWorld()->GetTimeSeconds();
 								ServerData = nullptr;
+
 								NetworkPredictionInterface->ForcePositionUpdate(PawnTimeSinceUpdate);
 							}
 						}
diff --git a/Engine/Source/Runtime/Engine/Private/ShaderCompiler/ShaderCompiler.cpp b/Engine/Source/Runtime/Engine/Private/ShaderCompiler/ShaderCompiler.cpp
--- a/Engine/Source/Runtime/Engine/Private/ShaderCompiler/ShaderCompiler.cpp
+++ b/Engine/Source/Runtime/Engine/Private/ShaderCompiler/ShaderCompiler.cpp
@@ -1822,8 +1822,7 @@ void FShaderCompilingManager::ProcessCompiledShaderMaps(
 								UE_LOG(LogShaderCompilers, Fatal,TEXT("Failed to compile default material %s!"), *Material->GetBaseMaterialPathName());
 							}
 
-							UE_LOG(LogShaderCompilers, Warning, TEXT("Failed to compile Material %s for platform %s, Default Material will be used in game."), 
-								*Material->GetBaseMaterialPathName(), 
+							UE_ASSET_LOG(LogShaderCompilers, Warning, *Material->GetBaseMaterialPathName(), TEXT("Failed to compile Material for platform %s, Default Material will be used in game."),
 								*LegacyShaderPlatformToShaderFormat(ShaderMap->GetShaderPlatform()).ToString());
 
 							for (int32 ErrorIndex = 0; ErrorIndex < Errors.Num(); ErrorIndex++)
@@ -1831,7 +1830,7 @@ void FShaderCompilingManager::ProcessCompiledShaderMaps(
 								FString ErrorMessage = Errors[ErrorIndex];
 								// Work around build machine string matching heuristics that will cause a cook to fail
 								ErrorMessage.ReplaceInline(TEXT("error "), TEXT("err0r "), ESearchCase::CaseSensitive);
-								UE_LOG(LogShaderCompilers, Warning, TEXT("	%s"), *ErrorMessage);
+								UE_LOG(LogShaderCompilers, Log, TEXT("	%s"), *ErrorMessage);
 							}
 						}
 						else
diff --git a/Engine/Source/Runtime/Engine/Private/SkeletalMeshComponentPhysics.cpp b/Engine/Source/Runtime/Engine/Private/SkeletalMeshComponentPhysics.cpp
--- a/Engine/Source/Runtime/Engine/Private/SkeletalMeshComponentPhysics.cpp
+++ b/Engine/Source/Runtime/Engine/Private/SkeletalMeshComponentPhysics.cpp
@@ -635,7 +635,7 @@ void USkeletalMeshComponent::InstantiatePhysicsAsset(const UPhysicsAsset& PhysAs
 			else
 			{
 				BodyInst->DOFMode = EDOFMode::None;
-				if (!CVarEnableRagdollPhysics.GetValueOnGameThread())
+				if (PhysScene != nullptr && !CVarEnableRagdollPhysics.GetValueOnGameThread())	//We only limit creation of the global physx scenes and not assets related to immedate mode
 				{
 					continue;
 				}
diff --git a/Engine/Source/Runtime/Engine/Private/Slate/SGameLayerManager.cpp b/Engine/Source/Runtime/Engine/Private/Slate/SGameLayerManager.cpp
--- a/Engine/Source/Runtime/Engine/Private/Slate/SGameLayerManager.cpp
+++ b/Engine/Source/Runtime/Engine/Private/Slate/SGameLayerManager.cpp
@@ -210,15 +210,20 @@ void SGameLayerManager::ClearWidgets()
 {
 	PlayerCanvas->ClearChildren();
 
-	for(const auto& LayerIt : PlayerLayers)
+	// Potential for removed layers to impact the map, so need to
+	// remove & delete as separate steps
+	while (PlayerLayers.Num())
 	{
-		const TSharedPtr<FPlayerLayer>& Layer = LayerIt.Value;
+		const auto LayerIt = PlayerLayers.CreateIterator();
+		const TSharedPtr<FPlayerLayer> Layer = LayerIt.Value();
+
 		if (Layer.IsValid())
 		{
 			Layer->Slot = nullptr;
 		}
+
+		PlayerLayers.Remove(LayerIt.Key());
 	}
-	PlayerLayers.Reset();
 
 	WindowTitleBarContentStack.Empty();
 	bIsWindowTitleBarVisible = false;
diff --git a/Engine/Source/Runtime/Engine/Private/StaticMesh.cpp b/Engine/Source/Runtime/Engine/Private/StaticMesh.cpp
--- a/Engine/Source/Runtime/Engine/Private/StaticMesh.cpp
+++ b/Engine/Source/Runtime/Engine/Private/StaticMesh.cpp
@@ -1953,6 +1953,13 @@ void UStaticMesh::GetAssetRegistryTags(TArray<FAssetRegistryTag>& OutTags) const
 		DefaultCollisionName = BodySetup->DefaultInstance.GetCollisionProfileName();
 	}
 
+	static const UEnum *ComplexityEnum = FindObject<UEnum>(ANY_PACKAGE, TEXT("ECollisionTraceFlag"), true);
+	FString ComplexityString;
+	if (BodySetup && ComplexityEnum)
+	{
+		ComplexityString = ComplexityEnum->GetNameStringByValue((int64)BodySetup->GetCollisionTraceFlag());
+	}
+
 	OutTags.Add( FAssetRegistryTag("Triangles", FString::FromInt(NumTriangles), FAssetRegistryTag::TT_Numerical) );
 	OutTags.Add( FAssetRegistryTag("Vertices", FString::FromInt(NumVertices), FAssetRegistryTag::TT_Numerical) );
 	OutTags.Add( FAssetRegistryTag("UVChannels", FString::FromInt(NumUVChannels), FAssetRegistryTag::TT_Numerical) );
@@ -1962,6 +1969,7 @@ void UStaticMesh::GetAssetRegistryTags(TArray<FAssetRegistryTag>& OutTags) const
 	OutTags.Add( FAssetRegistryTag("LODs", FString::FromInt(NumLODs), FAssetRegistryTag::TT_Numerical));
 	OutTags.Add( FAssetRegistryTag("SectionsWithCollision", FString::FromInt(NumSectionsWithCollision), FAssetRegistryTag::TT_Numerical));
 	OutTags.Add( FAssetRegistryTag("DefaultCollision", DefaultCollisionName.ToString(), FAssetRegistryTag::TT_Alphabetical));
+	OutTags.Add( FAssetRegistryTag("CollisionComplexity", ComplexityString, FAssetRegistryTag::TT_Alphabetical));
 
 #if WITH_EDITORONLY_DATA
 	if (AssetImportData)
diff --git a/Engine/Source/Runtime/Engine/Private/Texture.cpp b/Engine/Source/Runtime/Engine/Private/Texture.cpp
--- a/Engine/Source/Runtime/Engine/Private/Texture.cpp
+++ b/Engine/Source/Runtime/Engine/Private/Texture.cpp
@@ -977,6 +977,17 @@ uint32 UTexture::GetMaximumDimension() const
 }
 #endif // #if WITH_EDITOR
 
+void UTexture::MarkPendingKill()
+{
+	Super::MarkPendingKill();
+	if ( FPlatformProperties::RequiresCookedData() )
+	{
+		// Daniel L: Temporary code to track down OR-40458 REMOVE ME
+		UE_LOG(LogTexture, Warning, TEXT("Texture is being marked as pending kill!!! %s"), *GetFullName());
+	}
+}
+
+
 FName GetDefaultTextureFormatName( const ITargetPlatform* TargetPlatform, const UTexture* Texture, const FConfigFile& EngineSettings, bool bSupportDX11TextureFormats )
 {
 	FName TextureFormatName = NAME_None;
diff --git a/Engine/Source/Runtime/Engine/Private/TextureRenderTarget2D.cpp b/Engine/Source/Runtime/Engine/Private/TextureRenderTarget2D.cpp
--- a/Engine/Source/Runtime/Engine/Private/TextureRenderTarget2D.cpp
+++ b/Engine/Source/Runtime/Engine/Private/TextureRenderTarget2D.cpp
@@ -91,6 +91,16 @@ void UTextureRenderTarget2D::InitCustomFormat( uint32 InSizeX, uint32 InSizeY, E
 	OverrideFormat = InOverrideFormat;
 	bForceLinearGamma = bInForceLinearGamma;
 
+	if (!ensureMsgf(SizeX >= 0 && SizeX <= 65536, TEXT("Invalid SizeX=%u for RenderTarget %s"), SizeX, *GetName()))
+	{
+		SizeX = 1;
+	}
+
+	if (!ensureMsgf(SizeY >= 0 && SizeY <= 65536, TEXT("Invalid SizeY=%u for RenderTarget %s"), SizeY, *GetName()))
+	{
+		SizeY = 1;
+	}
+
 	// Recreate the texture's resource.
 	UpdateResource();
 }
@@ -361,6 +371,7 @@ FTextureRenderTarget2DResource::FTextureRenderTarget2DResource(const class UText
 	,	TargetSizeX(Owner->SizeX)
 	,	TargetSizeY(Owner->SizeY)
 {
+	
 }
 
 /**
@@ -379,6 +390,7 @@ void FTextureRenderTarget2DResource::ClampSize(int32 MaxSizeX,int32 MaxSizeY)
 		TargetSizeX = NewSizeX;
 		TargetSizeY = NewSizeY;
 		// reinit the resource with new TargetSizeX,TargetSizeY
+		check(TargetSizeX >= 0 && TargetSizeY >= 0);
 		UpdateRHI();
 	}	
 }
diff --git a/Engine/Source/Runtime/Engine/Private/TickTaskManager.cpp b/Engine/Source/Runtime/Engine/Private/TickTaskManager.cpp
--- a/Engine/Source/Runtime/Engine/Private/TickTaskManager.cpp
+++ b/Engine/Source/Runtime/Engine/Private/TickTaskManager.cpp
@@ -17,6 +17,7 @@
 #include "Engine/World.h"
 #include "TickTaskManagerInterface.h"
 #include "Async/ParallelFor.h"
+#include "Misc/TimeGuard.h"
 
 DEFINE_LOG_CATEGORY_STATIC(LogTick, Log, All);
 
@@ -264,6 +265,11 @@ public:
 		}
 		if (Target->IsTickFunctionEnabled())
 		{
+#if DO_TIMEGUARD
+			FTimerNameDelegate NameFunction = FTimerNameDelegate::CreateLambda( [&]{ return FString::Printf(TEXT("Slowtick %s "), *Target->DiagnosticMessage()); } );
+			SCOPE_TIME_GUARD_DELEGATE_MS(NameFunction, 2);
+#endif
+
 			Target->ExecuteTick(Target->CalculateDeltaTime(Context), Context.TickType, CurrentThread, MyCompletionGraphEvent);
 		}
 		Target->TaskPointer = nullptr;  // This is stale and a good time to clear it for safety
diff --git a/Engine/Source/Runtime/Engine/Private/UnrealClient.cpp b/Engine/Source/Runtime/Engine/Private/UnrealClient.cpp
--- a/Engine/Source/Runtime/Engine/Private/UnrealClient.cpp
+++ b/Engine/Source/Runtime/Engine/Private/UnrealClient.cpp
@@ -279,8 +279,8 @@ FUnorderedAccessViewRHIRef FRenderTarget::GetRenderTargetUAV() const
 
 void FScreenshotRequest::RequestScreenshot(bool bInShowUI)
 {
-	bShowUI = bInShowUI;
-	bIsScreenshotRequested = true;
+	// empty string means we'll later pick the name
+	RequestScreenshot(TEXT(""), bInShowUI, true);
 }
 
 void FScreenshotRequest::RequestScreenshot(const FString& InFilename, bool bInShowUI, bool bAddUniqueSuffix)
diff --git a/Engine/Source/Runtime/Engine/Private/World.cpp b/Engine/Source/Runtime/Engine/Private/World.cpp
--- a/Engine/Source/Runtime/Engine/Private/World.cpp
+++ b/Engine/Source/Runtime/Engine/Private/World.cpp
@@ -280,6 +280,7 @@ FWorldDelegates::FWorldRenameEvent FWorldDelegates::OnPreWorldRename;
 #endif // WITH_EDITOR
 FWorldDelegates::FWorldPostDuplicateEvent FWorldDelegates::OnPostDuplicate;
 FWorldDelegates::FWorldCleanupEvent FWorldDelegates::OnWorldCleanup;
+FWorldDelegates::FWorldCleanupEvent FWorldDelegates::OnPostWorldCleanup;
 FWorldDelegates::FWorldEvent FWorldDelegates::OnPreWorldFinishDestroy;
 FWorldDelegates::FOnLevelChanged FWorldDelegates::LevelAddedToWorld;
 FWorldDelegates::FOnLevelChanged FWorldDelegates::LevelRemovedFromWorld;
@@ -3605,6 +3606,8 @@ void UWorld::CleanupWorld(bool bSessionEnded, bool bCleanupResources, UWorld* Ne
 	{
 		LevelCollections.Empty();
 	}
+
+	FWorldDelegates::OnPostWorldCleanup.Broadcast(this, bSessionEnded, bCleanupResources);
 }
 
 UGameViewportClient* UWorld::GetGameViewport() const
diff --git a/Engine/Source/Runtime/Engine/Public/Animation/CachedAnimData.h b/Engine/Source/Runtime/Engine/Public/Animation/CachedAnimData.h
--- a/Engine/Source/Runtime/Engine/Public/Animation/CachedAnimData.h
+++ b/Engine/Source/Runtime/Engine/Public/Animation/CachedAnimData.h
@@ -49,6 +49,9 @@ struct ENGINE_API FCachedAnimStateData
 	/** Is State relevant? */
 	bool IsRelevant(UAnimInstance& InAnimInstance) const;
 
+	/** Is State active? */
+	bool IsActiveState(class UAnimInstance& InAnimInstance) const;
+
 private:
 	mutable int32 MachineIndex;
 	mutable int32 StateIndex;
diff --git a/Engine/Source/Runtime/Engine/Public/AnimationCompression.h b/Engine/Source/Runtime/Engine/Public/AnimationCompression.h
--- a/Engine/Source/Runtime/Engine/Public/AnimationCompression.h
+++ b/Engine/Source/Runtime/Engine/Public/AnimationCompression.h
@@ -18,27 +18,9 @@ DECLARE_LOG_CATEGORY_EXTERN(LogAnimationCompression, Log, All);
 #define SCALE_ZEROING_THRESHOLD (0.000001f)
 
 /** Size of Dummy bone used, when measuring error at an end effector which has a socket attached to it */
-#define END_EFFECTOR_SOCKET_DUMMY_BONE_SIZE	(50.f)
-
-// @Laurent: This is currently disabled. We don't want to do this during production as we would have to recompress data if Retargeting settings were to change.
-// We should consider only doing this when cooking content to a non editor shipping build.
-
-	/** This will reduce automatically bAnimRotationOnly Translation Tracks to 1 key, even if they have animated data. */
-	#define REDUCE_ANIMROTATIONONLY_TRACKS (0)
-
-	/** 
-	 * This will skip storing bAnimRotationOnly Translation Tracks, and assumes the AnimSet->bAnimRotationOnly flag won't be changed at runtime.
-	 * This is currently only used in the PerTrack compressor.
-	 */
-	#define SKIP_ANIMROTATIONONLY_TRACKS (0)
-
-	/**
-	 * Do not store translation tracks when those are marked 'ForceMeshTranslation'
-	 * true if SKIP_ANIMROTATIONONLY_TRACKS is defined to true.
-	 */
-	#define SKIP_FORCEMESHTRANSLATION_TRACKS (0 || SKIP_ANIMROTATIONONLY_TRACKS)
-
-// END
+#define END_EFFECTOR_DUMMY_BONE_LENGTH_SOCKET	(50.f)
+/** Dummy bone added to end effectors to make sure rotation doesn't get too aggressively compressed. */
+#define END_EFFECTOR_DUMMY_BONE_LENGTH	(5.f)
 
 #define Quant16BitDiv     (32767.f)
 #define Quant16BitFactor  (32767.f)
diff --git a/Engine/Source/Runtime/Engine/Public/AnimationUtils.h b/Engine/Source/Runtime/Engine/Public/AnimationUtils.h
--- a/Engine/Source/Runtime/Engine/Public/AnimationUtils.h
+++ b/Engine/Source/Runtime/Engine/Public/AnimationUtils.h
@@ -90,6 +90,13 @@ struct AnimationErrorStats
 	float MaxErrorTime;
 	/** Bone on which the worst error occurred */
 	int32 MaxErrorBone;
+
+	AnimationErrorStats()
+		: AverageError(0.f)
+		, MaxError(0.f)
+		, MaxErrorTime(0.f)
+		, MaxErrorBone(0)
+	{}
 };
 
 /** Different ways to create an additive animation */
diff --git a/Engine/Source/Runtime/GameplayTags/Classes/GameplayTagsManager.h b/Engine/Source/Runtime/GameplayTags/Classes/GameplayTagsManager.h
--- a/Engine/Source/Runtime/GameplayTags/Classes/GameplayTagsManager.h
+++ b/Engine/Source/Runtime/GameplayTags/Classes/GameplayTagsManager.h
@@ -89,6 +89,14 @@ struct FGameplayTagSource
 		static FName DefaultName = FName(TEXT("DefaultGameplayTags.ini"));
 		return DefaultName;
 	}
+
+#if WITH_EDITOR
+	static FName GetTransientEditorName()
+	{
+		static FName TransientEditorName = FName(TEXT("TransientEditor"));
+		return TransientEditorName;
+	}
+#endif
 };
 
 /** Simple tree node for gameplay tags, this stores metadata about specific tags */
@@ -229,6 +237,9 @@ class GAMEPLAYTAGS_API UGameplayTagsManager : public UObject
 
 	static FSimpleMulticastDelegate& OnLastChanceToAddNativeTags();
 
+
+	void CallOrRegister_OnDoneAddingNativeTagsDelegate(FSimpleMulticastDelegate::FDelegate Delegate);
+
 	/**
 	 * Gets a Tag Container containing the supplied tag and all of it's parents as explicit tags
 	 *
@@ -470,6 +481,9 @@ private:
 	/** Initializes the manager */
 	static void InitializeManager();
 
+	/** finished loading/adding native tags */
+	static FSimpleMulticastDelegate& OnDoneAddingNativeTagsDelegate();
+
 	/** The Tag Manager singleton */
 	static UGameplayTagsManager* SingletonManager;
 
@@ -539,9 +553,12 @@ private:
 	bool bDoneAddingNativeTags;
 
 #if WITH_EDITOR
-	// This critical section is to handle and editor-only issue where tag requests come from another thread when async loading from a background thread in FGameplayTagContainer::Serialize.
+	// This critical section is to handle an editor-only issue where tag requests come from another thread when async loading from a background thread in FGameplayTagContainer::Serialize.
 	// This class is not generically threadsafe.
 	mutable FCriticalSection GameplayTagMapCritical;
+
+	// Transient editor-only tags to support quick-iteration PIE workflows
+	TSet<FName> TransientEditorTags;
 #endif
 
 	/** Sorted list of nodes, used for network replication */
diff --git a/Engine/Source/Runtime/GameplayTags/Private/GameplayTagsManager.cpp b/Engine/Source/Runtime/GameplayTags/Private/GameplayTagsManager.cpp
--- a/Engine/Source/Runtime/GameplayTags/Private/GameplayTagsManager.cpp
+++ b/Engine/Source/Runtime/GameplayTags/Private/GameplayTagsManager.cpp
@@ -189,6 +189,14 @@ void UGameplayTagsManager::ConstructGameplayTagTree()
 			}
 		}
 
+#if WITH_EDITOR
+		// Add any transient editor-only tags
+		for (FName TransientTag : TransientEditorTags)
+		{
+			AddTagTableRow(FGameplayTagTableRow(TransientTag), FGameplayTagSource::GetTransientEditorName());
+		}
+#endif
+
 		// Grab the commonly replicated tags
 		CommonlyReplicatedTags.Empty();
 		for (FName TagName : MutableDefault->CommonlyReplicatedTags)
@@ -1181,6 +1189,27 @@ FGameplayTag UGameplayTagsManager::AddNativeGameplayTag(FName TagName)
 
 	return FGameplayTag();
 }
+void UGameplayTagsManager::CallOrRegister_OnDoneAddingNativeTagsDelegate(FSimpleMulticastDelegate::FDelegate Delegate)
+{
+	if (bDoneAddingNativeTags)
+	{
+		Delegate.Execute();
+	}
+	else
+	{
+		bool bAlreadyBound = Delegate.GetUObject() != nullptr ? OnDoneAddingNativeTagsDelegate().IsBoundToObject(Delegate.GetUObject()) : false;
+		if (!bAlreadyBound)
+		{
+			OnDoneAddingNativeTagsDelegate().Add(Delegate);
+		}
+	}
+}
+
+FSimpleMulticastDelegate& UGameplayTagsManager::OnDoneAddingNativeTagsDelegate()
+{
+	static FSimpleMulticastDelegate Delegate;
+	return Delegate;
+}
 
 FSimpleMulticastDelegate& UGameplayTagsManager::OnLastChanceToAddNativeTags()
 {
@@ -1201,6 +1230,7 @@ void UGameplayTagsManager::DoneAddingNativeTags()
 		{
 			ConstructNetIndex();
 		}
+		OnDoneAddingNativeTagsDelegate().Broadcast();
 	}
 }
 
diff --git a/Engine/Source/Runtime/GameplayTasks/Classes/GameplayTask.h b/Engine/Source/Runtime/GameplayTasks/Classes/GameplayTask.h
--- a/Engine/Source/Runtime/GameplayTasks/Classes/GameplayTask.h
+++ b/Engine/Source/Runtime/GameplayTasks/Classes/GameplayTask.h
@@ -319,6 +319,8 @@ protected:
 	 *	provided GameplayTasksComponent is configured to care about priorities (the default behavior)*/
 	uint8 Priority;
 
+	/** You should _never_ access it directly. We'll make it private once we have
+	 *	a good way of "deprecating: direct access */
 	EGameplayTaskState TaskState;
 
 	UPROPERTY(config)
diff --git a/Engine/Source/Runtime/GameplayTasks/Private/GameplayTask.cpp b/Engine/Source/Runtime/GameplayTasks/Private/GameplayTask.cpp
--- a/Engine/Source/Runtime/GameplayTasks/Private/GameplayTask.cpp
+++ b/Engine/Source/Runtime/GameplayTasks/Private/GameplayTask.cpp
@@ -278,7 +278,13 @@ void UGameplayTask::PerformActivation()
 
 	Activate();
 
-	TasksComponent->OnGameplayTaskActivated(*this);
+	// Activate call may result in the task actually "instantly" finishing.
+	// If this happens we don't want to bother the TaskComponent
+	// with information on this task
+	if (IsFinished() == false)
+	{
+		TasksComponent->OnGameplayTaskActivated(*this);
+	}
 }
 
 void UGameplayTask::Activate()
diff --git a/Engine/Source/Runtime/GameplayTasks/Private/GameplayTasksComponent.cpp b/Engine/Source/Runtime/GameplayTasks/Private/GameplayTasksComponent.cpp
--- a/Engine/Source/Runtime/GameplayTasks/Private/GameplayTasksComponent.cpp
+++ b/Engine/Source/Runtime/GameplayTasks/Private/GameplayTasksComponent.cpp
@@ -113,7 +113,10 @@ void UGameplayTasksComponent::OnGameplayTaskDeactivated(UGameplayTask& Task)
 
 	if (bIsFinished)
 	{
-		KnownTasks.RemoveSingleSwap(&Task);
+		// using RemoveSwap rather than RemoveSingleSwap since a Task can be added
+		// to KnownTasks both when activating as well as unpausing
+		// while removal happens only once. It's cheaper to handle it here.
+		KnownTasks.RemoveSwap(&Task);
 	}
 
 	if (Task.IsSimulatedTask())
diff --git a/Engine/Source/Runtime/PhysXCooking/Private/PhysXCooking.cpp b/Engine/Source/Runtime/PhysXCooking/Private/PhysXCooking.cpp
--- a/Engine/Source/Runtime/PhysXCooking/Private/PhysXCooking.cpp
+++ b/Engine/Source/Runtime/PhysXCooking/Private/PhysXCooking.cpp
@@ -294,12 +294,15 @@ bool FPhysXCooking::CookTriMeshImp(FName Format, EPhysXMeshCookFlags CookFlags,
 	}
 	else
 	{
+		if (PhysXCooking->validateTriangleMesh(PTriMeshDesc) == false)
+		{
+			NewParams.meshPreprocessParams = PxMeshPreprocessingFlag::eDISABLE_CLEAN_MESH;
+		}
 		//For non deformable meshes we can try to use BVH34
 		UseBVH34IfSupported(Format, NewParams);
 	}
 
 	PhysXCooking->setParams(NewParams);
-
 	bool bResult = false;
 
 	// Cook TriMesh Data
diff --git a/Engine/Source/Runtime/RenderCore/Private/RenderingThread.cpp b/Engine/Source/Runtime/RenderCore/Private/RenderingThread.cpp
--- a/Engine/Source/Runtime/RenderCore/Private/RenderingThread.cpp
+++ b/Engine/Source/Runtime/RenderCore/Private/RenderingThread.cpp
@@ -987,7 +987,19 @@ static void GameThreadWaitForTask(const FGraphEventRef& Task, bool bEmptyGameThr
 				if (!bDone && !WITH_EDITOR && !FPlatformMisc::IsDebuggerPresent())
 				{
 					static bool bDisabled = FParse::Param(FCommandLine::Get(), TEXT("nothreadtimeout"));
+					static bool bGPUDebugging = FParse::Param(FCommandLine::Get(), TEXT("gpucrashdebugging"));
 
+					if (bGPUDebugging && FPlatformTime::Seconds() > 2.0f)
+					{
+						bool IsGpuAlive = true;
+						if (GDynamicRHI)
+						{
+							IsGpuAlive = GDynamicRHI->CheckGpuHeartbeat();
+						}
+
+						UE_CLOG(!IsGpuAlive, LogRendererCore, Fatal, TEXT("CheckGpuHeartbeat returned false after %.02f secs of waiting for the GPU"), FPlatformTime::Seconds() - StartTime);
+					}
+		
 					// Fatal timeout if we run out of time and this thread is being monitor for heartbeats
 					// (We could just let the heartbeat monitor error for us, but this leads to better diagnostics).
 					if (FPlatformTime::Seconds() >= EndTime && FThreadHeartBeat::Get().IsBeating() && !bDisabled)
diff --git a/Engine/Source/Runtime/Slate/Private/Widgets/Input/SMenuAnchor.cpp b/Engine/Source/Runtime/Slate/Private/Widgets/Input/SMenuAnchor.cpp
--- a/Engine/Source/Runtime/Slate/Private/Widgets/Input/SMenuAnchor.cpp
+++ b/Engine/Source/Runtime/Slate/Private/Widgets/Input/SMenuAnchor.cpp
@@ -50,48 +50,23 @@ static FVector2D GetMenuOffsetForPlacement(const FGeometry& AllottedGeometry, EM
 		case MenuPlacement_RightLeftCenter:
 			return FVector2D( AllottedGeometry.GetLocalSize().X, - ( ( PopupSizeLocalSpace.Y / 2 ) - ( AllottedGeometry.GetLocalSize().Y / 2 ) ) );
 			break;
+		case MenuPlacement_MatchBottomLeft:
+			return FVector2D(0.0f, (AllottedGeometry.GetLocalSize().Y - PopupSizeLocalSpace.Y));
+			break;
 		default:
 			ensureMsgf( false, TEXT("Unhandled placement mode: %d"), PlacementMode );
 			return FVector2D::ZeroVector;
 	}
 }
 
-/** Compute the popup size, offset, and anchor rect in local space. */
-struct FPopupPlacement
-{
-	FPopupPlacement(const FGeometry& AllottedGeometry, const FVector2D& PopupDesiredSize, EMenuPlacement PlacementMode)
-	{
-		// Compute the popup size, offset, and anchor rect  in local space
-		LocalPopupSize = (PlacementMode == MenuPlacement_ComboBox || PlacementMode == MenuPlacement_ComboBoxRight)
-			? FVector2D(FMath::Max(AllottedGeometry.Size.X, PopupDesiredSize.X), PopupDesiredSize.Y)
-			: PopupDesiredSize;
-		LocalPopupOffset = GetMenuOffsetForPlacement(AllottedGeometry, PlacementMode, LocalPopupSize);
-		AnchorLocalSpace = FSlateRect::FromPointAndExtent(FVector2D::ZeroVector, AllottedGeometry.GetLocalSize());
-		Orientation = (PlacementMode == MenuPlacement_MenuRight || PlacementMode == MenuPlacement_MenuLeft) ? Orient_Horizontal : Orient_Vertical;
-	}
-
-	FVector2D LocalPopupSize;
-	FVector2D LocalPopupOffset;
-	FSlateRect AnchorLocalSpace;
-	EOrientation Orientation;
-};
-
-FGeometry ComputeMenuPlacement(const FGeometry& AllottedGeometry, const FVector2D& PopupDesiredSize, EMenuPlacement PlacementMode)
+SMenuAnchor::FPopupPlacement::FPopupPlacement(const FGeometry& PlacementGeometry, const FVector2D& PopupDesiredSize, EMenuPlacement PlacementMode)
 {
 	// Compute the popup size, offset, and anchor rect  in local space
-	const FPopupPlacement Placement(AllottedGeometry, PopupDesiredSize, PlacementMode);
-
-	// ask the application to compute the proper desktop offset for the anchor. This requires the offsets to be in desktop space.
-	const FVector2D NewPositionDesktopSpace = FSlateApplication::Get().CalculatePopupWindowPosition(
-		TransformRect(AllottedGeometry.GetAccumulatedLayoutTransform(), Placement.AnchorLocalSpace),
-		TransformVector(AllottedGeometry.GetAccumulatedLayoutTransform(), Placement.LocalPopupSize),
-		TransformPoint(AllottedGeometry.GetAccumulatedLayoutTransform(), Placement.LocalPopupOffset),
-		Placement.Orientation);
-
-	// transform the desktop offset into local space and use that as the layout transform for the child content.
-	return AllottedGeometry.MakeChild(
-		Placement.LocalPopupSize,
-		FSlateLayoutTransform(TransformPoint(Inverse(AllottedGeometry.GetAccumulatedLayoutTransform()), NewPositionDesktopSpace)));
+	const bool bIsComboBoxPopup = (PlacementMode == MenuPlacement_ComboBox || PlacementMode == MenuPlacement_ComboBoxRight);
+	LocalPopupSize = bIsComboBoxPopup ? FVector2D(FMath::Max(PlacementGeometry.Size.X, PopupDesiredSize.X), PopupDesiredSize.Y) : PopupDesiredSize;
+	LocalPopupOffset = GetMenuOffsetForPlacement(PlacementGeometry, PlacementMode, LocalPopupSize);
+	AnchorLocalSpace = FSlateRect::FromPointAndExtent(FVector2D::ZeroVector, PlacementGeometry.GetLocalSize());
+	Orientation = (PlacementMode == MenuPlacement_MenuRight || PlacementMode == MenuPlacement_MenuLeft) ? Orient_Horizontal : Orient_Vertical;
 }
 
 /*static*/ TArray<TWeakPtr<IMenu>> SMenuAnchor::OpenApplicationMenus;
@@ -122,8 +97,26 @@ void SMenuAnchor::Construct( const FArguments& InArgs )
 	bShouldDeferPaintingAfterWindowContent = InArgs._ShouldDeferPaintingAfterWindowContent;
 	bUseApplicationMenuStack               = InArgs._UseApplicationMenuStack;
 	bIsCollapsedByParent                   = InArgs._IsCollapsedByParent;
+	bApplyWidgetStyleToMenu                = InArgs._ApplyWidgetStyleToMenu;
 }
 
+FGeometry SMenuAnchor::ComputeNewWindowMenuPlacement(const FGeometry& AllottedGeometry, const FVector2D& PopupDesiredSize, EMenuPlacement PlacementMode) const
+{
+	// Compute the popup size, offset, and anchor rect  in local space
+	const FPopupPlacement PopupPlacement(AllottedGeometry, PopupDesiredSize, PlacementMode);
+
+	// ask the application to compute the proper desktop offset for the anchor. This requires the offsets to be in desktop space.
+	const FVector2D NewPositionDesktopSpace = FSlateApplication::Get().CalculatePopupWindowPosition(
+		TransformRect(AllottedGeometry.GetAccumulatedLayoutTransform(), PopupPlacement.AnchorLocalSpace),
+		TransformVector(AllottedGeometry.GetAccumulatedLayoutTransform(), PopupPlacement.LocalPopupSize),
+		TransformPoint(AllottedGeometry.GetAccumulatedLayoutTransform(), PopupPlacement.LocalPopupOffset),
+		PopupPlacement.Orientation);
+
+	// transform the desktop offset into local space and use that as the layout transform for the child content.
+	return AllottedGeometry.MakeChild(
+		PopupPlacement.LocalPopupSize,
+		FSlateLayoutTransform(TransformPoint(Inverse(AllottedGeometry.GetAccumulatedLayoutTransform()), NewPositionDesktopSpace)));
+}
 
 void SMenuAnchor::Tick( const FGeometry& AllottedGeometry, const double InCurrentTime, const float InDeltaTime )
 {
@@ -132,11 +125,11 @@ void SMenuAnchor::Tick( const FGeometry& AllottedGeometry, const double InCurren
 	{
 		// Figure out where our attached pop-up window should be placed.
 		const FVector2D PopupContentDesiredSize = PopupWindow->GetContent()->GetDesiredSize();
-		FGeometry PopupGeometry = ComputeMenuPlacement( AllottedGeometry, PopupContentDesiredSize, Placement.Get() );
+		FGeometry PopupGeometry = ComputeNewWindowMenuPlacement( AllottedGeometry, PopupContentDesiredSize, Placement.Get() );
 		const FVector2D NewPosition = PopupGeometry.LocalToAbsolute(FVector2D::ZeroVector);
 		// NOTE: In order to get the right size of the window, we need to take whatever the incoming scale of the menu anchor,
 		// then divide out the popup window's DPI scale.  Finally, we need to divide that remainder to the draw size.
-		// The idea here is to divide out any "extra" scale that's not associated with the DPI scale, since ComputeMenuPlacement
+		// The idea here is to divide out any "extra" scale that's not associated with the DPI scale, since ComputeNewWindowMenuPlacement
 		// makes a child transform for the new window, based on the geometry of it, which if the menu anchor is inside a zoom panel
 		// that would translate to a menu that had a smaller size window, if the scale was tiny, which we don't want - we only want
 		// the DPI scale if any, to be factored into the size.
@@ -242,14 +235,14 @@ int32 SMenuAnchor::OnPaint( const FPaintArgs& Args, const FGeometry& AllottedGeo
 				if (bShouldDeferPaintingAfterWindowContent)
 				{
 					OutDrawElements.QueueDeferredPainting(
-						FSlateWindowElementList::FDeferredPaint(PopupChild->Widget, Args, PopupChild->Geometry, InWidgetStyle, bParentEnabled));
+						FSlateWindowElementList::FDeferredPaint(PopupChild->Widget, Args, PopupChild->Geometry, bApplyWidgetStyleToMenu ? InWidgetStyle : FWidgetStyle(), bParentEnabled));
 				}
 				else
 				{
 					const TSharedPtr<SWindow> PresentingWindow = PopupWindowPtr.Pin();
 					if (PresentingWindow.IsValid())
 					{
-						PopupChild->Widget->Paint(Args.WithNewParent(this), PopupChild->Geometry, PresentingWindow->GetClippingRectangleInWindow(), OutDrawElements, LayerId + 1, InWidgetStyle, ShouldBeEnabled(bParentEnabled));
+						PopupChild->Widget->Paint(Args.WithNewParent(this), PopupChild->Geometry, PresentingWindow->GetClippingRectangleInWindow(), OutDrawElements, LayerId + 1, bApplyWidgetStyleToMenu ? InWidgetStyle : FWidgetStyle(), ShouldBeEnabled(bParentEnabled));
 					}
 				}
 			}
diff --git a/Engine/Source/Runtime/Slate/Private/Widgets/Layout/SScrollBox.cpp b/Engine/Source/Runtime/Slate/Private/Widgets/Layout/SScrollBox.cpp
--- a/Engine/Source/Runtime/Slate/Private/Widgets/Layout/SScrollBox.cpp
+++ b/Engine/Source/Runtime/Slate/Private/Widgets/Layout/SScrollBox.cpp
@@ -504,6 +504,11 @@ void SScrollBox::SetScrollBarThickness(FVector2D InThickness)
 	ScrollBar->SetThickness(InThickness);
 }
 
+void SScrollBox::SetScrollBarRightClickDragAllowed(bool bIsAllowed)
+{
+	bAllowsRightClickDragScrolling = bIsAllowed;
+}
+
 EActiveTimerReturnType SScrollBox::UpdateInertialScroll(double InCurrentTime, float InDeltaTime)
 {
 	bool bKeepTicking = bIsScrolling;
@@ -641,11 +646,11 @@ FReply SScrollBox::OnMouseButtonDown( const FGeometry& MyGeometry, const FPointe
 	}
 	else
 	{
-		if ( MouseEvent.GetEffectingButton() == EKeys::RightMouseButton && ScrollBar->IsNeeded() )
+		if ( MouseEvent.GetEffectingButton() == EKeys::RightMouseButton && ScrollBar->IsNeeded()  && bAllowsRightClickDragScrolling)
 		{
 			AmountScrolledWhileRightMouseDown = 0;
 
-		Invalidate(EInvalidateWidget::Layout);
+			Invalidate(EInvalidateWidget::Layout);
 
 			return FReply::Handled();
 		}
@@ -656,7 +661,7 @@ FReply SScrollBox::OnMouseButtonDown( const FGeometry& MyGeometry, const FPointe
 
 FReply SScrollBox::OnMouseButtonUp( const FGeometry& MyGeometry, const FPointerEvent& MouseEvent )
 {
-	if ( MouseEvent.GetEffectingButton() == EKeys::RightMouseButton )
+	if ( MouseEvent.GetEffectingButton() == EKeys::RightMouseButton && bAllowsRightClickDragScrolling)
 	{
 		if ( !bIsScrollingActiveTimerRegistered && IsRightClickScrolling() )
 		{
@@ -739,7 +744,7 @@ FReply SScrollBox::OnMouseMove( const FGeometry& MyGeometry, const FPointerEvent
 	}
 	else
 	{
-		if ( MouseEvent.IsMouseButtonDown(EKeys::RightMouseButton) )
+		if ( MouseEvent.IsMouseButtonDown(EKeys::RightMouseButton)  && bAllowsRightClickDragScrolling)
 		{
 			// If scrolling with the right mouse button, we need to remember how much we scrolled.
 			// If we did not scroll at all, we will bring up the context menu when the mouse is released.
