Manually resolving CL 3538106 from 4.17 branch with CL 3526124 from Dev-Editor. Copying shared pointer to captured local before assigning to null on render thread.

diff --git a/Engine/Source/Runtime/Engine/Private/Slate/DebugCanvas.cpp b/Engine/Source/Runtime/Engine/Private/Slate/DebugCanvas.cpp
--- a/Engine/Source/Runtime/Engine/Private/Slate/DebugCanvas.cpp
+++ b/Engine/Source/Runtime/Engine/Private/Slate/DebugCanvas.cpp
@@ -6,47 +6,8 @@
 #include "CanvasTypes.h"
 #include "Framework/Application/SlateApplication.h"
 
-/** Checks that all FCanvasProxy allocations were deleted */
-class FProxyCounter
-{
-public:
-	FProxyCounter()
-	{
-		Creations = 0;
-		Deletions = 0;
-	}
 
-	~FProxyCounter()
-	{
-		ensureMsgf( Creations == Deletions, TEXT("FProxyCounter::~FProxyCounter has a mismatch.  %d creations != %d deletions"), Creations, Deletions );
-	}
 
-	int32 Creations;
-	int32 Deletions;
-};
-	
-class FCanvasProxy
-{
-public:
-	FCanvasProxy( FRenderTarget* RenderTarget, UWorld* InWorld )
-		: Canvas(RenderTarget, NULL, InWorld, InWorld ? InWorld->FeatureLevel : GMaxRHIFeatureLevel)
-	{
-		// Do not allow the canvas to be flushed outside of our debug rendering path
-		Canvas.SetAllowedModes( FCanvas::Allow_DeleteOnRender );
-		++Counter.Creations;
-	}
-
-	~FCanvasProxy()
-	{
-		++Counter.Deletions;
-	}
-
-	/** The canvas on this proxy */
-	FCanvas Canvas;
-	static FProxyCounter Counter;
-};
-
-FProxyCounter FCanvasProxy::Counter;
 
 /**
  * Simple representation of the backbuffer that the debug canvas renders to
@@ -99,19 +60,18 @@ FDebugCanvasDrawer::~FDebugCanvasDrawer()
 	delete RenderTarget;
 
 	// We assume that the render thread is no longer utilizing any canvases
-	if( GameThreadCanvas && RenderThreadCanvas != GameThreadCanvas )
+	if( GameThreadCanvas.IsValid() && RenderThreadCanvas != GameThreadCanvas )
 	{
-		delete GameThreadCanvas;
+		GameThreadCanvas.Reset();
 	}
 
 	if( RenderThreadCanvas )
 	{
-		FCanvasProxy* RTCanvas = RenderThreadCanvas;
+		FCanvasPtr RTCanvas = RenderThreadCanvas;
 		ENQUEUE_RENDER_COMMAND(DeleteDebugRenderThreadCanvas)(
 			[RTCanvas](FRHICommandListImmediate& RHICmdList)
 		{
-			check(RTCanvas);
-			delete RTCanvas;
+			RTCanvas = nullptr;
 		});
 
 		RenderThreadCanvas = nullptr;
@@ -120,7 +80,7 @@ FDebugCanvasDrawer::~FDebugCanvasDrawer()
 
 FCanvas* FDebugCanvasDrawer::GetGameThreadDebugCanvas()
 {
-	return &GameThreadCanvas->Canvas;
+	return GameThreadCanvas.Get();
 }
 
 
@@ -132,26 +92,26 @@ void FDebugCanvasDrawer::BeginRenderingCanvas( const FIntRect& CanvasRect )
 		(
 			BeginRenderingDebugCanvas,
 			FDebugCanvasDrawer*, CanvasDrawer, this, 
-			FCanvasProxy*, CanvasToRender, GameThreadCanvas,
+			FCanvasPtr, CanvasToRender, GameThreadCanvas,
 			FIntRect, CanvasRect, CanvasRect,
 			{
 				// Delete the old rendering thread canvas
-				if( CanvasDrawer->GetRenderThreadCanvas() && CanvasToRender != NULL )
+				if( CanvasDrawer->GetRenderThreadCanvas().IsValid() && CanvasToRender.IsValid() )
 				{
 					CanvasDrawer->DeleteRenderThreadCanvas();
 				}
 
-				if( CanvasToRender == NULL )
+				if(!CanvasToRender.IsValid())
 				{
 					CanvasToRender = CanvasDrawer->GetRenderThreadCanvas();
 				}
 
-				CanvasDrawer->SetRenderThreadCanvas( CanvasRect, CanvasToRender ); 
+				CanvasDrawer->SetRenderThreadCanvas( CanvasRect, CanvasToRender );
 			}
 		);
 		
 		// Gave the canvas to the render thread
-		GameThreadCanvas = NULL;
+		GameThreadCanvas = nullptr;
 	}
 }
 
@@ -164,12 +124,9 @@ void FDebugCanvasDrawer::InitDebugCanvas(UWorld* InWorld)
 		// the same canvas
 	if (FSlateApplication::Get().IsNormalExecution())
 	{
-		if( GameThreadCanvas != NULL )
-		{
-			delete GameThreadCanvas;
-		}
+		FCanvas* Canvas = new FCanvas(RenderTarget, nullptr, InWorld, InWorld ? InWorld->FeatureLevel : GMaxRHIFeatureLevel);
 
-		GameThreadCanvas = new FCanvasProxy( RenderTarget, InWorld );
+		GameThreadCanvas = MakeShared<FCanvas, ESPMode::ThreadSafe>(*Canvas);
 	}
 }
 
@@ -177,28 +134,28 @@ void FDebugCanvasDrawer::DrawRenderThread(FRHICommandListImmediate& RHICmdList,
 {
 	check( IsInRenderingThread() );
 
-	if( RenderThreadCanvas )
+	if( RenderThreadCanvas.IsValid() )
 	{
 		FTexture2DRHIRef& RT = *(FTexture2DRHIRef*)InWindowBackBuffer;
 		RenderTarget->SetRenderTargetTexture( RT );
-		bool bNeedToFlipVertical = RenderThreadCanvas->Canvas.GetAllowSwitchVerticalAxis();
+		bool bNeedToFlipVertical = RenderThreadCanvas->GetAllowSwitchVerticalAxis();
 		// Do not flip when rendering to the back buffer
-		RenderThreadCanvas->Canvas.SetAllowSwitchVerticalAxis(false);
-		if (RenderThreadCanvas->Canvas.IsScaledToRenderTarget() && IsValidRef(RT)) 
+		RenderThreadCanvas->SetAllowSwitchVerticalAxis(false);
+		if (RenderThreadCanvas->IsScaledToRenderTarget() && IsValidRef(RT)) 
 		{
-			RenderThreadCanvas->Canvas.SetRenderTargetRect( FIntRect(0, 0, RT->GetSizeX(), RT->GetSizeY()) );
+			RenderThreadCanvas->SetRenderTargetRect( FIntRect(0, 0, RT->GetSizeX(), RT->GetSizeY()) );
 		}
 		else
 		{
-			RenderThreadCanvas->Canvas.SetRenderTargetRect( RenderTarget->GetViewRect() );
+			RenderThreadCanvas->SetRenderTargetRect( RenderTarget->GetViewRect() );
 		}
-		RenderThreadCanvas->Canvas.Flush_RenderThread(RHICmdList, true);
-		RenderThreadCanvas->Canvas.SetAllowSwitchVerticalAxis(bNeedToFlipVertical);
+		RenderThreadCanvas->Flush_RenderThread(RHICmdList, true);
+		RenderThreadCanvas->SetAllowSwitchVerticalAxis(bNeedToFlipVertical);
 		RenderTarget->ClearRenderTargetTexture();
 	}
 }
 
-FCanvasProxy* FDebugCanvasDrawer::GetRenderThreadCanvas() 
+FCanvasPtr FDebugCanvasDrawer::GetRenderThreadCanvas()
 {
 	check( IsInRenderingThread() );
 	return RenderThreadCanvas;
@@ -207,14 +164,10 @@ FCanvasProxy* FDebugCanvasDrawer::GetRenderThreadCanvas()
 void FDebugCanvasDrawer::DeleteRenderThreadCanvas()
 {
 	check( IsInRenderingThread() );
-	if( RenderThreadCanvas )
-	{
-		delete RenderThreadCanvas;
-		RenderThreadCanvas = NULL;
-	}
+	RenderThreadCanvas.Reset();
 }
 
-void FDebugCanvasDrawer::SetRenderThreadCanvas( const FIntRect& InCanvasRect, FCanvasProxy* Canvas )
+void FDebugCanvasDrawer::SetRenderThreadCanvas( const FIntRect& InCanvasRect, FCanvasPtr& Canvas )
 {
 	check( IsInRenderingThread() );
 	RenderTarget->SetViewRect( InCanvasRect );
