CIS fix.

diff --git a/Engine/Source/Runtime/Engine/Private/Slate/DebugCanvas.cpp b/Engine/Source/Runtime/Engine/Private/Slate/DebugCanvas.cpp
--- a/Engine/Source/Runtime/Engine/Private/Slate/DebugCanvas.cpp
+++ b/Engine/Source/Runtime/Engine/Private/Slate/DebugCanvas.cpp
@@ -65,13 +65,13 @@ FDebugCanvasDrawer::~FDebugCanvasDrawer()
 		GameThreadCanvas.Reset();
 	}
 
-	if( RenderThreadCanvas )
+	if( RenderThreadCanvas.IsValid() )
 	{
+		// Capture a copy of the canvas until the render thread can delete it
 		FCanvasPtr RTCanvas = RenderThreadCanvas;
 		ENQUEUE_RENDER_COMMAND(DeleteDebugRenderThreadCanvas)(
 			[RTCanvas](FRHICommandListImmediate& RHICmdList)
 		{
-			RTCanvas = nullptr;
 		});
 
 		RenderThreadCanvas = nullptr;
