Fix CIS warning; direct access to ComponentToWorld is deprecated.

diff --git a/Engine/Source/Runtime/Engine/Private/Components/SceneComponent.cpp b/Engine/Source/Runtime/Engine/Private/Components/SceneComponent.cpp
--- a/Engine/Source/Runtime/Engine/Private/Components/SceneComponent.cpp
+++ b/Engine/Source/Runtime/Engine/Private/Components/SceneComponent.cpp
@@ -2184,7 +2184,7 @@ void USceneComponent::UpdateChildTransforms(EUpdateTransformFlags UpdateTransfor
 #if ENABLE_NAN_DIAGNOSTIC
 	if (!GetComponentTransform().IsValid())
 	{
-		logOrEnsureNanError(TEXT("USceneComponent::UpdateChildTransforms found NaN/INF in ComponentToWorld: %s"), *ComponentToWorld.ToString());
+		logOrEnsureNanError(TEXT("USceneComponent::UpdateChildTransforms found NaN/INF in ComponentToWorld: %s"), *GetComponentTransform().ToString());
 	}
 #endif
 
