Copying //UE4/Dev-Framework to //UE4/Dev-Main (Source: //UE4/Dev-Framework @ 3582324)

diff --git a/Engine/Source/Runtime/Engine/Private/Kismet/BlueprintPlatformLibrary.cpp b/Engine/Source/Runtime/Engine/Private/Kismet/BlueprintPlatformLibrary.cpp
--- a/Engine/Source/Runtime/Engine/Private/Kismet/BlueprintPlatformLibrary.cpp
+++ b/Engine/Source/Runtime/Engine/Private/Kismet/BlueprintPlatformLibrary.cpp
@@ -37,7 +37,7 @@ void UPlatformGameInstance::BeginDestroy()
 	FCoreDelegates::ApplicationRegisteredForUserNotificationsDelegate.RemoveAll(this);
 	FCoreDelegates::ApplicationFailedToRegisterForRemoteNotificationsDelegate.RemoveAll(this);
 	FCoreDelegates::ApplicationReceivedRemoteNotificationDelegate.RemoveAll(this);
-	FCoreDelegates::ApplicationReceivedRemoteNotificationDelegate.RemoveAll(this);
+	FCoreDelegates::ApplicationReceivedLocalNotificationDelegate.RemoveAll(this);
 	FCoreDelegates::ApplicationReceivedScreenOrientationChangedNotificationDelegate.RemoveAll(this);
 
     Super::BeginDestroy();
diff --git a/Engine/Source/Runtime/Engine/Private/KismetSystemLibrary.cpp b/Engine/Source/Runtime/Engine/Private/KismetSystemLibrary.cpp
--- a/Engine/Source/Runtime/Engine/Private/KismetSystemLibrary.cpp
+++ b/Engine/Source/Runtime/Engine/Private/KismetSystemLibrary.cpp
@@ -6,6 +6,7 @@
 #include "Misc/CommandLine.h"
 #include "Misc/App.h"
 #include "Misc/EngineVersion.h"
+#include "Misc/RuntimeErrors.h"
 #include "UObject/GCObject.h"
 #include "EngineGlobals.h"
 #include "Components/ActorComponent.h"
@@ -811,74 +812,92 @@ void UKismetSystemLibrary::SetNamePropertyByName(UObject* Object, FName Property
 	}
 }
 
-void UKismetSystemLibrary::SetAssetPropertyByName(UObject* Object, FName PropertyName, const TAssetPtr<UObject>& Value)
+void UKismetSystemLibrary::SetSoftObjectPropertyByName(UObject* Object, FName PropertyName, const TSoftObjectPtr<UObject>& Value)
 {
 	if (Object != NULL)
 	{
-		UAssetObjectProperty* ObjectProp = FindField<UAssetObjectProperty>(Object->GetClass(), PropertyName);
-		const FAssetPtr* AssetPtr = (const FAssetPtr*)(&Value);
-		ObjectProp->SetPropertyValue_InContainer(Object, *AssetPtr);
+		USoftObjectProperty* ObjectProp = FindField<USoftObjectProperty>(Object->GetClass(), PropertyName);
+		const FSoftObjectPtr* SoftObjectPtr = (const FSoftObjectPtr*)(&Value);
+		ObjectProp->SetPropertyValue_InContainer(Object, *SoftObjectPtr);
 	}
 }
 
-void UKismetSystemLibrary::SetAssetClassPropertyByName(UObject* Object, FName PropertyName, const TAssetSubclassOf<UObject>& Value)
+void UKismetSystemLibrary::SetSoftClassPropertyByName(UObject* Object, FName PropertyName, const TSoftClassPtr<UObject>& Value)
 {
 	if (Object != NULL)
 	{
-		UAssetClassProperty* ObjectProp = FindField<UAssetClassProperty>(Object->GetClass(), PropertyName);
-		const FAssetPtr* AssetPtr = (const FAssetPtr*)(&Value);
-		ObjectProp->SetPropertyValue_InContainer(Object, *AssetPtr);
+		USoftClassProperty* ObjectProp = FindField<USoftClassProperty>(Object->GetClass(), PropertyName);
+		const FSoftObjectPtr* SoftObjectPtr = (const FSoftObjectPtr*)(&Value);
+		ObjectProp->SetPropertyValue_InContainer(Object, *SoftObjectPtr);
 	}
 }
 
-bool UKismetSystemLibrary::IsValidSoftObjectReference(const TAssetPtr<UObject>& SoftObjectReference)
+FSoftObjectPath UKismetSystemLibrary::MakeSoftObjectPath(const FString& PathString)
+{
+	FSoftObjectPath SoftObjectPath(PathString);
+	if (!PathString.IsEmpty() && !SoftObjectPath.IsValid())
+	{
+		FFormatNamedArguments Args;
+		Args.Add(TEXT("PathString"), FText::AsCultureInvariant(PathString));
+		LogRuntimeError(FText::Format(NSLOCTEXT("KismetSystemLibrary", "PathStringInvalid", "Object path {PathString} not valid for MakeSoftObjectPath."), Args));
+	}
+
+	return SoftObjectPath;
+}
+
+void UKismetSystemLibrary::BreakSoftObjectPath(FSoftObjectPath InSoftObjectPath, FString& PathString)
+{
+	PathString = InSoftObjectPath.ToString();
+}
+
+bool UKismetSystemLibrary::IsValidSoftObjectReference(const TSoftObjectPtr<UObject>& SoftObjectReference)
 {
 	return !SoftObjectReference.IsNull();
 }
 
-bool UKismetSystemLibrary::EqualEqual_SoftObjectReference(const TAssetPtr<UObject>& A, const TAssetPtr<UObject>& B)
+bool UKismetSystemLibrary::EqualEqual_SoftObjectReference(const TSoftObjectPtr<UObject>& A, const TSoftObjectPtr<UObject>& B)
 {
 	return A == B;
 }
 
-bool UKismetSystemLibrary::NotEqual_SoftObjectReference(const TAssetPtr<UObject>& A, const TAssetPtr<UObject>& B)
+bool UKismetSystemLibrary::NotEqual_SoftObjectReference(const TSoftObjectPtr<UObject>& A, const TSoftObjectPtr<UObject>& B)
 {
 	return A != B;
 }
 
-bool UKismetSystemLibrary::IsValidSoftClassReference(const TAssetSubclassOf<UObject>& SoftClassReference)
+bool UKismetSystemLibrary::IsValidSoftClassReference(const TSoftClassPtr<UObject>& SoftClassReference)
 {
 	return !SoftClassReference.IsNull();
 }
 
-bool UKismetSystemLibrary::EqualEqual_SoftClassReference(const TAssetSubclassOf<UObject>& A, const TAssetSubclassOf<UObject>& B)
+bool UKismetSystemLibrary::EqualEqual_SoftClassReference(const TSoftClassPtr<UObject>& A, const TSoftClassPtr<UObject>& B)
 {
 	return A == B;
 }
 
-bool UKismetSystemLibrary::NotEqual_SoftClassReference(const TAssetSubclassOf<UObject>& A, const TAssetSubclassOf<UObject>& B)
+bool UKismetSystemLibrary::NotEqual_SoftClassReference(const TSoftClassPtr<UObject>& A, const TSoftClassPtr<UObject>& B)
 {
 	return A != B;
 }
 
-UObject* UKismetSystemLibrary::Conv_SoftObjectReferenceToObject(const TAssetPtr<UObject>& Asset)
+UObject* UKismetSystemLibrary::Conv_SoftObjectReferenceToObject(const TSoftObjectPtr<UObject>& SoftObject)
 {
-	return ((const FAssetPtr*)&Asset)->Get();
+	return SoftObject.Get();
 }
 
-TSubclassOf<UObject> UKismetSystemLibrary::Conv_SoftClassReferenceToClass(const TAssetSubclassOf<UObject>& AssetClass)
+TSubclassOf<UObject> UKismetSystemLibrary::Conv_SoftClassReferenceToClass(const TSoftClassPtr<UObject>& SoftClass)
 {
-	return Cast<UClass>(((const FAssetPtr*)&AssetClass)->Get());
+	return SoftClass.Get();
 }
 
-TAssetPtr<UObject> UKismetSystemLibrary::Conv_ObjectToSoftObjectReference(UObject *Object)
+TSoftObjectPtr<UObject> UKismetSystemLibrary::Conv_ObjectToSoftObjectReference(UObject *Object)
 {
-	return TAssetPtr<UObject>(Object);
+	return TSoftObjectPtr<UObject>(Object);
 }
 
-TAssetSubclassOf<UObject> UKismetSystemLibrary::Conv_ClassToSoftClassReference(const TSubclassOf<UObject>& Class)
+TSoftClassPtr<UObject> UKismetSystemLibrary::Conv_ClassToSoftClassReference(const TSubclassOf<UObject>& Class)
 {
-	return TAssetSubclassOf<UObject>(*Class);
+	return TSoftClassPtr<UObject>(*Class);
 }
 
 void UKismetSystemLibrary::SetTextPropertyByName(UObject* Object, FName PropertyName, const FText& Value)
@@ -2314,7 +2333,7 @@ struct FLoadAssetActionBase : public FPendingLatentAction
 	// @TODO: it would be good to have static/global manager? 
 
 public:
-	FStringAssetReference AssetReference;
+	FSoftObjectPath SoftObjectPath;
 	FStreamableManager StreamableManager;
 	TSharedPtr<FStreamableHandle> Handle;
 	FName ExecutionFunction;
@@ -2323,13 +2342,13 @@ public:
 
 	virtual void OnLoaded() PURE_VIRTUAL(FLoadAssetActionBase::OnLoaded, );
 
-	FLoadAssetActionBase(const FStringAssetReference& InAssetReference, const FLatentActionInfo& InLatentInfo)
-		: AssetReference(InAssetReference)
+	FLoadAssetActionBase(const FSoftObjectPath& InSoftObjectPath, const FLatentActionInfo& InLatentInfo)
+		: SoftObjectPath(InSoftObjectPath)
 		, ExecutionFunction(InLatentInfo.ExecutionFunction)
 		, OutputLink(InLatentInfo.Linkage)
 		, CallbackTarget(InLatentInfo.CallbackTarget)
 	{
-		Handle = StreamableManager.RequestAsyncLoad(AssetReference);
+		Handle = StreamableManager.RequestAsyncLoad(SoftObjectPath);
 	}
 
 	virtual ~FLoadAssetActionBase()
@@ -2353,26 +2372,26 @@ public:
 #if WITH_EDITOR
 	virtual FString GetDescription() const override
 	{
-		return FString::Printf(TEXT("Load Asset Action Base: %s"), *AssetReference.ToString());
+		return FString::Printf(TEXT("Load Asset Action Base: %s"), *SoftObjectPath.ToString());
 	}
 #endif
 };
 
-void UKismetSystemLibrary::LoadAsset(UObject* WorldContextObject, TAssetPtr<UObject> Asset, UKismetSystemLibrary::FOnAssetLoaded OnLoaded, FLatentActionInfo LatentInfo)
+void UKismetSystemLibrary::LoadAsset(UObject* WorldContextObject, TSoftObjectPtr<UObject> Asset, UKismetSystemLibrary::FOnAssetLoaded OnLoaded, FLatentActionInfo LatentInfo)
 {
 	struct FLoadAssetAction : public FLoadAssetActionBase
 	{
 	public:
 		UKismetSystemLibrary::FOnAssetLoaded OnLoadedCallback;
 
-		FLoadAssetAction(const FStringAssetReference& InAssetReference, UKismetSystemLibrary::FOnAssetLoaded InOnLoadedCallback, const FLatentActionInfo& InLatentInfo)
-			: FLoadAssetActionBase(InAssetReference, InLatentInfo)
+		FLoadAssetAction(const FSoftObjectPath& InSoftObjectPath, UKismetSystemLibrary::FOnAssetLoaded InOnLoadedCallback, const FLatentActionInfo& InLatentInfo)
+			: FLoadAssetActionBase(InSoftObjectPath, InLatentInfo)
 			, OnLoadedCallback(InOnLoadedCallback)
 		{}
 
 		virtual void OnLoaded() override
 		{
-			UObject* LoadedObject = AssetReference.ResolveObject();
+			UObject* LoadedObject = SoftObjectPath.ResolveObject();
 			OnLoadedCallback.ExecuteIfBound(LoadedObject);
 		}
 	};
@@ -2382,27 +2401,27 @@ void UKismetSystemLibrary::LoadAsset(UObject* WorldContextObject, TAssetPtr<UObj
 		FLatentActionManager& LatentManager = World->GetLatentActionManager();
 		if (LatentManager.FindExistingAction<FLoadAssetAction>(LatentInfo.CallbackTarget, LatentInfo.UUID) == nullptr)
 		{
-			FLoadAssetAction* NewAction = new FLoadAssetAction(Asset.ToStringReference(), OnLoaded, LatentInfo);
+			FLoadAssetAction* NewAction = new FLoadAssetAction(Asset.ToSoftObjectPath(), OnLoaded, LatentInfo);
 			LatentManager.AddNewAction(LatentInfo.CallbackTarget, LatentInfo.UUID, NewAction);
 		}
 	}
 }
 
-void UKismetSystemLibrary::LoadAssetClass(UObject* WorldContextObject, TAssetSubclassOf<UObject> AssetClass, UKismetSystemLibrary::FOnAssetClassLoaded OnLoaded, FLatentActionInfo LatentInfo)
+void UKismetSystemLibrary::LoadAssetClass(UObject* WorldContextObject, TSoftClassPtr<UObject> AssetClass, UKismetSystemLibrary::FOnAssetClassLoaded OnLoaded, FLatentActionInfo LatentInfo)
 {
 	struct FLoadAssetClassAction : public FLoadAssetActionBase
 	{
 	public:
 		UKismetSystemLibrary::FOnAssetClassLoaded OnLoadedCallback;
 
-		FLoadAssetClassAction(const FStringAssetReference& InAssetReference, UKismetSystemLibrary::FOnAssetClassLoaded InOnLoadedCallback, const FLatentActionInfo& InLatentInfo)
-			: FLoadAssetActionBase(InAssetReference, InLatentInfo)
+		FLoadAssetClassAction(const FSoftObjectPath& InSoftObjectPath, UKismetSystemLibrary::FOnAssetClassLoaded InOnLoadedCallback, const FLatentActionInfo& InLatentInfo)
+			: FLoadAssetActionBase(InSoftObjectPath, InLatentInfo)
 			, OnLoadedCallback(InOnLoadedCallback)
 		{}
 
 		virtual void OnLoaded() override
 		{
-			UClass* LoadedObject = Cast<UClass>(AssetReference.ResolveObject());
+			UClass* LoadedObject = Cast<UClass>(SoftObjectPath.ResolveObject());
 			OnLoadedCallback.ExecuteIfBound(LoadedObject);
 		}
 	};
@@ -2412,7 +2431,7 @@ void UKismetSystemLibrary::LoadAssetClass(UObject* WorldContextObject, TAssetSub
 		FLatentActionManager& LatentManager = World->GetLatentActionManager();
 		if (LatentManager.FindExistingAction<FLoadAssetClassAction>(LatentInfo.CallbackTarget, LatentInfo.UUID) == nullptr)
 		{
-			FLoadAssetClassAction* NewAction = new FLoadAssetClassAction(AssetClass.ToStringReference(), OnLoaded, LatentInfo);
+			FLoadAssetClassAction* NewAction = new FLoadAssetClassAction(AssetClass.ToSoftObjectPath(), OnLoaded, LatentInfo);
 			LatentManager.AddNewAction(LatentInfo.CallbackTarget, LatentInfo.UUID, NewAction);
 		}
 	}
@@ -2464,27 +2483,27 @@ TSubclassOf<UObject> UKismetSystemLibrary::GetClassFromPrimaryAssetId(FPrimaryAs
 	return nullptr;
 }
 
-TAssetPtr<UObject> UKismetSystemLibrary::GetSoftObjectReferenceFromPrimaryAssetId(FPrimaryAssetId PrimaryAssetId)
+TSoftObjectPtr<UObject> UKismetSystemLibrary::GetSoftObjectReferenceFromPrimaryAssetId(FPrimaryAssetId PrimaryAssetId)
 {
 	if (UAssetManager* Manager = UAssetManager::GetIfValid())
 	{
 		FPrimaryAssetTypeInfo Info;
 		if (Manager->GetPrimaryAssetTypeInfo(PrimaryAssetId.PrimaryAssetType, Info) && !Info.bHasBlueprintClasses)
 		{
-			return TAssetPtr<UObject>(Manager->GetPrimaryAssetPath(PrimaryAssetId));
+			return TSoftObjectPtr<UObject>(Manager->GetPrimaryAssetPath(PrimaryAssetId));
 		}
 	}
 	return nullptr;
 }
 
-TAssetSubclassOf<UObject> UKismetSystemLibrary::GetSoftClassReferenceFromPrimaryAssetId(FPrimaryAssetId PrimaryAssetId)
+TSoftClassPtr<UObject> UKismetSystemLibrary::GetSoftClassReferenceFromPrimaryAssetId(FPrimaryAssetId PrimaryAssetId)
 {
 	if (UAssetManager* Manager = UAssetManager::GetIfValid())
 	{
 		FPrimaryAssetTypeInfo Info;
 		if (Manager->GetPrimaryAssetTypeInfo(PrimaryAssetId.PrimaryAssetType, Info) && Info.bHasBlueprintClasses)
 		{
-			return TAssetSubclassOf<UObject>(Manager->GetPrimaryAssetPath(PrimaryAssetId));
+			return TSoftClassPtr<UObject>(Manager->GetPrimaryAssetPath(PrimaryAssetId));
 		}
 	}
 	return nullptr;
@@ -2514,20 +2533,20 @@ FPrimaryAssetId UKismetSystemLibrary::GetPrimaryAssetIdFromClass(TSubclassOf<UOb
 	return FPrimaryAssetId();
 }
 
-FPrimaryAssetId UKismetSystemLibrary::GetPrimaryAssetIdFromSoftObjectReference(TAssetPtr<UObject> SoftObjectReference)
+FPrimaryAssetId UKismetSystemLibrary::GetPrimaryAssetIdFromSoftObjectReference(TSoftObjectPtr<UObject> SoftObjectReference)
 {
 	if (UAssetManager* Manager = UAssetManager::GetIfValid())
 	{
-		return Manager->GetPrimaryAssetIdForPath(SoftObjectReference.ToStringReference());
+		return Manager->GetPrimaryAssetIdForPath(SoftObjectReference.ToSoftObjectPath());
 	}
 	return FPrimaryAssetId();
 }
 
-FPrimaryAssetId UKismetSystemLibrary::GetPrimaryAssetIdFromSoftClassReference(TAssetSubclassOf<UObject> SoftClassReference)
+FPrimaryAssetId UKismetSystemLibrary::GetPrimaryAssetIdFromSoftClassReference(TSoftClassPtr<UObject> SoftClassReference)
 {
 	if (UAssetManager* Manager = UAssetManager::GetIfValid())
 	{
-		return Manager->GetPrimaryAssetIdForPath(SoftClassReference.ToStringReference());
+		return Manager->GetPrimaryAssetIdForPath(SoftClassReference.ToSoftObjectPath());
 	}
 	return FPrimaryAssetId();
 }
diff --git a/Engine/Source/Runtime/Engine/Private/LatentActionManager.cpp b/Engine/Source/Runtime/Engine/Private/LatentActionManager.cpp
--- a/Engine/Source/Runtime/Engine/Private/LatentActionManager.cpp
+++ b/Engine/Source/Runtime/Engine/Private/LatentActionManager.cpp
@@ -181,20 +181,7 @@ void FLatentActionManager::TickLatentActionForObject(float DeltaTime, FActionLis
 
 				if (UFunction* ExecutionFunction = CallbackTarget->FindFunction(LinkInfo.ExecutionFunction))
 				{
-					const bool bInstrumentLatentCall = CallbackTarget->GetClass()->HasInstrumentation();
-					if (CallbackTarget->GetClass()->HasInstrumentation())
-					{
-						// Place event markers around the latent blueprint call so we can match up the entry points.
-						FScriptInstrumentationSignal LatentEventStartInfo(EScriptInstrumentation::ResumeEvent, CallbackTarget, ExecutionFunction, LinkInfo.LinkID);
-						FBlueprintCoreDelegates::InstrumentScriptEvent(LatentEventStartInfo);
-						CallbackTarget->ProcessEvent(ExecutionFunction, &(LinkInfo.LinkID));
-						FScriptInstrumentationSignal LatentEventStopInfo(EScriptInstrumentation::Stop, CallbackTarget, ExecutionFunction);
-						FBlueprintCoreDelegates::InstrumentScriptEvent(LatentEventStopInfo);
-					}
-					else
-					{
-						CallbackTarget->ProcessEvent(ExecutionFunction, &(LinkInfo.LinkID));
-					}
+					CallbackTarget->ProcessEvent(ExecutionFunction, &(LinkInfo.LinkID));
 				}
 				else
 				{
diff --git a/Engine/Source/Runtime/Engine/Private/Level.cpp b/Engine/Source/Runtime/Engine/Private/Level.cpp
--- a/Engine/Source/Runtime/Engine/Private/Level.cpp
+++ b/Engine/Source/Runtime/Engine/Private/Level.cpp
@@ -8,6 +8,7 @@ Level.cpp: Level-related functions
 #include "Misc/ScopedSlowTask.h"
 #include "UObject/RenderingObjectVersion.h"
 #include "Templates/ScopedPointer.h"
+#include "Templates/UnrealTemplate.h"
 #include "UObject/Package.h"
 #include "Serialization/AsyncLoading.h"
 #include "EngineStats.h"
@@ -1948,6 +1949,35 @@ void ULevel::BeginCacheForCookedPlatformData(const ITargetPlatform *TargetPlatfo
 	}
 }
 
+void ULevel::FixupForPIE(int32 PIEInstanceID)
+{
+	TGuardValue<int32> SetPlayInEditorID(GPlayInEditorID, PIEInstanceID);
+
+	struct FSoftPathPIEFixupSerializer : public FArchiveUObject
+	{
+		FSoftPathPIEFixupSerializer() 
+		{
+			ArIsSaving = true;
+		}
+
+		FArchive& operator<<(FSoftObjectPath& Value)
+		{
+			Value.FixupForPIE();
+			return *this;
+		}
+	};
+
+	FSoftPathPIEFixupSerializer FixupSerializer;
+
+	TArray<UObject*> SubObjects;
+	GetObjectsWithOuter(this, SubObjects);
+
+	for (UObject* Object : SubObjects)
+	{
+		Object->Serialize(FixupSerializer);
+	}
+}
+
 #endif	//WITH_EDITOR
 
 bool ULevel::IsPersistentLevel() const
diff --git a/Engine/Source/Runtime/Engine/Private/LevelStreaming.cpp b/Engine/Source/Runtime/Engine/Private/LevelStreaming.cpp
--- a/Engine/Source/Runtime/Engine/Private/LevelStreaming.cpp
+++ b/Engine/Source/Runtime/Engine/Private/LevelStreaming.cpp
@@ -229,7 +229,7 @@ void ULevelStreaming::PostLoad()
 		const FString DeprecatedPackageNameString = PackageName_DEPRECATED.ToString();
 		if ( FPackageName::IsShortPackageName(PackageName_DEPRECATED) == false )
 		{
-			// Convert the FName reference to a TAssetPtr, then broadcast that we loaded a reference
+			// Convert the FName reference to a TSoftObjectPtr, then broadcast that we loaded a reference
 			// so this reference is gathered by the cooker without having to resave the package.
 			SetWorldAssetByPackageName(PackageName_DEPRECATED);
 			WorldAsset.GetUniqueID().PostLoadPath();
@@ -495,6 +495,13 @@ bool ULevelStreaming::RequestLevel(UWorld* PersistentWorld, bool bAllowLevelLoad
 #endif
 			if (World->PersistentLevel != LoadedLevel)
 			{
+#if WITH_EDITOR
+				if (PIEInstanceID != INDEX_NONE)
+				{
+					World->PersistentLevel->FixupForPIE(PIEInstanceID);
+				}
+#endif
+
 				SetLoadedLevel(World->PersistentLevel);
 				// Broadcast level loaded event to blueprints
 				OnLevelLoaded.Broadcast();
@@ -573,7 +580,16 @@ void ULevelStreaming::AsyncLevelLoadComplete(const FName& InPackageName, UPackag
 					}
 					else
 					{
-						check(PendingUnloadLevel == NULL);
+						check(PendingUnloadLevel == nullptr);
+					
+#if WITH_EDITOR
+						int32 PIEInstanceID = GetOutermost()->PIEInstanceID;
+						if (PIEInstanceID != INDEX_NONE)
+						{
+							World->PersistentLevel->FixupForPIE(PIEInstanceID);
+						}
+#endif
+
 						SetLoadedLevel(Level);
 						// Broadcast level loaded event to blueprints
 						OnLevelLoaded.Broadcast();
@@ -815,7 +831,7 @@ void ULevelStreaming::BroadcastLevelVisibleStatus(UWorld* PersistentWorld, FName
 	}
 }
 
-void ULevelStreaming::SetWorldAsset(const TAssetPtr<UWorld>& NewWorldAsset)
+void ULevelStreaming::SetWorldAsset(const TSoftObjectPtr<UWorld>& NewWorldAsset)
 {
 	WorldAsset = NewWorldAsset;
 	bHasCachedWorldAssetPackageFName = false;
@@ -841,7 +857,7 @@ void ULevelStreaming::SetWorldAssetByPackageName(FName InPackageName)
 {
 	const FString TargetWorldPackageName = InPackageName.ToString();
 	const FString TargetWorldObjectName = FPackageName::GetLongPackageAssetName(TargetWorldPackageName);
-	TAssetPtr<UWorld> NewWorld;
+	TSoftObjectPtr<UWorld> NewWorld;
 	NewWorld = TargetWorldPackageName + TEXT(".") + TargetWorldObjectName;
 	SetWorldAsset(NewWorld);
 }
@@ -859,8 +875,9 @@ void ULevelStreaming::RenameForPIE(int32 PIEInstanceID)
 				UWorld::BuildPIEPackagePrefix(PIEInstanceID));
 			PackageNameToLoad = FName(*NonPrefixedName);
 		}
-		const FString PlayWorldSteamingLevelName = UWorld::ConvertToPIEPackageName(GetWorldAssetPackageName(), PIEInstanceID);
-		SetWorldAssetByPackageName(FName(*PlayWorldSteamingLevelName));
+		FName PlayWorldStreamingPackageName = FName(*UWorld::ConvertToPIEPackageName(GetWorldAssetPackageName(), PIEInstanceID));
+		FSoftObjectPath::AddPIEPackageName(PlayWorldStreamingPackageName);
+		SetWorldAssetByPackageName(PlayWorldStreamingPackageName);
 	}
 	
 	// Rename LOD levels if any
@@ -871,8 +888,9 @@ void ULevelStreaming::RenameForPIE(int32 PIEInstanceID)
 		{
 			// Store LOD level original package name
 			LODPackageNamesToLoad.Add(LODPackageName); 
-			// Apply PIE prefix to package name
-			LODPackageName = FName(*UWorld::ConvertToPIEPackageName(LODPackageName.ToString(), PIEInstanceID)); 
+			// Apply PIE prefix to package name			
+			LODPackageName = FName(*UWorld::ConvertToPIEPackageName(LODPackageName.ToString(), PIEInstanceID));
+			FSoftObjectPath::AddPIEPackageName(LODPackageName);
 		}
 	}
 }
diff --git a/Engine/Source/Runtime/Engine/Private/Materials/Material.cpp b/Engine/Source/Runtime/Engine/Private/Materials/Material.cpp
--- a/Engine/Source/Runtime/Engine/Private/Materials/Material.cpp
+++ b/Engine/Source/Runtime/Engine/Private/Materials/Material.cpp
@@ -409,28 +409,30 @@ void UMaterialInterface::InitDefaultMaterials()
 
 		
 #if WITH_EDITOR
-		GPowerToRoughnessMaterialFunction = LoadObject< UMaterialFunction >( NULL, TEXT("/Engine/Functions/Engine_MaterialFunctions01/Shading/PowerToRoughness.PowerToRoughness"), NULL, LOAD_None, NULL );
+		GPowerToRoughnessMaterialFunction = LoadObject< UMaterialFunction >(nullptr, TEXT("/Engine/Functions/Engine_MaterialFunctions01/Shading/PowerToRoughness.PowerToRoughness"), nullptr, LOAD_None, nullptr);
 		checkf( GPowerToRoughnessMaterialFunction, TEXT("Cannot load PowerToRoughness") );
 		GPowerToRoughnessMaterialFunction->AddToRoot();
 
-		GConvertFromDiffSpecMaterialFunction = LoadObject< UMaterialFunction >( NULL, TEXT("/Engine/Functions/Engine_MaterialFunctions01/Shading/ConvertFromDiffSpec.ConvertFromDiffSpec"), NULL, LOAD_None, NULL );
+		GConvertFromDiffSpecMaterialFunction = LoadObject< UMaterialFunction >(nullptr, TEXT("/Engine/Functions/Engine_MaterialFunctions01/Shading/ConvertFromDiffSpec.ConvertFromDiffSpec"), nullptr, LOAD_None, nullptr);
 		checkf( GConvertFromDiffSpecMaterialFunction, TEXT("Cannot load ConvertFromDiffSpec") );
 		GConvertFromDiffSpecMaterialFunction->AddToRoot();
 #endif
 
 		for (int32 Domain = 0; Domain < MD_MAX; ++Domain)
 		{
-			if (GDefaultMaterials[Domain] == NULL)
+			if (GDefaultMaterials[Domain] == nullptr)
 			{
-				GDefaultMaterials[Domain] = FindObject<UMaterial>(NULL,GDefaultMaterialNames[Domain]);
-				if (GDefaultMaterials[Domain] == NULL
+				FString ResolvedPath = ResolveIniObjectsReference(GDefaultMaterialNames[Domain]);
+
+				GDefaultMaterials[Domain] = FindObject<UMaterial>(nullptr, *ResolvedPath);
+				if (GDefaultMaterials[Domain] == nullptr
 #if USE_EVENT_DRIVEN_ASYNC_LOAD_AT_BOOT_TIME
 					&& (RecursionLevel == 1 || !GEventDrivenLoaderEnabled)
 #endif
 					)
 				{
-					GDefaultMaterials[Domain] = LoadObject<UMaterial>(NULL, GDefaultMaterialNames[Domain], NULL, LOAD_DisableDependencyPreloading, NULL);
-					checkf(GDefaultMaterials[Domain] != NULL, TEXT("Cannot load default material '%s'"), GDefaultMaterialNames[Domain]);
+					GDefaultMaterials[Domain] = LoadObject<UMaterial>(nullptr, *ResolvedPath, nullptr, LOAD_DisableDependencyPreloading, nullptr);
+					checkf(GDefaultMaterials[Domain] != nullptr, TEXT("Cannot load default material '%s'"), GDefaultMaterialNames[Domain]);
 				}
 				if (GDefaultMaterials[Domain])
 				{
diff --git a/Engine/Source/Runtime/Engine/Private/Materials/MaterialExpressions.cpp b/Engine/Source/Runtime/Engine/Private/Materials/MaterialExpressions.cpp
--- a/Engine/Source/Runtime/Engine/Private/Materials/MaterialExpressions.cpp
+++ b/Engine/Source/Runtime/Engine/Private/Materials/MaterialExpressions.cpp
@@ -7357,6 +7357,7 @@ void UMaterialExpressionTransformPosition::GetCaption(TArray<FString>& OutCaptio
 UMaterialExpressionComment::UMaterialExpressionComment(const FObjectInitializer& ObjectInitializer)
 	: Super(ObjectInitializer)
 	, CommentColor(FLinearColor::White)
+	, FontSize(18)
 {
 	// Structure to hold one-time initialization
 	struct FConstructorStatics
@@ -7399,6 +7400,14 @@ void UMaterialExpressionComment::PostEditChangeProperty(FPropertyChangedEvent& P
 				CastChecked<UMaterialGraphNode_Comment>(GraphNode)->CommentColor = CommentColor;
 			}
 		}
+		else if (PropertyName == GET_MEMBER_NAME_CHECKED(UMaterialExpressionComment, FontSize))
+		{
+			if (GraphNode)
+			{
+				GraphNode->Modify();
+				CastChecked<UMaterialGraphNode_Comment>(GraphNode)->FontSize = FontSize;
+			}
+		}
 
 		// Don't need to update preview after changing comments
 		bNeedToUpdatePreview = false;
diff --git a/Engine/Source/Runtime/Engine/Private/Materials/MaterialInterface.cpp b/Engine/Source/Runtime/Engine/Private/Materials/MaterialInterface.cpp
--- a/Engine/Source/Runtime/Engine/Private/Materials/MaterialInterface.cpp
+++ b/Engine/Source/Runtime/Engine/Private/Materials/MaterialInterface.cpp
@@ -486,7 +486,7 @@ bool UMaterialInterface::FindTextureStreamingDataIndexRange(FName TextureName, i
 {
 #if WITH_EDITORONLY_DATA
 	// Because of redirectors (when textures are renammed), the texture names might be invalid and we need to udpate the data at every load.
-	// Normally we would do that in the post load, but since the process needs to resolve the StringAssetReference, this is forbidden at that place.
+	// Normally we would do that in the post load, but since the process needs to resolve the SoftObjectPaths, this is forbidden at that place.
 	// As a workaround, we do it on demand. Note that this is not required in cooked build as it is done in the presave.
 	const_cast<UMaterialInterface*>(this)->SortTextureStreamingData(false, false);
 #endif
diff --git a/Engine/Source/Runtime/Engine/Private/NetConnection.cpp b/Engine/Source/Runtime/Engine/Private/NetConnection.cpp
--- a/Engine/Source/Runtime/Engine/Private/NetConnection.cpp
+++ b/Engine/Source/Runtime/Engine/Private/NetConnection.cpp
@@ -1929,15 +1929,7 @@ void UNetConnection::HandleClientPlayer( APlayerController *PC, UNetConnection*
 			const ULevel* Level = LevelStreaming->GetLoadedLevel();
 			if ( Level != NULL && Level->bIsVisible && !Level->bClientOnlyVisible )
 			{
-				// Remap packagename for PIE networking before sending out to server
-				FName PackageName = Level->GetOutermost()->GetFName();
-				FString PackageNameStr = PackageName.ToString();
-				if (GEngine->NetworkRemapPath(Driver, PackageNameStr, false))
-				{
-					PackageName = FName(*PackageNameStr);
-				}
-
-				PC->ServerUpdateLevelVisibility( PackageName, true );
+				PC->ServerUpdateLevelVisibility(PC->NetworkRemapPath(Level->GetOutermost()->GetFName(), false), true );
 			}
 		}
 	}
diff --git a/Engine/Source/Runtime/Engine/Private/ObjectLibrary.cpp b/Engine/Source/Runtime/Engine/Private/ObjectLibrary.cpp
--- a/Engine/Source/Runtime/Engine/Private/ObjectLibrary.cpp
+++ b/Engine/Source/Runtime/Engine/Private/ObjectLibrary.cpp
@@ -493,7 +493,7 @@ int32 UObjectLibrary::LoadAssetsFromAssetData()
 	// Preload the packages with an async call, faster in cooked builds
 	if (bPreloadObjects)
 	{
-		TArray<FStringAssetReference> AssetsToStream;
+		TArray<FSoftObjectPath> AssetsToStream;
 
 		for (int32 AssetIdx = 0; AssetIdx < AssetDataList.Num(); AssetIdx++)
 		{
diff --git a/Engine/Source/Runtime/Engine/Private/PhysicsEngine/PhysicsAsset.cpp b/Engine/Source/Runtime/Engine/Private/PhysicsEngine/PhysicsAsset.cpp
--- a/Engine/Source/Runtime/Engine/Private/PhysicsEngine/PhysicsAsset.cpp
+++ b/Engine/Source/Runtime/Engine/Private/PhysicsEngine/PhysicsAsset.cpp
@@ -216,7 +216,7 @@ void UPhysicsAsset::Serialize(FArchive& Ar)
 #if WITH_EDITORONLY_DATA
 	if (DefaultSkelMesh_DEPRECATED != NULL)
 	{
-		PreviewSkeletalMesh = TAssetPtr<USkeletalMesh>(DefaultSkelMesh_DEPRECATED);
+		PreviewSkeletalMesh = TSoftObjectPtr<USkeletalMesh>(DefaultSkelMesh_DEPRECATED);
 		DefaultSkelMesh_DEPRECATED = NULL;
 	}
 #endif
diff --git a/Engine/Source/Runtime/Engine/Private/PlayerController.cpp b/Engine/Source/Runtime/Engine/Private/PlayerController.cpp
--- a/Engine/Source/Runtime/Engine/Private/PlayerController.cpp
+++ b/Engine/Source/Runtime/Engine/Private/PlayerController.cpp
@@ -209,17 +209,20 @@ void APlayerController::FailedToSpawnPawn()
 	ClientGotoState(NAME_Inactive);
 }
 
+FName APlayerController::NetworkRemapPath(FName InPackageName, bool bReading)
+{
+	// For PIE Networking: remap the packagename to our local PIE packagename
+	FString PackageNameStr = InPackageName.ToString();
+	GEngine->NetworkRemapPath(GetNetDriver(), PackageNameStr, bReading);
+	return FName(*PackageNameStr);
+}
+
 /// @cond DOXYGEN_WARNINGS
 
 void APlayerController::ClientUpdateLevelStreamingStatus_Implementation(FName PackageName, bool bNewShouldBeLoaded, bool bNewShouldBeVisible, bool bNewShouldBlockOnLoad, int32 LODIndex )
 {
-	// For PIE Networking: remap the packagename to our local PIE packagename
-	FString PackageNameStr = PackageName.ToString();
-	if (GEngine->NetworkRemapPath(GetNetDriver(), PackageNameStr, true))
-	{
-		PackageName = FName(*PackageNameStr);
-	}
-
+	PackageName = NetworkRemapPath(PackageName, true);
+	
 	// Distance dependent streaming levels should be controlled by client only
 	if (GetWorld() && GetWorld()->WorldComposition)
 	{
@@ -296,6 +299,8 @@ void APlayerController::ServerUpdateLevelVisibility_Implementation(FName Package
 	UNetConnection* Connection = Cast<UNetConnection>(Player);
 	if (Connection != NULL)
 	{
+		PackageName = NetworkRemapPath(PackageName, true);
+
 		// add or remove the level package name from the list, as requested
 		if (bIsVisible)
 		{
@@ -1093,7 +1098,7 @@ void APlayerController::CreateTouchInterface()
 		if (CurrentTouchInterface == nullptr)
 		{
 			// load what the game wants to show at startup
-			FStringAssetReference DefaultTouchInterfaceName = GetDefault<UInputSettings>()->DefaultTouchInterface;
+			FSoftObjectPath DefaultTouchInterfaceName = GetDefault<UInputSettings>()->DefaultTouchInterface;
 
 			if (DefaultTouchInterfaceName.IsValid())
 			{
@@ -3104,7 +3109,7 @@ void APlayerController::ClientForceGarbageCollection_Implementation()
 
 void APlayerController::LevelStreamingStatusChanged(ULevelStreaming* LevelObject, bool bNewShouldBeLoaded, bool bNewShouldBeVisible, bool bNewShouldBlockOnLoad, int32 LODIndex )
 {
-	ClientUpdateLevelStreamingStatus(LevelObject->GetWorldAssetPackageFName(),bNewShouldBeLoaded,bNewShouldBeVisible,bNewShouldBlockOnLoad,LODIndex);
+	ClientUpdateLevelStreamingStatus(NetworkRemapPath(LevelObject->GetWorldAssetPackageFName(), false), bNewShouldBeLoaded, bNewShouldBeVisible, bNewShouldBlockOnLoad, LODIndex);
 }
 
 /// @cond DOXYGEN_WARNINGS
diff --git a/Engine/Source/Runtime/Engine/Private/PrimaryAssetLabel.cpp b/Engine/Source/Runtime/Engine/Private/PrimaryAssetLabel.cpp
--- a/Engine/Source/Runtime/Engine/Private/PrimaryAssetLabel.cpp
+++ b/Engine/Source/Runtime/Engine/Private/PrimaryAssetLabel.cpp
@@ -46,11 +46,11 @@ void UPrimaryAssetLabel::UpdateAssetBundleData()
 		TArray<FAssetData> DirectoryAssets;
 		AssetRegistry.GetAssetsByPath(PackagePath, DirectoryAssets, true);
 
-		TArray<FStringAssetReference> NewPaths;
+		TArray<FSoftObjectPath> NewPaths;
 
 		for (const FAssetData& AssetData : DirectoryAssets)
 		{
-			FStringAssetReference AssetRef = Manager.GetAssetPathForData(AssetData);
+			FSoftObjectPath AssetRef = Manager.GetAssetPathForData(AssetData);
 
 			if (!AssetRef.IsNull())
 			{
@@ -64,7 +64,7 @@ void UPrimaryAssetLabel::UpdateAssetBundleData()
 
 	if (AssetCollection.CollectionName != NAME_None)
 	{
-		TArray<FStringAssetReference> NewPaths;
+		TArray<FSoftObjectPath> NewPaths;
 		TArray<FName> CollectionAssets;
 		ICollectionManager& CollectionManager = FCollectionManagerModule::GetModule().Get();
 		CollectionManager.GetAssetsInCollection(AssetCollection.CollectionName, ECollectionShareType::CST_All, CollectionAssets);
@@ -72,7 +72,7 @@ void UPrimaryAssetLabel::UpdateAssetBundleData()
 		{
 			FAssetData FoundAsset = Manager.GetAssetRegistry().GetAssetByObjectPath(CollectionAssets[Index]);
 
-			FStringAssetReference AssetRef = Manager.GetAssetPathForData(FoundAsset);
+			FSoftObjectPath AssetRef = Manager.GetAssetPathForData(FoundAsset);
 
 			if (!AssetRef.IsNull())
 			{
diff --git a/Engine/Source/Runtime/Engine/Private/SimpleConstructionScript.cpp b/Engine/Source/Runtime/Engine/Private/SimpleConstructionScript.cpp
--- a/Engine/Source/Runtime/Engine/Private/SimpleConstructionScript.cpp
+++ b/Engine/Source/Runtime/Engine/Private/SimpleConstructionScript.cpp
@@ -1137,7 +1137,7 @@ void USimpleConstructionScript::ValidateSceneRootNodes()
 }
 
 #if WITH_EDITOR
-void USimpleConstructionScript::GenerateListOfExistingNames(TArray<FName>& CurrentNames) const
+void USimpleConstructionScript::GenerateListOfExistingNames(TSet<FName>& CurrentNames) const
 {
 	TArray<const USCS_Node*> ChildrenNodes = GetAllNodesConst();
 	const UBlueprintGeneratedClass* OwnerClass = Cast<const UBlueprintGeneratedClass>(GetOuter());
@@ -1181,13 +1181,13 @@ void USimpleConstructionScript::GenerateListOfExistingNames(TArray<FName>& Curre
 
 	if (GetDefaultSceneRootNode())
 	{
-		CurrentNames.AddUnique(GetDefaultSceneRootNode()->GetVariableName());
+		CurrentNames.Add(GetDefaultSceneRootNode()->GetVariableName());
 	}
 }
 
 FName USimpleConstructionScript::GenerateNewComponentName(const UClass* ComponentClass, FName DesiredName ) const
 {
-	TArray<FName> CurrentNames;
+	TSet<FName> CurrentNames;
 	GenerateListOfExistingNames(CurrentNames);
 
 	FName NewName;
diff --git a/Engine/Source/Runtime/Engine/Private/SoundBase.cpp b/Engine/Source/Runtime/Engine/Private/SoundBase.cpp
--- a/Engine/Source/Runtime/Engine/Private/SoundBase.cpp
+++ b/Engine/Source/Runtime/Engine/Private/SoundBase.cpp
@@ -22,7 +22,7 @@ void USoundBase::PostInitProperties()
 
 	if (USoundBase::DefaultSoundClassObject == nullptr)
 	{
-		const FStringAssetReference DefaultSoundClassName = GetDefault<UAudioSettings>()->DefaultSoundClassName;
+		const FSoftObjectPath DefaultSoundClassName = GetDefault<UAudioSettings>()->DefaultSoundClassName;
 		if (DefaultSoundClassName.IsValid())
 		{
 			USoundBase::DefaultSoundClassObject = LoadObject<USoundClass>(nullptr, *DefaultSoundClassName.ToString());
@@ -32,7 +32,7 @@ void USoundBase::PostInitProperties()
 
 	if (USoundBase::DefaultSoundConcurrencyObject == nullptr)
 	{
-		const FStringAssetReference DefaultSoundConcurrencyName = GetDefault<UAudioSettings>()->DefaultSoundConcurrencyName;
+		const FSoftObjectPath DefaultSoundConcurrencyName = GetDefault<UAudioSettings>()->DefaultSoundConcurrencyName;
 		if (DefaultSoundConcurrencyName.IsValid())
 		{
 			USoundBase::DefaultSoundConcurrencyObject = LoadObject<USoundConcurrency>(nullptr, *DefaultSoundConcurrencyName.ToString());
diff --git a/Engine/Source/Runtime/Engine/Private/StreamableManager.cpp b/Engine/Source/Runtime/Engine/Private/StreamableManager.cpp
--- a/Engine/Source/Runtime/Engine/Private/StreamableManager.cpp
+++ b/Engine/Source/Runtime/Engine/Private/StreamableManager.cpp
@@ -167,29 +167,53 @@ EAsyncPackageState::Type FStreamableHandle::WaitUntilComplete(float Timeout)
 		return EAsyncPackageState::Complete;
 	}
 
-	if (bStalled)
+	// We need to recursively start any stalled handles
+	TArray<TSharedRef<FStreamableHandle>> HandlesToStart;
+
+	HandlesToStart.Add(AsShared());
+
+	for (int32 i = 0; i < HandlesToStart.Num(); i++)
 	{
-		// If we were stalled, start us now to avoid deadlocks
-		StartStalledHandle();
+		TSharedRef<FStreamableHandle> Handle = HandlesToStart[i];
+
+		if (Handle->IsStalled())
+		{
+			// If we were stalled, start us now to avoid deadlocks
+			UE_LOG(LogStreamableManager, Warning, TEXT("FStreamableHandle::WaitUntilComplete called on stalled handle %s, forcing load even though resources may not have been acquired yet"), *Handle->GetDebugName());
+			Handle->StartStalledHandle();
+		}
+
+		for (TSharedPtr<FStreamableHandle> ChildHandle : Handle->ChildHandles)
+		{
+			if (ChildHandle.IsValid())
+			{
+				HandlesToStart.Add(ChildHandle.ToSharedRef());
+			}
+		}
 	}
 
 	EAsyncPackageState::Type State = ProcessAsyncLoadingUntilComplete([this]() { return HasLoadCompleted(); }, Timeout);
 	
+	if (State == EAsyncPackageState::Complete)
+	{
+		ensureMsgf(HasLoadCompleted() || WasCanceled(), TEXT("WaitUntilComplete failed for streamable handle %s, async loading is done but handle is not complete"), *GetDebugName());
+	}
+
 	return State;
 }
 
-void FStreamableHandle::GetRequestedAssets(TArray<FStringAssetReference>& AssetList) const
+void FStreamableHandle::GetRequestedAssets(TArray<FSoftObjectPath>& AssetList) const
 {
 	AssetList = RequestedAssets;
 
 	// Check child handles
 	for (TSharedPtr<FStreamableHandle> ChildHandle : ChildHandles)
 	{
-		TArray<FStringAssetReference> ChildAssetList;
+		TArray<FSoftObjectPath> ChildAssetList;
 
 		ChildHandle->GetRequestedAssets(ChildAssetList);
 
-		for (const FStringAssetReference& ChildRef : ChildAssetList)
+		for (const FSoftObjectPath& ChildRef : ChildAssetList)
 		{
 			AssetList.AddUnique(ChildRef);
 		}
@@ -214,7 +238,7 @@ void FStreamableHandle::GetLoadedAssets(TArray<UObject *>& LoadedAssets) const
 {
 	if (HasLoadCompleted())
 	{
-		for (const FStringAssetReference& Ref : RequestedAssets)
+		for (const FSoftObjectPath& Ref : RequestedAssets)
 		{
 			// Try manager, should be faster and will handle redirects better
 			if (IsActive())
@@ -230,7 +254,7 @@ void FStreamableHandle::GetLoadedAssets(TArray<UObject *>& LoadedAssets) const
 		// Check child handles
 		for (TSharedPtr<FStreamableHandle> ChildHandle : ChildHandles)
 		{
-			for (const FStringAssetReference& Ref : ChildHandle->RequestedAssets)
+			for (const FSoftObjectPath& Ref : ChildHandle->RequestedAssets)
 			{
 				// Try manager, should be faster and will handle redirects better
 				if (IsActive())
@@ -310,7 +334,7 @@ void FStreamableHandle::CancelHandle()
 	ExecuteDelegate(CancelDelegate, SharedThis);
 
 	// Remove from referenced list
-	for (const FStringAssetReference& AssetRef : RequestedAssets)
+	for (const FSoftObjectPath& AssetRef : RequestedAssets)
 	{
 		OwningManager->RemoveReferencedAsset(AssetRef, SharedThis);
 	}
@@ -328,14 +352,18 @@ void FStreamableHandle::CancelHandle()
 
 	OwningManager = nullptr;
 
-	// Update any meta handles that are still active
-	for (TWeakPtr<FStreamableHandle> WeakHandle : ParentHandles)
+	if (ParentHandles.Num() > 0)
 	{
-		TSharedPtr<FStreamableHandle> Handle = WeakHandle.Pin();
-
-		if (Handle.IsValid())
+		// Update any meta handles that are still active. Copy the array first as elements may be removed from original while iterating
+		TArray<TWeakPtr<FStreamableHandle>> ParentHandlesCopy = ParentHandles;
+		for (TWeakPtr<FStreamableHandle> WeakHandle : ParentHandlesCopy)
 		{
-			Handle->UpdateCombinedHandle();
+			TSharedPtr<FStreamableHandle> Handle = WeakHandle.Pin();
+
+			if (Handle.IsValid())
+			{
+				Handle->UpdateCombinedHandle();
+			}
 		}
 	}
 }
@@ -357,7 +385,7 @@ void FStreamableHandle::ReleaseHandle()
 		TSharedRef<FStreamableHandle> SharedThis = AsShared();
 
 		// Remove from referenced list
-		for (const FStringAssetReference& AssetRef : RequestedAssets)
+		for (const FSoftObjectPath& AssetRef : RequestedAssets)
 		{
 			OwningManager->RemoveReferencedAsset(AssetRef, SharedThis);
 		}
@@ -459,10 +487,20 @@ void FStreamableHandle::UpdateCombinedHandle()
 	// If all our sub handles were canceled, cancel us. Otherwise complete us if at least one was completed and there are none in progress
 	if (bAllCanceled)
 	{
+		if (OwningManager)
+		{
+			OwningManager->PendingCombinedHandles.Remove(AsShared());
+		}
+
 		CancelHandle();
 	}
 	else if (bAllCompleted)
 	{
+		if (OwningManager)
+		{
+			OwningManager->PendingCombinedHandles.Remove(AsShared());
+		}
+
 		CompleteLoad();
 
 		if (bReleaseWhenLoaded)
@@ -488,7 +526,7 @@ void FStreamableHandle::CallUpdateDelegate()
 	}
 }
 
-void FStreamableHandle::AsyncLoadCallbackWrapper(const FName& PackageName, UPackage* Package, EAsyncLoadingResult::Type Result, FStringAssetReference TargetName)
+void FStreamableHandle::AsyncLoadCallbackWrapper(const FName& PackageName, UPackage* Package, EAsyncLoadingResult::Type Result, FSoftObjectPath TargetName)
 {
 	// Needed so we can bind with a shared pointer for safety
 	if (OwningManager)
@@ -519,7 +557,7 @@ void FStreamableHandle::ExecuteDelegate(const FStreamableDelegate& Delegate, TSh
 	}
 }
 
-/** Internal object, one of these per StringAssetReference managed by this system */
+/** Internal object, one of these per object paths managed by this system */
 struct FStreamable
 {
 	/** Hard Pointer to object */
@@ -606,7 +644,7 @@ FStreamableManager::~FStreamableManager()
 
 void FStreamableManager::OnPreGarbageCollect()
 {
-	TSet<FStringAssetReference> RedirectsToRemove;
+	TSet<FSoftObjectPath> RedirectsToRemove;
 
 	// Remove any streamables with no active handles, as GC may have freed them
 	for (TStreamableMap::TIterator It(StreamableItems); It; ++It)
@@ -666,7 +704,7 @@ void FStreamableManager::AddReferencedObjects(FReferenceCollector& Collector)
 	}
 }
 
-FStreamable* FStreamableManager::FindStreamable(const FStringAssetReference& Target) const
+FStreamable* FStreamableManager::FindStreamable(const FSoftObjectPath& Target) const
 {
 	FStreamable* Existing = StreamableItems.FindRef(Target);
 
@@ -678,12 +716,12 @@ FStreamable* FStreamableManager::FindStreamable(const FStringAssetReference& Tar
 	return Existing;
 }
 
-FStreamable* FStreamableManager::StreamInternal(const FStringAssetReference& InTargetName, TAsyncLoadPriority Priority, TSharedRef<FStreamableHandle> Handle)
+FStreamable* FStreamableManager::StreamInternal(const FSoftObjectPath& InTargetName, TAsyncLoadPriority Priority, TSharedRef<FStreamableHandle> Handle)
 {
 	check(IsInGameThread());
 	UE_LOG(LogStreamableManager, Verbose, TEXT("Asynchronous load %s"), *InTargetName.ToString());
 
-	FStringAssetReference TargetName = ResolveRedirects(InTargetName);
+	FSoftObjectPath TargetName = ResolveRedirects(InTargetName);
 	FStreamable* Existing = StreamableItems.FindRef(TargetName);
 	if (Existing)
 	{
@@ -736,7 +774,7 @@ FStreamable* FStreamableManager::StreamInternal(const FStringAssetReference& InT
 			if (Existing->Target)
 			{
 				UE_LOG(LogStreamableManager, Verbose, TEXT("     Static loaded %s"), *Existing->Target->GetFullName());
-				FStringAssetReference PossiblyNewName(Existing->Target->GetPathName());
+				FSoftObjectPath PossiblyNewName(Existing->Target->GetPathName());
 				if (PossiblyNewName != TargetName)
 				{
 					UE_LOG(LogStreamableManager, Verbose, TEXT("     Which redirected to %s"), *PossiblyNewName.ToString());
@@ -772,17 +810,17 @@ FStreamable* FStreamableManager::StreamInternal(const FStringAssetReference& InT
 	return Existing;
 }
 
-void FStreamableManager::SimpleAsyncLoad(const FStringAssetReference& Target, TAsyncLoadPriority Priority)
+void FStreamableManager::SimpleAsyncLoad(const FSoftObjectPath& Target, TAsyncLoadPriority Priority)
 {
 	RequestAsyncLoad(Target, FStreamableDelegate(), Priority, true);
 }
 
-UObject* FStreamableManager::SynchronousLoad(FStringAssetReference const& Target)
+UObject* FStreamableManager::SynchronousLoad(FSoftObjectPath const& Target)
 {
 	return LoadSynchronous(Target, true);
 }
 
-TSharedPtr<FStreamableHandle> FStreamableManager::RequestAsyncLoad(const TArray<FStringAssetReference>& TargetsToStream, FStreamableDelegate DelegateToCall, TAsyncLoadPriority Priority, bool bManageActiveHandle, bool bStartStalled, const FString& DebugName)
+TSharedPtr<FStreamableHandle> FStreamableManager::RequestAsyncLoad(const TArray<FSoftObjectPath>& TargetsToStream, FStreamableDelegate DelegateToCall, TAsyncLoadPriority Priority, bool bManageActiveHandle, bool bStartStalled, const FString& DebugName)
 {
 	// Schedule a new callback, this will get called when all related async loads are completed
 	TSharedRef<FStreamableHandle> NewRequest = MakeShareable(new FStreamableHandle());
@@ -796,7 +834,7 @@ TSharedPtr<FStreamableHandle> FStreamableManager::RequestAsyncLoad(const TArray<
 
 	for (int32 i = NewRequest->RequestedAssets.Num() - 1; i >= 0 ; i--)
 	{
-		FStringAssetReference& TargetName = NewRequest->RequestedAssets[i];
+		FSoftObjectPath& TargetName = NewRequest->RequestedAssets[i];
 		if (TargetName.IsNull())
 		{
 			// Remove null entries
@@ -822,7 +860,7 @@ TSharedPtr<FStreamableHandle> FStreamableManager::RequestAsyncLoad(const TArray<
 	{
 		FString RequestedSet;
 
-		for (const FStringAssetReference& Asset : TargetsToStream)
+		for (const FSoftObjectPath& Asset : TargetsToStream)
 		{
 			if (!RequestedSet.IsEmpty())
 			{
@@ -836,14 +874,14 @@ TSharedPtr<FStreamableHandle> FStreamableManager::RequestAsyncLoad(const TArray<
 	}
 
 	// Remove any duplicates
-	TSet<FStringAssetReference> TargetSet(NewRequest->RequestedAssets);
+	TSet<FSoftObjectPath> TargetSet(NewRequest->RequestedAssets);
 
 	if (TargetSet.Num() != NewRequest->RequestedAssets.Num())
 	{
 #if UE_BUILD_DEBUG
 		FString RequestedSet;
 
-		for (const FStringAssetReference& Asset : NewRequest->RequestedAssets)
+		for (const FSoftObjectPath& Asset : NewRequest->RequestedAssets)
 		{
 			if (!RequestedSet.IsEmpty())
 			{
@@ -876,22 +914,22 @@ TSharedPtr<FStreamableHandle> FStreamableManager::RequestAsyncLoad(const TArray<
 	return NewRequest;
 }
 
-TSharedPtr<FStreamableHandle> FStreamableManager::RequestAsyncLoad(const FStringAssetReference& TargetToStream, FStreamableDelegate DelegateToCall, TAsyncLoadPriority Priority, bool bManageActiveHandle, bool bStartStalled, const FString& DebugName)
+TSharedPtr<FStreamableHandle> FStreamableManager::RequestAsyncLoad(const FSoftObjectPath& TargetToStream, FStreamableDelegate DelegateToCall, TAsyncLoadPriority Priority, bool bManageActiveHandle, bool bStartStalled, const FString& DebugName)
 {
-	return RequestAsyncLoad(TArray<FStringAssetReference>{TargetToStream}, DelegateToCall, Priority, bManageActiveHandle, bStartStalled, DebugName);
+	return RequestAsyncLoad(TArray<FSoftObjectPath>{TargetToStream}, DelegateToCall, Priority, bManageActiveHandle, bStartStalled, DebugName);
 }
 
-TSharedPtr<FStreamableHandle>  FStreamableManager::RequestAsyncLoad(const TArray<FStringAssetReference>& TargetsToStream, TFunction<void()>&& Callback, TAsyncLoadPriority Priority, bool bManageActiveHandle, bool bStartStalled, const FString& DebugName)
+TSharedPtr<FStreamableHandle>  FStreamableManager::RequestAsyncLoad(const TArray<FSoftObjectPath>& TargetsToStream, TFunction<void()>&& Callback, TAsyncLoadPriority Priority, bool bManageActiveHandle, bool bStartStalled, const FString& DebugName)
 {
 	return RequestAsyncLoad(TargetsToStream, FStreamableDelegate::CreateLambda( MoveTemp( Callback ) ), Priority, bManageActiveHandle, bStartStalled, DebugName);
 }
 
-TSharedPtr<FStreamableHandle>  FStreamableManager::RequestAsyncLoad(const FStringAssetReference& TargetToStream, TFunction<void()>&& Callback, TAsyncLoadPriority Priority, bool bManageActiveHandle, bool bStartStalled, const FString& DebugName)
+TSharedPtr<FStreamableHandle>  FStreamableManager::RequestAsyncLoad(const FSoftObjectPath& TargetToStream, TFunction<void()>&& Callback, TAsyncLoadPriority Priority, bool bManageActiveHandle, bool bStartStalled, const FString& DebugName)
 {
 	return RequestAsyncLoad(TargetToStream, FStreamableDelegate::CreateLambda( MoveTemp( Callback ) ), Priority, bManageActiveHandle, bStartStalled, DebugName);
 }
 
-TSharedPtr<FStreamableHandle> FStreamableManager::RequestSyncLoad(const TArray<FStringAssetReference>& TargetsToStream, bool bManageActiveHandle, const FString& DebugName)
+TSharedPtr<FStreamableHandle> FStreamableManager::RequestSyncLoad(const TArray<FSoftObjectPath>& TargetsToStream, bool bManageActiveHandle, const FString& DebugName)
 {
 	// If in async loading thread or from callback always do sync as recursive tick is unsafe
 	// If in EDL always do sync as EDL internally avoids flushing
@@ -914,9 +952,9 @@ TSharedPtr<FStreamableHandle> FStreamableManager::RequestSyncLoad(const TArray<F
 	return Request;
 }
 
-TSharedPtr<FStreamableHandle> FStreamableManager::RequestSyncLoad(const FStringAssetReference& TargetToStream, bool bManageActiveHandle, const FString& DebugName)
+TSharedPtr<FStreamableHandle> FStreamableManager::RequestSyncLoad(const FSoftObjectPath& TargetToStream, bool bManageActiveHandle, const FString& DebugName)
 {
-	return RequestSyncLoad(TArray<FStringAssetReference>{TargetToStream}, bManageActiveHandle, DebugName);
+	return RequestSyncLoad(TArray<FSoftObjectPath>{TargetToStream}, bManageActiveHandle, DebugName);
 }
 
 void FStreamableManager::StartHandleRequests(TSharedRef<FStreamableHandle> Handle)
@@ -947,7 +985,7 @@ void FStreamableManager::StartHandleRequests(TSharedRef<FStreamableHandle> Handl
 	}
 }
 
-UObject* FStreamableManager::LoadSynchronous(const FStringAssetReference& Target, bool bManageActiveHandle, TSharedPtr<FStreamableHandle>* RequestHandlePointer)
+UObject* FStreamableManager::LoadSynchronous(const FSoftObjectPath& Target, bool bManageActiveHandle, TSharedPtr<FStreamableHandle>* RequestHandlePointer)
 {
 	TSharedPtr<FStreamableHandle> Request = RequestSyncLoad(Target, bManageActiveHandle, FString::Printf(TEXT("LoadSynchronous of %s"), *Target.ToString()));
 
@@ -971,7 +1009,7 @@ UObject* FStreamableManager::LoadSynchronous(const FStringAssetReference& Target
 	return nullptr;
 }
 
-void FStreamableManager::FindInMemory( FStringAssetReference& InOutTargetName, struct FStreamable* Existing )
+void FStreamableManager::FindInMemory( FSoftObjectPath& InOutTargetName, struct FStreamable* Existing )
 {
 	check(Existing);
 	check(!Existing->bAsyncLoadRequestOutstanding);
@@ -1000,7 +1038,7 @@ void FStreamableManager::FindInMemory( FStringAssetReference& InOutTargetName, s
 	}
 	if (Existing->Target)
 	{
-		FStringAssetReference PossiblyNewName(Existing->Target->GetPathName());
+		FSoftObjectPath PossiblyNewName(Existing->Target->GetPathName());
 		if (InOutTargetName != PossiblyNewName)
 		{
 			UE_LOG(LogStreamableManager, Verbose, TEXT("     Name changed to %s"), *PossiblyNewName.ToString());
@@ -1015,7 +1053,7 @@ void FStreamableManager::FindInMemory( FStringAssetReference& InOutTargetName, s
 	}
 }
 
-void FStreamableManager::AsyncLoadCallback(FStringAssetReference TargetName)
+void FStreamableManager::AsyncLoadCallback(FSoftObjectPath TargetName)
 {
 	FStreamable* Existing = FindStreamable(TargetName);
 
@@ -1053,7 +1091,7 @@ void FStreamableManager::AsyncLoadCallback(FStringAssetReference TargetName)
 	}
 }
 
-void FStreamableManager::CheckCompletedRequests(const FStringAssetReference& Target, struct FStreamable* Existing)
+void FStreamableManager::CheckCompletedRequests(const FSoftObjectPath& Target, struct FStreamable* Existing)
 {
 	static int32 RecursiveCount = 0;
 
@@ -1098,7 +1136,7 @@ void FStreamableManager::CheckCompletedRequests(const FStringAssetReference& Tar
 	RecursiveCount--;
 }
 
-void FStreamableManager::RemoveReferencedAsset(const FStringAssetReference& Target, TSharedRef<FStreamableHandle> Handle)
+void FStreamableManager::RemoveReferencedAsset(const FSoftObjectPath& Target, TSharedRef<FStreamableHandle> Handle)
 {
 	if (Target.IsNull())
 	{
@@ -1128,7 +1166,7 @@ void FStreamableManager::RemoveReferencedAsset(const FStringAssetReference& Targ
 	}
 }
 
-bool FStreamableManager::IsAsyncLoadComplete(const FStringAssetReference& Target) const
+bool FStreamableManager::IsAsyncLoadComplete(const FSoftObjectPath& Target) const
 {
 	// Failed loads count as success
 	check(IsInGameThread());
@@ -1136,7 +1174,7 @@ bool FStreamableManager::IsAsyncLoadComplete(const FStringAssetReference& Target
 	return !Existing || !Existing->bAsyncLoadRequestOutstanding;
 }
 
-UObject* FStreamableManager::GetStreamed(const FStringAssetReference& Target) const
+UObject* FStreamableManager::GetStreamed(const FSoftObjectPath& Target) const
 {
 	check(IsInGameThread());
 	FStreamable* Existing = FindStreamable(Target);
@@ -1147,7 +1185,7 @@ UObject* FStreamableManager::GetStreamed(const FStringAssetReference& Target) co
 	return nullptr;
 }
 
-void FStreamableManager::Unload(const FStringAssetReference& Target)
+void FStreamableManager::Unload(const FSoftObjectPath& Target)
 {
 	check(IsInGameThread());
 
@@ -1191,13 +1229,16 @@ TSharedPtr<FStreamableHandle> FStreamableManager::CreateCombinedHandle(const TAr
 		NewRequest->ChildHandles.Add(ChildHandle);
 	}
 
+	// Add to pending list so these handles don't free when not referenced
+	PendingCombinedHandles.Add(NewRequest);
+
 	// This may already be complete
 	NewRequest->UpdateCombinedHandle();
 
 	return NewRequest;
 }
 
-bool FStreamableManager::GetActiveHandles(const FStringAssetReference& Target, TArray<TSharedRef<FStreamableHandle>>& HandleList, bool bOnlyManagedHandles) const
+bool FStreamableManager::GetActiveHandles(const FSoftObjectPath& Target, TArray<TSharedRef<FStreamableHandle>>& HandleList, bool bOnlyManagedHandles) const
 {
 	check(IsInGameThread());
 	FStreamable* Existing = FindStreamable(Target);
@@ -1224,7 +1265,7 @@ bool FStreamableManager::GetActiveHandles(const FStringAssetReference& Target, T
 	return false;
 }
 
-FStringAssetReference FStreamableManager::ResolveRedirects(const FStringAssetReference& Target) const
+FSoftObjectPath FStreamableManager::ResolveRedirects(const FSoftObjectPath& Target) const
 {
 	FRedirectedPath const* Redir = StreamableRedirects.Find(Target);
 	if (Redir)
diff --git a/Engine/Source/Runtime/Engine/Private/Tests/AutomationCommon.cpp b/Engine/Source/Runtime/Engine/Private/Tests/AutomationCommon.cpp
--- a/Engine/Source/Runtime/Engine/Private/Tests/AutomationCommon.cpp
+++ b/Engine/Source/Runtime/Engine/Private/Tests/AutomationCommon.cpp
@@ -10,6 +10,7 @@
 #include "HardwareInfo.h"
 #include "UObject/UObjectHash.h"
 #include "UObject/UObjectIterator.h"
+#include "UObject/Package.h"
 #include "Kismet/GameplayStatics.h"
 #include "ContentStreaming.h"
 #include "Widgets/SWindow.h"
@@ -19,6 +20,7 @@
 #include "Scalability.h"
 #include "Matinee/MatineeActor.h"
 #include "IHeadMountedDisplay.h"
+#include "PackageName.h"
 
 #if (WITH_DEV_AUTOMATION_TESTS || WITH_PERF_AUTOMATION_TESTS)
 
@@ -82,7 +84,7 @@ namespace AutomationCommon
 	/** Gets a path used for automation testing (PNG sent to the AutomationTest folder) */
 	void GetScreenshotPath(const FString& TestName, FString& OutScreenshotName)
 	{
-		FString PathName = FPaths::AutomationDir() + TestName / FPlatformProperties::PlatformName();
+		FString PathName = FPaths::AutomationDir() + TestName / FPlatformProperties::IniPlatformName();
 		PathName = PathName + TEXT("_") + GetRenderDetailsString();
 
 		FPaths::MakePathRelativeTo(PathName, *FPaths::RootDir());
@@ -92,8 +94,6 @@ namespace AutomationCommon
 
 	FAutomationScreenshotData BuildScreenshotData(const FString& MapOrContext, const FString& TestName, int32 Width, int32 Height)
 	{
-		FString HardwareDetails = FHardwareInfo::GetHardwareDetailsString();
-
 		FAutomationScreenshotData Data;
 
 		Data.Name = TestName;
@@ -103,7 +103,7 @@ namespace AutomationCommon
 
 		Data.Width = Width;
 		Data.Height = Height;
-		Data.Platform = FPlatformProperties::PlatformName();
+		Data.Platform = FPlatformProperties::IniPlatformName();
 		Data.Rhi = FHardwareInfo::GetHardwareInfo(NAME_RHI);
 		GetFeatureLevelName(GMaxRHIFeatureLevel, Data.FeatureLevel);
 		Data.bIsStereo = GEngine->HMDDevice.IsValid() ? GEngine->HMDDevice->IsStereoEnabled() : false;
@@ -186,9 +186,16 @@ bool AutomationOpenMap(const FString& MapName)
 	{
 		UWorld* TestWorld = AutomationCommon::GetAnyGameWorld();
 
-		//Don't reload the map if it's already loaded. Check if the two are equal or one contains the other,
-		// since sometime one is a path and the other is the asset name, depending on whether we're in PIE.
-		if ( !TestWorld->GetMapName().Contains(MapName) && !MapName.Contains(TestWorld->GetMapName()) )
+		// Convert both to short names and strip PIE prefix
+		FString ShortMapName = FPackageName::GetShortName(MapName);
+		FString ShortWorldMapName = FPackageName::GetShortName(TestWorld->GetMapName());
+
+		if (TestWorld->GetOutermost()->PIEInstanceID != INDEX_NONE)
+		{
+			FString PIEPrefix = FString::Printf(PLAYWORLD_PACKAGE_PREFIX TEXT("_%d_"), TestWorld->GetOutermost()->PIEInstanceID);
+			ShortWorldMapName.ReplaceInline(*PIEPrefix, TEXT(""));
+		}
+		if (ShortMapName != ShortWorldMapName)
 		{
 			FString OpenCommand = FString::Printf(TEXT("Open %s"), *MapName);
 			GEngine->Exec(TestWorld, *OpenCommand);
diff --git a/Engine/Source/Runtime/Engine/Private/Tests/EngineAutomationTests.cpp b/Engine/Source/Runtime/Engine/Private/Tests/EngineAutomationTests.cpp
--- a/Engine/Source/Runtime/Engine/Private/Tests/EngineAutomationTests.cpp
+++ b/Engine/Source/Runtime/Engine/Private/Tests/EngineAutomationTests.cpp
@@ -76,6 +76,8 @@ namespace
 	}
 }
 
+#if PLATFORM_DESKTOP
+
 /**
  * SetRes Verification - Verify changing resolution works
  */
@@ -89,6 +91,7 @@ IMPLEMENT_SIMPLE_AUTOMATION_TEST(FSetResTest, "System.Windows.Set Resolution", E
  */
 bool FSetResTest::RunTest(const FString& Parameters)
 {
+
 	//Gets the default map that the game uses.
 	const UGameMapsSettings* GameMapsSettings = GetDefault<UGameMapsSettings>();
 	const FString& MapName = GameMapsSettings->GetGameDefaultMap();
@@ -106,10 +109,11 @@ bool FSetResTest::RunTest(const FString& Parameters)
 	ADD_LATENT_AUTOMATION_COMMAND(FExecStringLatentCommand(TEXT("setres 640x480")));
 	ADD_LATENT_AUTOMATION_COMMAND(FEngineWaitLatentCommand(2.0f));
 	ADD_LATENT_AUTOMATION_COMMAND(FExecStringLatentCommand(RestoreResolutionString));
-
 	return true;
 }
 
+#endif
+
 /**
  * Stats verification - Toggle various "stats" commands
  */
diff --git a/Engine/Source/Runtime/Engine/Private/UnrealEngine.cpp b/Engine/Source/Runtime/Engine/Private/UnrealEngine.cpp
--- a/Engine/Source/Runtime/Engine/Private/UnrealEngine.cpp
+++ b/Engine/Source/Runtime/Engine/Private/UnrealEngine.cpp
@@ -576,6 +576,7 @@ namespace
 		// Set the world type in the static map, so that UWorld::PostLoad can set the world type
 		const FName PIEPackageFName = FName(*PIEPackageName);
 		UWorld::WorldTypePreLoadMap.FindOrAdd( PIEPackageFName ) = WorldContext.WorldType;
+		FSoftObjectPath::AddPIEPackageName(PIEPackageFName);
 
 		uint32 LoadFlags = LOAD_None;
 		UPackage* NewPackage = CreatePackage(NULL, *PIEPackageName);
@@ -1760,7 +1761,7 @@ void LoadSpecialMaterial(const FString& MaterialName, UMaterial*& Material, bool
 
 
 template<typename ClassType>
-void LoadEngineClass(const FStringClassReference& ClassName, TSubclassOf<ClassType>& EngineClassRef)
+void LoadEngineClass(const FSoftClassPath& ClassName, TSubclassOf<ClassType>& EngineClassRef)
 {
 	if ( EngineClassRef == nullptr )
 	{
@@ -9715,6 +9716,13 @@ bool UEngine::MakeSureMapNameIsValid(FString& InOutMapName)
 		// If the user starts a multiplayer PIE session with an unsaved map,
 		// DoesPackageExist won't find it, so we have to try to find the package in memory as well.
 		bIsValid = (FindObjectFast<UPackage>(nullptr, FName(*TestMapName)) != nullptr) || FPackageName::DoesPackageExist(TestMapName);
+
+		// If we're not in the editor, then we always want to strip off the PIE prefix.  We might be connected to
+		// a PIE listen server.  In this case, we'll use our version of the map without the PIE prefix.
+		if (bIsValid && !GIsEditor)
+		{
+			InOutMapName = TestMapName;
+		}
 	}
 	else
 	{
@@ -10316,54 +10324,45 @@ bool UEngine::LoadMap( FWorldContext& WorldContext, FURL URL, class UPendingNetG
 	UPackage* WorldPackage = NULL;
 	UWorld*	NewWorld = NULL;
 	
-	// If this world is a PIE instance, we need to check if we are travelling to another PIE instance's world.
+	// If this world is a PIE instance, we need to check if we are traveling to another PIE instance's world.
 	// If we are, we need to set the PIERemapPrefix so that we load a copy of that world, instead of loading the
 	// PIE world directly.
 	if (!WorldContext.PIEPrefix.IsEmpty())
 	{
 		for (const FWorldContext& WorldContextFromList : WorldList)
 		{
-			// We want to ignore our own PIE instance so that we don't unnecessarily set the PIERemapPrefix if we are not travelling to
+			// We want to ignore our own PIE instance so that we don't unnecessarily set the PIERemapPrefix if we are not traveling to
 			// a server.
 			if (WorldContextFromList.World() != WorldContext.World())
 			{
 				if (!WorldContextFromList.PIEPrefix.IsEmpty() && URL.Map.Contains(WorldContextFromList.PIEPrefix))
 				{
-					WorldContext.PIERemapPrefix = WorldContextFromList.PIEPrefix;
-				}
-			}
-		}
-	}
+					FString SourceWorldPackage = UWorld::RemovePIEPrefix(URL.Map);
 
-	// Is this a PIE networking thing?
-	if (!WorldContext.PIERemapPrefix.IsEmpty())
-	{
-		if ( URL.Map.Contains(WorldContext.PIERemapPrefix) )
-		{
-			FString SourceWorldPackage = UWorld::RemovePIEPrefix(URL.Map);
+					// We are loading a new world for this context, so clear out PIE fixups that might be lingering.
+					// (note we dont want to do this in DuplicateWorldForPIE, since that is also called on streaming worlds.
+					GPlayInEditorID = WorldContext.PIEInstance;
+					FLazyObjectPtr::ResetPIEFixups();
 
-			// We are loading a new world for this context, so clear out PIE fixups that might be lingering.
-			// (note we dont want to do this in DuplicateWorldForPIE, since that is also called on streaming worlds.
-			GPlayInEditorID = WorldContext.PIEInstance;
-			FLazyObjectPtr::ResetPIEFixups();
+					NewWorld = UWorld::DuplicateWorldForPIE(SourceWorldPackage, nullptr);
+					if (NewWorld == nullptr)
+					{
+						NewWorld = CreatePIEWorldByLoadingFromPackage(WorldContext, SourceWorldPackage, WorldPackage);
+						if (NewWorld == nullptr)
+						{
+							Error = FString::Printf(TEXT("Failed to load package '%s' while in PIE"), *SourceWorldPackage);
+							return false;
+						}
+					}
+					else
+					{
+						WorldPackage = CastChecked<UPackage>(NewWorld->GetOuter());
+					}
 
-			NewWorld = UWorld::DuplicateWorldForPIE(SourceWorldPackage, NULL);
-			if (NewWorld == nullptr) 
-			{
-				NewWorld = CreatePIEWorldByLoadingFromPackage(WorldContext, SourceWorldPackage, WorldPackage);
-				if (NewWorld == nullptr)
-				{
-					Error = FString::Printf(TEXT("Failed to load package '%s' while in PIE"), *SourceWorldPackage);
-					return false;
+					NewWorld->StreamingLevelsPrefix = UWorld::BuildPIEPackagePrefix(WorldContext.PIEInstance);
+					GIsPlayInEditorWorld = true;
 				}
 			}
-			else
-			{
-				WorldPackage = CastChecked<UPackage>(NewWorld->GetOuter());
-			}
-
-			NewWorld->StreamingLevelsPrefix = UWorld::BuildPIEPackagePrefix(WorldContext.PIEInstance);
-			GIsPlayInEditorWorld = true;
 		}
 	}
 
@@ -10419,7 +10418,7 @@ bool UEngine::LoadMap( FWorldContext& WorldContext, FURL URL, class UPendingNetG
 		{
 			// If we are a PIE world and the world we just found is already initialized, then we're probably reloading the editor world and we
 			// need to create a PIE world by duplication instead
-			if (bPackageAlreadyLoaded)
+			if (bPackageAlreadyLoaded || NewWorld->WorldType == EWorldType::Editor)
 			{
 				if (WorldContext.PIEInstance == -1)
 				{
diff --git a/Engine/Source/Runtime/Engine/Private/UnrealExporter.cpp b/Engine/Source/Runtime/Engine/Private/UnrealExporter.cpp
--- a/Engine/Source/Runtime/Engine/Private/UnrealExporter.cpp
+++ b/Engine/Source/Runtime/Engine/Private/UnrealExporter.cpp
@@ -749,7 +749,6 @@ void ExportProperties
 				else
 				{
 					// If the array sizes are different, we will need to export each index so on import we maintain the size
-					bool bAnyElementDiffered = ArrayHelper.Num() != DiffArrayHelper.Num();
 					for (int32 DynamicArrayIndex = 0; DynamicArrayIndex < ArrayHelper.Num(); DynamicArrayIndex++)
 					{
 						FString	Value;
@@ -764,7 +763,6 @@ void ExportProperties
 						bool bExportItem = DiffData == NULL || (DiffData != SourceData && !InnerProp->Identical(SourceData, DiffData, ExportFlags));
 						if (bExportItem)
 						{
-							bAnyElementDiffered = true;
 							InnerProp->ExportTextItem(Value, SourceData, DiffData, Parent, ExportFlags, ExportRootScope);
 							if (ExportObjectProp)
 							{
@@ -802,14 +800,6 @@ void ExportProperties
 
 							Out.Logf(TEXT("%s%s(%i)=%s\r\n"), FCString::Spc(Indent), *Property->GetName(), DynamicArrayIndex, *Value);
 						}
-						// if some other element has already been determined to differ from the defaults, then export this item with no data so that
-						// the different array's size is maintained on import (this item will get the default values for that index, if any)
-						// however, if no elements of the array have changed, we still don't want to export anything
-						// so that the array size will also be taken from the defaults, which won't be the case if any element is exported
-						else if (bAnyElementDiffered)
-						{
-							Out.Logf(TEXT("%s%s(%i)=()\r\n"), FCString::Spc(Indent), *Property->GetName(), DynamicArrayIndex);
-						}
 					}
 					for (int32 DynamicArrayIndex = DiffArrayHelper.Num()-1; DynamicArrayIndex >= ArrayHelper.Num(); --DynamicArrayIndex)
 					{
diff --git a/Engine/Source/Runtime/Engine/Private/UserDefinedEnum.cpp b/Engine/Source/Runtime/Engine/Private/UserDefinedEnum.cpp
--- a/Engine/Source/Runtime/Engine/Private/UserDefinedEnum.cpp
+++ b/Engine/Source/Runtime/Engine/Private/UserDefinedEnum.cpp
@@ -4,6 +4,8 @@
 #include "EditorObjectVersion.h"
 #if WITH_EDITOR
 #include "Kismet2/EnumEditorUtils.h"
+#include "UObject/MetaData.h"
+#include "UObject/Package.h"
 #endif	// WITH_EDITOR
 
 UUserDefinedEnum::UUserDefinedEnum(const FObjectInitializer& ObjectInitializer)
@@ -91,6 +93,36 @@ void UUserDefinedEnum::PostEditUndo()
 	FEnumEditorUtils::PostEditUndo(this);
 }
 
+void UUserDefinedEnum::PostEditChangeProperty(FPropertyChangedEvent& PropertyChangedEvent)
+{
+	Super::PostEditChangeProperty(PropertyChangedEvent);
+
+	// Add the description to the tooltip
+ 	static const FName NAME_Tooltip(TEXT("Tooltip"));
+
+	UPackage* Package = GetOutermost();
+	check(Package);
+	UMetaData* PackageMetaData = Package->GetMetaData();
+
+	if (!EnumDescription.IsEmpty())
+	{
+		PackageMetaData->SetValue(this, NAME_Tooltip, *EnumDescription.ToString());
+	}
+	else
+	{
+		PackageMetaData->RemoveValue(this, NAME_Tooltip);
+	}
+}
+
+void UUserDefinedEnum::GetAssetRegistryTags(TArray<FAssetRegistryTag>& OutTags) const
+{
+	Super::GetAssetRegistryTags(OutTags);
+
+	FString DescriptionString;
+	FTextStringHelper::WriteToString(/*out*/ DescriptionString, EnumDescription);
+	OutTags.Emplace(GET_MEMBER_NAME_CHECKED(UUserDefinedEnum, EnumDescription), DescriptionString, FAssetRegistryTag::TT_Hidden);
+}
+
 FString UUserDefinedEnum::GenerateNewEnumeratorName()
 {
 	FString EnumNameString;
diff --git a/Engine/Source/Runtime/Engine/Private/World.cpp b/Engine/Source/Runtime/Engine/Private/World.cpp
--- a/Engine/Source/Runtime/Engine/Private/World.cpp
+++ b/Engine/Source/Runtime/Engine/Private/World.cpp
@@ -457,6 +457,9 @@ bool UWorld::Rename(const TCHAR* InName, UObject* NewOuter, ERenameFlags Flags)
 	bool bShouldFail = false;
 	FWorldDelegates::OnPreWorldRename.Broadcast(this, InName, NewOuter, Flags, bShouldFail);
 
+	// Make sure our legacy lightmap data is initialized so it can be renamed
+	PersistentLevel->HandleLegacyMapBuildData();
+
 	if (bShouldFail)
 	{
 		return false;
@@ -473,7 +476,7 @@ bool UWorld::Rename(const TCHAR* InName, UObject* NewOuter, ERenameFlags Flags)
 	// We're moving the world to a new package, rename UObjects which are map data but don't have the UWorld in their Outer chain.  There are two cases:
 	// 1) legacy lightmap textures and MapBuildData object will be in the same package as the UWorld
 	// 2) MapBuildData will be in a separate package with lightmap textures underneath it
-	if (PersistentLevel && PersistentLevel->MapBuildData)
+	if (PersistentLevel->MapBuildData)
 	{
 		FName NewMapBuildDataName = PersistentLevel->MapBuildData->GetFName();
 
@@ -2174,15 +2177,7 @@ void UWorld::AddToWorld( ULevel* Level, const FTransform& LevelTransform )
 				APlayerController* LocalPlayerController = It->GetPlayerController(this);
 				if (LocalPlayerController != NULL)
 				{
-					// Remap packagename for PIE networking before sending out to server
-					FName PackageName = Level->GetOutermost()->GetFName();
-					FString PackageNameStr = PackageName.ToString();
-					if (GEngine->NetworkRemapPath(LocalPlayerController->GetNetDriver(), PackageNameStr, false))
-					{
-						PackageName = FName(*PackageNameStr);
-					}
-
-					LocalPlayerController->ServerUpdateLevelVisibility(PackageName, true);
+					LocalPlayerController->ServerUpdateLevelVisibility(LocalPlayerController->NetworkRemapPath(Level->GetOutermost()->GetFName(), false), true);
 				}
 			}
 		}
@@ -2328,15 +2323,7 @@ void UWorld::RemoveFromWorld( ULevel* Level, bool bAllowIncrementalRemoval )
 					APlayerController* LocalPlayerController = It->GetPlayerController(this);
 					if (LocalPlayerController != NULL)
 					{
-						// Remap packagename for PIE networking before sending out to server
-						FName PackageName = Level->GetOutermost()->GetFName();
-						FString PackageNameStr = PackageName.ToString();
-						if (GEngine->NetworkRemapPath(LocalPlayerController->GetNetDriver(), PackageNameStr, false))
-						{
-							PackageName = FName(*PackageNameStr);
-						}
-
-						LocalPlayerController->ServerUpdateLevelVisibility(PackageName, false);
+						LocalPlayerController->ServerUpdateLevelVisibility(LocalPlayerController->NetworkRemapPath(Level->GetOutermost()->GetFName(), false), false);
 					}
 				}
 			}
@@ -2509,12 +2496,14 @@ int32 FLevelStreamingGCHelper::GetNumLevelsPendingPurge()
 
 void UWorld::RenameToPIEWorld(int32 PIEInstanceID)
 {
+#if WITH_EDITOR
 	UPackage* WorldPackage = GetOutermost();
 
 	WorldPackage->PIEInstanceID = PIEInstanceID;
 
 	const FString PIEPackageName = *UWorld::ConvertToPIEPackageName(WorldPackage->GetName(), PIEInstanceID);
 	WorldPackage->Rename(*PIEPackageName);
+	FSoftObjectPath::AddPIEPackageName(FName(*PIEPackageName));
 
 	StreamingLevelsPrefix = UWorld::BuildPIEPackagePrefix(PIEInstanceID);
 	
@@ -2527,6 +2516,9 @@ void UWorld::RenameToPIEWorld(int32 PIEInstanceID)
 	{
 		LevelStreaming->RenameForPIE(PIEInstanceID);
 	}
+
+	PersistentLevel->FixupForPIE(PIEInstanceID);
+#endif
 }
 
 FString UWorld::ConvertToPIEPackageName(const FString& PackageName, int32 PIEInstanceID)
@@ -2654,32 +2646,14 @@ UWorld* UWorld::DuplicateWorldForPIE(const FString& PackageName, UWorld* OwningW
 
 	FString PrefixedLevelName = ConvertToPIEPackageName(PackageName, PIEInstanceID);
 	const FName PrefixedLevelFName = FName(*PrefixedLevelName);
+	FSoftObjectPath::AddPIEPackageName(PrefixedLevelFName);
+
 	UWorld::WorldTypePreLoadMap.FindOrAdd(PrefixedLevelFName) = EWorldType::PIE;
 	UPackage* PIELevelPackage = CastChecked<UPackage>(CreatePackage(NULL,*PrefixedLevelName));
 	PIELevelPackage->SetPackageFlags(PKG_PlayInEditor);
 	PIELevelPackage->PIEInstanceID = PIEInstanceID;
 	PIELevelPackage->SetGuid( EditorLevelPackage->GetGuid() );
 
-	// Set up string asset reference fixups
-	TArray<FString> PackageNamesBeingDuplicatedForPIE;
-	PackageNamesBeingDuplicatedForPIE.Add(PrefixedLevelName);
-	if ( OwningWorld )
-	{
-		const FString PlayWorldMapName = ConvertToPIEPackageName(OwningWorld->GetOutermost()->GetName(), PIEInstanceID);
-
-		PackageNamesBeingDuplicatedForPIE.Add(PlayWorldMapName);
-		for (ULevelStreaming* StreamingLevel : OwningWorld->StreamingLevels)
-		{
-			if (StreamingLevel)
-			{
-				const FString StreamingLevelPIEName = UWorld::ConvertToPIEPackageName(StreamingLevel->GetWorldAssetPackageName(), PIEInstanceID);
-				PackageNamesBeingDuplicatedForPIE.AddUnique(StreamingLevelPIEName);
-			}
-		}
-	}
-
-	FStringAssetReference::SetPackageNamesBeingDuplicatedForPIE(PackageNamesBeingDuplicatedForPIE);
-
 	ULevel::StreamedLevelsOwningWorld.Add(PIELevelPackage->GetFName(), OwningWorld);
 	UWorld* PIELevelWorld = CastChecked<UWorld>(StaticDuplicateObject(EditorLevelWorld, PIELevelPackage, EditorLevelWorld->GetFName(), RF_AllFlags, nullptr, EDuplicateMode::PIE));
 
@@ -2693,9 +2667,6 @@ UWorld* UWorld::DuplicateWorldForPIE(const FString& PackageName, UWorld* OwningW
 	// Ensure the feature level matches the editor's, this is required as FeatureLevel is not a UPROPERTY and is not duplicated from EditorLevelWorld.
 	PIELevelWorld->FeatureLevel = EditorLevelWorld->FeatureLevel;
 
-	// Clean up string asset reference fixups
-	FStringAssetReference::ClearPackageNamesBeingDuplicatedForPIE();
-
 	// Clean up the world type list and owning world list now that PostLoad has occurred
 	UWorld::WorldTypePreLoadMap.Remove(PrefixedLevelFName);
 	ULevel::StreamedLevelsOwningWorld.Remove(PIELevelPackage->GetFName());
@@ -3401,20 +3372,13 @@ void UWorld::InitializeActorsForPlay(const FURL& InURL, bool bResetTime)
 		for (int32 LevelIndex = 1; LevelIndex < Levels.Num(); LevelIndex++)
 		{
 			ULevel*	SubLevel = Levels[LevelIndex];
-			// Remap packagename for PIE networking before sending out to server
-			FName PackageName = SubLevel->GetOutermost()->GetFName();
-			FString PackageNameStr = PackageName.ToString();
-			if (GEngine->NetworkRemapPath(GetNetDriver(), PackageNameStr, false))
-			{
-				PackageName = FName(*PackageNameStr);
-			}
 
 			for (FLocalPlayerIterator It(GEngine, this); It; ++It)
 			{
 				APlayerController* LocalPlayerController = It->GetPlayerController(this);
 				if (LocalPlayerController != NULL)
 				{
-					LocalPlayerController->ServerUpdateLevelVisibility(PackageName, SubLevel->bIsVisible);
+					LocalPlayerController->ServerUpdateLevelVisibility(LocalPlayerController->NetworkRemapPath(SubLevel->GetOutermost()->GetFName(), false), SubLevel->bIsVisible);
 				}
 			}
 		}
@@ -6309,6 +6273,8 @@ static ULevel* DuplicateLevelWithPrefix(ULevel* InLevel, int32 InstanceID )
 	NewPackage->FileName = OriginalPackage->FileName;
 	NewPackage->SetGuid( OriginalPackage->GetGuid() );
 
+	FSoftObjectPath::AddPIEPackageName(NewPackage->GetFName());
+
 	GPlayInEditorID = InstanceID;
 
 	// Create "vestigial" world for the persistent level - it's OwningWorld will still be the main world,
@@ -6318,11 +6284,6 @@ static ULevel* DuplicateLevelWithPrefix(ULevel* InLevel, int32 InstanceID )
 	NewWorld->WorldType = EWorldType::Game;
 	NewWorld->FeatureLevel = ERHIFeatureLevel::Num;
 
-	TArray<FString> PackageNamesBeingDuplicatedForPIE;
-	PackageNamesBeingDuplicatedForPIE.Add( PrefixedPackageName );
-
-	FStringAssetReference::SetPackageNamesBeingDuplicatedForPIE( PackageNamesBeingDuplicatedForPIE );
-
 	ULevel::StreamedLevelsOwningWorld.Add(NewPackage->GetFName(), OriginalOwningWorld);
 
 	FObjectDuplicationParameters Parameters( InLevel, NewWorld );
@@ -6337,7 +6298,6 @@ static ULevel* DuplicateLevelWithPrefix(ULevel* InLevel, int32 InstanceID )
 	ULevel* NewLevel = CastChecked<ULevel>( StaticDuplicateObjectEx( Parameters ) );
 
 	ULevel::StreamedLevelsOwningWorld.Remove(NewPackage->GetFName());
-	FStringAssetReference::ClearPackageNamesBeingDuplicatedForPIE();
 
 	// Fixup model components. The index buffers have been created for the components in the source world and the order
 	// in which components were post-loaded matters. So don't try to guarantee a particular order here, just copy the
diff --git a/Engine/Source/Runtime/Engine/Private/WorldComposition.cpp b/Engine/Source/Runtime/Engine/Private/WorldComposition.cpp
--- a/Engine/Source/Runtime/Engine/Private/WorldComposition.cpp
+++ b/Engine/Source/Runtime/Engine/Private/WorldComposition.cpp
@@ -93,10 +93,12 @@ void UWorldComposition::FixupForPIE(int32 PIEInstanceID)
 		
 		FString PIEPackageName = UWorld::ConvertToPIEPackageName(Tile.PackageName.ToString(), PIEInstanceID);
 		Tile.PackageName = FName(*PIEPackageName);
+		FSoftObjectPath::AddPIEPackageName(Tile.PackageName);
 		for (FName& LODPackageName : Tile.LODPackageNames)
 		{
 			FString PIELODPackageName = UWorld::ConvertToPIEPackageName(LODPackageName.ToString(), PIEInstanceID);
 			LODPackageName = FName(*PIELODPackageName);
+			FSoftObjectPath::AddPIEPackageName(LODPackageName);
 		}
 	}
 }
diff --git a/Engine/Source/Runtime/Engine/Public/Animation/NodeMappingContainer.h b/Engine/Source/Runtime/Engine/Public/Animation/NodeMappingContainer.h
--- a/Engine/Source/Runtime/Engine/Public/Animation/NodeMappingContainer.h
+++ b/Engine/Source/Runtime/Engine/Public/Animation/NodeMappingContainer.h
@@ -5,7 +5,7 @@
 #include "CoreMinimal.h"
 #include "UObject/ObjectMacros.h"
 #include "UObject/Object.h"
-#include "UObject/AssetPtr.h"
+#include "UObject/SoftObjectPtr.h"
 #include "NodeMappingProviderInterface.h"
 #include "NodeMappingContainer.generated.h"
 
@@ -37,7 +37,7 @@ private:
 	TMap<FName, FNodeMap>	NodeMapping;
 
 	UPROPERTY(EditAnywhere, Category = Mapping)
-	TAssetPtr<UBlueprint>	SourceAsset; 
+	TSoftObjectPtr<UBlueprint>	SourceAsset; 
 
 public:
 #if WITH_EDITOR
diff --git a/Engine/Source/Runtime/Engine/Public/EngineUtils.h b/Engine/Source/Runtime/Engine/Public/EngineUtils.h
--- a/Engine/Source/Runtime/Engine/Public/EngineUtils.h
+++ b/Engine/Source/Runtime/Engine/Public/EngineUtils.h
@@ -223,11 +223,22 @@ private:
 	}
 };
 
+/** Type enum, used to represent the special End iterator */
 enum class EActorIteratorType
 {
 	End
 };
 
+/** Iteration flags, specifies which types of levels and actors should be iterated */
+enum class EActorIteratorFlags
+{
+	AllActors			= 0x00000000, // No flags, iterate all actors
+	SkipPendingKill		= 0x00000001, // Skip pending kill actors
+	OnlySelectedActors	= 0x00000002, // Only iterate actors that are selected
+	OnlyActiveLevels	= 0x00000004, // Only iterate active levels
+};
+ENUM_CLASS_FLAGS(EActorIteratorFlags);
+
 /**
  * Template class used to filter actors by certain characteristics
  */
@@ -332,10 +343,12 @@ protected:
 	 * Hide the constructors as construction on this class should only be done by subclasses
 	 */
 	explicit TActorIteratorBase(EActorIteratorType)
+		: Flags(EActorIteratorFlags::AllActors)
 	{
 	}
 
-	TActorIteratorBase(UWorld* InWorld, TSubclassOf<AActor> InClass)
+	TActorIteratorBase(UWorld* InWorld, TSubclassOf<AActor> InClass, EActorIteratorFlags InFlags)
+		: Flags(InFlags)
 	{
 		State.Emplace(InWorld, InClass);
 	}
@@ -347,9 +360,19 @@ protected:
 	 * @param	Actor	Actor to check
 	 * @return	true
 	 */
-	FORCEINLINE static bool IsActorSuitable(AActor* Actor)
+	FORCEINLINE bool IsActorSuitable(AActor* Actor) const
 	{
-		return !Actor->IsPendingKill();
+		if (EnumHasAnyFlags(Flags, EActorIteratorFlags::SkipPendingKill) && Actor->IsPendingKill())
+		{
+			return false;
+		}
+
+		if (EnumHasAnyFlags(Flags, EActorIteratorFlags::OnlySelectedActors) && !Actor->IsSelected())
+		{
+			return false;
+		}
+
+		return true;
 	}
 
 	/**
@@ -359,12 +382,31 @@ protected:
 	 * @param Level the level to check for iteration
 	 * @return true if the level can be iterated, false otherwise
 	 */
-	FORCEINLINE static bool CanIterateLevel(ULevel* Level)
+	FORCEINLINE bool CanIterateLevel(ULevel* Level) const
 	{
+		if (EnumHasAnyFlags(Flags, EActorIteratorFlags::OnlyActiveLevels))
+		{
+			const bool bIsLevelVisibleOrAssociating = Level->bIsVisible || Level->bIsAssociatingLevel;
+
+			// Only allow iteration of Level if it's in the currently active level collection of the world, or is a static level.
+			const FLevelCollection* const ActorLevelCollection = Level->GetCachedLevelCollection();
+			const FLevelCollection* const ActiveLevelCollection = Level->OwningWorld ? Level->OwningWorld->GetActiveLevelCollection() : nullptr;
+
+			// If the world's active level collection is null, we can't apply any meaningful filter,
+			// so just allow iteration in this case.
+			const bool bIsCurrentLevelCollectionTicking = !ActiveLevelCollection || (ActorLevelCollection == ActiveLevelCollection);
+
+			const bool bIsLevelCollectionNullOrStatic = !ActorLevelCollection || ActorLevelCollection->GetType() == ELevelCollectionType::StaticLevels;
+			const bool bShouldIterateLevelCollection = bIsCurrentLevelCollectionTicking || bIsLevelCollectionNullOrStatic;
+
+			return bIsLevelVisibleOrAssociating && bShouldIterateLevelCollection;
+		}
+
 		return true;
 	}
 
 private:
+	EActorIteratorFlags Flags;
 	TOptional<FActorIteratorState> State;
 
 	friend bool operator==(const TActorIteratorBase& Lhs, const TActorIteratorBase& Rhs) { check(!Rhs.State); return  !Lhs; }
@@ -386,8 +428,8 @@ public:
 	 *
 	 * @param  InWorld  The world whose actors are to be iterated over.
 	 */
-	explicit FActorIterator(UWorld* InWorld)
-		: Super(InWorld, AActor::StaticClass())
+	explicit FActorIterator(UWorld* InWorld, EActorIteratorFlags InFlags = EActorIteratorFlags::OnlyActiveLevels | EActorIteratorFlags::SkipPendingKill)
+		: Super(InWorld, AActor::StaticClass(), InFlags)
 	{
 		++(*this);
 	}
@@ -399,31 +441,6 @@ public:
 		: Super(EActorIteratorType::End)
 	{
 	}
-
-private:
-	/**
-	 * Used to examine whether this level is valid for iteration or not
-	 *
-	 * @param Level the level to check for iteration
-	 * @return true if the level can be iterated, false otherwise
-	 */
-	static bool CanIterateLevel(ULevel* Level)
-	{
-		const bool bIsLevelVisibleOrAssociating = Level->bIsVisible || Level->bIsAssociatingLevel;
-
-		// Only allow iteration of Level if it's in the currently active level collection of the world, or is a static level.
-		const FLevelCollection* const ActorLevelCollection = Level->GetCachedLevelCollection();
-		const FLevelCollection* const ActiveLevelCollection = Level->OwningWorld ? Level->OwningWorld->GetActiveLevelCollection() : nullptr;
-
-		// If the world's active level collection is null, we can't apply any meaningful filter,
-		// so just allow iteration in this case.
-		const bool bIsCurrentLevelCollectionTicking = !ActiveLevelCollection || (ActorLevelCollection == ActiveLevelCollection);
-		
-		const bool bIsLevelCollectionNullOrStatic = !ActorLevelCollection || ActorLevelCollection->GetType() == ELevelCollectionType::StaticLevels;
-		const bool bShouldIterateLevelCollection = bIsCurrentLevelCollectionTicking || bIsLevelCollectionNullOrStatic;
-
-		return bIsLevelVisibleOrAssociating && bShouldIterateLevelCollection;
-	}
 };
 
 /**
@@ -438,15 +455,17 @@ public:
 	 *
 	 * @param  InWorld  The world whose actors are to be iterated over.
 	 */
-	explicit FActorRange(UWorld* InWorld)
-		: World(InWorld)
+	explicit FActorRange(UWorld* InWorld, EActorIteratorFlags InFlags = EActorIteratorFlags::OnlyActiveLevels | EActorIteratorFlags::SkipPendingKill)
+		: Flags(InFlags)
+		, World(InWorld)
 	{
 	}
 
 private:
-	UWorld* World;
+	EActorIteratorFlags	Flags;
+	UWorld*				World;
 
-	friend FActorIterator begin(const FActorRange& Range) { return FActorIterator(Range.World); }
+	friend FActorIterator begin(const FActorRange& Range) { return FActorIterator(Range.World, Range.Flags); }
 	friend FActorIterator end  (const FActorRange& Range) { return FActorIterator(EActorIteratorType::End); }
 };
 
@@ -466,8 +485,8 @@ public:
 	 * @param  InWorld  The world whose actors are to be iterated over.
 	 * @param  InClass  The subclass of actors to be iterated over.
 	 */
-	explicit TActorIterator( UWorld* InWorld, TSubclassOf<ActorType> InClass = ActorType::StaticClass() )
-		: Super( InWorld, InClass )
+	explicit TActorIterator(UWorld* InWorld, TSubclassOf<ActorType> InClass = ActorType::StaticClass(), EActorIteratorFlags InFlags = EActorIteratorFlags::OnlyActiveLevels | EActorIteratorFlags::SkipPendingKill)
+		: Super(InWorld, InClass, InFlags)
 	{
 		++(*this);
 	}
@@ -499,18 +518,6 @@ public:
 	{
 		return **this;
 	}
-
-private:
-	/**
-	 * Used to examine whether this level is valid for iteration or not
-	 *
-	 * @param Level the level to check for iteration
-	 * @return true if the level can be iterated, false otherwise
-	 */
-	static bool CanIterateLevel(ULevel* Level)
-	{
-		return Level->bIsVisible || Level->bIsAssociatingLevel;
-	}
 };
 
 /**
@@ -526,22 +533,24 @@ public:
 	 * @param  InWorld  The world whose actors are to be iterated over.
 	 * @param  InClass  The subclass of actors to be iterated over.
 	 */
-	explicit TActorRange(UWorld* InWorld, TSubclassOf<ActorType> InClass = ActorType::StaticClass())
-		: World(InWorld)
+	explicit TActorRange(UWorld* InWorld, TSubclassOf<ActorType> InClass = ActorType::StaticClass(), EActorIteratorFlags InFlags = EActorIteratorFlags::OnlyActiveLevels | EActorIteratorFlags::SkipPendingKill)
+		: Flags(InFlags)
+		, World(InWorld)
 		, Class(InClass)
 	{
 	}
 
 private:
-	UWorld*                World;
-	TSubclassOf<ActorType> Class;
+	EActorIteratorFlags		Flags;
+	UWorld*					World;
+	TSubclassOf<ActorType>	Class;
 
-	friend TActorIterator<ActorType> begin(const TActorRange& Range) { return TActorIterator<ActorType>(Range.World, Range.Class); }
-	friend TActorIterator<ActorType> end  (const TActorRange& Range) { return TActorIterator<ActorType>(EActorIteratorType::End); }
+	friend TActorIterator<ActorType> begin(const TActorRange& Range) { return TActorIterator<ActorType>(Range.World, Range.Class, Range.Flags); }
+	friend TActorIterator<ActorType> end(const TActorRange& Range) { return TActorIterator<ActorType>(EActorIteratorType::End); }
 };
 
 /**
- * Selected actor iterator
+ * Selected actor iterator, this is for ease of use but the same can be done by adding EActorIteratorFlags::OnlySelectedActors to 
  */
 class FSelectedActorIterator : public TActorIteratorBase<FSelectedActorIterator>
 {
@@ -554,8 +563,8 @@ public:
 	 *
 	 * @param  InWorld  The world whose actors are to be iterated over.
 	 */
-	explicit FSelectedActorIterator( UWorld* InWorld )
-		:Super( InWorld, AActor::StaticClass() )
+	explicit FSelectedActorIterator(UWorld* InWorld)
+		: Super(InWorld, AActor::StaticClass(), EActorIteratorFlags::SkipPendingKill | EActorIteratorFlags::OnlySelectedActors)
 	{
 		++(*this);
 	}
@@ -567,18 +576,6 @@ public:
 		: Super(EActorIteratorType::End)
 	{
 	}
-
-protected:
-	/**
-	 * Determines if the actor should be returned during iteration or not
-	 *
-	 * @param	Actor	Actor to check
-	 * @return	true if actor is not null and is selected, false otherwise
-	 */
-	FORCEINLINE static bool IsActorSuitable(AActor* Actor)
-	{
-		return !Actor->IsPendingKill() && Actor->IsSelected();
-	}
 };
 
 /**
diff --git a/Engine/Source/Runtime/Engine/Public/GeneratedCodeHelpers.h b/Engine/Source/Runtime/Engine/Public/GeneratedCodeHelpers.h
--- a/Engine/Source/Runtime/Engine/Public/GeneratedCodeHelpers.h
+++ b/Engine/Source/Runtime/Engine/Public/GeneratedCodeHelpers.h
@@ -490,19 +490,6 @@ public:
 		return UKismetSystemLibrary::Generic_SetStructurePropertyByName(Object, PropertyName, &Value);
 	}
 
-	// Replacements for CustomThunk functions from UBlueprintFunctionLibrary
-	static FStringAssetReference MakeStringAssetReference(const FString& AssetLongPathname)
-	{
-		FStringAssetReference Ref(AssetLongPathname);
-		if (!AssetLongPathname.IsEmpty() && !Ref.IsValid())
-		{
-			ExecutionMessage(*FString::Printf(TEXT("Asset path \"%s\" not valid. Only long path name is allowed."), *AssetLongPathname), ELogVerbosity::Error);
-			return FStringAssetReference();
-		}
-
-		return Ref;
-	}
-
 	// Replacements for CustomThunk functions from KismetMathLibrary
 	static float Divide_FloatFloat(float A, float B)
 	{
diff --git a/Engine/Source/Runtime/EngineSettings/Classes/GameMapsSettings.h b/Engine/Source/Runtime/EngineSettings/Classes/GameMapsSettings.h
--- a/Engine/Source/Runtime/EngineSettings/Classes/GameMapsSettings.h
+++ b/Engine/Source/Runtime/EngineSettings/Classes/GameMapsSettings.h
@@ -5,8 +5,7 @@
 #include "CoreMinimal.h"
 #include "UObject/ObjectMacros.h"
 #include "UObject/Object.h"
-#include "Misc/StringAssetReference.h"
-#include "Misc/StringClassReference.h"
+#include "UObject/SoftObjectPath.h"
 #include "GameMapsSettings.generated.h"
 
 /** Ways the screen can be split with two players. */
@@ -44,7 +43,7 @@ struct FGameModeName
 
 	/** GameMode class to load */
 	UPROPERTY(EditAnywhere, Category = DefaultModes, meta = (MetaClass = "GameModeBase"))
-	FStringClassReference GameMode;
+	FSoftClassPath GameMode;
 };
 
 UCLASS(config=Engine, defaultconfig)
@@ -59,7 +58,7 @@ class ENGINESETTINGS_API UGameMapsSettings
 	 *
 	 * @return the default map specified in the settings
 	 */
-	static const FString GetGameDefaultMap( );
+	static FString GetGameDefaultMap( );
 
 	/**
 	 * Get the global default game type specified in the configuration
@@ -67,7 +66,7 @@ class ENGINESETTINGS_API UGameMapsSettings
 	 * 
 	 * @return the proper global default game type
 	 */
-	static const FString& GetGlobalDefaultGameMode( );
+	static FString GetGlobalDefaultGameMode( );
 
 	/**
 	 * Searches the GameModeClassAliases list for a named game mode, if not found will return passed in string
@@ -104,7 +103,7 @@ public:
 
 	/** If set, this map will be loaded when the Editor starts up. */
 	UPROPERTY(config, EditAnywhere, Category=DefaultMaps, meta=(AllowedClasses="World"))
-	FStringAssetReference EditorStartupMap;
+	FSoftObjectPath EditorStartupMap;
 
 	/** The default options that will be appended to a map being loaded. */
 	UPROPERTY(config, EditAnywhere, Category=DefaultMaps, AdvancedDisplay)
@@ -112,7 +111,7 @@ public:
 
 	/** The map loaded when transition from one map to another. */
 	UPROPERTY(config, EditAnywhere, Category=DefaultMaps, AdvancedDisplay, meta=(AllowedClasses="World"))
-	FStringAssetReference TransitionMap;
+	FSoftObjectPath TransitionMap;
 
 	/** Whether the screen should be split or not when multiple local players are present */
 	UPROPERTY(config, EditAnywhere, Category=LocalMultiplayer)
@@ -135,28 +134,28 @@ public:
 
 	/** The class to use when instantiating the transient GameInstance class */
 	UPROPERTY(config, noclear, EditAnywhere, Category=GameInstance, meta=(MetaClass="GameInstance"))
-	FStringClassReference GameInstanceClass;
+	FSoftClassPath GameInstanceClass;
 
 private:
 
 	/** The map that will be loaded by default when no other map is loaded. */
 	UPROPERTY(config, EditAnywhere, Category=DefaultMaps, meta=(AllowedClasses="World"))
-	FStringAssetReference GameDefaultMap;
+	FSoftObjectPath GameDefaultMap;
 
 	/** The map that will be loaded by default when no other map is loaded (DEDICATED SERVER). */
 	UPROPERTY(config, EditAnywhere, Category=DefaultMaps, AdvancedDisplay, meta=(AllowedClasses="World"))
-	FStringAssetReference ServerDefaultMap;
+	FSoftObjectPath ServerDefaultMap;
 
 	/** GameMode to use if not specified in any other way. (e.g. per-map DefaultGameMode or on the URL). */
 	UPROPERTY(config, noclear, EditAnywhere, Category=DefaultModes, meta=(MetaClass="GameModeBase", DisplayName="Default GameMode"))
-	FStringClassReference GlobalDefaultGameMode;
+	FSoftClassPath GlobalDefaultGameMode;
 
 	/**
 	 * GameMode to use if not specified in any other way. (e.g. per-map DefaultGameMode or on the URL) (DEDICATED SERVERS)
 	 * If not set, the GlobalDefaultGameMode value will be used.
 	 */
 	UPROPERTY(config, EditAnywhere, Category=DefaultModes, meta=(MetaClass="GameModeBase"), AdvancedDisplay)
-	FStringClassReference GlobalDefaultServerGameMode;
+	FSoftClassPath GlobalDefaultServerGameMode;
 
 	/** Overrides the GameMode to use when loading a map that starts with a specific prefix */
 	UPROPERTY(config, EditAnywhere, Category = DefaultModes, AdvancedDisplay)
diff --git a/Engine/Source/Runtime/EngineSettings/Private/EngineSettingsModule.cpp b/Engine/Source/Runtime/EngineSettings/Private/EngineSettingsModule.cpp
--- a/Engine/Source/Runtime/EngineSettings/Private/EngineSettingsModule.cpp
+++ b/Engine/Source/Runtime/EngineSettings/Private/EngineSettingsModule.cpp
@@ -86,7 +86,7 @@ UHudSettings::UHudSettings( const FObjectInitializer& ObjectInitializer )
 /* Static functions
  *****************************************************************************/
 
-const FString UGameMapsSettings::GetGameDefaultMap( )
+FString UGameMapsSettings::GetGameDefaultMap( )
 {
 	return IsRunningDedicatedServer()
 		? GetDefault<UGameMapsSettings>()->ServerDefaultMap.GetLongPackageName()
@@ -94,7 +94,7 @@ const FString UGameMapsSettings::GetGameDefaultMap( )
 }
 
 
-const FString& UGameMapsSettings::GetGlobalDefaultGameMode( )
+FString UGameMapsSettings::GetGlobalDefaultGameMode( )
 {
 	UGameMapsSettings* GameMapsSettings = Cast<UGameMapsSettings>(UGameMapsSettings::StaticClass()->GetDefaultObject());
 
@@ -212,7 +212,7 @@ void UGameMapsSettings::SetGlobalDefaultGameMode( const FString& NewGameMode )
 }
 
 // Backwards compat for map strings
-void FixMapAssetRef(FStringAssetReference& MapAssetReference)
+void FixMapAssetRef(FSoftObjectPath& MapAssetReference)
 {
 	const FString AssetRefStr = MapAssetReference.ToString();
 	int32 DummyIndex;
diff --git a/Engine/Source/Runtime/Foliage/Private/FoliageInstanceBase.cpp b/Engine/Source/Runtime/Foliage/Private/FoliageInstanceBase.cpp
--- a/Engine/Source/Runtime/Foliage/Private/FoliageInstanceBase.cpp
+++ b/Engine/Source/Runtime/Foliage/Private/FoliageInstanceBase.cpp
@@ -128,7 +128,7 @@ FFoliageInstanceBaseId FFoliageInstanceBaseCache::AddInstanceBaseId(UActorCompon
 				UWorld* ComponentWorld = Cast<UWorld>(ComponentLevel->GetOuter());
 				if (ComponentWorld)
 				{
-					auto WorldKey = TAssetPtr<UWorld>(ComponentWorld);
+					auto WorldKey = TSoftObjectPtr<UWorld>(ComponentWorld);
 					InstanceBaseLevelMap.FindOrAdd(WorldKey).Add(BaseInfo.BasePtr);
 				}
 			}
@@ -226,7 +226,7 @@ void FFoliageInstanceBaseCache::CompactInstanceBaseCache(AInstancedFoliageActor*
 		// Check sub-levels
 		if (!bExists)
 		{
-			const FName PackageName = FName(*FPackageName::ObjectPathToPackageName(WorldAsset.ToStringReference().ToString()));
+			const FName PackageName = FName(*FPackageName::ObjectPathToPackageName(WorldAsset.ToString()));
 			if (World->WorldComposition)
 			{
 				bExists = World->WorldComposition->DoesTileExists(PackageName);
diff --git a/Engine/Source/Runtime/Foliage/Private/InstancedFoliage.cpp b/Engine/Source/Runtime/Foliage/Private/InstancedFoliage.cpp
--- a/Engine/Source/Runtime/Foliage/Private/InstancedFoliage.cpp
+++ b/Engine/Source/Runtime/Foliage/Private/InstancedFoliage.cpp
@@ -193,7 +193,7 @@ static void ConvertDeprecatedFoliageMeshes(
 	check(IFA->InstanceBaseCache.InstanceBaseLevelMap.Num() <= 1);
 	// populate WorldAsset->BasePtr map
 	IFA->InstanceBaseCache.InstanceBaseLevelMap.Empty();
-	auto& BaseList = IFA->InstanceBaseCache.InstanceBaseLevelMap.Add(TAssetPtr<UWorld>(Cast<UWorld>(IFA->GetLevel()->GetOuter())));
+	auto& BaseList = IFA->InstanceBaseCache.InstanceBaseLevelMap.Add(TSoftObjectPtr<UWorld>(Cast<UWorld>(IFA->GetLevel()->GetOuter())));
 	for (auto& BaseInfoPair : IFA->InstanceBaseCache.InstanceBaseMap)
 	{
 		BaseList.Add(BaseInfoPair.Value.BasePtr);
diff --git a/Engine/Source/Runtime/Foliage/Public/FoliageInstanceBase.h b/Engine/Source/Runtime/Foliage/Public/FoliageInstanceBase.h
--- a/Engine/Source/Runtime/Foliage/Public/FoliageInstanceBase.h
+++ b/Engine/Source/Runtime/Foliage/Public/FoliageInstanceBase.h
@@ -3,7 +3,7 @@
 
 #include "CoreMinimal.h"
 #include "UObject/LazyObjectPtr.h"
-#include "UObject/AssetPtr.h"
+#include "UObject/SoftObjectPtr.h"
 
 class AInstancedFoliageActor;
 class UActorComponent;
@@ -82,7 +82,7 @@ struct FFoliageInstanceBaseCache
 	
 	// Map for detecting removed bases
 	// serialized
-	TMap<TAssetPtr<UWorld>, TArray<FFoliageInstanceBasePtr>> InstanceBaseLevelMap;
+	TMap<TSoftObjectPtr<UWorld>, TArray<FFoliageInstanceBasePtr>> InstanceBaseLevelMap;
 };
 
 #endif// WITH_EDITORONLY_DATA
diff --git a/Engine/Source/Runtime/GameplayTags/Classes/GameplayTagsSettings.h b/Engine/Source/Runtime/GameplayTags/Classes/GameplayTagsSettings.h
--- a/Engine/Source/Runtime/GameplayTags/Classes/GameplayTagsSettings.h
+++ b/Engine/Source/Runtime/GameplayTags/Classes/GameplayTagsSettings.h
@@ -5,7 +5,7 @@
 #include "CoreMinimal.h"
 #include "UObject/ObjectMacros.h"
 #include "UObject/Object.h"
-#include "Misc/StringAssetReference.h"
+#include "UObject/SoftObjectPath.h"
 #include "GameplayTagsManager.h"
 #include "GameplayTagsSettings.generated.h"
 
@@ -99,7 +99,7 @@ class GAMEPLAYTAGS_API UGameplayTagsSettings : public UGameplayTagsList
 
 	/** List of data tables to load tags from */
 	UPROPERTY(config, EditAnywhere, Category = GameplayTags, meta = (AllowedClasses = "DataTable"))
-	TArray<FStringAssetReference> GameplayTagTableList;
+	TArray<FSoftObjectPath> GameplayTagTableList;
 
 	/** List of active tag redirects */
 	UPROPERTY(config, EditAnywhere, Category = GameplayTags)
diff --git a/Engine/Source/Runtime/GameplayTags/Private/GameplayTagsManager.cpp b/Engine/Source/Runtime/GameplayTags/Private/GameplayTagsManager.cpp
--- a/Engine/Source/Runtime/GameplayTags/Private/GameplayTagsManager.cpp
+++ b/Engine/Source/Runtime/GameplayTags/Private/GameplayTagsManager.cpp
@@ -48,7 +48,7 @@ void UGameplayTagsManager::LoadGameplayTagTables()
 
 	UGameplayTagsSettings* MutableDefault = GetMutableDefault<UGameplayTagsSettings>();
 
-	for (FStringAssetReference DataTablePath : MutableDefault->GameplayTagTableList)
+	for (FSoftObjectPath DataTablePath : MutableDefault->GameplayTagTableList)
 	{
 		UDataTable* TagTable = LoadObject<UDataTable>(nullptr, *DataTablePath.ToString(), nullptr, LOAD_None, nullptr);
 
diff --git a/Engine/Source/Runtime/Landscape/Classes/LandscapeSplineControlPoint.h b/Engine/Source/Runtime/Landscape/Classes/LandscapeSplineControlPoint.h
--- a/Engine/Source/Runtime/Landscape/Classes/LandscapeSplineControlPoint.h
+++ b/Engine/Source/Runtime/Landscape/Classes/LandscapeSplineControlPoint.h
@@ -159,7 +159,7 @@ protected:
 #if WITH_EDITORONLY_DATA
 	/** World reference for if mesh component is stored in another streaming level */
 	UPROPERTY(TextExportTransient, NonPIEDuplicateTransient)
-	TAssetPtr<UWorld> ForeignWorld;
+	TSoftObjectPtr<UWorld> ForeignWorld;
 
 	/** Key for tracking whether this segment has been modified relative to the mesh component stored in another streaming level */
 	UPROPERTY(TextExportTransient, NonPIEDuplicateTransient)
@@ -199,7 +199,7 @@ public:
 	/** Delete spline points */
 	virtual void DeleteSplinePoints();
 
-	const TAssetPtr<UWorld>& GetForeignWorld() const { return ForeignWorld; }
+	const TSoftObjectPtr<UWorld>& GetForeignWorld() const { return ForeignWorld; }
 	FGuid GetModificationKey() const { return ModificationKey; }
 #endif // WITH_EDITOR
 
diff --git a/Engine/Source/Runtime/Landscape/Classes/LandscapeSplineSegment.h b/Engine/Source/Runtime/Landscape/Classes/LandscapeSplineSegment.h
--- a/Engine/Source/Runtime/Landscape/Classes/LandscapeSplineSegment.h
+++ b/Engine/Source/Runtime/Landscape/Classes/LandscapeSplineSegment.h
@@ -240,7 +240,7 @@ protected:
 #if WITH_EDITORONLY_DATA
 	/** World references for mesh components stored in other streaming levels */
 	UPROPERTY(TextExportTransient, NonPIEDuplicateTransient)
-	TArray<TAssetPtr<UWorld>> ForeignWorlds;
+	TArray<TSoftObjectPtr<UWorld>> ForeignWorlds;
 
 	/** Key for tracking whether this segment has been modified relative to the mesh components stored in other streaming levels */
 	UPROPERTY(TextExportTransient, NonPIEDuplicateTransient)
@@ -264,7 +264,7 @@ public:
 	void UpdateSplineEditorMesh();
 	virtual void DeleteSplinePoints();
 
-	const TArray<TAssetPtr<UWorld>>& GetForeignWorlds() const { return ForeignWorlds; }
+	const TArray<TSoftObjectPtr<UWorld>>& GetForeignWorlds() const { return ForeignWorlds; }
 	FGuid GetModificationKey() const { return ModificationKey; }
 #endif
 
