accidentally brought a change from WEX which wasn't supposed to come over to main

diff --git a/Engine/Source/Runtime/Engine/Private/PhysicsEngine/PhysicsAsset.cpp b/Engine/Source/Runtime/Engine/Private/PhysicsEngine/PhysicsAsset.cpp
--- a/Engine/Source/Runtime/Engine/Private/PhysicsEngine/PhysicsAsset.cpp
+++ b/Engine/Source/Runtime/Engine/Private/PhysicsEngine/PhysicsAsset.cpp
@@ -332,9 +332,7 @@ FBox UPhysicsAsset::CalcAABB(const USkinnedMeshComponent* MeshComp, const FTrans
 		Box = FBox( LocalToWorld.GetLocation(), LocalToWorld.GetLocation() );
 	}
 
-	// @todo: - DO NOT copy this change to UE4/Main
-	const float MinBoundSize = 30.f;
-	// @todo: - DO NOT copy this change to UE4/Main
+	const float MinBoundSize = 1.f;
 	const FVector BoxSize = Box.GetSize();
 
 	if(BoxSize.GetMin() < MinBoundSize)
