Copying //UE4/Dev-Networking to //UE4/Dev-Main (Source: //UE4/Dev-Networking @ 3596412)

diff --git a/Engine/Source/Runtime/Engine/Private/NetConnection.cpp b/Engine/Source/Runtime/Engine/Private/NetConnection.cpp
--- a/Engine/Source/Runtime/Engine/Private/NetConnection.cpp
+++ b/Engine/Source/Runtime/Engine/Private/NetConnection.cpp
@@ -24,6 +24,7 @@
 #include "Net/DataChannel.h"
 #include "Engine/PackageMapClient.h"
 #include "Engine/NetworkObjectList.h"
+#include "EncryptionComponent.h"
 
 #include "Net/PerfCountersHelpers.h"
 #include "GameDelegates.h"
@@ -35,6 +36,8 @@ static TAutoConsoleVariable<int32> CVarPingDisplayServerTime( TEXT( "net.PingDis
 
 static TAutoConsoleVariable<int32> CVarTickAllOpenChannels( TEXT( "net.TickAllOpenChannels" ), 0, TEXT( "If nonzero, each net connection will tick all of its open channels every tick. Leaving this off will improve performance." ) );
 
+static TAutoConsoleVariable<int32> CVarRandomizeSequence(TEXT("net.RandomizeSequence"), 1, TEXT("Randomize initial packet sequence"));
+
 DECLARE_CYCLE_STAT(TEXT("NetConnection SendAcks"), Stat_NetConnectionSendAck, STATGROUP_Net);
 DECLARE_CYCLE_STAT(TEXT("NetConnection Tick"), Stat_NetConnectionTick, STATGROUP_Net);
 
@@ -48,6 +51,7 @@ UNetConnection* UNetConnection::GNetConnectionBeingCleanedUp = NULL;
 UNetConnection::UNetConnection(const FObjectInitializer& ObjectInitializer)
 :	UPlayer(ObjectInitializer)
 ,	Driver				( NULL )
+,	PackageMapClass		( UPackageMapClient::StaticClass() )
 ,	PackageMap			( NULL )
 ,	ViewTarget			( NULL )
 ,   OwningActor			( NULL )
@@ -100,6 +104,9 @@ UNetConnection::UNetConnection(const FObjectInitializer& ObjectInitializer)
 ,	GameNetworkProtocolVersion( FNetworkVersion::GetGameNetworkProtocolVersion() )
 ,	ClientWorldPackageName( NAME_None )
 ,	bResendAllDataSinceOpen( false )
+#if !UE_BUILD_SHIPPING
+,	ReceivedRawPacketDel()
+#endif
 ,	PlayerOnlinePlatformName( NAME_None )
 {
 }
@@ -171,9 +178,13 @@ void UNetConnection::InitBase(UNetDriver* InDriver,class FSocket* InSocket, cons
 	}
 
 	// Create package map.
-	auto PackageMapClient = NewObject<UPackageMapClient>(this);
-	PackageMapClient->Initialize(this, Driver->GuidCache);
-	PackageMap = PackageMapClient;
+	UPackageMapClient* PackageMapClient = NewObject<UPackageMapClient>(this, PackageMapClass);
+
+	if (ensure(PackageMapClient != nullptr))
+	{
+		PackageMapClient->Initialize(this, Driver->GuidCache);
+		PackageMap = PackageMapClient;
+	}
 
 	// Create the voice channel
 	CreateChannel(CHTYPE_Voice, true, VOICE_CHANNEL_INDEX);
@@ -277,28 +288,65 @@ void UNetConnection::InitHandler()
 
 void UNetConnection::InitSequence(int32 IncomingSequence, int32 OutgoingSequence)
 {
-	// Make sure the sequence hasn't already been initialized
-	check(InPacketId == -1);
+	// Make sure the sequence hasn't already been initialized on the server, and ignore multiple initializations on the client
+	check(InPacketId == -1 || Driver->ServerConnection != nullptr);
+
+	if (InPacketId == -1 && CVarRandomizeSequence.GetValueOnAnyThread() > 0)
+	{
+		// Initialize the base UNetConnection packet sequence (not very useful/effective at preventing attacks)
+		InPacketId = IncomingSequence - 1;
+		OutPacketId = OutgoingSequence;
+		OutAckPacketId = OutgoingSequence - 1;
+
+		// Initialize the reliable packet sequence (more useful/effective at preventing attacks)
+		InitInReliable = IncomingSequence & (MAX_CHSEQUENCE - 1);
+		InitOutReliable = OutgoingSequence & (MAX_CHSEQUENCE - 1);
 
-	// Initialize the base UNetConnection packet sequence (not very useful/effective at preventing attacks)
-	InPacketId = IncomingSequence - 1;
-	OutPacketId = OutgoingSequence;
-	OutAckPacketId = OutgoingSequence - 1;
+		UE_LOG(LogNet, Verbose, TEXT("InitSequence: IncomingSequence: %i, OutgoingSequence: %i, InitInReliable: %i, InitOutReliable: %i"), IncomingSequence, OutgoingSequence, InitInReliable, InitOutReliable);
 
-	// Initialize the reliable packet sequence (more useful/effective at preventing attacks)
-	InitInReliable = IncomingSequence & (MAX_CHSEQUENCE - 1);
-	InitOutReliable = OutgoingSequence & (MAX_CHSEQUENCE - 1);
+		for (int32 i=0; i<ARRAY_COUNT(InReliable); i++)
+		{
+			InReliable[i] = InitInReliable;
+		}
 
-	UE_LOG(LogNet, Verbose, TEXT("InitSequence: IncomingSequence: %i, OutgoingSequence: %i, InitInReliable: %i, InitOutReliable: %i"), IncomingSequence, OutgoingSequence, InitInReliable, InitOutReliable);
+		for (int32 i=0; i<ARRAY_COUNT(OutReliable); i++)
+		{
+			OutReliable[i] = InitOutReliable;
+		}
+	}
+}
 
-	for (int32 i=0; i<ARRAY_COUNT(InReliable); i++)
+void UNetConnection::EnableEncryptionWithKey(TArrayView<const uint8> Key)
+{
+	if (Handler.IsValid())
 	{
-		InReliable[i] = InitInReliable;
+		UE_LOG(LogNet, Verbose, TEXT("UNetConnection::EnableEncryptionWithKey, %s"), *Describe());
+
+		TSharedPtr<FEncryptionComponent> EncryptionComponent = Handler->GetEncryptionComponent();
+		if (EncryptionComponent.IsValid())
+		{
+			EncryptionComponent->SetEncryptionKey(Key);
+			EncryptionComponent->EnableEncryption();
+		}
+		else
+		{
+			UE_LOG(LogNet, Warning, TEXT("UNetConnection::EnableEncryptionWithKey, encryption component not found!"));
+		}
 	}
+}
 
-	for (int32 i=0; i<ARRAY_COUNT(OutReliable); i++)
+void UNetConnection::EnableEncryptionWithKeyServer(TArrayView<const uint8> Key)
+{
+	if (State != USOCK_Invalid && State != USOCK_Closed && Driver)
+	{
+		FNetControlMessage<NMT_EncryptionAck>::Send(this);
+		FlushNet();
+
+		EnableEncryptionWithKey(Key);
+	}
+	else
 	{
-		OutReliable[i] = InitOutReliable;
+		UE_LOG(LogNet, Log, TEXT("UNetConnection::EnableEncryptionWithKeyServer, connection in invalid state. %s"), *Describe());
 	}
 }
 
@@ -591,6 +639,18 @@ void UNetConnection::InitSendBuffer()
 
 void UNetConnection::ReceivedRawPacket( void* InData, int32 Count )
 {
+#if !UE_BUILD_SHIPPING
+	// Add an opportunity for the hook to block further processing
+	bool bBlockReceive = false;
+
+	ReceivedRawPacketDel.ExecuteIfBound(InData, Count, bBlockReceive);
+
+	if (bBlockReceive)
+	{
+		return;
+	}
+#endif
+
 	uint8* Data = (uint8*)InData;
 
 	if (Handler.IsValid())
@@ -618,6 +678,13 @@ void UNetConnection::ReceivedRawPacket( void* InData, int32 Count )
 
 			return;
 		}
+		
+		// See if we receive a packet that wasn't fully consumed by the handler before the handler is initialized.
+		if (!Handler->IsFullyInitialized())
+		{
+			UE_LOG(LogNet, Warning, TEXT("PacketHander isn't fully initialized and also didn't fully consume a packet! This will cause the connection to try to send a packet before the initial packet sequence has been established. Ignoring. Connection: %s"), *Describe());
+			return;
+		}
 	}
 
 
@@ -679,10 +746,6 @@ void UNetConnection::FlushNet(bool bIgnoreSimulation)
 {
 	check(Driver);
 
-	// Due to the PacketHandler handshake code, servers must never send the client data,
-	// before first receiving a client control packet (which is taken as an indication of a complete handshake).
-	check(HasReceivedClientPacket())
-
 	// Update info.
 	ValidateSendBuffer();
 	LastEnd = FBitWriterMark();
@@ -691,6 +754,15 @@ void UNetConnection::FlushNet(bool bIgnoreSimulation)
 	// If there is any pending data to send, send it.
 	if ( SendBuffer.GetNumBits() || ( Driver->Time-LastSendTime > Driver->KeepAliveTime && !InternalAck && State != USOCK_Closed ) )
 	{
+		// Due to the PacketHandler handshake code, servers must never send the client data,
+		// before first receiving a client control packet (which is taken as an indication of a complete handshake).
+		if ( !HasReceivedClientPacket() && CVarRandomizeSequence.GetValueOnAnyThread() != 0 )
+		{
+			UE_LOG(LogNet, Log, TEXT("Attempting to send data before handshake is complete. %s"), *Describe());
+			Close();
+			return;
+		}
+
 		// If sending keepalive packet, still write the packet id
 		if ( SendBuffer.GetNumBits() == 0 )
 		{
@@ -1208,7 +1280,7 @@ void UNetConnection::ReceivedPacket( FBitReader& Reader )
 			if( !Channel )
 			{
 				// Validate channel type.
-				if ( !UChannel::IsKnownChannelType( Bunch.ChType ) )
+				if ( !Driver->IsKnownChannelType( Bunch.ChType ) )
 				{
 					// Unknown type.
 					CLOSE_CONNECTION_DUE_TO_SECURITY_VIOLATION(this, ESecurityEvent::Invalid_Data, TEXT( "UNetConnection::ReceivedPacket: Connection unknown channel type (%i)" ), (int)Bunch.ChType);
@@ -1283,7 +1355,10 @@ int32 UNetConnection::WriteBitsToSendBuffer(
 #if !UE_BUILD_SHIPPING
 	// Now that the stateless handshake is responsible for initializing the packet sequence numbers,
 	//	we can't allow any packets to be written to the send buffer until after this has completed
-	check(!Handler.IsValid() || Handler->IsFullyInitialized());
+	if (CVarRandomizeSequence.GetValueOnAnyThread() > 0)
+	{
+		checkf(!Handler.IsValid() || Handler->IsFullyInitialized(), TEXT("Attempted to write to send buffer before packet handler was fully initialized. Connection: %s"), *Describe());
+	}
 #endif
 
 	const int32 TotalSizeInBits = SizeInBits + ExtraSizeInBits;
@@ -1517,7 +1592,7 @@ int32 UNetConnection::SendRawBunch( FOutBunch& Bunch, bool InAllowMerge )
 
 UChannel* UNetConnection::CreateChannel( EChannelType ChType, bool bOpenedLocally, int32 ChIndex )
 {
-	check(UChannel::IsKnownChannelType(ChType));
+	check(Driver->IsKnownChannelType(ChType));
 	AssertValid();
 
 	// If no channel index was specified, find the first available.
@@ -1556,7 +1631,7 @@ UChannel* UNetConnection::CreateChannel( EChannelType ChType, bool bOpenedLocall
 	check(Channels[ChIndex]==NULL);
 
 	// Create channel.
-	UChannel* Channel = NewObject<UChannel>(GetTransientPackage(), UChannel::ChannelClasses[ChType]);
+	UChannel* Channel = NewObject<UChannel>(GetTransientPackage(), Driver->ChannelClasses[ChType]);
 	Channel->Init( this, ChIndex, bOpenedLocally );
 	Channels[ChIndex] = Channel;
 	OpenChannels.Add(Channel);
diff --git a/Engine/Source/Runtime/Engine/Private/NetworkDriver.cpp b/Engine/Source/Runtime/Engine/Private/NetworkDriver.cpp
--- a/Engine/Source/Runtime/Engine/Private/NetworkDriver.cpp
+++ b/Engine/Source/Runtime/Engine/Private/NetworkDriver.cpp
@@ -44,6 +44,7 @@
 #include "Engine/PackageMapClient.h"
 #include "Net/RepLayout.h"
 #include "Net/DataReplication.h"
+#include "Engine/ControlChannel.h"
 #include "Engine/ActorChannel.h"
 #include "Engine/VoiceChannel.h"
 #include "Engine/NetworkObjectList.h"
@@ -168,6 +169,11 @@ static TAutoConsoleVariable<int32> CVarUseAdaptiveNetUpdateFrequency(
 	1, 
 	TEXT( "If 1, NetUpdateFrequency will be calculated based on how often actors actually send something when replicating" ) );
 
+TAutoConsoleVariable<int32> CVarNetAllowEncryption(
+	TEXT("net.AllowEncryption"),
+	1,
+	TEXT("If true, the engine will attempt to load an encryption PacketHandler component and fill in the EncryptionToken parameter of the NMT_Hello message based on the ?EncryptionToken= URL option and call callbacks if it's non-empty."));
+
 /*-----------------------------------------------------------------------------
 	UNetDriver implementation.
 -----------------------------------------------------------------------------*/
@@ -204,11 +210,17 @@ UNetDriver::UNetDriver(const FObjectInitializer& ObjectInitializer)
 ,	LastCleanupTime(0.0)
 ,	NetTag(0)
 ,	DebugRelevantActors(false)
+#if !UE_BUILD_SHIPPING
+,	SendRPCDel()
+#endif
 ,	ProcessQueuedBunchesCurrentFrameMilliseconds(0.0f)
 ,	NetworkObjects(new FNetworkObjectList)
 ,	LagState(ENetworkLagState::NotLagging)
 ,	DuplicateLevelID(INDEX_NONE)
 {
+	ChannelClasses[CHTYPE_Control]	= UControlChannel::StaticClass();
+	ChannelClasses[CHTYPE_Actor]	= UActorChannel::StaticClass();
+	ChannelClasses[CHTYPE_Voice]	= UVoiceChannel::StaticClass();
 }
 
 void UNetDriver::InitPacketSimulationSettings()
@@ -249,6 +261,8 @@ void UNetDriver::PostInitProperties()
 		// Do not time out in PIE since the server is local.
 		bNoTimeouts = bNoTimeouts || (GEditor && GEditor->PlayWorld);
 #endif // WITH_EDITOR
+	
+		OnLevelRemovedFromWorldHandle = FWorldDelegates::LevelRemovedFromWorld.AddUObject(this, &UNetDriver::OnLevelRemovedFromWorld);
 	}
 }
 
@@ -1349,6 +1363,7 @@ void UNetDriver::InternalProcessRemoteFunction
 		// Bump the ReplicationFrame value to invalidate any properties marked as "unchanged" for this frame.
 		ReplicationFrame++;
 		
+		Ch->GetActor()->CallPreReplication(this);
 		Ch->ReplicateActor();
 	}
 
@@ -1638,6 +1653,7 @@ void UNetDriver::Serialize( FArchive& Ar )
 		ClientConnections.CountBytes(Ar);
 	}
 }
+
 void UNetDriver::FinishDestroy()
 {
 	if ( !HasAnyFlags(RF_ClassDefaultObject) )
@@ -1665,6 +1681,8 @@ void UNetDriver::FinishDestroy()
 
 		// Delete the guid cache
 		GuidCache.Reset();
+
+		FWorldDelegates::LevelRemovedFromWorld.Remove(OnLevelRemovedFromWorldHandle);
 	}
 	else
 	{
@@ -1971,23 +1989,13 @@ FActorDestructionInfo *	CreateDestructionInfo( UNetDriver * NetDriver, AActor* T
 	FActorDestructionInfo &NewInfo = NetDriver->DestroyedStartupOrDormantActors.FindOrAdd( NetGUID );
 	NewInfo.DestroyedPosition = ThisActor->GetActorLocation();
 	NewInfo.NetGUID = NetGUID;
+	NewInfo.Level = ThisActor->GetLevel();
 	NewInfo.ObjOuter = ThisActor->GetOuter();
 	NewInfo.PathName = ThisActor->GetName();
 
-	// get the level for the object
-	const ULevel* Level = NULL;
-	for (const UObject* Obj = ThisActor; Obj != NULL; Obj = Obj->GetOuter())
-	{
-		Level = Cast<const ULevel>(Obj);
-		if (Level != NULL)
-		{
-			break;
-		}
-	}
-
-	if (Level && !Level->IsPersistentLevel() )
+	if (NewInfo.Level.IsValid() && !NewInfo.Level->IsPersistentLevel() )
 	{
-		NewInfo.StreamingLevelName = Level->GetOutermost()->GetFName();
+		NewInfo.StreamingLevelName = NewInfo.Level->GetOutermost()->GetFName();
 	}
 	else
 	{
@@ -2925,16 +2933,22 @@ int32 UNetDriver::ServerReplicateActors_PrioritizeActors( UNetConnection* Connec
 
 int32 UNetDriver::ServerReplicateActors_ProcessPrioritizedActors( UNetConnection* Connection, const TArray<FNetViewer>& ConnectionViewers, FActorPriority** PriorityActors, const int32 FinalSortedCount, int32& OutUpdated )
 {
+	int32 ActorUpdatesThisConnection		= 0;
+	int32 ActorUpdatesThisConnectionSent	= 0;
+	int32 FinalRelevantCount				= 0;
+
 	if ( !Connection->IsNetReady( 0 ) )
 	{
 		// Connection saturated, don't process any actors
+
+		// Update stats even though there was no processing.
+		SET_DWORD_STAT( STAT_NumReplicatedActorAttempts, ActorUpdatesThisConnection );
+		SET_DWORD_STAT( STAT_NumReplicatedActors, ActorUpdatesThisConnectionSent );
+		SET_DWORD_STAT( STAT_NumRelevantActors, FinalRelevantCount );
+
 		return 0;
 	}
 
-	int32 ActorUpdatesThisConnection		= 0;
-	int32 ActorUpdatesThisConnectionSent	= 0;
-	int32 FinalRelevantCount				= 0;
-
 	for ( int32 j = 0; j < FinalSortedCount; j++ )
 	{
 		FNetworkObjectInfo*	ActorInfo = PriorityActors[j]->ActorInfo;
@@ -3635,6 +3649,42 @@ TSharedPtr< FReplicationChangelistMgr > UNetDriver::GetReplicationChangeListMgr(
 	return *ReplicationChangeListMgrPtr;
 }
 
+void UNetDriver::OnLevelRemovedFromWorld(class ULevel* InLevel, class UWorld* InWorld)
+{
+	if (InWorld == World)
+	{
+		for (AActor* Actor : InLevel->Actors)
+		{
+			if (Actor)
+			{
+				NotifyActorLevelUnloaded(Actor);
+				GetNetworkObjectList().Remove(Actor);
+			}
+		}
+
+		TArray<FNetworkGUID> RemovedGUIDs;
+		for (auto It = DestroyedStartupOrDormantActors.CreateIterator(); It; ++It)
+		{
+			if (It->Value.Level == InLevel)
+			{
+				RemovedGUIDs.Add(It->Key);
+				It.RemoveCurrent();
+			}
+		}
+
+		if (RemovedGUIDs.Num())
+		{
+			for (UNetConnection* Connection : ClientConnections)
+			{
+				for (const FNetworkGUID& GUIDToRemove : RemovedGUIDs)
+				{
+					Connection->DestroyedStartupOrDormantActors.Remove(GUIDToRemove);
+				}
+			}
+		}
+	}
+}
+
 FAutoConsoleCommandWithWorld	DumpRelevantActorsCommand(
 	TEXT("net.DumpRelevantActors"), 
 	TEXT( "Dumps information on relevant actors during next network update" ), 
diff --git a/Engine/Source/Runtime/Engine/Private/Pawn.cpp b/Engine/Source/Runtime/Engine/Private/Pawn.cpp
--- a/Engine/Source/Runtime/Engine/Private/Pawn.cpp
+++ b/Engine/Source/Runtime/Engine/Private/Pawn.cpp
@@ -27,6 +27,7 @@
 #include "Interfaces/NetworkPredictionInterface.h"
 #include "GameFramework/PlayerState.h"
 #include "Components/PawnNoiseEmitterComponent.h"
+#include "GameFramework/GameNetworkManager.h"
 
 DEFINE_LOG_CATEGORY(LogDamage);
 DEFINE_LOG_CATEGORY_STATIC(LogPawn, Warning, All);
@@ -1052,7 +1053,8 @@ bool APawn::IsNetRelevantFor(const AActor* RealViewer, const AActor* ViewTarget,
 		}
 	}
 
-	return ((SrcLocation - GetActorLocation()).SizeSquared() < NetCullDistanceSquared);
+	return !GetDefault<AGameNetworkManager>()->bUseDistanceBasedRelevancy ||
+			IsWithinNetRelevancyDistance(SrcLocation);
 }
 
 void APawn::GetLifetimeReplicatedProps( TArray< FLifetimeProperty > & OutLifetimeProps ) const
diff --git a/Engine/Source/Runtime/Engine/Private/PendingNetGame.cpp b/Engine/Source/Runtime/Engine/Private/PendingNetGame.cpp
--- a/Engine/Source/Runtime/Engine/Private/PendingNetGame.cpp
+++ b/Engine/Source/Runtime/Engine/Private/PendingNetGame.cpp
@@ -93,7 +93,13 @@ void UPendingNetGame::SendInitialJoin()
 
 			UE_LOG(LogNet, Log, TEXT( "UPendingNetGame::SendInitialJoin: Sending hello. %s" ), *ServerConn->Describe());
 
-			FNetControlMessage<NMT_Hello>::Send(ServerConn, IsLittleEndian, LocalNetworkVersion);
+			FString EncryptionToken;
+			if (CVarNetAllowEncryption.GetValueOnGameThread() != 0)
+			{
+				EncryptionToken = URL.GetOption(TEXT("EncryptionToken="), TEXT(""));
+			}
+
+			FNetControlMessage<NMT_Hello>::Send(ServerConn, IsLittleEndian, LocalNetworkVersion, EncryptionToken);
 
 			ServerConn->FlushNet();
 		}
@@ -305,12 +311,66 @@ void UPendingNetGame::NotifyControlMessage(UNetConnection* Connection, uint8 Mes
 			NetDriver->ServerConnection->PackageMap->ResolvePathAndAssignNetGUID(NetGUID, Path);
 			break;
 		}
+		case NMT_EncryptionAck:
+		{
+			if (FNetDelegates::OnReceivedNetworkEncryptionAck.IsBound())
+			{
+				TWeakObjectPtr<UNetConnection> WeakConnection = Connection;
+				FNetDelegates::OnReceivedNetworkEncryptionAck.Execute(FOnEncryptionKeyResponse::CreateUObject(this, &UPendingNetGame::FinalizeEncryptedConnection, WeakConnection));
+			}
+			else
+			{
+				// This error will be resolved in TickWorldTravel()
+				ConnectionError = TEXT("No encryption ack handler");
+
+				// Force close the session
+				UE_LOG(LogNet, Warning, TEXT("%s: No delegate available to handle encryption ack, disconnecting."), *Connection->GetName());
+				Connection->Close();
+			}
+			break;
+		}
 		default:
 			UE_LOG(LogNet, Log, TEXT(" --- Unknown/unexpected message for pending level"));
 			break;
 	}
 }
 
+void UPendingNetGame::FinalizeEncryptedConnection(const FEncryptionKeyResponse& Response, TWeakObjectPtr<UNetConnection> WeakConnection)
+{
+	UNetConnection* Connection = WeakConnection.Get();
+	if (Connection)
+	{
+		if (Connection->State != USOCK_Invalid && Connection->State != USOCK_Closed && Connection->Driver)
+		{
+			if (Response.Response == EEncryptionResponse::Success)
+			{
+				Connection->EnableEncryptionWithKey(Response.EncryptionKey);
+			}
+			else
+			{
+				// This error will be resolved in TickWorldTravel()
+				FString ResponseStr(Lex::ToString(Response.Response));
+				UE_LOG(LogNet, Warning, TEXT("UPendingNetGame::FinalizeEncryptedConnection: encryption failure [%s] %s"), *ResponseStr, *Response.ErrorMsg);
+				ConnectionError = TEXT("Encryption ack failure");
+				Connection->Close();
+			}
+		}
+		else
+		{
+			// This error will be resolved in TickWorldTravel()
+			UE_LOG(LogNet, Warning, TEXT("UPendingNetGame::FinalizeEncryptedConnection: connection in invalid state. %s"), *Connection->Describe());
+			ConnectionError = TEXT("Connection encryption state failure");
+			Connection->Close();
+		}
+	}
+	else
+	{
+		// This error will be resolved in TickWorldTravel()
+		UE_LOG(LogNet, Warning, TEXT("UPendingNetGame::FinalizeEncryptedConnection: Connection is null."));
+		ConnectionError = TEXT("Connection missing during encryption ack");
+	}
+}
+
 void UPendingNetGame::Tick( float DeltaTime )
 {
 	check(NetDriver && NetDriver->ServerConnection);
diff --git a/Engine/Source/Runtime/Engine/Private/World.cpp b/Engine/Source/Runtime/Engine/Private/World.cpp
--- a/Engine/Source/Runtime/Engine/Private/World.cpp
+++ b/Engine/Source/Runtime/Engine/Private/World.cpp
@@ -1717,7 +1717,7 @@ void UWorld::EnsureCollisionTreeIsBuilt()
 		PhysicsScene->EnsureCollisionTreeIsBuilt(this);
 	}
 
-    bIsBuilt = true;
+	bIsBuilt = true;
 }
 
 void UWorld::InvalidateModelGeometry(ULevel* InLevel)
@@ -2272,10 +2272,10 @@ void UWorld::RemoveFromWorld( ULevel* Level, bool bAllowIncrementalRemoval )
 				} while (!IsTimeLimitExceeded(TEXT("unregistering components"), StartTime, Level, GLevelStreamingUnregisterComponentsTimeLimit));
 			}
 		}
-        else
-        {
+		else
+		{
 			Level->bIsBeingRemoved = true;
-        }
+		}
 
 		if ( bFinishRemovingLevel )
 		{
@@ -2284,11 +2284,6 @@ void UWorld::RemoveFromWorld( ULevel* Level, bool bAllowIncrementalRemoval )
 				if (AActor* Actor = Level->Actors[ActorIdx])
 				{
 					Actor->RouteEndPlay(EEndPlayReason::RemovedFromWorld);
-
-					if (NetDriver)
-					{
-						NetDriver->NotifyActorLevelUnloaded(Actor);
-					}
 				}
 			}
 
@@ -2327,15 +2322,6 @@ void UWorld::RemoveFromWorld( ULevel* Level, bool bAllowIncrementalRemoval )
 					}
 				}
 			}
-		
-			// Remove any actors in the network list
-			for (AActor* Actor : Level->Actors)
-			{
-				if (Actor)
-				{
-					RemoveNetworkActor(Actor);
-				}
-			}
 
 			Level->bIsVisible = false;
 
@@ -4162,8 +4148,9 @@ void UWorld::NotifyControlMessage(UNetConnection* Connection, uint8 MessageType,
 				uint8 IsLittleEndian = 0;
 				uint32 RemoteNetworkVersion = 0;
 				uint32 LocalNetworkVersion = FNetworkVersion::GetLocalNetworkVersion();
+				FString EncryptionToken;
 
-				FNetControlMessage<NMT_Hello>::Receive(Bunch, IsLittleEndian, RemoteNetworkVersion);
+				FNetControlMessage<NMT_Hello>::Receive(Bunch, IsLittleEndian, RemoteNetworkVersion, EncryptionToken);
 
 				if (!FNetworkVersion::IsNetworkCompatible(LocalNetworkVersion, RemoteNetworkVersion))
 				{
@@ -4176,10 +4163,25 @@ void UWorld::NotifyControlMessage(UNetConnection* Connection, uint8 MessageType,
 				}
 				else
 				{
-					Connection->Challenge = FString::Printf(TEXT("%08X"), FPlatformTime::Cycles());
-					Connection->SetExpectedClientLoginMsgType( NMT_Login );
-					FNetControlMessage<NMT_Challenge>::Send(Connection, Connection->Challenge);
-					Connection->FlushNet();
+					if (EncryptionToken.IsEmpty())
+					{
+						SendChallengeControlMessage(Connection);
+					}
+					else
+					{
+						if (FNetDelegates::OnReceivedNetworkEncryptionToken.IsBound())
+						{
+							TWeakObjectPtr<UNetConnection> WeakConnection = Connection;
+							FNetDelegates::OnReceivedNetworkEncryptionToken.Execute(EncryptionToken, FOnEncryptionKeyResponse::CreateUObject(this, &UWorld::SendChallengeControlMessage, WeakConnection));
+						}
+						else
+						{
+							FString FailureMsg(TEXT("Encryption failure"));
+							UE_LOG(LogNet, Warning, TEXT("%s: No delegate available to handle encryption token, disconnecting."), *Connection->GetName());
+							FNetControlMessage<NMT_Failure>::Send(Connection, FailureMsg);
+							Connection->FlushNet(true);
+						}
+					}
 				}
 				break;
 			}
@@ -4458,6 +4460,61 @@ void UWorld::NotifyControlMessage(UNetConnection* Connection, uint8 MessageType,
 	}
 }
 
+void UWorld::SendChallengeControlMessage(UNetConnection* Connection)
+{
+	if (Connection)
+	{
+		if (Connection->State != USOCK_Invalid && Connection->State != USOCK_Closed && Connection->Driver)
+		{
+			Connection->Challenge = FString::Printf(TEXT("%08X"), FPlatformTime::Cycles());
+			Connection->SetExpectedClientLoginMsgType(NMT_Login);
+			FNetControlMessage<NMT_Challenge>::Send(Connection, Connection->Challenge);
+			Connection->FlushNet();
+		}
+		else
+		{
+			UE_LOG(LogNet, Log, TEXT("UWorld::SendChallengeControlMessage: connection in invalid state. %s"), *Connection->Describe());
+		}
+	}
+	else
+	{
+		UE_LOG(LogNet, Log, TEXT("UWorld::SendChallengeControlMessage: Connection is null."));
+	}
+}
+
+void UWorld::SendChallengeControlMessage(const FEncryptionKeyResponse& Response, TWeakObjectPtr<UNetConnection> WeakConnection)
+{
+	UNetConnection* Connection = WeakConnection.Get();
+	if (Connection)
+	{
+		if (Connection->State != USOCK_Invalid && Connection->State != USOCK_Closed && Connection->Driver)
+		{
+			if (Response.Response == EEncryptionResponse::Success)
+			{
+				Connection->EnableEncryptionWithKeyServer(Response.EncryptionKey);
+				SendChallengeControlMessage(Connection);
+			}
+			else
+			{
+				FString ResponseStr(Lex::ToString(Response.Response));
+				UE_LOG(LogNet, Warning, TEXT("UWorld::SendChallengeControlMessage: encryption failure [%s] %s"), *ResponseStr, *Response.ErrorMsg);
+				FNetControlMessage<NMT_Failure>::Send(Connection, ResponseStr);
+				Connection->FlushNet();
+				// Can't close the connection here since it will leave the failure message in the send buffer and just close the socket. 
+				// Connection->Close();
+			}
+		}
+		else
+		{
+			UE_LOG(LogNet, Warning, TEXT("UWorld::SendChallengeControlMessage: connection in invalid state. %s"), *Connection->Describe());
+		}
+	}
+	else
+	{
+		UE_LOG(LogNet, Warning, TEXT("UWorld::SendChallengeControlMessage: Connection is null."));
+	}
+}
+
 bool UWorld::Listen( FURL& InURL )
 {
 #if WITH_SERVER_CODE
@@ -5269,7 +5326,7 @@ UWorld* FSeamlessTravelHandler::Tick()
 				ProcessActor(Player);
 			}
 
- 			bool bCreateNewGameMode = !bIsClient;
+			bool bCreateNewGameMode = !bIsClient;
 			{
 				// scope because after GC the kept pointers will be bad
 				AGameModeBase* KeptGameMode = nullptr;
@@ -6247,6 +6304,28 @@ void UWorld::SetActiveLevelCollection(const FLevelCollection* InCollection)
 	GameState = InCollection->GetGameState();
 	NetDriver = InCollection->GetNetDriver();
 	DemoNetDriver = InCollection->GetDemoNetDriver();
+
+	// TODO: START TEMP FIX FOR UE-42508
+	if (NetDriver && NetDriver->NetDriverName != NAME_None)
+	{
+		UNetDriver* TempNetDriver = GEngine->FindNamedNetDriver(this, NetDriver->NetDriverName);
+		if (TempNetDriver != NetDriver)
+		{
+			UE_LOG(LogWorld, Warning, TEXT("SetActiveLevelCollection attempted to use an out of date NetDriver: %s"), *(NetDriver->NetDriverName.ToString()));
+			NetDriver = TempNetDriver;
+		}
+	}
+
+	if (DemoNetDriver && DemoNetDriver->NetDriverName != NAME_None)
+	{
+		UDemoNetDriver* TempDemoNetDriver = Cast<UDemoNetDriver>(GEngine->FindNamedNetDriver(this, DemoNetDriver->NetDriverName));
+		if (TempDemoNetDriver != DemoNetDriver)
+		{
+			UE_LOG(LogWorld, Warning, TEXT("SetActiveLevelCollection attempted to use an out of date DemoNetDriver: %s"), *(DemoNetDriver->NetDriverName.ToString()));
+			DemoNetDriver = TempDemoNetDriver;
+		}
+	}
+	// TODO: END TEMP FIX FOR UE-42508
 }
 
 static ULevel* DuplicateLevelWithPrefix(ULevel* InLevel, int32 InstanceID )
@@ -6373,25 +6452,25 @@ void UWorld::ChangeFeatureLevel(ERHIFeatureLevel::Type InFeatureLevel, bool bSho
 	{
 		UE_LOG(LogWorld, Log, TEXT("Changing Feature Level (Enum) from %i to %i"), (int)FeatureLevel, (int)InFeatureLevel);
 		FScopedSlowTask SlowTask(100.f, NSLOCTEXT("Engine", "ChangingPreviewRenderingLevelMessage", "Changing Preview Rendering Level"), bShowSlowProgressDialog);
-        SlowTask.MakeDialog();
-        {
-            SlowTask.EnterProgressFrame(10.0f);
-            // Give all scene components the opportunity to prepare for pending feature level change.
-            for (TObjectIterator<USceneComponent> It; It; ++It)
-            {
-                USceneComponent* SceneComponent = *It;
-                if (SceneComponent->GetWorld() == this)
-                {
-                    SceneComponent->PreFeatureLevelChange(InFeatureLevel);
-                }
-            }
+		SlowTask.MakeDialog();
+		{
+			SlowTask.EnterProgressFrame(10.0f);
+			// Give all scene components the opportunity to prepare for pending feature level change.
+			for (TObjectIterator<USceneComponent> It; It; ++It)
+			{
+				USceneComponent* SceneComponent = *It;
+				if (SceneComponent->GetWorld() == this)
+				{
+					SceneComponent->PreFeatureLevelChange(InFeatureLevel);
+				}
+			}
 
-            SlowTask.EnterProgressFrame(10.0f);
-            FGlobalComponentReregisterContext RecreateComponents;
-            FlushRenderingCommands();
+			SlowTask.EnterProgressFrame(10.0f);
+			FGlobalComponentReregisterContext RecreateComponents;
+			FlushRenderingCommands();
 
-            // Decrement refcount on old feature level
-            UMaterialInterface::SetGlobalRequiredFeatureLevel(InFeatureLevel, true);
+			// Decrement refcount on old feature level
+			UMaterialInterface::SetGlobalRequiredFeatureLevel(InFeatureLevel, true);
 
             SlowTask.EnterProgressFrame(10.0f);
             UMaterial::AllMaterialsCacheResourceShadersForRendering();
@@ -6402,21 +6481,21 @@ void UWorld::ChangeFeatureLevel(ERHIFeatureLevel::Type InFeatureLevel, bool bSho
             SlowTask.EnterProgressFrame(10.0f);
             GShaderCompilingManager->ProcessAsyncResults(false, true);
 
-            SlowTask.EnterProgressFrame(10.0f);
-            //invalidate global bound shader states so they will be created with the new shaders the next time they are set (in SetGlobalBoundShaderState)
-            for (TLinkedList<FGlobalBoundShaderStateResource*>::TIterator It(FGlobalBoundShaderStateResource::GetGlobalBoundShaderStateList()); It; It.Next())
-            {
-                BeginUpdateResourceRHI(*It);
-            }
+			SlowTask.EnterProgressFrame(10.0f);
+			//invalidate global bound shader states so they will be created with the new shaders the next time they are set (in SetGlobalBoundShaderState)
+			for (TLinkedList<FGlobalBoundShaderStateResource*>::TIterator It(FGlobalBoundShaderStateResource::GetGlobalBoundShaderStateList()); It; It.Next())
+			{
+				BeginUpdateResourceRHI(*It);
+			}
 
-            FeatureLevel = InFeatureLevel;
+			FeatureLevel = InFeatureLevel;
 
-            SlowTask.EnterProgressFrame(10.0f);
+			SlowTask.EnterProgressFrame(10.0f);
 			RecreateScene(InFeatureLevel);
 
-            SlowTask.EnterProgressFrame(10.0f);
-            TriggerStreamingDataRebuild();
-        }
+			SlowTask.EnterProgressFrame(10.0f);
+			TriggerStreamingDataRebuild();
+		}
 	}
 }
 
diff --git a/Engine/Source/Runtime/Engine/Public/Net/DataChannel.h b/Engine/Source/Runtime/Engine/Public/Net/DataChannel.h
--- a/Engine/Source/Runtime/Engine/Public/Net/DataChannel.h
+++ b/Engine/Source/Runtime/Engine/Public/Net/DataChannel.h
@@ -13,6 +13,9 @@
 #include "Net/DataBunch.h"
 #include "Engine/ChildConnection.h"
 
+// Backwards-compatibility macro for online game plugins to test
+#define SUPPORTS_ENCRYPTION_TOKEN 1
+
 /*-----------------------------------------------------------------------------
 	UControlChannel base class.
 -----------------------------------------------------------------------------*/
@@ -376,7 +379,7 @@ private:
 #define IMPLEMENT_CONTROL_CHANNEL_MESSAGE(Name) static uint8 Dummy##_FNetControlMessage_##Name = FNetControlMessage<NMT_##Name>::Initialize();
 
 // message type definitions
-DEFINE_CONTROL_CHANNEL_MESSAGE_TWOPARAM(Hello, 0, uint8, uint32); // initial client connection message
+DEFINE_CONTROL_CHANNEL_MESSAGE_THREEPARAM(Hello, 0, uint8, uint32, FString); // initial client connection message
 DEFINE_CONTROL_CHANNEL_MESSAGE_THREEPARAM(Welcome, 1, FString, FString, FString); // server tells client they're ok'ed to load the server's level
 DEFINE_CONTROL_CHANNEL_MESSAGE_ONEPARAM(Upgrade, 2, uint32); // server tells client their version is incompatible
 DEFINE_CONTROL_CHANNEL_MESSAGE_ONEPARAM(Challenge, 3, FString); // server sends client challenge string to verify integrity
@@ -393,6 +396,7 @@ DEFINE_CONTROL_CHANNEL_MESSAGE_ONEPARAM(DebugText, 17, FString); // debug text s
 DEFINE_CONTROL_CHANNEL_MESSAGE_TWOPARAM(NetGUIDAssign, 18, FNetworkGUID, FString); // Explicit NetworkGUID assignment. This is rare and only happens if a netguid is only serialized client->server (this msg goes server->client to tell client what ID to use in that case)
 DEFINE_CONTROL_CHANNEL_MESSAGE_ONEPARAM(SecurityViolation, 19, FString); // server tells client that it has violated security and has been disconnected
 DEFINE_CONTROL_CHANNEL_MESSAGE_TWOPARAM(GameSpecific, 20, uint8, FString); // custom game-specific message routed to UGameInstance for processing
+DEFINE_CONTROL_CHANNEL_MESSAGE_ZEROPARAM(EncryptionAck, 21);
 
 // 			Beacon control channel flow
 // Client												Server
diff --git a/Engine/Source/Runtime/Launch/Private/LaunchEngineLoop.cpp b/Engine/Source/Runtime/Launch/Private/LaunchEngineLoop.cpp
--- a/Engine/Source/Runtime/Launch/Private/LaunchEngineLoop.cpp
+++ b/Engine/Source/Runtime/Launch/Private/LaunchEngineLoop.cpp
@@ -253,8 +253,17 @@ class FOutputDeviceStdOutput : public FOutputDevice
 public:
 
 	FOutputDeviceStdOutput()
+		: AllowedLogVerbosity(ELogVerbosity::Display)
 	{
-		bAllowLogVerbosity = FParse::Param(FCommandLine::Get(), TEXT("AllowStdOutLogVerbosity"));
+		if (FParse::Param(FCommandLine::Get(), TEXT("AllowStdOutLogVerbosity")))
+		{
+			AllowedLogVerbosity = ELogVerbosity::Log;
+		}
+
+		if (FParse::Param(FCommandLine::Get(), TEXT("FullStdOutLogOutput")))
+		{
+			AllowedLogVerbosity = ELogVerbosity::All;
+		}
 	}
 
 	virtual ~FOutputDeviceStdOutput()
@@ -268,7 +277,7 @@ public:
 
 	virtual void Serialize( const TCHAR* V, ELogVerbosity::Type Verbosity, const class FName& Category ) override
 	{
-		if ( (bAllowLogVerbosity && Verbosity <= ELogVerbosity::Log) || (Verbosity <= ELogVerbosity::Display) )
+		if (Verbosity <= AllowedLogVerbosity)
 		{
 #if PLATFORM_USE_LS_SPEC_FOR_WIDECHAR
 			wprintf(TEXT("%ls\n"), *FOutputDeviceHelper::FormatLogLine(Verbosity, Category, V, GPrintLogTimes));
@@ -280,7 +289,7 @@ public:
 	}
 
 private:
-	bool bAllowLogVerbosity;
+	ELogVerbosity::Type AllowedLogVerbosity;
 };
 
 
diff --git a/Engine/Source/Runtime/PacketHandlers/PacketHandler/Private/PacketHandler.cpp b/Engine/Source/Runtime/PacketHandlers/PacketHandler/Private/PacketHandler.cpp
--- a/Engine/Source/Runtime/PacketHandlers/PacketHandler/Private/PacketHandler.cpp
+++ b/Engine/Source/Runtime/PacketHandlers/PacketHandler/Private/PacketHandler.cpp
@@ -2,11 +2,13 @@
 
 #include "PacketHandler.h"
 #include "PacketAudit.h"
+#include "EncryptionComponent.h"
 
 #include "Misc/ConfigCacheIni.h"
 #include "Modules/ModuleManager.h"
 #include "UObject/ObjectMacros.h"
 #include "UObject/Package.h"
+#include "ConsoleManager.h"
 
 #include "HandlerComponentFactory.h"
 #include "ReliabilityHandlerComponent.h"
@@ -95,6 +97,20 @@ void PacketHandler::Initialize(Handler::Mode InMode, uint32 InMaxPacketBits, boo
 		}
 	}
 
+	// Add encryption component, if configured.
+	FString EncryptionComponentName;
+	if (GConfig->GetString(TEXT("PacketHandlerComponents"), TEXT("EncryptionComponent"), EncryptionComponentName, GEngineIni))
+	{
+		static IConsoleVariable* const AllowEncryptionCVar = IConsoleManager::Get().FindConsoleVariable(TEXT("net.AllowEncryption"));
+		if (AllowEncryptionCVar == nullptr || AllowEncryptionCVar->GetInt() != 0)
+		{
+			EncryptionComponent = StaticCastSharedPtr<FEncryptionComponent>(AddHandler(EncryptionComponentName, true));
+		}
+		else
+		{
+			UE_LOG(PacketHandlerLog, Warning, TEXT("PacketHandler encryption component is configured as %s, but it won't be used because the cvar net.AllowEncryption is false."), *EncryptionComponentName);
+		}
+	}
 
 	bool bEnableReliability = false;
 
@@ -297,6 +313,11 @@ void PacketHandler::OutgoingHigh(FBitWriter& Writer)
 	// @todo #JohnB
 }
 
+TSharedPtr<FEncryptionComponent> PacketHandler::GetEncryptionComponent()
+{
+	return EncryptionComponent;
+}
+
 const ProcessedPacket PacketHandler::Incoming_Internal(uint8* Packet, int32 CountBytes, bool bConnectionless, FString Address)
 {
 	// @todo #JohnB: Try to optimize this function more, seeing as it will be a common codepath DoS attacks pass through
diff --git a/Engine/Source/Runtime/PacketHandlers/PacketHandler/Public/PacketHandler.h b/Engine/Source/Runtime/PacketHandlers/PacketHandler/Public/PacketHandler.h
--- a/Engine/Source/Runtime/PacketHandlers/PacketHandler/Public/PacketHandler.h
+++ b/Engine/Source/Runtime/PacketHandlers/PacketHandler/Public/PacketHandler.h
@@ -15,6 +15,7 @@ PACKETHANDLER_API DECLARE_LOG_CATEGORY_EXTERN(PacketHandlerLog, Log, All);
 
 
 class HandlerComponent;
+class FEncryptionComponent;
 class ReliabilityHandlerComponent;
 
 
@@ -329,6 +330,8 @@ public:
 		return Outgoing_Internal(Packet, CountBits, true, Address);
 	}
 
+	/** Returns a pointer to the component set as the encryption handler, if any. */
+	TSharedPtr<FEncryptionComponent> GetEncryptionComponent();
 
 protected:
 	/**
@@ -510,6 +513,8 @@ private:
 	/** The HandlerComponent pipeline, for processing incoming/outgoing packets */
 	TArray<TSharedPtr<HandlerComponent>> HandlerComponents;
 
+	/** A direct pointer to the component configured as the encryption component. Will also be present in the HandlerComponents array. */
+	TSharedPtr<FEncryptionComponent> EncryptionComponent;
 
 	/** The maximum supported packet size (reflects UNetConnection::MaxPacket) */
 	uint32 MaxPacketBits;
