Fixed crash when running with nullrhi

#tests preflighted Orion Win64 tests

diff --git a/Engine/Source/Runtime/Engine/Private/UserInterface/Canvas.cpp b/Engine/Source/Runtime/Engine/Private/UserInterface/Canvas.cpp
--- a/Engine/Source/Runtime/Engine/Private/UserInterface/Canvas.cpp
+++ b/Engine/Source/Runtime/Engine/Private/UserInterface/Canvas.cpp
@@ -941,7 +941,8 @@ void FCanvas::Clear(const FLinearColor& ClearColor)
 			SCOPED_DRAW_EVENT(RHICmdList, CanvasClear);
 			if (CanvasRenderTarget)
 			{
-				if (CanvasRenderTarget->GetRenderTargetTexture()->GetClearBinding() == FClearValueBinding(ClearColor))
+				// possibility for the RTT to be null for nullrhi
+				if (CanvasRenderTarget->GetRenderTargetTexture() && CanvasRenderTarget->GetRenderTargetTexture()->GetClearBinding() == FClearValueBinding(ClearColor))
 				{
 					// do fast clear
 					SetRenderTarget(RHICmdList, CanvasRenderTarget->GetRenderTargetTexture(), FTextureRHIRef(), ESimpleRenderTargetMode::EClearColorAndDepth);
