Updating comments to improve API documentation/tooltips on asset management framework.

diff --git a/Engine/Source/Runtime/CoreUObject/Public/UObject/PrimaryAssetId.h b/Engine/Source/Runtime/CoreUObject/Public/UObject/PrimaryAssetId.h
--- a/Engine/Source/Runtime/CoreUObject/Public/UObject/PrimaryAssetId.h
+++ b/Engine/Source/Runtime/CoreUObject/Public/UObject/PrimaryAssetId.h
@@ -72,9 +72,11 @@ private:
  */
 struct FPrimaryAssetId
 {
-	/** The Type of this object, by default it's base class's name */
+	/** An FName describing the logical type of this object, usually the name of a base UClass. For example, any Blueprint derived from APawn will have a Primary Asset Type of "Pawn".
+	"PrimaryAssetType:PrimaryAssetName" should form a unique name across your project. */
 	FPrimaryAssetType PrimaryAssetType;
-	/** The Name of this object, by default it's short name */
+	/** An FName describing this asset. This is usually the short name of the object, but could be a full asset path for things like maps, or objects with GetPrimaryId() overridden.
+	"PrimaryAssetType:PrimaryAssetName" should form a unique name across your project. */
 	FName PrimaryAssetName;
 
 	/** Static names to represent the AssetRegistry tags for the above data */
diff --git a/Engine/Source/Runtime/Engine/Classes/Engine/AssetManagerTypes.h b/Engine/Source/Runtime/Engine/Classes/Engine/AssetManagerTypes.h
--- a/Engine/Source/Runtime/Engine/Classes/Engine/AssetManagerTypes.h
+++ b/Engine/Source/Runtime/Engine/Classes/Engine/AssetManagerTypes.h
@@ -13,16 +13,16 @@
 UENUM()
 enum class EPrimaryAssetCookRule : uint8
 {
-	/** Nothing specifically known about this asset, it will cook if something else depends on it */
+	/** Nothing is known about this asset specifically. It will cook if something else depends on it. */
 	Unknown,
 
-	/** Asset should never be cooked/shipped in any situation, error if it something depends on it */
+	/** Asset should never be cooked/shipped in any situation. An error will be generated if something depends on it. */
 	NeverCook,
 
-	/** Asset can be cooked for development and testing, but should never be shipped in a production build */
+	/** Asset can be cooked for development and testing, but should never be shipped in a production build. */
 	DevelopmentCook,
 
-	/** Asset will always be cooked, in both production and development */
+	/** Asset will always be cooked, in both production and development. */
 	AlwaysCook,
 };
 
@@ -32,19 +32,19 @@ struct FPrimaryAssetRules
 {
 	GENERATED_BODY()
 	
-	/** Primary Assets with a higher priority will take precedence over lower priorities when assigning management for referenced assets. If priorities match, both will manage the same secondary asset */
+	/** Primary Assets with a higher priority will take precedence over lower priorities when assigning management for referenced assets. If priorities match, both will manage the same Secondary Asset. */
 	UPROPERTY(EditAnywhere, Category = Rules)
 	int32 Priority;
 
-	/** If true, will apply this rule to all referenced secondary assets recursively, if they are not managed by a higher priority Primary Asset */
+	/** If true, this rule will apply to all referenced Secondary Assets recursively, as long as they are not managed by a higher-priority Primary Asset. */
 	UPROPERTY(EditAnywhere, Category = Rules)
 	bool bApplyRecursively;
 
-	/** Assets will be put into this chunk id specifically, if set to something other than -1 */
-	UPROPERTY(EditAnywhere, Category = Rules)
+	/** Assets will be put into this Chunk ID specifically, if set to something other than -1. The default Chunk is Chunk 0. */
+	UPROPERTY(EditAnywhere, Category = Rules, meta = (DisplayName = "Chunk ID"))
 	int32 ChunkId;
 
-	/** Rule describing when this asset should be cooked */
+	/** Rule describing when this asset should be cooked. */
 	UPROPERTY(EditAnywhere, Category = Rules)
 	EPrimaryAssetCookRule CookRule;
 
@@ -61,17 +61,17 @@ struct FPrimaryAssetRules
 			&& CookRule == Other.CookRule;
 	}
 
-	/** Checks if all rules are the same as the default, if so this will be ignored */
+	/** Checks if all rules are the same as the default. If so this will be ignored. */
 	ENGINE_API bool IsDefault() const;
 
-	/** Override non-default rules from an override struct */
+	/** Override non-default rules from an override struct. */
 	ENGINE_API void OverrideRules(const FPrimaryAssetRules& OverrideRules);
 
-	/** Propagate cook rules from parent to child, won't override non default values */
+	/** Propagate cook rules from parent to child, won't override non-default values. */
 	ENGINE_API void PropagateCookRules(const FPrimaryAssetRules& ParentRules);
 };
 
-/** Structure with publicly exposed information about an asset type. These can be loaded out of a config file */
+/** Structure with publicly exposed information about an asset type. These can be loaded out of a config file. */
 USTRUCT()
 struct FPrimaryAssetTypeInfo
 {
diff --git a/Engine/Source/Runtime/Engine/Classes/Engine/StreamableManager.h b/Engine/Source/Runtime/Engine/Classes/Engine/StreamableManager.h
--- a/Engine/Source/Runtime/Engine/Classes/Engine/StreamableManager.h
+++ b/Engine/Source/Runtime/Engine/Classes/Engine/StreamableManager.h
@@ -46,31 +46,32 @@ struct ENGINE_API FStreamableHandle : public TSharedFromThis<FStreamableHandle,
 		return bStalled;
 	}
 
-	/** Returns true if this is a combined handle that depends on child handles */
+	/** Returns true if this is a combined handle that depends on child handles. */
 	bool IsCombinedHandle() const
 	{
 		return bIsCombinedHandle;
 	}
 
-	/** Returns the debug name for this handle */
+	/** Returns the debug name for this handle. */
 	const FString& GetDebugName() const
 	{
 		return DebugName;
 	}
 
-	/** Returns the streaming priority */
+	/** Returns the streaming priority. */
 	TAsyncLoadPriority GetPriority() const
 	{
 		return Priority;
 	}
 
-	/** Release this handle, called from normal gameplay code to indicate that the loaded assets are no longer needed. If called before completion will release on completion */
+	/** Release this handle. This can be called from normal gameplay code to indicate that the loaded assets are no longer needed. Will be called implicitly if all shared pointers
+	to this handle are destroyed. If called before the completion delegate, the release will be delayed until after completion. */
 	void ReleaseHandle();
 
 	/** Cancel a request, callable from within the manager or externally. Will stop delegate from being called */
 	void CancelHandle();
 
-	/** Tells a stalled handle to start it's actual request */
+	/** Tells a stalled handle to start its actual request. */
 	void StartStalledHandle();
 
 	/** Bind delegate that is called when load completes, only works if loading is in progress. This will overwrite any already bound delegate! */
@@ -82,7 +83,8 @@ struct ENGINE_API FStreamableHandle : public TSharedFromThis<FStreamableHandle,
 	/** Bind delegate that is called periodically as delegate updates, only works if loading is in progress. This will overwrite any already bound delegate! */
 	bool BindUpdateDelegate(FStreamableUpdateDelegate NewDelegate);
 
-	/** Wait for this load to finish */
+	/** Blocks until the requested assets have loaded. This pushes the requested asset to the top of the priority list, but does not flush all async loading, usually resulting
+	in faster completion than a LoadObject call. */
 	EAsyncPackageState::Type WaitUntilComplete(float Timeout = 0.0f);
 
 	/** Gets list of assets references this load was started with. This will be the paths before redirectors, and not all of these are guaranteed to be loaded */
@@ -97,7 +99,7 @@ struct ENGINE_API FStreamableHandle : public TSharedFromThis<FStreamableHandle,
 	/** Returns number of assets that have completed loading out of initial list, failed loads will count as loaded */
 	void GetLoadedCount(int32& LoadedCount, int32& RequestedCount) const;
 
-	/** Returns % complete as a float 0-1 */
+	/** Returns progress as a value between 0.0 and 1.0. */
 	float GetProgress() const;
 
 	/** Get the StreamableManager for this handle */
@@ -199,7 +201,7 @@ struct ENGINE_API FStreamableManager : public FGCObject
 	static const TAsyncLoadPriority AsyncLoadHighPriority = 100;
 
 	/** 
-	 * Request streaming of one or more target objects, and call a delegate on completion. 
+	 * This is the primary streamable operation. Requests streaming of one or more target objects. When complete, a delegate function is called. Returns a Streamable Handle.
 	 *
 	 * @param TargetsToStream		Assets to load off disk
 	 * @param DelegateToCall		Delegate to call when load finishes. Will be called on the next tick if asset is already loaded, or many seconds later
@@ -211,13 +213,12 @@ struct ENGINE_API FStreamableManager : public FGCObject
 	TSharedPtr<FStreamableHandle> RequestAsyncLoad(const TArray<FSoftObjectPath>& TargetsToStream, FStreamableDelegate DelegateToCall = FStreamableDelegate(), TAsyncLoadPriority Priority = DefaultAsyncLoadPriority, bool bManageActiveHandle = false, bool bStartStalled = false, const FString& DebugName = TEXT("RequestAsyncLoad ArrayDelegate"));
 	TSharedPtr<FStreamableHandle> RequestAsyncLoad(const FSoftObjectPath& TargetToStream, FStreamableDelegate DelegateToCall = FStreamableDelegate(), TAsyncLoadPriority Priority = DefaultAsyncLoadPriority, bool bManageActiveHandle = false, bool bStartStalled = false, const FString& DebugName = TEXT("RequestAsyncLoad SingleDelegate"));
 
-	/** Lambda Wrappers. Be aware that Callback may go off multiple seconds in the future */
+	/** Lambda Wrappers. Be aware that Callback may go off multiple seconds in the future. */
 	TSharedPtr<FStreamableHandle> RequestAsyncLoad(const TArray<FSoftObjectPath>& TargetsToStream, TFunction<void()>&& Callback, TAsyncLoadPriority Priority = DefaultAsyncLoadPriority, bool bManageActiveHandle = false, bool bStartStalled = false, const FString& DebugName = TEXT("RequestAsyncLoad ArrayLambda"));
 	TSharedPtr<FStreamableHandle> RequestAsyncLoad(const FSoftObjectPath& TargetToStream, TFunction<void()>&& Callback, TAsyncLoadPriority Priority = DefaultAsyncLoadPriority, bool bManageActiveHandle = false, bool bStartStalled = false, const FString& DebugName = TEXT("RequestAsyncLoad SingleLambda"));
 
 	/** 
-	 * Synchronously load a set of assets, and return a handle.
-	 * This can be very slow and may stall the game thread for several seconds.
+	 * Synchronously load a set of assets, and return a handle. This can be very slow and may stall the game thread for several seconds.
 	 * 
 	 * @param TargetsToStream		Assets to load off disk
 	 * @param bManageActiveHandle	If true, the manager will keep the streamable handle active until explicitly released
@@ -227,8 +228,7 @@ struct ENGINE_API FStreamableManager : public FGCObject
 	TSharedPtr<FStreamableHandle> RequestSyncLoad(const FSoftObjectPath& TargetToStream, bool bManageActiveHandle = false, const FString& DebugName = TEXT("RequestSyncLoad Single"));
 
 	/** 
-	 * Synchronously load the referred asset and return the loaded object, or nullptr if it can't be found.
-	 * This can be very slow and may stall the game thread for several seconds.
+	 * Synchronously load the referred asset and return the loaded object, or nullptr if it can't be found. This can be very slow and may stall the game thread for several seconds.
 	 * 
 	 * @param Target				Specific asset to load off disk
 	 * @param bManageActiveHandle	If true, the manager will keep the streamable handle active until explicitly released
