Fix for FRepChangedPropertyTracker::SetExternalData not clearing the data array before appending. Would cause corrupted replay data if called twice for the same object within one replay frame.

diff --git a/Engine/Source/Runtime/Engine/Public/Net/RepLayout.h b/Engine/Source/Runtime/Engine/Public/Net/RepLayout.h
--- a/Engine/Source/Runtime/Engine/Public/Net/RepLayout.h
+++ b/Engine/Source/Runtime/Engine/Public/Net/RepLayout.h
@@ -60,8 +60,10 @@ public:
 	virtual void SetExternalData( const uint8* Src, const int32 NumBits ) override
 	{
 		ExternalDataNumBits = NumBits;
-		ExternalData.AddUninitialized( ( NumBits + 7 ) >> 3 );
-		FMemory::Memcpy( ExternalData.GetData(), Src, ( NumBits + 7 ) >> 3 );
+		const int32 NumBytes = ( NumBits + 7 ) >> 3;
+		ExternalData.Reset( NumBytes );
+		ExternalData.AddUninitialized( NumBytes );
+		FMemory::Memcpy( ExternalData.GetData(), Src, NumBytes );
 	}
 
 	virtual bool IsReplay() const override
