Merging CL 3593597 from //UE4/Dev-Editor to Main (//UE4/Main)

Only update instance in non archetype or CDO.

diff --git a/Engine/Source/Runtime/Engine/Classes/Components/InstancedStaticMeshComponent.h b/Engine/Source/Runtime/Engine/Classes/Components/InstancedStaticMeshComponent.h
--- a/Engine/Source/Runtime/Engine/Classes/Components/InstancedStaticMeshComponent.h
+++ b/Engine/Source/Runtime/Engine/Classes/Components/InstancedStaticMeshComponent.h
@@ -272,7 +272,7 @@ public:
 	void ClearInstanceSelection();
 
 	/** Initialize the Per Instance Render Data */
-	void InitPerInstanceRenderData(FStaticMeshInstanceData* InSharedInstanceBufferData = nullptr);
+	void InitPerInstanceRenderData(bool InitializeFromCurrentData, FStaticMeshInstanceData* InSharedInstanceBufferData = nullptr);
 
 	/** Transfers ownership of instance render data to a render thread. Instance render data will be released in scene proxy destructor or on render thread task. */
 	void ReleasePerInstanceRenderData();
diff --git a/Engine/Source/Runtime/Engine/Private/HierarchicalInstancedStaticMesh.cpp b/Engine/Source/Runtime/Engine/Private/HierarchicalInstancedStaticMesh.cpp
--- a/Engine/Source/Runtime/Engine/Private/HierarchicalInstancedStaticMesh.cpp
+++ b/Engine/Source/Runtime/Engine/Private/HierarchicalInstancedStaticMesh.cpp
@@ -1822,7 +1822,10 @@ void UHierarchicalInstancedStaticMeshComponent::PostDuplicate(bool bDuplicateFor
 {
 	Super::PostDuplicate(bDuplicateForPIE);
 
-	BuildTreeIfOutdated(false, false);
+	if (!HasAnyFlags(RF_ClassDefaultObject | RF_ArchetypeObject) && bDuplicateForPIE)
+	{
+		BuildTreeIfOutdated(false, false);
+	}
 }
 
 void UHierarchicalInstancedStaticMeshComponent::RemoveInstanceInternal(int32 InstanceIndex)
@@ -2162,7 +2165,7 @@ void UHierarchicalInstancedStaticMeshComponent::BuildTree()
 
 		if (!PerInstanceRenderData.IsValid())
 		{
-			InitPerInstanceRenderData();
+			InitPerInstanceRenderData(false);
 		}
 
 		// Resync RenderData with newly built cluster tree so we take into account the newly generated InstanceReorderTable generated from the cluster tree
@@ -2350,29 +2353,29 @@ bool UHierarchicalInstancedStaticMeshComponent::BuildTreeIfOutdated(bool Async,
 		|| UnbuiltInstanceBoundsList.Num() > 0
 		|| GetLinkerUE4Version() < VER_UE4_REBUILD_HIERARCHICAL_INSTANCE_TREES)
 	{
-		if (GetStaticMesh())
+		if (!GetStaticMesh()->HasAnyFlags(RF_NeedLoad)) // we can build the tree if the static mesh is not even loaded, and we can't call PostLoad as the load is not even done
 		{
 			GetStaticMesh()->ConditionalPostLoad();
-		}
 
-		if (Async)
-		{
-			if (IsAsyncBuilding())
+			if (Async)
 			{
-				// invalidate the results of the current async build we need to modify the tree
-				bConcurrentRemoval = true;
+				if (IsAsyncBuilding())
+				{
+					// invalidate the results of the current async build we need to modify the tree
+					bConcurrentRemoval = true;
+				}
+				else
+				{
+					BuildTreeAsync();
+				}
 			}
 			else
 			{
-				BuildTreeAsync();
+				BuildTree();
 			}
-		}
-		else
-		{
-			BuildTree();
-		}
 
-		return true;
+			return true;
+		}
 	}
 
 	return false;
@@ -2507,31 +2510,34 @@ void UHierarchicalInstancedStaticMeshComponent::PostLoad()
 
 	Super::PostLoad();
 
-	NumBuiltRenderInstances = ClusterTreePtr.IsValid() && ClusterTreePtr->Num() > 0 ? (*ClusterTreePtr.Get())[0].LastInstance - (*ClusterTreePtr.Get())[0].FirstInstance + 1 : 0;
-
-	if (bEnableDensityScaling && GetWorld() && GetWorld()->IsGameWorld())
+	if (!HasAnyFlags(RF_ClassDefaultObject | RF_ArchetypeObject))
 	{
-		const float ScalabilityDensity = FMath::Clamp(CVarFoliageDensityScale.GetValueOnGameThread(), 0.0f, 1.0f);
-		if (ScalabilityDensity == 0)
-		{
-			// exclude all instances
-			ExcludedDueToDensityScaling.Init(true, PerInstanceSMData.Num());
-			NumBuiltRenderInstances = 0;
-		}
-		else if (ScalabilityDensity > 0.0f && ScalabilityDensity < 1.0f)
-		{
-			FRandomStream Rand(InstancingRandomSeed);
-			ExcludedDueToDensityScaling.Init(false, PerInstanceSMData.Num());
+		NumBuiltRenderInstances = ClusterTreePtr.IsValid() && ClusterTreePtr->Num() > 0 ? (*ClusterTreePtr.Get())[0].LastInstance - (*ClusterTreePtr.Get())[0].FirstInstance + 1 : 0;
 
-			for (int32 i = 0; i < ExcludedDueToDensityScaling.Num(); ++i)
+		if (bEnableDensityScaling && GetWorld() && GetWorld()->IsGameWorld())
+		{
+			const float ScalabilityDensity = FMath::Clamp(CVarFoliageDensityScale.GetValueOnGameThread(), 0.0f, 1.0f);
+			if (ScalabilityDensity == 0)
+			{
+				// exclude all instances
+				ExcludedDueToDensityScaling.Init(true, PerInstanceSMData.Num());
+				NumBuiltRenderInstances = 0;
+			}
+			else if (ScalabilityDensity > 0.0f && ScalabilityDensity < 1.0f)
 			{
-				ExcludedDueToDensityScaling[i] = (Rand.FRand() > ScalabilityDensity);
-			}			
+				FRandomStream Rand(InstancingRandomSeed);
+				ExcludedDueToDensityScaling.Init(false, PerInstanceSMData.Num());
+
+				for (int32 i = 0; i < ExcludedDueToDensityScaling.Num(); ++i)
+				{
+					ExcludedDueToDensityScaling[i] = (Rand.FRand() > ScalabilityDensity);
+				}
+			}
 		}
-	}
 
-	// If any of the data is out of sync, build the tree now!
-	BuildTreeIfOutdated(true, false);	
+		// If any of the data is out of sync, build the tree now!
+		BuildTreeIfOutdated(true, false);
+	}
 }
 
 static void GatherInstanceTransformsInArea(const UHierarchicalInstancedStaticMeshComponent& Component, const FBox& AreaBox, int32 Child, TArray<FTransform>& InstanceData)
diff --git a/Engine/Source/Runtime/Engine/Private/InstancedStaticMesh.cpp b/Engine/Source/Runtime/Engine/Private/InstancedStaticMesh.cpp
--- a/Engine/Source/Runtime/Engine/Private/InstancedStaticMesh.cpp
+++ b/Engine/Source/Runtime/Engine/Private/InstancedStaticMesh.cpp
@@ -100,10 +100,6 @@ void FStaticMeshInstanceBuffer::UpdateInstanceData(UInstancedStaticMeshComponent
 	// Allocate the vertex data storage type.
 	InstanceData->AllocateInstances(NumInstances, false);
 
-	// Setup our random number generator such that random values are generated consistently for any
-	// given instance index between reattaches
-	FRandomStream RandomStream(InComponent->InstancingRandomSeed);
-
 	const FMeshMapBuildData* MeshMapBuildData = NULL;
 
 	if (InComponent->LODData.Num() > 0)
@@ -161,7 +157,7 @@ void FStaticMeshInstanceBuffer::UpdateInstanceData(UInstancedStaticMeshComponent
 	InComponent->RemovedInstances.Reset();
 }
 
-void FStaticMeshInstanceBuffer::Init(UInstancedStaticMeshComponent* InComponent, const TArray<TRefCountPtr<HHitProxy> >& InHitProxies)
+void FStaticMeshInstanceBuffer::Init(UInstancedStaticMeshComponent* InComponent, const TArray<TRefCountPtr<HHitProxy> >& InHitProxies, bool InitializeBufferFromData)
 {
 	QUICK_SCOPE_CYCLE_COUNTER(STAT_FStaticMeshInstanceBuffer_Init);
 	
@@ -172,7 +168,15 @@ void FStaticMeshInstanceBuffer::Init(UInstancedStaticMeshComponent* InComponent,
 
 	SetupCPUAccess(InComponent);
 
-	UpdateInstanceData(InComponent, InHitProxies, 0, InComponent->PerInstanceSMData.Num());
+	// Setup our random number generator such that random values are generated consistently for any
+	// given instance index between reattaches
+	check(InComponent->InstancingRandomSeed != 0)
+	RandomStream.Initialize(InComponent->InstancingRandomSeed);
+
+	if (InitializeBufferFromData)
+	{
+		UpdateInstanceData(InComponent, InHitProxies, 0, InComponent->PerInstanceSMData.Num());
+	}
 
 	float ThisTime = (StartTime - FPlatformTime::Seconds()) * 1000.0f;
 	if (ThisTime > 30.0f)
@@ -184,6 +188,11 @@ void FStaticMeshInstanceBuffer::Init(UInstancedStaticMeshComponent* InComponent,
 void FStaticMeshInstanceBuffer::InitFromPreallocatedData(UInstancedStaticMeshComponent* InComponent, FStaticMeshInstanceData& Other)
 {
 	QUICK_SCOPE_CYCLE_COUNTER(STAT_FStaticMeshInstanceBuffer_InitFromPreallocatedData);
+
+	// Setup our random number generator such that random values are generated consistently for any
+	// given instance index between reattaches
+	RandomStream.Initialize(InComponent->InstancingRandomSeed);
+
 	const uint32 NewNumInstances = Other.NumInstances();
 	AllocateData(Other);
 	NumInstances = NewNumInstances;
@@ -254,6 +263,11 @@ void FStaticMeshInstanceBuffer::UpdateRHIVertexBuffer(const TSet<int32>& InIndex
 {
 	check(IsInRenderingThread());
 
+	if (InstanceData->NumInstances() == 0)
+	{
+		return;
+	}
+
 	//TODO: to uncomment when RHI interface support proper vertex buffer partial update lock/unlock
 //	if (InIndexList.Num() == 0 || InIndexList.Num() == InstanceData->NumInstances())
 	{
@@ -872,14 +886,18 @@ UInstancedStaticMeshComponent::UInstancedStaticMeshComponent(const FObjectInitia
 	PhysicsSerializer = ObjectInitializer.CreateDefaultSubobject<UPhysicsSerializer>(this, TEXT("PhysicsSerializer"));
 	bDisallowMeshPaintPerInstance = true;
 
-	if (FApp::CanEverRender() && !HasAnyFlags(RF_ClassDefaultObject))
+	if (FApp::CanEverRender() && !HasAnyFlags(RF_ClassDefaultObject|RF_ArchetypeObject))
 	{
-		InitPerInstanceRenderData();
+		while (InstancingRandomSeed == 0)
+		{
+			InstancingRandomSeed = FMath::Rand();
+		}
 	}
 }
 
 UInstancedStaticMeshComponent::~UInstancedStaticMeshComponent()
 {
+	ReleasePerInstanceRenderData();
 	FlushAsyncBuildInstanceBufferTask();
 }
 
@@ -1015,9 +1033,7 @@ void UInstancedStaticMeshComponent::ApplyComponentInstanceData(FInstancedStaticM
 
 	// Force recreation of the render data
 	ReleasePerInstanceRenderData();
-	InitPerInstanceRenderData();
-	PerInstanceRenderData->UpdateInstanceData(this, 0, PerInstanceSMData.Num());
-
+	InitPerInstanceRenderData(true);
 	MarkRenderStateDirty();
 #endif
 }
@@ -1823,7 +1839,7 @@ void UInstancedStaticMeshComponent::ClearInstances()
 	// Indicate we need to update render state to reflect changes
 	bPerInstanceRenderDataWasPrebuilt = false;
 	ReleasePerInstanceRenderData();
-	InitPerInstanceRenderData();
+	InitPerInstanceRenderData(false);
 	MarkRenderStateDirty();
 
 	UNavigationSystem::UpdateComponentInNavOctree(*this);
@@ -1863,8 +1879,13 @@ void UInstancedStaticMeshComponent::SetupNewInstanceData(FInstancedStaticMeshIns
 	}
 }
 
-void UInstancedStaticMeshComponent::InitPerInstanceRenderData(FStaticMeshInstanceData* InSharedInstanceBufferData)
+void UInstancedStaticMeshComponent::InitPerInstanceRenderData(bool InitializeFromCurrentData, FStaticMeshInstanceData* InSharedInstanceBufferData)
 {
+	while (InstancingRandomSeed == 0)
+	{
+		InstancingRandomSeed = FMath::Rand();
+	}
+
 	UWorld* World = GetWorld();
 
 	ERHIFeatureLevel::Type FeatureLevel = World != nullptr ? World->FeatureLevel : GMaxRHIFeatureLevel;
@@ -1890,7 +1911,7 @@ void UInstancedStaticMeshComponent::InitPerInstanceRenderData(FStaticMeshInstanc
 		}
 		else
 		{
-			PerInstanceRenderData = MakeShareable(new FPerInstanceRenderData(this, FeatureLevel, IsDynamic));
+			PerInstanceRenderData = MakeShareable(new FPerInstanceRenderData(this, FeatureLevel, IsDynamic, InitializeFromCurrentData));
 		}
 	}
 }
@@ -1899,11 +1920,11 @@ void UInstancedStaticMeshComponent::PostLoad()
 {
 	Super::PostLoad();
 
-	if (!HasAnyFlags(RF_ClassDefaultObject))
+	if (!HasAnyFlags(RF_ClassDefaultObject|RF_ArchetypeObject))
 	{
-		InitPerInstanceRenderData();
+		InitPerInstanceRenderData(false);
 
-		if (PerInstanceSMData.Num() > 0)
+		if (PerInstanceSMData.Num() > 0 && PerInstanceRenderData->InstanceBuffer.GetNumInstances() == 0) // only load the data if it's not already loaded
 		{
 			UWorld* World = GetWorld();
 
@@ -1911,10 +1932,7 @@ void UInstancedStaticMeshComponent::PostLoad()
 			if (CVarASyncInstaneBufferConversion.GetValueOnGameThread() > 0 && World != nullptr && World->IsGameWorld())
 			{
 				World->AsyncPreRegisterLevelStreamingTasks.Increment();
-				while (InstancingRandomSeed == 0)
-				{
-					InstancingRandomSeed = FMath::Rand();
-				}
+
 				AsyncBuildInstanceBufferTask = new FAsyncTask<FAsyncBuildInstanceBuffer>(this, World);
 				AsyncBuildInstanceBufferTask->StartBackgroundTask();
 			}
@@ -2034,20 +2052,18 @@ void UInstancedStaticMeshComponent::GetResourceSizeEx(FResourceSizeEx& Cumulativ
 
 void UInstancedStaticMeshComponent::BeginDestroy()
 {
+	ReleasePerInstanceRenderData();
 	bPerInstanceRenderDataWasPrebuilt = false;
 	Super::BeginDestroy();
-	ReleasePerInstanceRenderData();
 }
 
 void UInstancedStaticMeshComponent::PostDuplicate(bool bDuplicateForPIE)
 {
 	Super::PostDuplicate(bDuplicateForPIE);
 
-	InitPerInstanceRenderData();
-
-	if (PerInstanceSMData.Num() > 0)
+	if (!HasAnyFlags(RF_ClassDefaultObject|RF_ArchetypeObject) && bDuplicateForPIE)
 	{
-		PerInstanceRenderData->UpdateInstanceData(this, 0, PerInstanceSMData.Num());
+		InitPerInstanceRenderData(true);		
 	}
 }
 
@@ -2092,13 +2108,7 @@ void UInstancedStaticMeshComponent::PostEditChangeChainProperty(FPropertyChanged
 		{
 			// Force a full refresh of the instance buffer
 			ReleasePerInstanceRenderData();
-			InitPerInstanceRenderData();
-
-			if (PerInstanceRenderData.IsValid())
-			{
-				PerInstanceRenderData->UpdateInstanceData(this, 0, PerInstanceSMData.Num());
-			}
-
+			InitPerInstanceRenderData(true);
 			MarkRenderStateDirty();
 		}
 	}
diff --git a/Engine/Source/Runtime/Engine/Private/InstancedStaticMesh.h b/Engine/Source/Runtime/Engine/Private/InstancedStaticMesh.h
--- a/Engine/Source/Runtime/Engine/Private/InstancedStaticMesh.h
+++ b/Engine/Source/Runtime/Engine/Private/InstancedStaticMesh.h
@@ -67,7 +67,7 @@ public:
 	 * @param InComponent - The owning component
 	 * @param InHitProxies - Array of hit proxies for each instance, if desired.
 	 */
-	void Init(UInstancedStaticMeshComponent* InComponent, const TArray<TRefCountPtr<HHitProxy> >& InHitProxies);
+	void Init(UInstancedStaticMeshComponent* InComponent, const TArray<TRefCountPtr<HHitProxy> >& InHitProxies, bool InitializeBufferFromData);
 
 	/**
 	 * Update the specified instance range (called on game thread)
@@ -145,6 +145,9 @@ private:
 	/** Is the vertex buffer considered dynamic */
 	bool IsDynamic;
 
+	/** Is used to generate random value for each instance consistently between Update call */
+	FRandomStream RandomStream;
+
 	/** Allocates the vertex data storage type. */
 	void AllocateData();
 
@@ -334,26 +337,28 @@ private:
 struct FPerInstanceRenderData
 {
 	// Should be always constructed on main thread
-	FPerInstanceRenderData(UInstancedStaticMeshComponent* InComponent, ERHIFeatureLevel::Type InFeaureLevel, bool IsDynamicBuffer)
+	FPerInstanceRenderData(UInstancedStaticMeshComponent* InComponent, ERHIFeatureLevel::Type InFeaureLevel, bool IsDynamicBuffer, bool InitializeBufferFromData)
 		: InstanceBuffer(InFeaureLevel, IsDynamicBuffer)
 	{
-		// Create hit proxies for each instance if the component wants
-		if (GIsEditor && InComponent->bHasPerInstanceHitProxies)
+		if (InitializeBufferFromData)
 		{
-			QUICK_SCOPE_CYCLE_COUNTER(STAT_FPerInstanceRenderData_HitProxies);
-			HitProxies.Reserve(InComponent->PerInstanceSMData.Num());
-			
-			for (int32 InstanceIdx=0; InstanceIdx < InComponent->PerInstanceSMData.Num(); InstanceIdx++)
+			// Create hit proxies for each instance if the component wants
+			if (GIsEditor && InComponent->bHasPerInstanceHitProxies)
 			{
-				HitProxies.Add(new HInstancedStaticMeshInstance(InComponent, InstanceIdx));
+				QUICK_SCOPE_CYCLE_COUNTER(STAT_FPerInstanceRenderData_HitProxies);
+				HitProxies.Reserve(InComponent->PerInstanceSMData.Num());
+
+				for (int32 InstanceIdx = 0; InstanceIdx < InComponent->PerInstanceSMData.Num(); InstanceIdx++)
+				{
+					HitProxies.Add(new HInstancedStaticMeshInstance(InComponent, InstanceIdx));
+				}
 			}
 		}
-			
+
 		// initialize the instance buffer from the component's instances
-		InstanceBuffer.Init(InComponent, HitProxies);
+		InstanceBuffer.Init(InComponent, HitProxies, InitializeBufferFromData);
 
 		InitResource();
-
 	}
 
 	FPerInstanceRenderData(UInstancedStaticMeshComponent* InComponent, FStaticMeshInstanceData& Other, ERHIFeatureLevel::Type InFeaureLevel, bool IsDynamicBuffer)
@@ -502,7 +507,7 @@ public:
 	 */
 	void UpdatePerInstanceRenderData(const TSet<int32>& InNeedUpdatingInstanceIndexList)
 	{
-		if (PerInstanceRenderData.IsValid())
+		if (PerInstanceRenderData.IsValid() && PerInstanceRenderData->InstanceBuffer.GetNumInstances() > 0)
 		{
 			if (!PerInstanceRenderData->InstanceBuffer.IsInitialized())
 			{
diff --git a/Engine/Source/Runtime/Foliage/Private/InstancedFoliage.cpp b/Engine/Source/Runtime/Foliage/Private/InstancedFoliage.cpp
--- a/Engine/Source/Runtime/Foliage/Private/InstancedFoliage.cpp
+++ b/Engine/Source/Runtime/Foliage/Private/InstancedFoliage.cpp
@@ -764,12 +764,12 @@ void FFoliageMeshInfo::CreateNewComponent(AInstancedFoliageActor* InIFA, const U
 	}
 
 	UFoliageInstancedStaticMeshComponent* FoliageComponent = NewObject<UFoliageInstancedStaticMeshComponent>(InIFA, ComponentClass, NAME_None, RF_Transactional);
+	FoliageComponent->InitPerInstanceRenderData(false);
 
 	Component = FoliageComponent;
 	Component->SetStaticMesh(InSettings->GetStaticMesh());
 	Component->bSelectable = true;
 	Component->bHasPerInstanceHitProxies = true;
-	Component->InstancingRandomSeed = FMath::Rand();
 
 	if (Component->GetStaticMesh() != nullptr)
 	{
@@ -1282,7 +1282,7 @@ void FFoliageMeshInfo::ReapplyInstancesToComponent()
 		const bool bWasRegistered = Component->IsRegistered();
 		Component->UnregisterComponent();
 		Component->ClearInstances();
-		Component->InitPerInstanceRenderData();
+		Component->InitPerInstanceRenderData(false);
 
 		Component->bAutoRebuildTreeOnInstanceChanges = false;
 
diff --git a/Engine/Source/Runtime/Landscape/Private/LandscapeGrass.cpp b/Engine/Source/Runtime/Landscape/Private/LandscapeGrass.cpp
--- a/Engine/Source/Runtime/Landscape/Private/LandscapeGrass.cpp
+++ b/Engine/Source/Runtime/Landscape/Private/LandscapeGrass.cpp
@@ -2395,7 +2395,7 @@ void ALandscapeProxy::UpdateGrass(const TArray<FVector>& Cameras, bool bForceSyn
 
 						if (!HierarchicalInstancedStaticMeshComponent->PerInstanceRenderData.IsValid())
 						{
-							HierarchicalInstancedStaticMeshComponent->InitPerInstanceRenderData(&Inner.Builder->InstanceBuffer);
+							HierarchicalInstancedStaticMeshComponent->InitPerInstanceRenderData(true, &Inner.Builder->InstanceBuffer);
 						}
 						else
 						{
