Copying //UE4/Orion-Staging to //UE4/Main (Source: //Orion/Dev-General @ 3606378)

diff --git a/Engine/Source/Editor/ViewportInteraction/ViewportWorldInteraction.cpp b/Engine/Source/Editor/ViewportInteraction/ViewportWorldInteraction.cpp
--- a/Engine/Source/Editor/ViewportInteraction/ViewportWorldInteraction.cpp
+++ b/Engine/Source/Editor/ViewportInteraction/ViewportWorldInteraction.cpp
@@ -2608,6 +2608,10 @@ void UViewportWorldInteraction::RefreshTransformGizmo( const bool bNewObjectsSel
 		{
 			ViewerLocation = GetHeadTransform().GetLocation();
 		}
+		else if (DefaultOptionalViewportClient != nullptr)
+		{
+			ViewerLocation = DefaultOptionalViewportClient->GetViewLocation();
+		}
 		else if( GCurrentLevelEditingViewportClient != nullptr )
 		{
 			ViewerLocation = GCurrentLevelEditingViewportClient->GetViewLocation();
@@ -2859,19 +2863,19 @@ void UViewportWorldInteraction::ApplyVelocityDamping( FVector& Velocity, const b
 void UViewportWorldInteraction::CycleTransformGizmoCoordinateSpace()
 {
 	const bool bGetRawValue = true;
-	const ECoordSystem CurrentCoordSystem = GLevelEditorModeTools().GetCoordSystem( bGetRawValue );
+	const ECoordSystem CurrentCoordSystem = GetModeTools().GetCoordSystem( bGetRawValue );
 	SetTransformGizmoCoordinateSpace( CurrentCoordSystem == COORD_World ? COORD_Local : COORD_World );
 }
 
 void UViewportWorldInteraction::SetTransformGizmoCoordinateSpace( const ECoordSystem NewCoordSystem )
 {
-	GLevelEditorModeTools().SetCoordSystem( NewCoordSystem );
+	GetModeTools().SetCoordSystem( NewCoordSystem );
 }
 
 ECoordSystem UViewportWorldInteraction::GetTransformGizmoCoordinateSpace() const
 {
 	const bool bGetRawValue = false;
-	const ECoordSystem CurrentCoordSystem = GLevelEditorModeTools().GetCoordSystem( bGetRawValue );
+	const ECoordSystem CurrentCoordSystem = GetModeTools().GetCoordSystem( bGetRawValue );
 	return CurrentCoordSystem;
 }
 
@@ -3042,7 +3046,7 @@ EGizmoHandleTypes UViewportWorldInteraction::GetCurrentGizmoType() const
 	}
 	else
 	{
-		switch( GLevelEditorModeTools().GetWidgetMode() )
+		switch( GetModeTools().GetWidgetMode() )
 		{
 			case FWidget::WM_TranslateRotateZ:
 				return EGizmoHandleTypes::All;
@@ -3069,19 +3073,19 @@ void UViewportWorldInteraction::SetGizmoHandleType( const EGizmoHandleTypes InGi
 	{
 		case EGizmoHandleTypes::All:
 			GizmoType = InGizmoHandleType;
-			GLevelEditorModeTools().SetWidgetMode( FWidget::WM_Translate );
+			GetModeTools().SetWidgetMode( FWidget::WM_Translate );
 			break;
 
 		case EGizmoHandleTypes::Translate:
-			GLevelEditorModeTools().SetWidgetMode( FWidget::WM_Translate );
+			GetModeTools().SetWidgetMode( FWidget::WM_Translate );
 			break;
 
 		case EGizmoHandleTypes::Rotate:
-			GLevelEditorModeTools().SetWidgetMode( FWidget::WM_Rotate );
+			GetModeTools().SetWidgetMode( FWidget::WM_Rotate );
 			break;
 
 		case EGizmoHandleTypes::Scale:
-			GLevelEditorModeTools().SetWidgetMode( FWidget::WM_Scale );
+			GetModeTools().SetWidgetMode( FWidget::WM_Scale );
 			break;
 
 		check(0);
@@ -3537,6 +3541,11 @@ bool UViewportWorldInteraction::HasTransformableWithVelocityInSimulate() const
 	return bResult;
 }
 
+FEditorModeTools& UViewportWorldInteraction::GetModeTools() const
+{
+	return DefaultOptionalViewportClient == nullptr ? GLevelEditorModeTools() : *DefaultOptionalViewportClient->GetModeTools();
+}
+
 FVector UViewportWorldInteraction::SnapLocation(const bool bLocalSpaceSnapping, const FVector& DesiredGizmoLocation, const FTransform &GizmoStartTransform, const FVector SnapGridBase, const bool bShouldConstrainMovement, const FVector AlignAxes)
 {
 	bool bTransformableAlignmentUsed = false;
diff --git a/Engine/Source/Runtime/Core/Private/Apple/ApplePlatformStackWalk.cpp b/Engine/Source/Runtime/Core/Private/Apple/ApplePlatformStackWalk.cpp
--- a/Engine/Source/Runtime/Core/Private/Apple/ApplePlatformStackWalk.cpp
+++ b/Engine/Source/Runtime/Core/Private/Apple/ApplePlatformStackWalk.cpp
@@ -174,6 +174,20 @@ void FApplePlatformStackWalk::CaptureStackBackTrace( uint64* BackTrace, uint32 M
 
 bool FApplePlatformStackWalk::ProgramCounterToHumanReadableString( int32 CurrentCallDepth, uint64 ProgramCounter, ANSICHAR* HumanReadableString, SIZE_T HumanReadableStringSize, FGenericCrashContext* Context )
 {
+	
+	//
+	// Callstack lines should be written in this standard format
+	//
+	//	0xaddress module!func [file]
+	// 
+	// E.g. 0x045C8D01 OrionClient.self!UEngine::PerformError() [D:\Epic\Orion\Engine\Source\Runtime\Engine\Private\UnrealEngine.cpp:6481]
+	//
+	// Module may be omitted, everything else should be present, or substituted with a string that conforms to the expected type
+	//
+	// E.g 0x00000000 UnknownFunction []
+	//
+	// 
+
 	Dl_info DylibInfo;
 	int32 Result = dladdr((const void*)ProgramCounter, &DylibInfo);
 	if (Result == 0)
@@ -190,9 +204,13 @@ bool FApplePlatformStackWalk::ProgramCounterToHumanReadableString( int32 Current
 	{
 		AsyncSafeProgramCounterToSymbolInfo(ProgramCounter, SymbolInfo);
 	}
+
+	ANSICHAR TempArray[MAX_SPRINTF];
+
+	// Write out prefix, address, module, and function na,e
+	FCStringAnsi::Sprintf(TempArray, "0x%08x %s!%s ", (uint32)ProgramCounter, SymbolInfo.ModuleName, SymbolInfo.FunctionName);
 	
-	// Write out function name.
-	FCStringAnsi::Strcat(HumanReadableString, HumanReadableStringSize, SymbolInfo.FunctionName);
+	FCStringAnsi::Strcat(HumanReadableString, HumanReadableStringSize, TempArray);
 
 	// Get filename.
 	{
@@ -201,25 +219,17 @@ bool FApplePlatformStackWalk::ProgramCounterToHumanReadableString( int32 Current
 		if(SymbolInfo.LineNumber == 0)
 		{
 			// No line number. Print out the logical address instead.
-			FCStringAnsi::Sprintf(FileNameLine, "Address = 0x%-8x (filename not found) ", ProgramCounter);
+			FCStringAnsi::Sprintf(FileNameLine, " [UnknownFile]) ", ProgramCounter);
 		}
 		else
 		{
 			// try to add source file and line number, too
-			FCStringAnsi::Sprintf(FileNameLine, "[%s, line %d] ", SymbolInfo.Filename, SymbolInfo.LineNumber);
+			FCStringAnsi::Sprintf(FileNameLine, " [%s:%d] ", SymbolInfo.Filename, SymbolInfo.LineNumber);
 		}
 		
 		FCStringAnsi::Strcat(HumanReadableString, HumanReadableStringSize, FileNameLine);
 	}
 
-	// Get module name.
-	{
-		ANSICHAR ModuleName[MAX_SPRINTF];
-		// Write out Module information if there is sufficient space.
-		FCStringAnsi::Sprintf(ModuleName, "[in %s]", SymbolInfo.ModuleName);
-		FCStringAnsi::Strcat(HumanReadableString, HumanReadableStringSize, ModuleName);
-	}
-	
 	// For the crash reporting code this needs a Windows line ending, the caller is responsible for the '\n'
 	FCStringAnsi::Strcat(HumanReadableString, HumanReadableStringSize, "\r");
 
diff --git a/Engine/Source/Runtime/Core/Private/GenericPlatform/GenericPlatformStackWalk.cpp b/Engine/Source/Runtime/Core/Private/GenericPlatform/GenericPlatformStackWalk.cpp
--- a/Engine/Source/Runtime/Core/Private/GenericPlatform/GenericPlatformStackWalk.cpp
+++ b/Engine/Source/Runtime/Core/Private/GenericPlatform/GenericPlatformStackWalk.cpp
@@ -47,11 +47,18 @@ bool FGenericPlatformStackWalk::ProgramCounterToHumanReadableString( int32 Curre
 bool FGenericPlatformStackWalk::SymbolInfoToHumanReadableString( const FProgramCounterSymbolInfo& SymbolInfo, ANSICHAR* HumanReadableString, SIZE_T HumanReadableStringSize )
 {
 	const int32 MAX_TEMP_SPRINTF = 256;
-	// Valid callstack line 
-	// ModuleName!FunctionName [Filename:LineNumber]
+
+	//
+	// Callstack lines should be written in this standard format
+	//
+	//	0xaddress module!func [file]
 	// 
-	// Invalid callstack line
-	// ModuleName!
+	// E.g. 0x045C8D01 OrionClient.self!UEngine::PerformError() [D:\Epic\Orion\Engine\Source\Runtime\Engine\Private\UnrealEngine.cpp:6481]
+	//
+	// Module may be omitted, everything else should be present, or substituted with a string that conforms to the expected type
+	//
+	// E.g 0x00000000 UnknownFunction []
+	//
 	// 
 	if( HumanReadableString && HumanReadableStringSize > 0 )
 	{
@@ -63,8 +70,12 @@ bool FGenericPlatformStackWalk::SymbolInfoToHumanReadableString( const FProgramC
 		const UPTRINT RealPos = FMath::Max( (UPTRINT)Pos0, (UPTRINT)Pos1 );
 		const ANSICHAR* StrippedModuleName = RealPos > 0 ? (const ANSICHAR*)(RealPos + 1) : SymbolInfo.ModuleName;
 
-		//FCStringAnsi::Sprintf( StackLine, "%s!%s [%s:%i]", StrippedModuleName, (const ANSICHAR*)SymbolInfo.FunctionName, (const ANSICHAR*)SymbolInfo.Filename, SymbolInfo.LineNumber );
-
+		// Start with address
+		ANSICHAR PCAddress[MAX_TEMP_SPRINTF] = { 0 };
+		FCStringAnsi::Snprintf(PCAddress, MAX_TEMP_SPRINTF, "0x%016X ", SymbolInfo.ProgramCounter);
+		FCStringAnsi::Strncat(StackLine, PCAddress, MAX_SPRINTF);
+		
+		// Module if it's present
 		const bool bHasValidModuleName = FCStringAnsi::Strlen(StrippedModuleName) > 0;
 		if (bHasValidModuleName)
 		{
@@ -72,6 +83,7 @@ bool FGenericPlatformStackWalk::SymbolInfoToHumanReadableString( const FProgramC
 			FCStringAnsi::Strncat(StackLine, "!", MAX_SPRINTF);
 		}
 
+		// Function if it's available, unknown if it's not
 		const bool bHasValidFunctionName = FCStringAnsi::Strlen( SymbolInfo.FunctionName ) > 0;
 		if( bHasValidFunctionName )
 		{
@@ -79,11 +91,10 @@ bool FGenericPlatformStackWalk::SymbolInfoToHumanReadableString( const FProgramC
 		}
 		else
 		{
-			ANSICHAR PCAddress[MAX_TEMP_SPRINTF] = {0};
-			FCStringAnsi::Snprintf(PCAddress, MAX_TEMP_SPRINTF, "0x%016X", SymbolInfo.ProgramCounter);
-			FCStringAnsi::Strncat(StackLine, PCAddress, MAX_SPRINTF);
+			FCStringAnsi::Strncat(StackLine, "UnknownFunction", MAX_SPRINTF);
 		}
 
+		// file info
 		const bool bHasValidFilename = FCStringAnsi::Strlen( SymbolInfo.Filename ) > 0 && SymbolInfo.LineNumber > 0;
 		if( bHasValidFilename )
 		{
@@ -91,7 +102,10 @@ bool FGenericPlatformStackWalk::SymbolInfoToHumanReadableString( const FProgramC
 			FCStringAnsi::Snprintf( FilenameAndLineNumber, MAX_TEMP_SPRINTF, " [%s:%i]", SymbolInfo.Filename, SymbolInfo.LineNumber );
 			FCStringAnsi::Strncat(StackLine, FilenameAndLineNumber, MAX_SPRINTF);
 		}
-
+		else
+		{
+			FCStringAnsi::Strcat(StackLine, " []");
+		}
 
 		// Append the stack line.
 		FCStringAnsi::Strncat(HumanReadableString, StackLine, HumanReadableStringSize);
diff --git a/Engine/Source/Runtime/Core/Private/Linux/LinuxPlatformStackWalk.cpp b/Engine/Source/Runtime/Core/Private/Linux/LinuxPlatformStackWalk.cpp
--- a/Engine/Source/Runtime/Core/Private/Linux/LinuxPlatformStackWalk.cpp
+++ b/Engine/Source/Runtime/Core/Private/Linux/LinuxPlatformStackWalk.cpp
@@ -888,6 +888,19 @@ void FLinuxPlatformStackWalk::ProgramCounterToSymbolInfo( uint64 ProgramCounter,
 
 bool FLinuxPlatformStackWalk::ProgramCounterToHumanReadableString( int32 CurrentCallDepth, uint64 ProgramCounter, ANSICHAR* HumanReadableString, SIZE_T HumanReadableStringSize, FGenericCrashContext* Context )
 {
+	//
+	// Callstack lines should be written in this standard format
+	//
+	//	0xaddress module!func [file]
+	// 
+	// E.g. 0x045C8D01 OrionClient.self!UEngine::PerformError() [D:\Epic\Orion\Engine\Source\Runtime\Engine\Private\UnrealEngine.cpp:6481]
+	//
+	// Module may be omitted, everything else should be present, or substituted with a string that conforms to the expected type
+	//
+	// E.g 0x00000000 UnknownFunction []
+	//
+	// 
+
 	if (HumanReadableString && HumanReadableStringSize > 0)
 	{
 		ANSICHAR TempArray[MAX_SPRINTF];
@@ -895,11 +908,11 @@ bool FLinuxPlatformStackWalk::ProgramCounterToHumanReadableString( int32 Current
 		{
 			if (PLATFORM_64BITS)
 			{
-				FCStringAnsi::Sprintf(TempArray, "[Callstack] 0x%016llx ", ProgramCounter);
+				FCStringAnsi::Sprintf(TempArray, "0x%016llx ", ProgramCounter);
 			}
 			else
 			{
-				FCStringAnsi::Sprintf(TempArray, "[Callstack] 0x%08x ", (uint32) ProgramCounter);
+				FCStringAnsi::Sprintf(TempArray, "0x%08x ", (uint32) ProgramCounter);
 			}
 			LinuxStackWalkHelpers::AppendToString(HumanReadableString, HumanReadableStringSize, Context, TempArray);
 
@@ -909,11 +922,11 @@ bool FLinuxPlatformStackWalk::ProgramCounterToHumanReadableString( int32 Current
 		{
 			if (PLATFORM_64BITS)
 			{
-				FCStringAnsi::Sprintf(TempArray, "[Callstack]  %02d  0x%016llx  ", CurrentCallDepth, ProgramCounter);
+				FCStringAnsi::Sprintf(TempArray, "0x%016llx ", ProgramCounter);
 			}
 			else
 			{
-				FCStringAnsi::Sprintf(TempArray, "[Callstack]  %02d  0x%08x  ", CurrentCallDepth, (uint32) ProgramCounter);
+				FCStringAnsi::Sprintf(TempArray, "0x%08x ", (uint32) ProgramCounter);
 			}
 			LinuxStackWalkHelpers::AppendToString(HumanReadableString, HumanReadableStringSize, Context, TempArray);
 
@@ -934,9 +947,9 @@ bool FLinuxPlatformStackWalk::ProgramCounterToHumanReadableString( int32 Current
 
 				if (bAddDetailedInfo)
 				{
-					// append FunctionName() [Source.cpp, line X] to HumanReadableString
+					// append FunctionName() [Source.cpp:X] to HumanReadableString
 					LinuxStackWalkHelpers::AppendToString(HumanReadableString, HumanReadableStringSize, Context, FunctionName);
-					FCStringAnsi::Sprintf(TempArray, " [%s, line %d]", SourceFilename, LineNumber);
+					FCStringAnsi::Sprintf(TempArray, " [%s:%d]", SourceFilename, LineNumber);
 					LinuxStackWalkHelpers::AppendToString(HumanReadableString, HumanReadableStringSize, Context, TempArray);
 
 					// append Module!FunctioName [Source.cpp:X] to MinidumpCallstackInfo
@@ -950,10 +963,8 @@ bool FLinuxPlatformStackWalk::ProgramCounterToHumanReadableString( int32 Current
 				{
 					// get the function name for backtrace, may be incorrect
 					FunctionName = LinuxStackWalkHelpers::GetFunctionName(Context, CurrentCallDepth);
-					if (FunctionName)
-					{
-						LinuxStackWalkHelpers::AppendToString(HumanReadableString, HumanReadableStringSize, Context, FunctionName);
-					}
+				
+					LinuxStackWalkHelpers::AppendToString(HumanReadableString, HumanReadableStringSize, Context, FunctionName != nullptr ? FunctionName : "UnknownFunction");	
 
 					FCStringAnsi::Strncat(LinuxContext->MinidumpCallstackInfo, "Unknown!", ARRAY_COUNT( LinuxContext->MinidumpCallstackInfo ) - 1);
 					LinuxStackWalkHelpers::AppendFunctionNameIfAny(*LinuxContext, FunctionName, ProgramCounter);
diff --git a/Engine/Source/Runtime/Core/Private/Misc/AssertionMacros.cpp b/Engine/Source/Runtime/Core/Private/Misc/AssertionMacros.cpp
--- a/Engine/Source/Runtime/Core/Private/Misc/AssertionMacros.cpp
+++ b/Engine/Source/Runtime/Core/Private/Misc/AssertionMacros.cpp
@@ -27,6 +27,42 @@ bool FDebug::bHasAsserted = false;
 
 #define FILE_LINE_DESC TEXT(" [File:%s] [Line: %i] ")
 
+
+/*
+	Ensure behavior
+
+	* ensure() macro calls OptionallyLogFormattedEnsureMessageReturningFalse 
+	* OptionallyLogFormattedEnsureMessageReturningFalse calls EnsureFailed()
+	* EnsureFailed() -
+		* Formats the ensure failure and calls StaticFailDebug to populate the global error info (without callstack)
+		* Prints the script callstack (if any)
+		* Halts if a debugger is attached 
+		* If not, logs the callstack and attempts to submit an error report
+	* execution continues as normal, (on some platforms this can take ~30 secs to perform)
+
+	Check behavior
+
+	* check() macro calls LogAssertFailedMessage
+	* LogAssertFailedMessage formats the assertion message and calls StaticFailDebug
+	* StaticFailDebug populates global error info with the failure message and if supported (AllowsCallStackDumpDuringAssert) the callstack
+	* If a debugger is attached execution halts
+	* If not FDebug::AssertFailed is called
+	* FDebug::AssertFailed logs the assert message and description to GError
+	* At this point behavior depends on the platform-specific error output device implementation
+		* Desktop platforms (Windows, Mac, Linux) will generally throw an exception and in the handler attempt to submit a crash report and exit
+		* Console platforms will generally dump the info to the log and abort()
+
+	Fatal-error behavior
+
+	* The UE_LOG macro calls FMsg::Logf which checks for "Fatal" verbosity
+	* FMsg::Logf formats the failure message and calls StaticFailDebug
+	* StaticFailDebug populates global error info with the failure message and if supported (AllowsCallStackDumpDuringAssert) the callstack
+	* FDebug::AssertFailed is then called, and from this point behavior is identical to an assert but with a different message
+
+*/
+
+
+
 /** Lock used to synchronize the fail debug calls. */
 static FCriticalSection	FailDebugCriticalSection;
 
@@ -71,50 +107,85 @@ void PrintScriptCallstack(bool bEmptyWhenDone)
  */
 void StaticFailDebug( const TCHAR* Error, const ANSICHAR* File, int32 Line, const TCHAR* Description, bool bIsEnsure )
 {
-	// For ensure log should be flushed in the engine loop.
-	if( !bIsEnsure )
+	TCHAR DescriptionAndTrace[4096];
+
+	FCString::Strncpy(DescriptionAndTrace, Description, ARRAY_COUNT(DescriptionAndTrace) - 1);
+
+	// some platforms (Windows, Mac, Linux) generate this themselves by throwing an exception and capturing
+	// the backtrace later on
+	if (FPlatformProperties::AllowsCallStackDumpDuringAssert() && bIsEnsure == false)
 	{
-		GLog->PanicFlushThreadedLogs();
+		ANSICHAR StackTrace[4096];
+		if (StackTrace != NULL)
+		{
+			StackTrace[0] = 0;
+			FPlatformStackWalk::StackWalkAndDump(StackTrace, ARRAY_COUNT(StackTrace), CALLSTACK_IGNOREDEPTH);
+
+			FCString::Strncat(DescriptionAndTrace, TEXT("\n"), ARRAY_COUNT(DescriptionAndTrace) - 1);
+			FCString::Strncat(DescriptionAndTrace, ANSI_TO_TCHAR(StackTrace), ARRAY_COUNT(DescriptionAndTrace) - 1);
+		}
 	}
 
 	FScopeLock Lock( &FailDebugCriticalSection );
-	FPlatformMisc::LowLevelOutputDebugStringf(TEXT("%s") FILE_LINE_DESC TEXT("\n%s\n"), Error, ANSI_TO_TCHAR(File), Line, Description);
 
 	// Copy the detailed error into the error message.
-	FCString::Snprintf( GErrorMessage, ARRAY_COUNT( GErrorMessage ), TEXT( "%s" ) FILE_LINE_DESC TEXT( "\n%s\n" ), Error, ANSI_TO_TCHAR( File ), Line, Description );
+	FCString::Snprintf( GErrorMessage, ARRAY_COUNT( GErrorMessage ), TEXT( "%s" ) FILE_LINE_DESC TEXT( "\n%s\n" ), Error, ANSI_TO_TCHAR( File ), Line, DescriptionAndTrace);
 
 	// Copy the error message to the error history.
 	FCString::Strncpy( GErrorHist, GErrorMessage, ARRAY_COUNT( GErrorHist ) );
 	FCString::Strncat( GErrorHist, TEXT( "\r\n\r\n" ), ARRAY_COUNT( GErrorHist ) );
 }
 
-void OutputMultiLineCallstack(const ANSICHAR* File, int32 Line, const FName& LogName, const TCHAR* Heading, TCHAR* Message, ELogVerbosity::Type Verbosity)
+
+void FDebug::LogFormattedMessageWithCallstack(const FName& LogName, const ANSICHAR* File, int32 Line, const TCHAR* Heading, const TCHAR* Message, ELogVerbosity::Type Verbosity)
 {
-	const bool bWriteUATMarkers = FParse::Param(FCommandLine::Get(), TEXT("CrashForUAT")) && FParse::Param(FCommandLine::Get(), TEXT("stdout"));
+	const bool bLowLevel = LogName == NAME_None;
+	const bool bWriteUATMarkers = FParse::Param(FCommandLine::Get(), TEXT("CrashForUAT")) && FParse::Param(FCommandLine::Get(), TEXT("stdout")) && !bLowLevel;
 
 	if (bWriteUATMarkers)
 	{
 		FMsg::Logf(File, Line, LogName, Verbosity, TEXT("begin: stack for UAT"));
 	}
 
-	FMsg::Logf(File, Line, LogName, Verbosity, TEXT("%s"), Heading);
-	FMsg::Logf(File, Line, LogName, Verbosity, TEXT(""));
+	if (bLowLevel)
+	{
+		FPlatformMisc::LowLevelOutputDebugStringf(TEXT("%s\n"), Heading);
+	}
+	else
+	{
+		FMsg::Logf(File, Line, LogName, Verbosity, TEXT("%s"), Heading);
+		FMsg::Logf(File, Line, LogName, Verbosity, TEXT(""));
+	}
 
-	for (TCHAR* LineStart = Message;; )
+	for (const TCHAR* LineStart = Message;; )
 	{
+		TCHAR SingleLine[1024];
+
 		// Find the end of the current line
-		TCHAR* LineEnd = LineStart;
+		const TCHAR* LineEnd = LineStart;
+		TCHAR* SingleLineWritePos = SingleLine;
+
 		while (*LineEnd != 0 && *LineEnd != '\r' && *LineEnd != '\n')
 		{
-			LineEnd++;
+			*SingleLineWritePos++ = *LineEnd++;
 		}
 
-		// Output it
-		TCHAR LineEndCharacter = *LineEnd;
-		*LineEnd = 0;
-		FMsg::Logf(File, Line, LogName, Verbosity, TEXT("%s"), LineStart);
-		*LineEnd = LineEndCharacter;
+		// cap it it
+		*SingleLineWritePos = 0;
 
+		// prefix function lines with [Callstack] for parsing tools
+		const TCHAR* Prefix = (FCString::Strnicmp(LineStart, TEXT("0x"), 2) == 0) ? TEXT("[Callstack] ") : TEXT("");
+
+		// if this is an address line, prefix it with [Callstack]
+		if (bLowLevel)
+		{
+			FPlatformMisc::LowLevelOutputDebugStringf(TEXT("%s%s\n"), Prefix, SingleLine);
+		}
+		else
+		{
+			FMsg::Logf(File, Line, LogName, Verbosity, TEXT("%s%s"), Prefix, SingleLine);
+		}
+		
 		// Quit if this was the last line
 		if (*LineEnd == 0)
 		{
@@ -138,30 +209,18 @@ void OutputMultiLineCallstack(const ANSICHAR* File, int32 Line, const FName& Log
 //
 void VARARGS FDebug::LogAssertFailedMessage(const ANSICHAR* Expr, const ANSICHAR* File, int32 Line, const TCHAR* Format/*=TEXT("")*/, ...)
 {
-	// Print out the blueprint callstack
-	PrintScriptCallstack(true);
-
 	// Ignore this assert if we're already forcibly shutting down because of a critical error.
 	if( !GIsCriticalError )
 	{
+		// Print out the blueprint callstack
+		PrintScriptCallstack(true);
+
 		TCHAR DescriptionString[4096];
 		GET_VARARGS( DescriptionString, ARRAY_COUNT( DescriptionString ), ARRAY_COUNT( DescriptionString ) - 1, Format, Format );
 
 		TCHAR ErrorString[MAX_SPRINTF];
 		FCString::Sprintf( ErrorString, TEXT( "Assertion failed: %s" ), ANSI_TO_TCHAR( Expr ) );
 
-		if( FPlatformProperties::AllowsCallStackDumpDuringAssert() )
-		{
-			ANSICHAR StackTrace[4096];
-			if( StackTrace != NULL )
-			{
-				StackTrace[0] = 0;
-				FPlatformStackWalk::StackWalkAndDump( StackTrace, ARRAY_COUNT(StackTrace), CALLSTACK_IGNOREDEPTH );
-
-				FCString::Strncat( DescriptionString, ANSI_TO_TCHAR( StackTrace ), ARRAY_COUNT( DescriptionString ) - 1 );
-			}
-		}
-
 		StaticFailDebug( ErrorString, File, Line, DescriptionString, false );
 	}
 }
@@ -243,7 +302,7 @@ void FDebug::EnsureFailed(const ANSICHAR* Expr, const ANSICHAR* File, int32 Line
 
 		// Dump the error and flush the log.
 #if !NO_LOGGING
-		OutputMultiLineCallstack(__FILE__, __LINE__, LogOutputDevice.GetCategoryName(), TEXT("=== Handled ensure: ==="), ErrorMsg, ELogVerbosity::Error);
+		FDebug::LogFormattedMessageWithCallstack(LogOutputDevice.GetCategoryName(), __FILE__, __LINE__, TEXT("=== Handled ensure: ==="), ErrorMsg, ELogVerbosity::Error);
 #endif
 		GLog->Flush();
 
@@ -345,7 +404,7 @@ void VARARGS FDebug::AssertFailed(const ANSICHAR* Expr, const ANSICHAR* File, in
 	GET_VARARGS(DescriptionString, ARRAY_COUNT(DescriptionString), ARRAY_COUNT(DescriptionString) - 1, Format, Format);
 
 	TCHAR ErrorString[MAX_SPRINTF];
-	FCString::Sprintf(ErrorString, TEXT("Assertion failed: %s"), ANSI_TO_TCHAR(Expr));
+	FCString::Sprintf(ErrorString, TEXT("%s"), ANSI_TO_TCHAR(Expr));
 	GError->Logf(TEXT("Assertion failed: %s") FILE_LINE_DESC TEXT("\n%s\n"), ErrorString, ANSI_TO_TCHAR(File), Line, DescriptionString);
 }
 
@@ -390,7 +449,7 @@ void FDebug::DumpStackTraceToLog()
 
 	// Dump the error and flush the log.
 	// ELogVerbosity::Error to make sure it gets printed in log for conveniency.
-	OutputMultiLineCallstack(__FILE__, __LINE__, LogOutputDevice.GetCategoryName(), TEXT("=== FDebug::DumpStackTrace(): ==="), ANSI_TO_TCHAR(StackTrace), ELogVerbosity::Error);
+	FDebug::LogFormattedMessageWithCallstack(LogOutputDevice.GetCategoryName(), __FILE__, __LINE__, TEXT("=== FDebug::DumpStackTrace(): ==="), ANSI_TO_TCHAR(StackTrace), ELogVerbosity::Error);
 	GLog->Flush();
 	FMemory::SystemFree(StackTrace);
 #endif
diff --git a/Engine/Source/Runtime/Core/Private/Misc/DefaultValueHelper.cpp b/Engine/Source/Runtime/Core/Private/Misc/DefaultValueHelper.cpp
--- a/Engine/Source/Runtime/Core/Private/Misc/DefaultValueHelper.cpp
+++ b/Engine/Source/Runtime/Core/Private/Misc/DefaultValueHelper.cpp
@@ -4,6 +4,7 @@
 #include "Math/Color.h"
 #include "Math/Vector2D.h"
 #include "Math/Vector.h"
+#include "Math/Vector4.h"
 #include "Math/Rotator.h"
 
 
@@ -576,6 +577,25 @@ bool FDefaultValueHelper::ParseVector2D(const FString& Source, FVector2D& OutVal
 }
 
 
+bool FDefaultValueHelper::ParseVector4(const FString& Source, FVector4& OutVal)
+{
+	float X, Y, Z, W;
+	TArray<FString> SourceParts;
+
+	if (Source.ParseIntoArray(SourceParts, TEXT(",")) == 4 &&
+		ParseFloat(SourceParts[0], X) &&
+		ParseFloat(SourceParts[1], Y) && 
+		ParseFloat(SourceParts[2], Z) &&
+		ParseFloat(SourceParts[3], W))
+	{
+		OutVal = FVector4(X, Y, Z, W);
+		return true;
+	}
+
+	return false;
+}
+
+
 bool FDefaultValueHelper::ParseRotator(const FString& Source, FRotator& OutVal)
 {
 	FVector Vector;
diff --git a/Engine/Source/Runtime/Core/Private/Windows/WindowsApplication.cpp b/Engine/Source/Runtime/Core/Private/Windows/WindowsApplication.cpp
--- a/Engine/Source/Runtime/Core/Private/Windows/WindowsApplication.cpp
+++ b/Engine/Source/Runtime/Core/Private/Windows/WindowsApplication.cpp
@@ -800,10 +800,13 @@ int32 FWindowsApplication::ProcessMessage( HWND hwnd, uint32 msg, WPARAM wParam,
 		}
 
 #if WINVER >= 0x0601
+		// ORION - REMOVED FOR WACOM SUPPORT!
+		/*
 		if (IsFakeMouseInputMessage(msg))
 		{
 			return 0;
 		}
+		*/
 #endif
 
 		switch(msg)
diff --git a/Engine/Source/Runtime/Core/Private/Windows/WindowsPlatformOutputDevices.cpp b/Engine/Source/Runtime/Core/Private/Windows/WindowsPlatformOutputDevices.cpp
--- a/Engine/Source/Runtime/Core/Private/Windows/WindowsPlatformOutputDevices.cpp
+++ b/Engine/Source/Runtime/Core/Private/Windows/WindowsPlatformOutputDevices.cpp
@@ -149,8 +149,7 @@ void FOutputDeviceWindowsError::HandleError()
 
 	// Dump the error and flush the log.
 #if !NO_LOGGING
-	extern void OutputMultiLineCallstack(const ANSICHAR* File, int32 Line, const FName& LogName, const TCHAR* Heading, TCHAR* Message, ELogVerbosity::Type Verbosity);
-	OutputMultiLineCallstack(__FILE__, __LINE__, LogWindows.GetCategoryName(), TEXT("=== Critical error: ==="), GErrorHist, ELogVerbosity::Error);
+	FDebug::LogFormattedMessageWithCallstack(LogWindows.GetCategoryName(), __FILE__, __LINE__, TEXT("=== Critical error: ==="), GErrorHist, ELogVerbosity::Error);
 #endif
 	GLog->PanicFlushThreadedLogs();
 
diff --git a/Engine/Source/Runtime/Core/Private/Windows/WindowsPlatformStackWalk.cpp b/Engine/Source/Runtime/Core/Private/Windows/WindowsPlatformStackWalk.cpp
--- a/Engine/Source/Runtime/Core/Private/Windows/WindowsPlatformStackWalk.cpp
+++ b/Engine/Source/Runtime/Core/Private/Windows/WindowsPlatformStackWalk.cpp
@@ -808,8 +808,22 @@ bool FWindowsPlatformStackWalk::InitStackWalking()
 		SymSetOptions( SymOpts );
 	
 		// Initialize the symbol engine.		
-		const FString RemoteStorage = GetRemoteStorage(GetDownstreamStorage());
+		FString RemoteStorage = GetRemoteStorage(GetDownstreamStorage());
 #if WINVER > 0x502
+		if (RemoteStorage.IsEmpty())
+		{
+			// by default passing null to SymInitialize will use the current working dir to search for a pdb, 
+			// but to support the basedir argument that allows an exe to run against data in a different location
+			// we'll put the path of the executing module first.
+			TCHAR ModulePath[MAX_PATH] = { 0 };
+			if (::GetModuleFileName(::GetModuleHandle(NULL), ModulePath, MAX_PATH))
+			{
+				RemoteStorage = FPaths::GetPath(ModulePath);
+				RemoteStorage += ";";
+				RemoteStorage += FPlatformProcess::GetCurrentWorkingDirectory();
+			}
+		}
+		
 		SymInitializeW( GetCurrentProcess(), RemoteStorage.IsEmpty() ? nullptr : *RemoteStorage, true );
 #else
 		SymInitialize( GetCurrentProcess(), nullptr, true );
diff --git a/Engine/Source/Runtime/Core/Public/HAL/MallocBinned2.h b/Engine/Source/Runtime/Core/Public/HAL/MallocBinned2.h
--- a/Engine/Source/Runtime/Core/Public/HAL/MallocBinned2.h
+++ b/Engine/Source/Runtime/Core/Public/HAL/MallocBinned2.h
@@ -55,7 +55,7 @@
 	#if UE_BUILD_SHIPPING && !WITH_EDITOR
 		#define BINNED2_ALLOCATOR_STATS 0	
 	#else
-		#define BINNED2_ALLOCATOR_STATS 0
+		#define BINNED2_ALLOCATOR_STATS 1
 	#endif
 #endif
 
diff --git a/Engine/Source/Runtime/Core/Public/Misc/AssertionMacros.h b/Engine/Source/Runtime/Core/Public/Misc/AssertionMacros.h
--- a/Engine/Source/Runtime/Core/Public/Misc/AssertionMacros.h
+++ b/Engine/Source/Runtime/Core/Public/Misc/AssertionMacros.h
@@ -5,6 +5,11 @@
 #include "CoreTypes.h"
 #include "HAL/PlatformMisc.h"
 
+namespace ELogVerbosity
+{
+	enum Type : uint8;
+}
+
 /**
  * FDebug
  * These functions offer debugging and diagnostic functionality and its presence 
@@ -52,6 +57,19 @@ struct CORE_API FDebug
 	 */
 	static bool VARARGS OptionallyLogFormattedEnsureMessageReturningFalse(bool bLog, const ANSICHAR* Expr, const ANSICHAR* File, int32 Line, const TCHAR* FormattedMsg, ...);
 #endif // DO_CHECK || DO_GUARD_SLOW
+
+	/**
+	* Logs an a message to the provided log channel. If a callstack is included (detected by lines starting with 0x) if will be logged in the standard Unreal 
+	* format of [Callstack] Address FunctionInfo [File]
+	*
+	* @param	LogName		Log channel. If NAME_None then LowLevelOutputDebugStringf is used
+	* @param	File		File name ANSI string (__FILE__)
+	* @param	Line		Line number (__LINE__)
+	* @param	Heading		Informative heading displayed above the message callstack
+	* @param	Message		Multi-line message with a callstack
+	*
+	*/
+	static void LogFormattedMessageWithCallstack(const FName& LogName, const ANSICHAR* File, int32 Line, const TCHAR* Heading, const TCHAR* Message, ELogVerbosity::Type Verbosity);
 };
 
 /*----------------------------------------------------------------------------
diff --git a/Engine/Source/Runtime/Core/Public/Misc/DefaultValueHelper.h b/Engine/Source/Runtime/Core/Public/Misc/DefaultValueHelper.h
--- a/Engine/Source/Runtime/Core/Public/Misc/DefaultValueHelper.h
+++ b/Engine/Source/Runtime/Core/Public/Misc/DefaultValueHelper.h
@@ -143,6 +143,16 @@ public:
 	 */
 	static bool ParseVector2D(const FString& Source, FVector2D& OutVal);
 
+	/**
+	* Converts a string into a FVector4. Accepted form: " %f, %f, %f, %f "
+	*
+	* @param Source the input string to try to convert
+	* @param OutVal the output vector4
+	*
+	* @return true if the conversion happened, false otherwise
+	*/
+	static bool ParseVector4(const FString& Source, FVector4& OutVal);
+
 	/**
 	 * Converts a string into a FRotator. Accepted form: " %f, %f, %f "
 	 *
diff --git a/Engine/Source/Runtime/CoreUObject/Public/UObject/UObjectBaseUtility.h b/Engine/Source/Runtime/CoreUObject/Public/UObject/UObjectBaseUtility.h
--- a/Engine/Source/Runtime/CoreUObject/Public/UObject/UObjectBaseUtility.h
+++ b/Engine/Source/Runtime/CoreUObject/Public/UObject/UObjectBaseUtility.h
@@ -134,8 +134,7 @@ public:
 	/**
 	 * Marks this object as RF_PendingKill.
 	 */
-	// Daniel L: Temporarily Virtual to track down OR-40458 remove virtual!
-	virtual void MarkPendingKill()
+	FORCEINLINE void MarkPendingKill()
 	{
 		check(!IsRooted());
 		GUObjectArray.IndexToObject(InternalIndex)->SetPendingKill();
diff --git a/Engine/Source/Runtime/Engine/Classes/Animation/AnimInstance.h b/Engine/Source/Runtime/Engine/Classes/Animation/AnimInstance.h
--- a/Engine/Source/Runtime/Engine/Classes/Animation/AnimInstance.h
+++ b/Engine/Source/Runtime/Engine/Classes/Animation/AnimInstance.h
@@ -981,6 +981,8 @@ public:
 	void ResetDynamics();
 
 public:
+	/** Access a read only version of the Updater Counter from the AnimInstanceProxy on the GameThread. */
+	const FGraphTraversalCounter& GetUpdateCounter() const; 
 
 	/** Access the required bones array */
 	FBoneContainer& GetRequiredBones();
@@ -1123,6 +1125,7 @@ private:
 private:
 	// update montage
 	void UpdateMontage(float DeltaSeconds);
+	void UpdateMontageSyncGroup();
 
 protected:
 	// Updates the montage data used for evaluation based on the current playing montages
diff --git a/Engine/Source/Runtime/Engine/Classes/Animation/AnimMontage.h b/Engine/Source/Runtime/Engine/Classes/Animation/AnimMontage.h
--- a/Engine/Source/Runtime/Engine/Classes/Animation/AnimMontage.h
+++ b/Engine/Source/Runtime/Engine/Classes/Animation/AnimMontage.h
@@ -16,6 +16,7 @@
 #include "Animation/AnimationAsset.h"
 #include "AlphaBlend.h"
 #include "Animation/AnimCompositeBase.h"
+#include "Animation/TimeStretchCurve.h"
 #include "AnimMontage.generated.h"
 
 class UAnimInstance;
@@ -151,6 +152,15 @@ struct FBranchingPointMarker
 	}
 };
 
+UENUM()
+enum class EMontageSubStepResult : uint8
+{
+	Moved,
+	NotMoved,
+	InvalidSection,
+	InvalidMontage,
+};
+
 /**
  * Delegate for when Montage is completed, whether interrupted or finished
  * Weight of this montage is 0.f, so it stops contributing to output pose
@@ -166,11 +176,115 @@ DECLARE_DELEGATE_TwoParams( FOnMontageEnded, class UAnimMontage*, bool /*bInterr
  */
 DECLARE_DELEGATE_TwoParams( FOnMontageBlendingOutStarted, class UAnimMontage*, bool /*bInterrupted*/) 
 
+/**
+	Helper struct to sub step through Montages when advancing time.
+	These require stopping at sections and branching points to potential jumps and loops.
+	And also stepping through TimeStretchMarkers to adjust play rate based on TimeStretchCurve.
+ */
+struct FMontageSubStepper
+{
+private:
+	const struct FAnimMontageInstance* MontageInstance;
+	const class UAnimMontage* Montage;
+
+	float TimeRemaining;
+	float Cached_CombinedPlayRate;
+	float PlayRate;
+	float DeltaMove;
+	bool bPlayingForward;
+
+	int32 CurrentSectionIndex;
+	float CurrentSectionStartTime;
+	float CurrentSectionLength;
+	bool bReachedEndOfSection;
+	bool bHasValidTimeStretchCurveData;
+
+	int32 TimeStretchMarkerIndex;
+
+	mutable TArray<float> SectionStartPositions_Target;
+	mutable TArray<float> SectionEndPositions_Target;
+
+	float Cached_P_Target;
+	float Cached_P_Original;
+
+	FTimeStretchCurveInstance TimeStretchCurveInstance;
+
+public:
+	FMontageSubStepper()
+		: MontageInstance(nullptr)
+		, Montage(nullptr)
+		, TimeRemaining(0.f)
+		, Cached_CombinedPlayRate(0.f)
+		, PlayRate(0.f)
+		, DeltaMove(0.f)
+		, bPlayingForward(true)
+		, CurrentSectionIndex(INDEX_NONE)
+		, CurrentSectionStartTime(0.f)
+		, CurrentSectionLength(0.f)
+		, bReachedEndOfSection(false)
+		, bHasValidTimeStretchCurveData(false)
+		, TimeStretchMarkerIndex(INDEX_NONE)
+		, Cached_P_Target(FLT_MAX)
+		, Cached_P_Original(FLT_MAX)
+	{}
+
+	void Initialize(const struct FAnimMontageInstance& InAnimInstance);
+
+	void AddEvaluationTime(float InDeltaTime) { TimeRemaining += InDeltaTime; }
+	bool HasTimeRemaining() const { return (TimeRemaining > SMALL_NUMBER); }
+	float GetRemainingTime() const { return TimeRemaining; }
+	EMontageSubStepResult Advance(float& InOut_P_Original, const FBranchingPointMarker** OutBranchingPointMarkerPtr);
+	bool HasReachedEndOfSection() const { return bReachedEndOfSection; }
+	float GetRemainingPlayTimeToSectionEnd(const float In_P_Original);
+
+	bool GetbPlayingForward() const { return bPlayingForward; }
+	float GetDeltaMove() const { return DeltaMove; }
+	int32 GetCurrentSectionIndex() const { return CurrentSectionIndex; }
+
+	/** Invalidate Cached_CombinedPlayRate to force data to be recached in 'ConditionallyUpdateCachedData' */
+	void ClearCachedData() { Cached_CombinedPlayRate = FLT_MAX; }
+
+private:
+	/** 
+		Updates markers P_Marker_Original and P_Marker_Target
+		*only* if T_Target has changed. 
+	*/
+	void ConditionallyUpdateTimeStretchCurveCachedData();
+
+	/** 
+		Finds montage position in 'target' space, given current position in 'original' space.
+		This means given a montage position, we find his play back time.
+		This should only be used for montage position, as we cache results and lazily update it for performance.
+	*/
+	float FindMontagePosition_Target(float In_P_Original);
+
+	/**
+		Finds montage position in 'original' space, given current position in 'target' space.
+		This means given a montage play back time, we find his actual position.
+		This should only be used for montage position, as we cache results and lazily update it for performance.
+	*/
+	float FindMontagePosition_Original(float In_P_Target);
+
+	/** 
+		Gets current section end position in target space. 
+		this is lazily cached, as it can be called every frame to test when to blend out. 
+	*/
+	float GetCurrSectionEndPosition_Target() const;
+
+	/** 
+		Gets current section start position in target space.
+		this is lazily cached, as it can be called every frame to test when to blend out. 
+	*/
+	float GetCurrSectionStartPosition_Target() const;
+};
+
 USTRUCT()
 struct FAnimMontageInstance
 {
 	GENERATED_USTRUCT_BODY()
 
+	friend struct FMontageSubStepper;
+
 	// Montage reference
 	UPROPERTY()
 	class UAnimMontage* Montage;
@@ -196,6 +310,8 @@ struct FAnimMontageInstance
 	bool bDidUseMarkerSyncThisTick;
 
 private:
+	struct FMontageSubStepper MontageSubStepper;
+
 	// list of next sections per section - index of array is section id
 	UPROPERTY()
 	TArray<int32> NextSections;
@@ -337,6 +453,7 @@ public:
 	float GetPlayRate() const { return PlayRate; }
 	float GetDeltaMoved() const { return DeltaMoved; }
 	float GetPreviousPosition() const { return PreviousPosition;  }
+
 	/** 
 	 * Setters
 	 */
@@ -362,7 +479,11 @@ public:
 	 * to accumulate and update curve data/notifies/branching points
 	 */
 	ENGINE_API void UpdateWeight(float DeltaTime);
-	//~ @fixme laurent can we make Advance use that, so we don't have 2 code paths which risk getting out of sync?
+
+#if WITH_EDITOR	
+	ENGINE_API void EditorOnly_PreAdvance();
+#endif
+
 	/** Simulate is same as Advance, but without calling any events or touching any of the instance data. So it performs a simulation of advancing the timeline. */
 	bool SimulateAdvance(float DeltaTime, float& InOutPosition, struct FRootMotionMovementParams & OutRootMotionParams) const;
 	void Advance(float DeltaTime, struct FRootMotionMovementParams * OutRootMotionParams, bool bBlendRootMotion);
@@ -390,6 +511,8 @@ private:
 	void BranchingPointEventHandler(const FBranchingPointMarker* BranchingPointMarker);
 	void RefreshNextPrevSections();
 
+	float GetRemainingPlayTimeToSectionEnd(const FMontageSubStepper& MontageSubStepper) const;
+
 public:
 	/** static functions that are used by matinee functionality */
 	ENGINE_API static UAnimMontage* SetMatineeAnimPositionInner(FName SlotName, USkeletalMeshComponent* SkeletalMeshComponent, UAnimSequenceBase* InAnimSequence, float InPosition, bool bLooping);
@@ -412,6 +535,8 @@ class UAnimMontage : public UAnimCompositeBase
 {
 	GENERATED_UCLASS_BODY()
 
+	friend struct FAnimMontageInstance;
+
 	/** Blend in option. */
 	UPROPERTY(EditAnywhere, Category=BlendOption)
 	FAlphaBlend BlendIn;
@@ -482,6 +607,7 @@ class UAnimMontage : public UAnimCompositeBase
 public:
 	//~ Begin UObject Interface
 	virtual void PostLoad() override;
+	virtual void PreSave(const class ITargetPlatform* TargetPlatform) override;
 
 	// Gets the sequence length of the montage by calculating it from the lengths of the segments in the montage
 	ENGINE_API float CalculateSequenceLength();
@@ -665,7 +791,8 @@ public:
 	TArray<int32> BranchingPointStateNotifyIndices;
 
 	/** Find first branching point marker between track positions */
-	const FBranchingPointMarker* FindFirstBranchingPointMarker(float StartTrackPos, float EndTrackPos);
+	const FBranchingPointMarker* FindFirstBranchingPointMarker(float StartTrackPos, float EndTrackPos) const;
+
 	/** Filter out notifies from array that are marked as 'BranchingPoints' */
 	void FilterOutNotifyBranchingPoints(TArray<const FAnimNotifyEvent*>& InAnimNotifies);
 
@@ -678,4 +805,18 @@ public:
 	virtual void InvalidateRecursiveAsset() override;
 	virtual bool ContainRecursive(TArray<UAnimCompositeBase*>& CurrentAccumulatedList) override;
 	//~End UAnimCompositeBase Interface
+
+	//~Begin Time Stretch Curve
+public:
+	UPROPERTY(EditAnywhere, Category = TimeStretchCurve)
+	FTimeStretchCurve TimeStretchCurve;
+
+	/** Name of optional TimeStretchCurveName to look for in Montage. */
+	UPROPERTY(EditAnywhere, Category = TimeStretchCurve)
+	FName TimeStretchCurveName;
+
+private:
+	void BakeTimeStretchCurve();
+	//~End Time Stretch Curve
+
 };
diff --git a/Engine/Source/Runtime/Engine/Classes/Engine/StaticMesh.h b/Engine/Source/Runtime/Engine/Classes/Engine/StaticMesh.h
--- a/Engine/Source/Runtime/Engine/Classes/Engine/StaticMesh.h
+++ b/Engine/Source/Runtime/Engine/Classes/Engine/StaticMesh.h
@@ -830,6 +830,9 @@ public:
 	 */
 	ENGINE_API void SetVertexColorData(const TMap<FVector, FColor>& VertexColorData);
 
+	/** Removes all vertex colors from this mesh and rebuilds it (Editor only */
+	ENGINE_API void RemoveVertexColors();
+
 	void EnforceLightmapRestrictions();
 
 	/** Calculates the extended bounds */
diff --git a/Engine/Source/Runtime/Engine/Classes/Engine/Texture.h b/Engine/Source/Runtime/Engine/Classes/Engine/Texture.h
--- a/Engine/Source/Runtime/Engine/Classes/Engine/Texture.h
+++ b/Engine/Source/Runtime/Engine/Classes/Engine/Texture.h
@@ -730,9 +730,6 @@ public:
 	ENGINE_API virtual uint32 GetMaximumDimension() const;
 #endif
 
-	// Daniel L: Temporary code to track down OR-40458 REMOVE ME
-	ENGINE_API virtual void MarkPendingKill() override;
-
 	/** @return the width of the surface represented by the texture. */
 	virtual float GetSurfaceWidth() const PURE_VIRTUAL(UTexture::GetSurfaceWidth,return 0;);
 
diff --git a/Engine/Source/Runtime/Engine/Classes/Engine/TextureStreamingTypes.h b/Engine/Source/Runtime/Engine/Classes/Engine/TextureStreamingTypes.h
--- a/Engine/Source/Runtime/Engine/Classes/Engine/TextureStreamingTypes.h
+++ b/Engine/Source/Runtime/Engine/Classes/Engine/TextureStreamingTypes.h
@@ -68,6 +68,15 @@ struct FStreamingTexturePrimitiveInfo
 		PackedRelativeBox(0)
 	{
 	}
+
+	FStreamingTexturePrimitiveInfo(UTexture2D* InTexture, const FBoxSphereBounds& InBounds, float InTexelFactor, uint32 InPackedRelativeBox = 0) :
+		Texture(InTexture), 
+		Bounds(InBounds), 
+		TexelFactor(InTexelFactor),
+		PackedRelativeBox(InPackedRelativeBox)
+	{
+	}
+
 };
 
 /** 
diff --git a/Engine/Source/Runtime/Engine/Classes/Engine/World.h b/Engine/Source/Runtime/Engine/Classes/Engine/World.h
--- a/Engine/Source/Runtime/Engine/Classes/Engine/World.h
+++ b/Engine/Source/Runtime/Engine/Classes/Engine/World.h
@@ -3233,6 +3233,9 @@ public:
 	DECLARE_MULTICAST_DELEGATE_TwoParams(FOnWorldTickStart, ELevelTick, float);
 	static FOnWorldTickStart OnWorldTickStart;
 
+	DECLARE_MULTICAST_DELEGATE_ThreeParams(FOnWorldPostActorTick, UWorld* /*World*/, ELevelTick/**Tick Type*/, float/**Delta Seconds*/);
+	static FOnWorldPostActorTick OnWorldPostActorTick;
+
 	// Callback for world creation
 	static FWorldEvent OnPostWorldCreation;
 	
diff --git a/Engine/Source/Runtime/Engine/Classes/Materials/MaterialInterface.h b/Engine/Source/Runtime/Engine/Classes/Materials/MaterialInterface.h
--- a/Engine/Source/Runtime/Engine/Classes/Materials/MaterialInterface.h
+++ b/Engine/Source/Runtime/Engine/Classes/Materials/MaterialInterface.h
@@ -220,6 +220,8 @@ protected:
 	/** Because of redirector, the texture names need to be resorted at each load in case they changed. */
 	UPROPERTY(transient)
 	bool bTextureStreamingDataSorted;
+	UPROPERTY()
+	int32 TextureStreamingDataVersion;
 #endif
 
 	/** Data used by the texture streaming to know how each texture is sampled by the material. Sorted by names for quick access. */
diff --git a/Engine/Source/Runtime/Engine/Engine.Build.cs b/Engine/Source/Runtime/Engine/Engine.Build.cs
--- a/Engine/Source/Runtime/Engine/Engine.Build.cs
+++ b/Engine/Source/Runtime/Engine/Engine.Build.cs
@@ -174,7 +174,6 @@ public class Engine : ModuleRules
 				"HeadMountedDisplay",
 				"MRMesh",
 				"StreamingPauseRendering",
-				"Niagara",
 			}
 		);
 
@@ -422,7 +421,5 @@ public class Engine : ModuleRules
 		{
 			RuntimeDependencies.Add("$(EngineDir)/Content/Stats/...", StagedFileType.UFS);
 		}
-
-		DynamicallyLoadedModuleNames.Add("Niagara");
 	}
 }
diff --git a/Engine/Source/Runtime/Engine/Private/Actor.cpp b/Engine/Source/Runtime/Engine/Private/Actor.cpp
--- a/Engine/Source/Runtime/Engine/Private/Actor.cpp
+++ b/Engine/Source/Runtime/Engine/Private/Actor.cpp
@@ -2249,6 +2249,18 @@ FString AActor::GetHumanReadableName() const
 
 void AActor::DisplayDebug(UCanvas* Canvas, const FDebugDisplayInfo& DebugDisplay, float& YL, float& YPos)
 {
+	// Draw box around Actor being debugged.
+#if ENABLE_DRAW_DEBUG
+	{
+		FVector BoundsOrigin, BoundsExtent;
+		GetActorBounds(true, BoundsOrigin, BoundsExtent);
+
+		// Expand extent a little bit
+		BoundsExtent *= 1.1f;
+		DrawDebugBox(GetWorld(), BoundsOrigin, BoundsExtent, FColor::Green, false, -1.f, 0, 2.f);
+	}
+#endif
+
 	FDisplayDebugManager& DisplayDebugManager = Canvas->DisplayDebugManager;
 	DisplayDebugManager.SetDrawColor(FColor(255, 0, 0));
 
diff --git a/Engine/Source/Runtime/Engine/Private/Animation/AnimInstance.cpp b/Engine/Source/Runtime/Engine/Private/Animation/AnimInstance.cpp
--- a/Engine/Source/Runtime/Engine/Private/Animation/AnimInstance.cpp
+++ b/Engine/Source/Runtime/Engine/Private/Animation/AnimInstance.cpp
@@ -293,10 +293,11 @@ void UAnimInstance::UpdateMontage(float DeltaSeconds)
 	// if we do multi threading, make sure this stays in game thread. 
 	// This is because branch points need to execute arbitrary code inside this call.
 	Montage_Advance(DeltaSeconds);
+}
 
-	// now we know all montage has advanced
-	// time to test sync groups
-	for (auto& MontageInstance : MontageInstances)
+void UAnimInstance::UpdateMontageSyncGroup()
+{
+	for (FAnimMontageInstance* MontageInstance : MontageInstances)
 	{
 		bool bRecordNeedsResetting = true;
 		if (MontageInstance->bDidUseMarkerSyncThisTick)
@@ -304,13 +305,13 @@ void UAnimInstance::UpdateMontage(float DeltaSeconds)
 			const int32 GroupIndexToUse = MontageInstance->GetSyncGroupIndex();
 
 			// that is public data, so if anybody decided to play with it
-			if (ensure (GroupIndexToUse != INDEX_NONE))
+			if (ensure(GroupIndexToUse != INDEX_NONE))
 			{
 				bRecordNeedsResetting = false;
 				FAnimGroupInstance* SyncGroup;
 				FAnimTickRecord& TickRecord = GetProxyOnGameThread<FAnimInstanceProxy>().CreateUninitializedTickRecord(GroupIndexToUse, /*out*/ SyncGroup);
-				MakeMontageTickRecord(TickRecord, MontageInstance->Montage, MontageInstance->GetPosition(), 
-					MontageInstance->GetPreviousPosition(), MontageInstance->GetDeltaMoved(), MontageInstance->GetWeight(), 
+				MakeMontageTickRecord(TickRecord, MontageInstance->Montage, MontageInstance->GetPosition(),
+					MontageInstance->GetPreviousPosition(), MontageInstance->GetDeltaMoved(), MontageInstance->GetWeight(),
 					MontageInstance->MarkersPassedThisTick, MontageInstance->MarkerTickRecord);
 
 				// Update the sync group if it exists
@@ -327,9 +328,6 @@ void UAnimInstance::UpdateMontage(float DeltaSeconds)
 			MontageInstance->MarkerTickRecord.Reset();
 		}
 	}
-
-	// update montage eval data
-	UpdateMontageEvaluationData();
 }
 
 void UAnimInstance::UpdateAnimation(float DeltaSeconds, bool bNeedsValidRootMotion)
@@ -341,6 +339,40 @@ void UAnimInstance::UpdateAnimation(float DeltaSeconds, bool bNeedsValidRootMoti
 	SCOPE_CYCLE_COUNTER(STAT_UpdateAnimation);
 	FScopeCycleCounterUObject AnimScope(this);
 
+	if (const USkeletalMeshComponent* SkelMeshComp = GetSkelMeshComponent())
+	{
+		/**
+			If we're set to OnlyTickMontagesWhenNotRendered and we haven't been recently rendered,
+			then only update montages and skip everything else.
+		*/
+		if ((SkelMeshComp->MeshComponentUpdateFlag == EMeshComponentUpdateFlag::OnlyTickMontagesWhenNotRendered)
+			&& !SkelMeshComp->bRecentlyRendered)
+		{
+			/**
+				Clear NotifyQueue prior to ticking montages. 
+				This is typically done in 'PreUpdate', but we're skipping this here since we're not updating the graph.
+				A side effect of this, is that we're stopping all state notifies in the graph, until ticking resumes.
+				This should be fine. But if it is ever a problem, we should keep two versions of them. One for montages and one for the graph.
+			*/
+			NotifyQueue.Reset(GetSkelMeshComponent());
+
+			UpdateMontage(DeltaSeconds);
+
+			/**
+				We intentionally skip UpdateMontageSyncGroup(), since SyncGroup update is skipped along with AnimGraph update
+				when EMeshComponentUpdateFlag::OnlyTickMontagesWhenNotRendered.
+
+				We also intentionally do not call UpdateMontageEvaluationData after the call to UpdateMontage.
+				As we would have to call 'UpdateAnimation' on the graph as well, so weights could be in sync with this new data.
+				The problem lies in the fact that 'Evaluation' can be called without a call to 'Update' prior.
+				This means our data would be out of sync. So we only call UpdateMontageEvaluationData below
+				when we also update the AnimGraph as well.
+				This means that calls to 'Evaluation' without a call to 'Update' prior will render stale data, but that's to be expected.
+			*/
+			return;
+		}
+	}
+
 	// acquire the proxy as we need to update
 	FAnimInstanceProxy& Proxy = GetProxyOnGameThread<FAnimInstanceProxy>();
 
@@ -375,7 +407,16 @@ void UAnimInstance::UpdateAnimation(float DeltaSeconds, bool bNeedsValidRootMoti
 
 	// need to update montage BEFORE node update or Native Update.
 	// so that node knows where montage is
-	UpdateMontage(DeltaSeconds);
+	{
+		UpdateMontage(DeltaSeconds);
+
+		// now we know all montage has advanced
+		// time to test sync groups
+		UpdateMontageSyncGroup();
+
+		// Update montage eval data, to be used by AnimGraph Update and Evaluate phases.
+		UpdateMontageEvaluationData();
+	}
 
 	{
 		SCOPE_CYCLE_COUNTER(STAT_NativeUpdateAnimation);
@@ -2626,6 +2667,11 @@ void UAnimInstance::RecordStateWeight(const int32 InMachineClassIndex, const int
 	GetProxyOnAnyThread<FAnimInstanceProxy>().RecordStateWeight(InMachineClassIndex, InStateIndex, InStateWeight);
 }
 
+const FGraphTraversalCounter& UAnimInstance::GetUpdateCounter() const
+{
+	return GetProxyOnGameThread<FAnimInstanceProxy>().GetUpdateCounter();
+}
+
 FBoneContainer& UAnimInstance::GetRequiredBones()
 {
 	return GetProxyOnGameThread<FAnimInstanceProxy>().GetRequiredBones();
diff --git a/Engine/Source/Runtime/Engine/Private/Animation/AnimInstanceProxy.cpp b/Engine/Source/Runtime/Engine/Private/Animation/AnimInstanceProxy.cpp
--- a/Engine/Source/Runtime/Engine/Private/Animation/AnimInstanceProxy.cpp
+++ b/Engine/Source/Runtime/Engine/Private/Animation/AnimInstanceProxy.cpp
@@ -129,6 +129,21 @@ void FAnimInstanceProxy::Initialize(UAnimInstance* InAnimInstance)
 #endif
 
 	ReinitializeSlotNodes();
+
+	if (const USkeletalMeshComponent* SkelMeshComp = InAnimInstance->GetOwningComponent())
+	{
+		ComponentTransform = SkelMeshComp->GetComponentTransform();
+		ComponentRelativeTransform = SkeletalMeshComponent->GetRelativeTransform();
+
+		const AActor* OwningActor = SkeletalMeshComponent->GetOwner();
+		ActorTransform = OwningActor ? OwningActor->GetActorTransform() : FTransform::Identity;
+	}
+	else
+	{
+		ComponentTransform = FTransform::Identity;
+		ComponentRelativeTransform = FTransform::Identity;
+		ActorTransform = FTransform::Identity;
+	}
 }
 
 void FAnimInstanceProxy::InitializeRootNode()
diff --git a/Engine/Source/Runtime/Engine/Private/Animation/AnimMontage.cpp b/Engine/Source/Runtime/Engine/Private/Animation/AnimMontage.cpp
--- a/Engine/Source/Runtime/Engine/Private/Animation/AnimMontage.cpp
+++ b/Engine/Source/Runtime/Engine/Private/Animation/AnimMontage.cpp
@@ -15,8 +15,16 @@
 #include "Animation/AnimNotifies/AnimNotify.h"
 #include "Animation/AnimNotifies/AnimNotifyState.h"
 #include "Animation/AnimSingleNodeInstance.h"
+#include "Engine/Engine.h"
 
 DEFINE_LOG_CATEGORY(LogAnimMontage);
+
+// Pre-built FNames so we don't take the hit of constructing FNames at spawn time
+namespace MontageFNames
+{
+	static FName TimeStretchCurveName(TEXT("MontageTimeStretchCurve"));
+}
+
 ///////////////////////////////////////////////////////////////////////////
 //
 UAnimMontage::UAnimMontage(const FObjectInitializer& ObjectInitializer)
@@ -31,6 +39,8 @@ UAnimMontage::UAnimMontage(const FObjectInitializer& ObjectInitializer)
 	BlendOutTime_DEPRECATED = -1.f;
 
 	AddSlot(FAnimSlotGroup::DefaultSlotName);
+
+	TimeStretchCurveName = MontageFNames::TimeStretchCurveName;
 }
 
 FSlotAnimationTrack& UAnimMontage::AddSlot(FName SlotName)
@@ -316,6 +326,13 @@ void UAnimMontage::UnregisterOnMontageChanged(void* Unregister)
 }
 #endif	//WITH_EDITOR
 
+void UAnimMontage::PreSave(const class ITargetPlatform* TargetPlatform)
+{
+	BakeTimeStretchCurve();
+
+	Super::PreSave(TargetPlatform);
+}
+
 void UAnimMontage::PostLoad()
 {
 	Super::PostLoad();
@@ -594,15 +611,8 @@ void UAnimMontage::AddBranchingPointMarker(FBranchingPointMarker TickMarker, TMa
 	}
 }
 
-const FBranchingPointMarker* UAnimMontage::FindFirstBranchingPointMarker(float StartTrackPos, float EndTrackPos)
+const FBranchingPointMarker* UAnimMontage::FindFirstBranchingPointMarker(float StartTrackPos, float EndTrackPos) const
 {
-	// Auto refresh this in editor to catch changes being made to AnimNotifies.
-	// RefreshCacheData should handle this but I'm not 100% sure it will cover all existing cases
-	if (GIsEditor && (!GetWorld() || !GetWorld()->IsPlayInEditor()))
-	{
-		RefreshBranchingPointMarkers();
-	}
-
 	if (BranchingPointMarkers.Num() > 0)
 	{
 		const bool bSearchBackwards = (EndTrackPos < StartTrackPos);
@@ -639,7 +649,7 @@ const FBranchingPointMarker* UAnimMontage::FindFirstBranchingPointMarker(float S
 			}
 		}
 	}
-	return NULL;
+	return nullptr;
 }
 
 void UAnimMontage::FilterOutNotifyBranchingPoints(TArray<const FAnimNotifyEvent*>& InAnimNotifies)
@@ -1387,6 +1397,8 @@ void FAnimMontageInstance::Initialize(class UAnimMontage * InMontage)
 		{
 			SyncGroupIndex = AnimInstance.Get()->GetSyncGroupIndexFromName(Montage->SyncGroup);
 		}
+
+		MontageSubStepper.Initialize(*this);
 	}
 }
 
@@ -1698,125 +1710,434 @@ void FAnimMontageInstance::UpdateWeight(float DeltaTime)
 
 bool FAnimMontageInstance::SimulateAdvance(float DeltaTime, float& InOutPosition, FRootMotionMovementParams & OutRootMotionParams) const
 {
-	if( !IsValid() )
+	if (!IsValid())
 	{
 		return false;
 	}
 
-	const float CombinedPlayRate = PlayRate * Montage->RateScale;
-	const bool bPlayingForward = (CombinedPlayRate > 0.f);
-
 	const bool bExtractRootMotion = Montage->HasRootMotion() && !IsRootMotionDisabled();
 
-	float DesiredDeltaMove = CombinedPlayRate * DeltaTime;
-	float OriginalMoveDelta = DesiredDeltaMove;
+	FMontageSubStepper SimulateMontageSubStepper;
+	SimulateMontageSubStepper.Initialize(*this);
+	SimulateMontageSubStepper.AddEvaluationTime(DeltaTime);
+	while (SimulateMontageSubStepper.HasTimeRemaining())
+	{
+		const float PreviousSubStepPosition = InOutPosition;
+		EMontageSubStepResult SubStepResult = SimulateMontageSubStepper.Advance(InOutPosition, nullptr);
 
-	while( (FMath::Abs(DesiredDeltaMove) > KINDA_SMALL_NUMBER) && ((OriginalMoveDelta * DesiredDeltaMove) > 0.f) )
+		if (SubStepResult != EMontageSubStepResult::Moved)
+		{
+			// stop and leave this loop
+			break;
+		}
+
+		// Extract Root Motion for this time slice, and accumulate it.
+		if (bExtractRootMotion)
+		{
+			OutRootMotionParams.Accumulate(Montage->ExtractRootMotionFromTrackRange(PreviousSubStepPosition, InOutPosition));
+		}
+
+		// if we reached end of section, and we were not processing a branching point, and no events has messed with out current position..
+		// .. Move to next section.
+		// (this also handles looping, the same as jumping to a different section).
+		if (SimulateMontageSubStepper.HasReachedEndOfSection())
+		{
+			const int32 CurrentSectionIndex = SimulateMontageSubStepper.GetCurrentSectionIndex();
+			const bool bPlayingForward = SimulateMontageSubStepper.GetbPlayingForward();
+
+			// Get recent NextSectionIndex in case it's been changed by previous events.
+			const int32 RecentNextSectionIndex = bPlayingForward ? NextSections[CurrentSectionIndex] : PrevSections[CurrentSectionIndex];
+			if (RecentNextSectionIndex != INDEX_NONE)
+			{
+				float LatestNextSectionStartTime;
+				float LatestNextSectionEndTime;
+				Montage->GetSectionStartAndEndTime(RecentNextSectionIndex, LatestNextSectionStartTime, LatestNextSectionEndTime);
+
+				// Jump to next section's appropriate starting point (start or end).
+				InOutPosition = bPlayingForward ? LatestNextSectionStartTime : (LatestNextSectionEndTime - KINDA_SMALL_NUMBER); // remain within section
+			}
+			else
+			{
+				// Reached end of last section. Exit.
+				break;
+			}
+		}
+	}
+
+	return true;
+}
+
+void FMontageSubStepper::Initialize(const struct FAnimMontageInstance& InAnimInstance)
+{
+	MontageInstance = &InAnimInstance;
+	Montage = MontageInstance->Montage;
+}
+
+EMontageSubStepResult FMontageSubStepper::Advance(float& InOut_P_Original, const FBranchingPointMarker** OutBranchingPointMarkerPtr)
+{
+	DeltaMove = 0.f;
+
+	if (MontageInstance == nullptr || (Montage == nullptr))
 	{
-		// Get position relative to current montage section.
-		float PosInSection;
-		int32 CurrentSectionIndex = Montage->GetAnimCompositeSectionIndexFromPos(InOutPosition, PosInSection);
+		return EMontageSubStepResult::InvalidMontage;
+	}
 
-		if( Montage->IsValidSectionIndex(CurrentSectionIndex) )
+	bReachedEndOfSection = false;
+
+	// Update Current Section info in case it's needed by the montage's update loop.
+	// We need to do this even if we're not going to move this frame.
+	// We could have been moved externally via a SetPosition() call.
+	float PositionInSection;
+	CurrentSectionIndex = Montage->GetAnimCompositeSectionIndexFromPos(InOut_P_Original, PositionInSection);
+	if (!Montage->IsValidSectionIndex(CurrentSectionIndex))
+	{
+		return EMontageSubStepResult::InvalidSection;
+	}
+
+	const FCompositeSection& CurrentSection = Montage->GetAnimCompositeSection(CurrentSectionIndex);
+	CurrentSectionStartTime = CurrentSection.GetTime();
+
+	// Find end of current section. We only update one section at a time.
+	CurrentSectionLength = Montage->GetSectionLength(CurrentSectionIndex);
+
+	if (!MontageInstance->bPlaying || FMath::IsNearlyZero(TimeRemaining))
+	{
+		return EMontageSubStepResult::NotMoved;
+	}
+
+	// If we're forcing next position, this is our DeltaMove.
+	// We don't use play rate and delta time to move.
+	if (MontageInstance->ForcedNextPosition.IsSet())
+	{
+		const float NewPosition = MontageInstance->ForcedNextPosition.GetValue();
+		DeltaMove = NewPosition - InOut_P_Original;
+		PlayRate = DeltaMove / TimeRemaining;
+		bPlayingForward = (DeltaMove >= 0.f);
+		TimeStretchMarkerIndex = INDEX_NONE;
+	}
+	else 
+	{
+		PlayRate = MontageInstance->PlayRate * Montage->RateScale;
+
+		if (FMath::IsNearlyZero(PlayRate))
 		{
-			const FCompositeSection& CurrentSection = Montage->GetAnimCompositeSection(CurrentSectionIndex);
+			return EMontageSubStepResult::NotMoved;
+		}
 
-			// we need to advance within section
-			check(NextSections.IsValidIndex(CurrentSectionIndex));
-			const int32 NextSectionIndex = bPlayingForward ? NextSections[CurrentSectionIndex] : PrevSections[CurrentSectionIndex];
+		// See if we can attempt to use a TimeStretchCurve.
+		const bool bAttemptTimeStretchCurve = Montage->TimeStretchCurve.IsValid() && !FMath::IsNearlyEqual(PlayRate, 1.f);
+		if (bAttemptTimeStretchCurve)
+		{
+			// First we need to see if we have valid cached data and if it is up to date.
+			ConditionallyUpdateTimeStretchCurveCachedData();
+		}
 
-			// Find end of current section. We only update one section at a time.
-			const float CurrentSectionLength = Montage->GetSectionLength(CurrentSectionIndex);
-			float StoppingPosInSection = CurrentSectionLength;
+		// If we're not using a TimeStretchCurve, play rate is constant.
+		if (!bAttemptTimeStretchCurve || !bHasValidTimeStretchCurveData)
+		{
+			bPlayingForward = (PlayRate > 0.f);
+			DeltaMove = TimeRemaining * PlayRate;
+			TimeStretchMarkerIndex = INDEX_NONE;
+		}
+		else
+		{
+			// We're using a TimeStretchCurve.
 
-			// Also look for a branching point. If we have one, stop there first to handle it.
-			const float CurrentSectionEndPos = CurrentSection.GetTime() + CurrentSectionLength;
+			// Find P_Target for current InOut_P_Original.
+			// Not that something external could have modified the montage's position.
+			// So we need to refresh our P_Target.
+			float P_Target = FindMontagePosition_Target(InOut_P_Original);
 
-			// Update position within current section.
-			// Note that we explicitly disallow looping there, we handle it ourselves to simplify code by not having to handle time wrapping around in multiple places.
-			const float PrevPosInSection = PosInSection;
-			const ETypeAdvanceAnim AdvanceType = FAnimationRuntime::AdvanceTime(false, DesiredDeltaMove, PosInSection, StoppingPosInSection);
-			// ActualDeltaPos == ActualDeltaMove since we break down looping in multiple steps. So we don't have to worry about time wrap around here.
-			const float ActualDeltaPos = PosInSection - PrevPosInSection;
-			const float ActualDeltaMove = ActualDeltaPos;
+			// With P_Target, we're in 'play back time' space. 
+			// So we can add our delta time there directly.
+			P_Target += bPlayingForward ? TimeRemaining : -TimeRemaining;
+			// Make sure we don't exceed our boundaries.
+			P_Target = TimeStretchCurveInstance.Clamp_P_Target(P_Target);
 
-			// Decrease actual move from desired move. We'll take another iteration if there is anything left.
-			DesiredDeltaMove -= ActualDeltaMove;
+			// Now we can map this back into 'original' space and find which frame of animation we should play.
+			const float NewP_Original = FindMontagePosition_Original(P_Target);
 
-			float PrevPosition = Position;
-			InOutPosition = FMath::Clamp<float>(InOutPosition + ActualDeltaPos, CurrentSection.GetTime(), CurrentSectionEndPos);
+			// And from there, derive our DeltaMove and actual PlayRate for this substep.
+			DeltaMove = NewP_Original - InOut_P_Original;
+			PlayRate = DeltaMove / TimeRemaining;
+		}
+	}
 
-			if( FMath::Abs(ActualDeltaMove) > 0.f )
-			{
-				// Extract Root Motion for this time slice, and accumulate it.
-				if( bExtractRootMotion )
-				{
-					OutRootMotionParams.Accumulate(Montage->ExtractRootMotionFromTrackRange(PrevPosition, InOutPosition));
-				}
+	// Now look for a branching point. If we have one, stop there first to handle it.
+	// We need to stop at branching points, because they can trigger events that can cause side effects
+	// (jumping to a new position, changing sections, changing play rate, etc).
+	if (OutBranchingPointMarkerPtr)
+	{
+		*OutBranchingPointMarkerPtr = Montage->FindFirstBranchingPointMarker(InOut_P_Original, InOut_P_Original + DeltaMove);
+		if (*OutBranchingPointMarkerPtr)
+		{
+			// If we have a branching point, adjust DeltaMove so we stop there.
+			DeltaMove = (*OutBranchingPointMarkerPtr)->TriggerTime - InOut_P_Original;
+		}
+	}
 
-				// if we reached end of section, and we were not processing a branching point, and no events has messed with out current position..
-				// .. Move to next section.
-				// (this also handles looping, the same as jumping to a different section).
-				if( AdvanceType != ETAA_Default )
-				{
-					// Get recent NextSectionIndex in case it's been changed by previous events.
-					const int32 RecentNextSectionIndex = bPlayingForward ? NextSections[CurrentSectionIndex] : PrevSections[CurrentSectionIndex];
-					if( RecentNextSectionIndex != INDEX_NONE )
-					{
-						float LatestNextSectionStartTime;
-						float LatestNextSectionEndTime;
-						Montage->GetSectionStartAndEndTime(RecentNextSectionIndex, LatestNextSectionStartTime, LatestNextSectionEndTime);
-						// Jump to next section's appropriate starting point (start or end).
-						InOutPosition = bPlayingForward ? LatestNextSectionStartTime : (LatestNextSectionEndTime - KINDA_SMALL_NUMBER); // remain within section
-					}
-				}
-			}
-			else
+	// Finally clamp DeltaMove by section markers.
+	{
+		// Clamp DeltaMove based on move allowed within current section
+		// We stop at each section marker to evaluate whether we should jump to another section marker or not.
+		// Test is inclusive, so we know if we've reached marker or not.
+		if (bPlayingForward)
+		{
+			const float MaxSectionMove = CurrentSectionLength - PositionInSection;
+			if (DeltaMove >= MaxSectionMove)
 			{
-				break;
+				DeltaMove = MaxSectionMove;
+				bReachedEndOfSection = true;
 			}
 		}
 		else
 		{
-			// stop and leave this loop
-			break;
+			const float MinSectionMove = /* 0.f */ - PositionInSection;
+			if (DeltaMove <= MinSectionMove)
+			{
+				DeltaMove = MinSectionMove;
+				bReachedEndOfSection = true;
+			}
 		}
 	}
 
-	return true;
+	// DeltaMove is now final, see if it has any effect on our position.
+	if (FMath::Abs(DeltaMove) > 0.f)
+	{
+		// Note that we don't worry about looping and wrapping around here.
+		// We step per section to simplify code to extract notifies/root motion/etc.
+		InOut_P_Original += DeltaMove;
+
+		// Decrease RemainingTime with actual time elapsed 
+		// So we can take more substeps as needed.
+		const float TimeStep = DeltaMove / PlayRate;
+		ensure(TimeStep >= 0.f);
+		TimeRemaining = FMath::Max(TimeRemaining - TimeStep, 0.f);
+
+		return EMontageSubStepResult::Moved;
+	}
+	else
+	{
+		return EMontageSubStepResult::NotMoved;
+	}
 }
 
-void FAnimMontageInstance::Advance(float DeltaTime, struct FRootMotionMovementParams * OutRootMotionParams, bool bBlendRootMotion)
+void FMontageSubStepper::ConditionallyUpdateTimeStretchCurveCachedData()
 {
-	SCOPE_CYCLE_COUNTER(STAT_AnimMontageInstance_Advance);
-	FScopeCycleCounterUObject MontageScope(Montage);
+	// CombinedPlayRate defines our overall desired play back time, aka T_Target.
+	// When using a TimeStretchCurve, this also defines S and U.
+	// Only update these if CombinedPlayRate has changed.
+	const float CombinedPlayRate = MontageInstance->PlayRate * Montage->RateScale;
+	if (CombinedPlayRate == Cached_CombinedPlayRate)
+	{
+		return;
+	}
+	Cached_CombinedPlayRate = CombinedPlayRate;
+	
+	// We'll set this to true at the end, if we succeed with valid data.
+	bHasValidTimeStretchCurveData = false;
 
-	if( IsValid() )
+	// We should not be using this code path with a 0 play rate
+	// or a 1 play rate. we can use traditional cheaper update without curve.
+	ensure(!FMath::IsNearlyZero(CombinedPlayRate));
+	ensure(!FMath::IsNearlyEqual(CombinedPlayRate, 1.f));
+
+	bPlayingForward = (CombinedPlayRate > 0.f);
+	TimeStretchCurveInstance.InitializeFromPlayRate(CombinedPlayRate, Montage->TimeStretchCurve);
+
+	/*
+		Section Segment Positions in Target space will have to be re-cached, as needed.
+		This is to determine 'remaining time until end' to trigger blend outs.
+		But most montages don't use sections.
+		So this is optional and done on demand.
+	*/
 	{
-#if WITH_EDITOR
-		// this is necessary and it is not easy to do outside of here
-		// since undo also can change composite sections
-		if( (Montage->CompositeSections.Num() != NextSections.Num()) || (Montage->CompositeSections.Num() != PrevSections.Num()))
+		const int32 NumSections = Montage->CompositeSections.Num();
+		SectionStartPositions_Target.Reset(NumSections);
+		SectionStartPositions_Target.Init(-1.f, NumSections);
+		SectionEndPositions_Target.Reset(NumSections);
+		SectionEndPositions_Target.Init(-1.f, NumSections);
+	}
+
+	bHasValidTimeStretchCurveData = TimeStretchCurveInstance.HasValidData();
+}
+
+float FMontageSubStepper::FindMontagePosition_Target(float In_P_Original)
+{
+	check(bHasValidTimeStretchCurveData);
+
+	// See if our cached version is not up to date.
+	// Then we need to update it.
+	if (In_P_Original != Cached_P_Original)
+	{
+		// Update cached value.
+		Cached_P_Original = In_P_Original;
+
+		// Update TimeStretchMarkerIndex if needed.
+		// This would happen if we jumped position due to sections or external input.
+		TimeStretchCurveInstance.UpdateMarkerIndexForPosition(TimeStretchMarkerIndex, Cached_P_Original, TimeStretchCurveInstance.GetMarkers_Original());
+
+		// With an accurate TimeStretchMarkerIndex, we can map P_Original to P_Target
+		Cached_P_Target = TimeStretchCurveInstance.Convert_P_Original_To_Target(TimeStretchMarkerIndex, Cached_P_Original);
+	}
+
+	return Cached_P_Target;
+}
+
+float FMontageSubStepper::FindMontagePosition_Original(float In_P_Target)
+{
+	check(bHasValidTimeStretchCurveData);
+
+	// See if our cached version is not up to date.
+	// Then we need to update it.
+	if (In_P_Target != Cached_P_Target)
+	{
+		// Update cached value.
+		Cached_P_Target = In_P_Target;
+
+		// Update TimeStretchMarkerIndex if needed.
+		// This would happen if we jumped position due to sections or external input.
+		TimeStretchCurveInstance.UpdateMarkerIndexForPosition(TimeStretchMarkerIndex, Cached_P_Target, TimeStretchCurveInstance.GetMarkers_Target());
+
+		// With an accurate TimeStretchMarkerIndex, we can map P_Original to P_Target
+		Cached_P_Original = TimeStretchCurveInstance.Convert_P_Target_To_Original(TimeStretchMarkerIndex, Cached_P_Target);
+	}
+
+	return Cached_P_Original;
+}
+
+float FMontageSubStepper::GetCurrSectionStartPosition_Target() const
+{
+	check(bHasValidTimeStretchCurveData);
+
+	const float CachedSectionStartPosition_Target = SectionStartPositions_Target[CurrentSectionIndex];
+	if (CachedSectionStartPosition_Target >= 0.f)
+	{
+		return CachedSectionStartPosition_Target;
+	}
+
+	const int32 SectionStartMarkerIndex = TimeStretchCurveInstance.BinarySearchMarkerIndex(CurrentSectionStartTime, TimeStretchCurveInstance.GetMarkers_Original());
+	const float SectionStart_Target = TimeStretchCurveInstance.Convert_P_Original_To_Target(SectionStartMarkerIndex, CurrentSectionStartTime);
+
+	SectionStartPositions_Target[CurrentSectionIndex] = SectionStart_Target;
+
+	return SectionStart_Target;
+}
+
+float FMontageSubStepper::GetCurrSectionEndPosition_Target() const
+{
+	check(bHasValidTimeStretchCurveData);
+
+	const float CachedSectionEndPosition_Target = SectionEndPositions_Target[CurrentSectionIndex];
+	if (CachedSectionEndPosition_Target >= 0.f)
+	{
+		return CachedSectionEndPosition_Target;
+	}
+
+	const float SectionEnd_Original = CurrentSectionStartTime + CurrentSectionLength;
+	const int32 SectionEndMarkerIndex = TimeStretchCurveInstance.BinarySearchMarkerIndex(SectionEnd_Original, TimeStretchCurveInstance.GetMarkers_Original());
+	const float SectionEnd_Target = TimeStretchCurveInstance.Convert_P_Original_To_Target(SectionEndMarkerIndex, SectionEnd_Original);
+
+	SectionEndPositions_Target[CurrentSectionIndex] = SectionEnd_Target;
+
+	return SectionEnd_Target;
+}
+
+float FMontageSubStepper::GetRemainingPlayTimeToSectionEnd(const float In_P_Original)
+{
+	// If our current play rate is zero, we can't predict our remaining play time.
+	if (FMath::IsNearlyZero(PlayRate))
+	{
+		return BIG_NUMBER;
+	}
+
+	// Find position in montage where current section ends.
+	const float CurrSectionEnd_Original = bPlayingForward
+		? (CurrentSectionStartTime + CurrentSectionLength)
+		: CurrentSectionStartTime;
+
+	// If we have no TimeStretchCurve, it's pretty straight forward.
+	// Assume constant play rate.
+	if (TimeStretchMarkerIndex == INDEX_NONE)
+	{
+		const float DeltaPositionToEnd = CurrSectionEnd_Original - In_P_Original;
+		const float RemainingPlayTime = FMath::Abs(DeltaPositionToEnd / PlayRate);
+		return RemainingPlayTime;
+	}
+
+	// We're using a TimeStretchCurve.
+	check(bHasValidTimeStretchCurveData);
+
+	// Find our position in 'target' space. This is in play back time.
+	const float P_Target = FindMontagePosition_Target(In_P_Original);
+	if (bPlayingForward)
+	{
+		// Find CurrSectionEnd_Target.
+		if (FMath::IsNearlyEqual(CurrSectionEnd_Original, TimeStretchCurveInstance.Get_T_Original()))
+		{
+			const float RemainingPlayTime = (TimeStretchCurveInstance.Get_T_Target() - P_Target);
+			return RemainingPlayTime;
+		}
+		else
+		{
+			const float CurrSectionEnd_Target = GetCurrSectionEndPosition_Target();
+			const float RemainingPlayTime = (CurrSectionEnd_Target - P_Target);
+			return RemainingPlayTime;
+		}
+	}
+	// Playing Backwards
+	else
+	{
+		// Find CurrSectionEnd_Target.
+		if (FMath::IsNearlyEqual(CurrSectionEnd_Original, 0.f))
 		{
-			RefreshNextPrevSections();
+			const float RemainingPlayTime = P_Target;
+			return RemainingPlayTime;
 		}
+		else
+		{
+			const float CurrSectionStart_Target = GetCurrSectionStartPosition_Target();
+			const float RemainingPlayTime = (P_Target - CurrSectionStart_Target);
+			return RemainingPlayTime;
+		}
+	}
+}
+
+#if WITH_EDITOR
+void FAnimMontageInstance::EditorOnly_PreAdvance()
+{
+	// this is necessary and it is not easy to do outside of here
+	// since undo also can change composite sections
+	if ((Montage->CompositeSections.Num() != NextSections.Num()) || (Montage->CompositeSections.Num() != PrevSections.Num()))
+	{
+		RefreshNextPrevSections();
+	}
+
+	// Auto refresh this in editor to catch changes being made to AnimNotifies.
+	// RefreshCacheData should handle this but I'm not 100% sure it will cover all existing cases
+	Montage->RefreshBranchingPointMarkers();
+
+	// Bake TimeStretchCurve in editor to catch any edits made to source curve.
+	Montage->BakeTimeStretchCurve();
+	// Clear cached data, so it can be recached from updated time stretch curve.
+	MontageSubStepper.ClearCachedData();
+}
 #endif
 
+void FAnimMontageInstance::Advance(float DeltaTime, struct FRootMotionMovementParams* OutRootMotionParams, bool bBlendRootMotion)
+{
+	SCOPE_CYCLE_COUNTER(STAT_AnimMontageInstance_Advance);
+	FScopeCycleCounterUObject MontageScope(Montage);
+
+	if (IsValid())
+	{
 		// with custom curves, we can't just filter by weight
 		// also if you have custom curve with longer 0, you'll likely to pause montage during that blending time
 		// I think that is a bug. It still should move, the weight might come back later. 
-		if( bPlaying )
+		if (bPlaying)
 		{
-			const float CombinedPlayRate = PlayRate * Montage->RateScale;
-			const bool bPlayingForward = (CombinedPlayRate > 0.f);
-
-			const bool bExtractRootMotion = (OutRootMotionParams != NULL) && Montage->HasRootMotion();
+			const bool bExtractRootMotion = (OutRootMotionParams != nullptr) && Montage->HasRootMotion();
 			
-			float DesiredDeltaMove = ForcedNextPosition.IsSet() ? ForcedNextPosition.GetValue() - Position : CombinedPlayRate * DeltaTime;
-			float OriginalMoveDelta = DesiredDeltaMove;
-
-			ForcedNextPosition.Reset();
-
 			DeltaMoved = 0.f;
 			PreviousPosition = Position;
 
@@ -1825,150 +2146,145 @@ void FAnimMontageInstance::Advance(float DeltaTime, struct FRootMotionMovementPa
 			{
 				MarkersPassedThisTick.Reset();
 			}
+			
+			/** 
+				Limit number of iterations for performance.
+				This can get out of control if PlayRate is set really high, or there is a hitch, and Montage is looping for example.
+			*/
+			const int32 MaxIterations = 10;
+			int32 NumIterations = 0;
+
+			/** 
+				If we're hitting our max number of iterations for whatever reason,
+				make sure we're not accumulating too much time, and go out of range.
+			*/
+			if (MontageSubStepper.GetRemainingTime() < 10.f)
+			{
+				MontageSubStepper.AddEvaluationTime(DeltaTime);
+			}
 
-			while( bPlaying && (FMath::Abs(DesiredDeltaMove) > KINDA_SMALL_NUMBER) && ((OriginalMoveDelta * DesiredDeltaMove) > 0.f) )
+			while (bPlaying && MontageSubStepper.HasTimeRemaining() && (++NumIterations < MaxIterations))
 			{
 				SCOPE_CYCLE_COUNTER(STAT_AnimMontageInstance_Advance_Iteration);
 
-				// Get position relative to current montage section.
-				float PosInSection;
-				int32 CurrentSectionIndex = Montage->GetAnimCompositeSectionIndexFromPos(Position, PosInSection);
-				
-				if( Montage->IsValidSectionIndex(CurrentSectionIndex) )
+				const float PreviousSubStepPosition = Position;
+				const FBranchingPointMarker* BranchingPointMarker = nullptr;
+				EMontageSubStepResult SubStepResult = MontageSubStepper.Advance(Position, &BranchingPointMarker);
+
+				if (SubStepResult == EMontageSubStepResult::InvalidSection
+					|| SubStepResult == EMontageSubStepResult::InvalidMontage)
 				{
-					const FCompositeSection& CurrentSection = Montage->GetAnimCompositeSection(CurrentSectionIndex);
+					// stop and leave this loop
+					Stop(FAlphaBlend(Montage->BlendOut, Montage->BlendOut.GetBlendTime() * DefaultBlendTimeMultiplier), false);
+					break;
+				}
 
-					// we need to advance within section
+				const float SubStepDeltaMove = MontageSubStepper.GetDeltaMove();
+				DeltaMoved += SubStepDeltaMove;
+				const bool bPlayingForward = MontageSubStepper.GetbPlayingForward();
+
+				// If current section is last one, check to trigger a blend out and if it hasn't stopped yet, see if we should stop
+				// We check this even if we haven't moved, in case our position was different from last frame.
+				// (Code triggered a position jump).
+				if (!IsStopped())
+				{
+					const int32 CurrentSectionIndex = MontageSubStepper.GetCurrentSectionIndex();
 					check(NextSections.IsValidIndex(CurrentSectionIndex));
 					const int32 NextSectionIndex = bPlayingForward ? NextSections[CurrentSectionIndex] : PrevSections[CurrentSectionIndex];
-
-					// Find end of current section. We only update one section at a time.
-					const float CurrentSectionLength = Montage->GetSectionLength(CurrentSectionIndex);
-					float StoppingPosInSection = CurrentSectionLength;
-
-					// Also look for a branching point. If we have one, stop there first to handle it.
-					const float CurrentSectionEndPos = CurrentSection.GetTime() + CurrentSectionLength;
-					const FBranchingPointMarker* BranchingPointMarker = Montage->FindFirstBranchingPointMarker(Position, FMath::Clamp(Position + DesiredDeltaMove, CurrentSection.GetTime(), CurrentSectionEndPos));
-					if (BranchingPointMarker)
+					if (NextSectionIndex == INDEX_NONE)
 					{
-						// get the first one to see if it's less than AnimEnd
-						StoppingPosInSection = BranchingPointMarker->TriggerTime - CurrentSection.GetTime();
-					}
+						const float PlayTimeToEnd = MontageSubStepper.GetRemainingPlayTimeToSectionEnd(Position);
 
-					// Update position within current section.
-					// Note that we explicitly disallow looping there, we handle it ourselves to simplify code by not having to handle time wrapping around in multiple places.
-					const float PrevPosInSection = PosInSection;
-					const ETypeAdvanceAnim AdvanceType = FAnimationRuntime::AdvanceTime(false, DesiredDeltaMove, PosInSection, StoppingPosInSection);
-					// ActualDeltaPos == ActualDeltaMove since we break down looping in multiple steps. So we don't have to worry about time wrap around here.
-					const float ActualDeltaPos = PosInSection - PrevPosInSection;
-					const float ActualDeltaMove = ActualDeltaPos;
-					DeltaMoved += ActualDeltaMove;
-
-					// Decrease actual move from desired move. We'll take another iteration if there is anything left.
-					DesiredDeltaMove -= ActualDeltaMove;
+						const bool bCustomBlendOutTriggerTime = (Montage->BlendOutTriggerTime >= 0);
+						const float DefaultBlendOutTime = Montage->BlendOut.GetBlendTime() * DefaultBlendTimeMultiplier;
+						const float BlendOutTriggerTime = bCustomBlendOutTriggerTime ? Montage->BlendOutTriggerTime : DefaultBlendOutTime;
 
-					float PrevPosition = Position;
-					Position = FMath::Clamp<float>(Position + ActualDeltaPos, CurrentSection.GetTime(), CurrentSectionEndPos);
+						// ... trigger blend out if within blend out time window.
+						if (PlayTimeToEnd <= FMath::Max<float>(BlendOutTriggerTime, KINDA_SMALL_NUMBER))
+						{
+							const float BlendOutTime = bCustomBlendOutTriggerTime ? DefaultBlendOutTime : PlayTimeToEnd;
+							Stop(FAlphaBlend(Montage->BlendOut, BlendOutTime), false);
+						}
+					}
+				}
 
+				const bool bHaveMoved = (SubStepResult == EMontageSubStepResult::Moved);
+				if (bHaveMoved)
+				{
 					if (bDidUseMarkerSyncThisTick)
 					{
-						Montage->MarkerData.CollectMarkersInRange(PrevPosition, Position, MarkersPassedThisTick, ActualDeltaMove);
+						Montage->MarkerData.CollectMarkersInRange(PreviousSubStepPosition, Position, MarkersPassedThisTick, SubStepDeltaMove);
 					}
 
-					const float PositionBeforeFiringEvents = Position;
-
-					const bool bHaveMoved = FMath::Abs(ActualDeltaMove) > 0.f;
-
-					const float Weight = Blend.GetBlendedValue();
-
-					if (bHaveMoved)
+					// Extract Root Motion for this time slice, and accumulate it.
+					// IsRootMotionDisabled() can be changed by AnimNotifyState BranchingPoints while advancing, so it needs to be checked here.
+					if (bExtractRootMotion && AnimInstance.IsValid() && !IsRootMotionDisabled())
 					{
-						// Extract Root Motion for this time slice, and accumulate it.
-						// IsRootMotionDisabled() can be changed by AnimNotifyState BranchingPoints while advancing, so it needs to be checked here.
-						if (bExtractRootMotion && AnimInstance.IsValid() && !IsRootMotionDisabled())
+						const FTransform RootMotion = Montage->ExtractRootMotionFromTrackRange(PreviousSubStepPosition, Position);
+						if (bBlendRootMotion)
 						{
-							const FTransform RootMotion = Montage->ExtractRootMotionFromTrackRange(PrevPosition, Position);
-							if (bBlendRootMotion)
-							{
-								// Defer blending in our root motion until after we get our slot weight updated
-								AnimInstance.Get()->QueueRootMotionBlend(RootMotion, Montage->SlotAnimTracks[0].SlotName, Weight);
-							}
-							else
-							{
-								OutRootMotionParams->Accumulate(RootMotion);
-							}
-
-							UE_LOG(LogRootMotion, Log, TEXT("\tFAnimMontageInstance::Advance ExtractedRootMotion: %s, AccumulatedRootMotion: %s, bBlendRootMotion: %d")
-								, *RootMotion.GetTranslation().ToCompactString()
-								, *OutRootMotionParams->GetRootMotionTransform().GetTranslation().ToCompactString()
-								, bBlendRootMotion
-								);
+							// Defer blending in our root motion until after we get our slot weight updated
+							const float Weight = Blend.GetBlendedValue();
+							AnimInstance.Get()->QueueRootMotionBlend(RootMotion, Montage->SlotAnimTracks[0].SlotName, Weight);
 						}
-					}
-
-					// If current section is last one, check to trigger a blend out and if it hasn't stopped yet, see if we should stop
-					if (NextSectionIndex == INDEX_NONE && !IsStopped())
-					{
-						const float DeltaPosToEnd = bPlayingForward ? (CurrentSectionLength - PosInSection) : PosInSection;
-						const float DeltaTimeToEnd = DeltaPosToEnd / FMath::Abs(CombinedPlayRate);
-
-						const bool bCustomBlendOutTriggerTime = (Montage->BlendOutTriggerTime >= 0);
-						const float DefaultBlendOutTime = Montage->BlendOut.GetBlendTime() * DefaultBlendTimeMultiplier;
-						const float BlendOutTriggerTime = bCustomBlendOutTriggerTime ? Montage->BlendOutTriggerTime : DefaultBlendOutTime;
-							
-						// ... trigger blend out if within blend out time window.
-						if (DeltaTimeToEnd <= FMath::Max<float>(BlendOutTriggerTime, KINDA_SMALL_NUMBER))
+						else
 						{
-							const float BlendOutTime = bCustomBlendOutTriggerTime ? DefaultBlendOutTime : DeltaTimeToEnd;
-							Stop(FAlphaBlend(Montage->BlendOut, BlendOutTime), false);
+							OutRootMotionParams->Accumulate(RootMotion);
 						}
+
+						UE_LOG(LogRootMotion, Log, TEXT("\tFAnimMontageInstance::Advance ExtractedRootMotion: %s, AccumulatedRootMotion: %s, bBlendRootMotion: %d")
+							, *RootMotion.GetTranslation().ToCompactString()
+							, *OutRootMotionParams->GetRootMotionTransform().GetTranslation().ToCompactString()
+							, bBlendRootMotion
+							);
 					}
 
-					if (bHaveMoved)
-					{
-						// Delegate has to be called last in this loop
-						// so that if this changes position, the new position will be applied in the next loop
-						// first need to have event handler to handle it
+					// Delegate has to be called last in this loop
+					// so that if this changes position, the new position will be applied in the next loop
+					// first need to have event handler to handle it
+					// Save off position before triggering events, in case they cause a jump to another position
+					const float PositionBeforeFiringEvents = Position;
 
-						// Save position before firing events.
-						if (!bInterrupted)
-						{
-							HandleEvents(PrevPosition, Position, BranchingPointMarker);
-						}
+					// Save position before firing events.
+					if (!bInterrupted)
+					{
+						HandleEvents(PreviousSubStepPosition, Position, BranchingPointMarker);
 					}
 
 					// if we reached end of section, and we were not processing a branching point, and no events has messed with out current position..
 					// .. Move to next section.
 					// (this also handles looping, the same as jumping to a different section).
-					if ((AdvanceType != ETAA_Default) && !BranchingPointMarker && (PositionBeforeFiringEvents == Position))
+					if (MontageSubStepper.HasReachedEndOfSection() && !BranchingPointMarker && (PositionBeforeFiringEvents == Position))
 					{
 						// Get recent NextSectionIndex in case it's been changed by previous events.
+						const int32 CurrentSectionIndex = MontageSubStepper.GetCurrentSectionIndex();
 						const int32 RecentNextSectionIndex = bPlayingForward ? NextSections[CurrentSectionIndex] : PrevSections[CurrentSectionIndex];
-
-						if( RecentNextSectionIndex != INDEX_NONE )
+						if (RecentNextSectionIndex != INDEX_NONE)
 						{
 							float LatestNextSectionStartTime, LatestNextSectionEndTime;
 							Montage->GetSectionStartAndEndTime(RecentNextSectionIndex, LatestNextSectionStartTime, LatestNextSectionEndTime);
 
 							// Jump to next section's appropriate starting point (start or end).
-							float EndOffset = KINDA_SMALL_NUMBER/2.f; //KINDA_SMALL_NUMBER/2 because we use KINDA_SMALL_NUMBER to offset notifies for triggering and SMALL_NUMBER is too small
+							const float EndOffset = KINDA_SMALL_NUMBER / 2.f; //KINDA_SMALL_NUMBER/2 because we use KINDA_SMALL_NUMBER to offset notifies for triggering and SMALL_NUMBER is too small
 							Position = bPlayingForward ? LatestNextSectionStartTime : (LatestNextSectionEndTime - EndOffset);
 						}
-					}
-
-					if (!bHaveMoved)
-					{
-						// If it hasn't moved, there is nothing much to do but weight update
-						break;
+						else
+						{
+							// If there is no next section and we've reached the end of this one, exit
+							break;
+						}
 					}
 				}
-				else
+
+				if (SubStepResult == EMontageSubStepResult::NotMoved)
 				{
-					// stop and leave this loop
-					Stop(FAlphaBlend(Montage->BlendOut, Montage->BlendOut.GetBlendTime() * DefaultBlendTimeMultiplier), false);
+					// If it hasn't moved, there is nothing much to do but weight update
 					break;
 				}
 			}
+		
+			// if we had a ForcedNextPosition set, reset it.
+			ForcedNextPosition.Reset();
 		}
 	}
 
@@ -2319,3 +2635,27 @@ bool FAnimMontageInstance::CanUseMarkerSync() const
 	return SyncGroupIndex != INDEX_NONE && IsStopped() && Blend.IsComplete() == false;
 }
 
+void UAnimMontage::BakeTimeStretchCurve()
+{
+	TimeStretchCurve.Reset();
+
+	// See if Montage is hosting a curve named 'TimeStretchCurveName'
+	FFloatCurve* TimeStretchFloatCurve = nullptr;
+	if (const USkeleton* MySkeleton = GetSkeleton())
+	{
+		if (const FSmartNameMapping* CurveNameMapping = MySkeleton->GetSmartNameContainer(USkeleton::AnimCurveMappingName))
+		{
+			if (const USkeleton::AnimCurveUID CurveUID = CurveNameMapping->FindUID(TimeStretchCurveName))
+			{
+				TimeStretchFloatCurve = (FFloatCurve*)(GetCurveData().GetCurveData(CurveUID));
+			}
+		}
+	}
+
+	if (TimeStretchFloatCurve == nullptr)
+	{
+		return;
+	}
+	
+	TimeStretchCurve.BakeFromFloatCurve(*TimeStretchFloatCurve, SequenceLength);
+}
diff --git a/Engine/Source/Runtime/Engine/Private/Animation/AnimSingleNodeInstanceProxy.cpp b/Engine/Source/Runtime/Engine/Private/Animation/AnimSingleNodeInstanceProxy.cpp
--- a/Engine/Source/Runtime/Engine/Private/Animation/AnimSingleNodeInstanceProxy.cpp
+++ b/Engine/Source/Runtime/Engine/Private/Animation/AnimSingleNodeInstanceProxy.cpp
@@ -53,6 +53,8 @@ void FAnimSingleNodeInstanceProxy::PropagatePreviewCurve(FPoseContext& Output)
 
 void FAnimSingleNodeInstanceProxy::UpdateAnimationNode(float DeltaSeconds)
 {
+	UpdateCounter.Increment();
+
 	FAnimationUpdateContext UpdateContext(this, DeltaSeconds);
 	SingleNode.Update_AnyThread(UpdateContext);
 }
diff --git a/Engine/Source/Runtime/Engine/Private/Animation/AnimationRuntime.cpp b/Engine/Source/Runtime/Engine/Private/Animation/AnimationRuntime.cpp
--- a/Engine/Source/Runtime/Engine/Private/Animation/AnimationRuntime.cpp
+++ b/Engine/Source/Runtime/Engine/Private/Animation/AnimationRuntime.cpp
@@ -112,7 +112,7 @@ FORCEINLINE void BlendCurves(const TArrayView<const FBlendedCurve* const> Source
 				SumOfWeight += Weight;
 			}
 
-			if (SumOfWeight > ZERO_ANIMWEIGHT_THRESH)
+			if (FAnimWeight::IsRelevant(SumOfWeight))
 			{
 				TArray<float> NormalizeSourceWeights;
 				NormalizeSourceWeights.AddUninitialized(SourceWeights.Num());
@@ -454,13 +454,81 @@ void FAnimationRuntime::BlendPosesTogetherPerBoneInMeshSpace(
 	}
 }
 
+void FAnimationRuntime::LerpPoses(FCompactPose& PoseA, const FCompactPose& PoseB, FBlendedCurve& CurveA, const FBlendedCurve& CurveB, float Alpha)
+{
+	// If pose A is full weight, we're set.
+	if (FAnimWeight::IsRelevant(Alpha))
+	{
+		// Make sure poses are compatible with each other.
+		check(&PoseA.GetBoneContainer() == &PoseB.GetBoneContainer());
+
+		// If pose 2 is full weight, just copy, no need to blend.
+		if (FAnimWeight::IsFullWeight(Alpha))
+		{
+			PoseA.CopyBonesFrom(PoseB);
+			CurveA.CopyFrom(CurveB);
+		}
+		else
+		{
+			const ScalarRegister VWeightOfPose1(1.f - Alpha);
+			const ScalarRegister VWeightOfPose2(Alpha);
+			for (FCompactPoseBoneIndex BoneIndex : PoseA.ForEachBoneIndex())
+			{
+				FTransform& InOutBoneTransform1 = PoseA[BoneIndex];
+				InOutBoneTransform1 *= VWeightOfPose1;
+
+				const FTransform& BoneTransform2 = PoseB[BoneIndex];
+				InOutBoneTransform1.AccumulateWithShortestRotation(BoneTransform2, VWeightOfPose2);
+
+				InOutBoneTransform1.NormalizeRotation();
+			}
+
+			CurveA.LerpTo(CurveB, Alpha);
+		}
+	}
+}
+
+void FAnimationRuntime::LerpPosesPerBone(FCompactPose& PoseA, const FCompactPose& PoseB, FBlendedCurve& CurveA, const FBlendedCurve& CurveB, float Alpha, const TArray<float>& PerBoneWeights)
+{
+	// If pose A is full weight, we're set.
+	if (FAnimWeight::IsRelevant(Alpha))
+	{
+		// Make sure poses are compatible with each other.
+		check(&PoseA.GetBoneContainer() == &PoseB.GetBoneContainer());
+
+		for (FCompactPoseBoneIndex BoneIndex : PoseA.ForEachBoneIndex())
+		{
+			const float BoneAlpha = Alpha * PerBoneWeights[BoneIndex.GetInt()];
+			if (FAnimWeight::IsRelevant(BoneAlpha))
+			{
+				const ScalarRegister VWeightOfPose1(1.f - BoneAlpha);
+				const ScalarRegister VWeightOfPose2(BoneAlpha);
+
+				FTransform& InOutBoneTransform1 = PoseA[BoneIndex];
+				InOutBoneTransform1 *= VWeightOfPose1;
+
+				const FTransform& BoneTransform2 = PoseB[BoneIndex];
+				InOutBoneTransform1.AccumulateWithShortestRotation(BoneTransform2, VWeightOfPose2);
+
+				InOutBoneTransform1.NormalizeRotation();
+			}
+		}
+
+		// @note : This isn't perfect as curve can link to joint, and it would be the best to use that information
+		// but that is very expensive option as we have to have another indirect look up table to search. 
+		// For now, replacing with combine (non-zero will be overridden)
+		// in the future, we might want to do this outside if we want per bone blend to apply curve also UE-39182
+		CurveA.Combine(CurveB);
+	}
+}
+
 void FAnimationRuntime::LerpBoneTransforms(TArray<FTransform>& A, const TArray<FTransform>& B, float Alpha, const TArray<FBoneIndexType>& RequiredBonesArray)
 {
-	if (Alpha >= (1.f - ZERO_ANIMWEIGHT_THRESH))
+	if (FAnimWeight::IsFullWeight(Alpha))
 	{
 		A = B;
 	}
-	else if (Alpha > ZERO_ANIMWEIGHT_THRESH)
+	else if (FAnimWeight::IsRelevant(Alpha))
 	{
 		FTransform* ATransformData = A.GetData(); 
 		const FTransform* BTransformData = B.GetData();
@@ -605,10 +673,10 @@ void FAnimationRuntime::AccumulateAdditivePose(FCompactPose& BasePose, const FCo
 
 void FAnimationRuntime::AccumulateLocalSpaceAdditivePose(FCompactPose& BasePose, const FCompactPose& AdditivePose, FBlendedCurve& BaseCurve, const FBlendedCurve& AdditiveCurve, float Weight)
 {
-	if (Weight > ZERO_ANIMWEIGHT_THRESH)
+	if (FAnimWeight::IsRelevant(Weight))
 	{
 		const ScalarRegister VBlendWeight(Weight);
-		if (Weight >= (1.f - ZERO_ANIMWEIGHT_THRESH))
+		if (FAnimWeight::IsFullWeight(Weight))
 		{
 			// fast path, no need to weight additive.
 			for (FCompactPoseBoneIndex BoneIndex : BasePose.ForEachBoneIndex())
@@ -633,7 +701,7 @@ void FAnimationRuntime::AccumulateMeshSpaceRotationAdditiveToLocalPose(FCompactP
 {
 	SCOPE_CYCLE_COUNTER(STAT_AccumulateMeshSpaceRotAdditiveToLocalPose);
 
-	if (Weight > ZERO_ANIMWEIGHT_THRESH)
+	if (FAnimWeight::IsRelevant(Weight))
 	{
 		// Convert base pose from local space to mesh space rotation.
 		FAnimationRuntime::ConvertPoseToMeshRotation(BasePose);
@@ -950,12 +1018,12 @@ void FAnimationRuntime::BlendMeshPosesPerBoneWeights(
 		const float BlendWeight = FMath::Clamp(BoneBlendWeights[BoneIndex.GetInt()].BlendWeight, 0.f, 1.f);
 		MaxPoseWeights[PoseIndex] = FMath::Max(MaxPoseWeights[PoseIndex], BlendWeight);
 
-		if (BlendWeight < ZERO_ANIMWEIGHT_THRESH)
+		if (!FAnimWeight::IsRelevant(BlendWeight))
 		{
 			BlendAtom = BaseAtom;
 			BlendRotations[BoneIndex] = SourceRotations[BoneIndex];
 		}
-		else if ((1.0 - BlendWeight) < ZERO_ANIMWEIGHT_THRESH)
+		else if (FAnimWeight::IsFullWeight(BlendWeight))
 		{
 			BlendAtom = TargetAtom;
 			BlendRotations[BoneIndex] = TargetRotations[BoneIndex];
@@ -1043,11 +1111,11 @@ void FAnimationRuntime::BlendLocalPosesPerBoneWeights(
 		const float BlendWeight = FMath::Clamp(BoneBlendWeights[BoneIndex.GetInt()].BlendWeight, 0.f, 1.f);
 		MaxPoseWeights[PoseIndex] = FMath::Max(MaxPoseWeights[PoseIndex], BlendWeight);
 
-		if (BlendWeight < ZERO_ANIMWEIGHT_THRESH)
+		if (!FAnimWeight::IsRelevant(BlendWeight))
 		{
 			OutPose[BoneIndex] = BaseAtom;
 		}
-		else if ((1.0 - BlendWeight) < ZERO_ANIMWEIGHT_THRESH)
+		else if (FAnimWeight::IsFullWeight(BlendWeight))
 		{
 			OutPose[BoneIndex] = BlendPoses[PoseIndex][BoneIndex];
 		}
diff --git a/Engine/Source/Runtime/Engine/Private/Components/DirectionalLightComponent.cpp b/Engine/Source/Runtime/Engine/Private/Components/DirectionalLightComponent.cpp
--- a/Engine/Source/Runtime/Engine/Private/Components/DirectionalLightComponent.cpp
+++ b/Engine/Source/Runtime/Engine/Private/Components/DirectionalLightComponent.cpp
@@ -308,11 +308,11 @@ public:
 		return true;
 	}
 
-	virtual FVector2D GetDirectionalLightDistanceFadeParameters(ERHIFeatureLevel::Type InFeatureLevel, bool bPrecomputedLightingIsValid) const override
+	virtual FVector2D GetDirectionalLightDistanceFadeParameters(ERHIFeatureLevel::Type InFeatureLevel, bool bPrecomputedLightingIsValid, int32 MaxNearCascades) const override
 	{
-		float FarDistance = GetCSMMaxDistance(bPrecomputedLightingIsValid);
+		float FarDistance = GetCSMMaxDistance(bPrecomputedLightingIsValid, MaxNearCascades);
 		{
-			if (ShouldCreateRayTracedCascade(InFeatureLevel, bPrecomputedLightingIsValid))
+			if (ShouldCreateRayTracedCascade(InFeatureLevel, bPrecomputedLightingIsValid, MaxNearCascades))
 			{
 				FarDistance = GetDistanceFieldShadowDistance();
 			}
@@ -340,12 +340,11 @@ public:
 		return true;
 	}
 
-	virtual bool ShouldCreateRayTracedCascade(ERHIFeatureLevel::Type InFeatureLevel, bool bPrecomputedLightingIsValid) const override
+	virtual bool ShouldCreateRayTracedCascade(ERHIFeatureLevel::Type InFeatureLevel, bool bPrecomputedLightingIsValid, int32 MaxNearCascades) const override
 	{
-		//@todo - handle View.MaxShadowCascades properly
-		const uint32 NumCascades = GetNumShadowMappedCascades(10, bPrecomputedLightingIsValid);
+		const uint32 NumCascades = GetNumShadowMappedCascades(MaxNearCascades, bPrecomputedLightingIsValid);
 		const float RaytracedShadowDistance = GetDistanceFieldShadowDistance();
-		const bool bCreateWithCSM = NumCascades > 0 && RaytracedShadowDistance > GetCSMMaxDistance(bPrecomputedLightingIsValid);
+		const bool bCreateWithCSM = NumCascades > 0 && RaytracedShadowDistance > GetCSMMaxDistance(bPrecomputedLightingIsValid, MaxNearCascades);
 		const bool bCreateWithoutCSM = NumCascades == 0 && RaytracedShadowDistance > 0;
 		return DoesPlatformSupportDistanceFieldShadowing(GShaderPlatformForFeatureLevel[InFeatureLevel]) && (bCreateWithCSM || bCreateWithoutCSM);
 	}
@@ -360,14 +359,19 @@ private:
 	uint32 GetNumShadowMappedCascades(uint32 MaxShadowCascades, bool bPrecomputedLightingIsValid) const
 	{
 		const int32 EffectiveNumDynamicShadowCascades = bPrecomputedLightingIsValid ? DynamicShadowCascades : FMath::Max(0, CVarUnbuiltNumWholeSceneDynamicShadowCascades.GetValueOnAnyThread());
-		const int32 NumCascades = GetCSMMaxDistance(bPrecomputedLightingIsValid) > 0.0f ? EffectiveNumDynamicShadowCascades : 0;
+		const int32 NumCascades = GetCSMMaxDistance(bPrecomputedLightingIsValid, MaxShadowCascades) > 0.0f ? EffectiveNumDynamicShadowCascades : 0;
 		return FMath::Min<int32>(NumCascades, MaxShadowCascades);
 	}
 
-	float GetCSMMaxDistance(bool bPrecomputedLightingIsValid) const
+	float GetCSMMaxDistance(bool bPrecomputedLightingIsValid, int32 MaxShadowCascades) const
 	{
 		static const auto CVar = IConsoleManager::Get().FindTConsoleVariableDataFloat(TEXT("r.Shadow.DistanceScale"));
-		 
+		
+		if (MaxShadowCascades <= 0)
+		{
+			return 0.0f;
+		}
+
 		float Scale = FMath::Clamp(CVar->GetValueOnRenderThread(), 0.0f, 2.0f);
 		float Distance = GetEffectiveWholeSceneDynamicShadowRadius(bPrecomputedLightingIsValid) * Scale;
 		return Distance;
@@ -532,7 +536,7 @@ private:
 		// near cascade means non far and non distance field cascade
 		uint32 NumNearCascades = GetNumShadowMappedCascades(View.MaxShadowCascades, bPrecomputedLightingIsValid);
 
-		float CascadeDistanceWithoutFar = GetCSMMaxDistance(bPrecomputedLightingIsValid);
+		float CascadeDistanceWithoutFar = GetCSMMaxDistance(bPrecomputedLightingIsValid, View.MaxShadowCascades);
 		float ShadowNear = View.NearClippingDistance;
 		const float EffectiveCascadeDistributionExponent = GetEffectiveCascadeDistributionExponent(bPrecomputedLightingIsValid);
 
@@ -634,7 +638,7 @@ private:
 	{
 		uint32 NumNearCascades = GetNumShadowMappedCascades(View.MaxShadowCascades, bPrecomputedLightingIsValid);
 
-		const bool bHasRayTracedCascade = ShouldCreateRayTracedCascade(View.GetFeatureLevel(), bPrecomputedLightingIsValid);
+		const bool bHasRayTracedCascade = ShouldCreateRayTracedCascade(View.GetFeatureLevel(), bPrecomputedLightingIsValid, View.MaxShadowCascades);
 
 		// this checks for WholeSceneDynamicShadowRadius and DynamicShadowCascades
 		uint32 NumNearAndFarCascades = GetNumViewDependentWholeSceneShadows(View, bPrecomputedLightingIsValid);
