Copying //UE4/Orion-Staging to //UE4/Main (Source: //Orion/Dev-General @ 3606378)

diff --git a/Engine/Source/Runtime/Engine/Private/Components/SkeletalMeshComponent.cpp b/Engine/Source/Runtime/Engine/Private/Components/SkeletalMeshComponent.cpp
--- a/Engine/Source/Runtime/Engine/Private/Components/SkeletalMeshComponent.cpp
+++ b/Engine/Source/Runtime/Engine/Private/Components/SkeletalMeshComponent.cpp
@@ -588,10 +588,10 @@ void USkeletalMeshComponent::InitAnim(bool bForceReinit)
 		// this has to be called before Initialize Animation because it will required RequiredBones list when InitializeAnimScript
 		RecalcRequiredBones(0);
 
-		bool bDoRefreshBoneTransform = true;
-		if (InitializeAnimScriptInstance(bForceReinit))
+		const bool bInitializedAnimInstance = InitializeAnimScriptInstance(bForceReinit);
+		// Make sure we have a valid pose.
+		if (bInitializedAnimInstance || (AnimScriptInstance == nullptr))
 		{
-			//Make sure we have a valid pose
 			if (bUseRefPoseOnInitAnim)
 			{
 				BoneSpaceTransforms = SkeletalMesh->RefSkeleton.GetRefBonePose();
@@ -599,20 +599,19 @@ void USkeletalMeshComponent::InitAnim(bool bForceReinit)
 				FillComponentSpaceTransforms(SkeletalMesh, BoneSpaceTransforms, GetEditableComponentSpaceTransforms());
 				bNeedToFlipSpaceBaseBuffers = true; // Have updated space bases so need to flip
 				FlipEditableSpaceBases();
-				bDoRefreshBoneTransform = false;
 			}
 			else
 			{
 				TickAnimation(0.f, false);
+				RefreshBoneTransforms();
 			}
 
-			OnAnimInitialized.Broadcast();
+			if (bInitializedAnimInstance)
+			{
+				OnAnimInitialized.Broadcast();
+			}
 		}
 
-		if (bDoRefreshBoneTransform)
-		{
-			RefreshBoneTransforms();
-		}
 		UpdateComponentToWorld();
 	}
 }
@@ -924,6 +923,16 @@ void USkeletalMeshComponent::TickAnimation(float DeltaTime, bool bNeedsValidRoot
 		{
 			PostProcessAnimInstance->UpdateAnimation(DeltaTime * GlobalAnimRateScale, false);
 		}
+
+		/**
+			If we're called directly for autonomous proxies, TickComponent is not guaranteed to get called.
+			So dispatch all queued events here if we're doing MontageOnly ticking.
+		*/
+		if ((MeshComponentUpdateFlag == EMeshComponentUpdateFlag::OnlyTickMontagesWhenNotRendered)
+			&& !bRecentlyRendered)
+		{
+			ConditionallyDispatchQueuedAnimEvents();
+		}
 	}
 }
 
@@ -940,9 +949,7 @@ bool USkeletalMeshComponent::UpdateLODStatus()
 
 bool USkeletalMeshComponent::ShouldUpdateTransform(bool bLODHasChanged) const
 {
-
-#if WITH_EDITOR
-	
+#if WITH_EDITOR	
 	// If we're in an editor world (Non running, WorldType will be PIE when simulating or in PIE) then we only want transform updates on LOD changes as the
 	// animation isn't running so it would just waste CPU time
 	if(GetWorld()->WorldType == EWorldType::Editor)
@@ -972,8 +979,7 @@ bool USkeletalMeshComponent::ShouldUpdateTransform(bool bLODHasChanged) const
 	// If forcing RefPose we can skip updating the skeleton for perf, except if it's using MorphTargets.
 	const bool bSkipBecauseOfRefPose = bForceRefpose && bOldForceRefPose && (MorphTargetCurves.Num() == 0) && ((AnimScriptInstance) ? !AnimScriptInstance->HasMorphTargetCurves() : true);
 
-	// LOD changing should always trigger an update.
-	return (bLODHasChanged || !bRequiredBonesUpToDate || (!bNoSkeletonUpdate && !bSkipBecauseOfRefPose && Super::ShouldUpdateTransform(bLODHasChanged)));
+	return (!bNoSkeletonUpdate && !bSkipBecauseOfRefPose && Super::ShouldUpdateTransform(bLODHasChanged));
 }
 
 bool USkeletalMeshComponent::ShouldTickPose() const
@@ -1847,6 +1853,29 @@ void USkeletalMeshComponent::RefreshBoneTransforms(FActorComponentTickFunction*
 		CachedCurve.Empty();
 	}
 
+	// If we need to eval the graph, and we're not going to update it.
+	// make sure it's been ticked at least once!
+	if (bShouldDoEvaluation)
+	{
+		bool bShouldTickAnimation = false;		
+		if (AnimScriptInstance && !AnimScriptInstance->NeedsUpdate())
+		{
+			bShouldTickAnimation = bShouldTickAnimation || !AnimScriptInstance->GetUpdateCounter().HasEverBeenUpdated();
+			for (const UAnimInstance* SubInstance : SubInstances)
+			{
+				bShouldTickAnimation = bShouldTickAnimation || (SubInstance && !SubInstance->GetUpdateCounter().HasEverBeenUpdated());
+			}
+		}
+
+		bShouldTickAnimation = bShouldTickAnimation || (PostProcessAnimInstance && !PostProcessAnimInstance->NeedsUpdate() && !PostProcessAnimInstance->GetUpdateCounter().HasEverBeenUpdated());
+
+		if (bShouldTickAnimation)
+		{
+			// We bypass TickPose() and call TickAnimation directly, so URO doesn't intercept us.
+			TickAnimation(0.f, false);
+		}
+	}
+
 	if(AnimScriptInstance)
 	{
 		AnimScriptInstance->PreEvaluateAnimation();
diff --git a/Engine/Source/Runtime/Engine/Private/Components/SkinnedMeshComponent.cpp b/Engine/Source/Runtime/Engine/Private/Components/SkinnedMeshComponent.cpp
--- a/Engine/Source/Runtime/Engine/Private/Components/SkinnedMeshComponent.cpp
+++ b/Engine/Source/Runtime/Engine/Private/Components/SkinnedMeshComponent.cpp
@@ -618,7 +618,7 @@ bool USkinnedMeshComponent::ShouldTickPose() const
 
 bool USkinnedMeshComponent::ShouldUpdateTransform(bool bLODHasChanged) const
 {
-	return (bLODHasChanged || bRecentlyRendered || (MeshComponentUpdateFlag == EMeshComponentUpdateFlag::AlwaysTickPoseAndRefreshBones));
+	return (bRecentlyRendered || (MeshComponentUpdateFlag == EMeshComponentUpdateFlag::AlwaysTickPoseAndRefreshBones));
 }
 
 bool USkinnedMeshComponent::ShouldUseUpdateRateOptimizations() const
@@ -2374,11 +2374,6 @@ bool USkinnedMeshComponent::UpdateLODStatus()
 		}
 	}
 
-	if (bLODChanged)
-	{
-		MarkRenderDynamicDataDirty();
-	}
-
 	return bLODChanged;
 }
 
diff --git a/Engine/Source/Runtime/Engine/Private/Components/StaticMeshComponent.cpp b/Engine/Source/Runtime/Engine/Private/Components/StaticMeshComponent.cpp
--- a/Engine/Source/Runtime/Engine/Private/Components/StaticMeshComponent.cpp
+++ b/Engine/Source/Runtime/Engine/Private/Components/StaticMeshComponent.cpp
@@ -851,25 +851,21 @@ void UStaticMeshComponent::GetStreamingTextureInfo(FStreamingTextureLevelContext
 			const FVector2D& Scale = Lightmap->GetCoordinateScale();
 			if (Scale.X > SMALL_NUMBER && Scale.Y > SMALL_NUMBER)
 			{
-				FStreamingTexturePrimitiveInfo& StreamingTexture = *new(OutStreamingTextures) FStreamingTexturePrimitiveInfo;
-				StreamingTexture.Bounds		 = Bounds;
-				StreamingTexture.PackedRelativeBox = PackedRelativeBox_Identity; // Set the PackedRelativeBox to incremental level insertion.
-				StreamingTexture.TexelFactor = GetStaticMesh()->LightmapUVDensity * TransformScale / FMath::Min(Scale.X, Scale.Y);
-				StreamingTexture.Texture	 = Lightmap->GetTexture(LightmapIndex);
+				const float TexelFactor = GetStaticMesh()->LightmapUVDensity * TransformScale / FMath::Min(Scale.X, Scale.Y);
+				new (OutStreamingTextures) FStreamingTexturePrimitiveInfo(Lightmap->GetTexture(LightmapIndex), Bounds, TexelFactor, PackedRelativeBox_Identity);
+				new (OutStreamingTextures) FStreamingTexturePrimitiveInfo(Lightmap->GetAOMaterialMaskTexture(), Bounds, TexelFactor, PackedRelativeBox_Identity);
+				new (OutStreamingTextures) FStreamingTexturePrimitiveInfo(Lightmap->GetSkyOcclusionTexture(), Bounds, TexelFactor, PackedRelativeBox_Identity);
 			}
 		}
 
 		FShadowMap2D* Shadowmap = MeshMapBuildData && MeshMapBuildData->ShadowMap ? MeshMapBuildData->ShadowMap->GetShadowMap2D() : NULL;
-		if (Shadowmap != NULL && Shadowmap->IsValid())
+		if (Shadowmap && Shadowmap->IsValid())
 		{
 			const FVector2D& Scale = Shadowmap->GetCoordinateScale();
 			if (Scale.X > SMALL_NUMBER && Scale.Y > SMALL_NUMBER)
 			{
-				FStreamingTexturePrimitiveInfo& StreamingTexture = *new(OutStreamingTextures) FStreamingTexturePrimitiveInfo;
-				StreamingTexture.Bounds		 = Bounds;
-				StreamingTexture.PackedRelativeBox = PackedRelativeBox_Identity;
-				StreamingTexture.TexelFactor = GetStaticMesh()->LightmapUVDensity * TransformScale / FMath::Min(Scale.X, Scale.Y);
-				StreamingTexture.Texture	 = Shadowmap->GetTexture();
+				const float TexelFactor = GetStaticMesh()->LightmapUVDensity * TransformScale / FMath::Min(Scale.X, Scale.Y);
+				new (OutStreamingTextures) FStreamingTexturePrimitiveInfo(Shadowmap->GetTexture(), Bounds, TexelFactor, PackedRelativeBox_Identity);
 			}
 		}
 	}
diff --git a/Engine/Source/Runtime/Engine/Private/ContentStreaming.cpp b/Engine/Source/Runtime/Engine/Private/ContentStreaming.cpp
--- a/Engine/Source/Runtime/Engine/Private/ContentStreaming.cpp
+++ b/Engine/Source/Runtime/Engine/Private/ContentStreaming.cpp
@@ -1023,15 +1023,25 @@ void FStreamingManagerCollection::NotifyPrimitiveDetached( const UPrimitiveCompo
  * Replaces previous info.
  */
 void FStreamingManagerCollection::NotifyPrimitiveUpdated_Concurrent( const UPrimitiveComponent* Primitive )
+{
+	// Route to streaming managers.
+	for( int32 ManagerIndex=0; ManagerIndex<StreamingManagers.Num(); ManagerIndex++ )
 	{
-		// Route to streaming managers.
-		for( int32 ManagerIndex=0; ManagerIndex<StreamingManagers.Num(); ManagerIndex++ )
-		{
-			IStreamingManager* StreamingManager = StreamingManagers[ManagerIndex];
+		IStreamingManager* StreamingManager = StreamingManagers[ManagerIndex];
 		StreamingManager->NotifyPrimitiveUpdated_Concurrent( Primitive );
 	}
 }
 
+void FStreamingManagerCollection::PropagateLightingScenarioChange()
+{
+	// Route to streaming managers.
+	for( int32 ManagerIndex=0; ManagerIndex<StreamingManagers.Num(); ManagerIndex++ )
+	{
+		IStreamingManager* StreamingManager = StreamingManagers[ManagerIndex];
+		StreamingManager->PropagateLightingScenarioChange();
+	}
+}
+
 void FStreamingManagerCollection::AddOrRemoveTextureStreamingManagerIfNeeded(bool bIsInit)
 {
 	bool bUseTextureStreaming = false;
diff --git a/Engine/Source/Runtime/Engine/Private/GameViewportClient.cpp b/Engine/Source/Runtime/Engine/Private/GameViewportClient.cpp
--- a/Engine/Source/Runtime/Engine/Private/GameViewportClient.cpp
+++ b/Engine/Source/Runtime/Engine/Private/GameViewportClient.cpp
@@ -1394,10 +1394,6 @@ void UGameViewportClient::Draw(FViewport* InViewport, FCanvas* SceneCanvas)
 
 		//ensure canvas has been flushed before rendering UI
 		SceneCanvas->Flush_GameThread();
-		if (DebugCanvas != NULL)
-		{
-			DebugCanvas->Flush_GameThread();
-		}
 
 		DrawnDelegate.Broadcast();
 
diff --git a/Engine/Source/Runtime/Engine/Private/LODActor.cpp b/Engine/Source/Runtime/Engine/Private/LODActor.cpp
--- a/Engine/Source/Runtime/Engine/Private/LODActor.cpp
+++ b/Engine/Source/Runtime/Engine/Private/LODActor.cpp
@@ -19,6 +19,8 @@
 
 #if WITH_EDITOR
 #include "Editor.h"
+
+#include "HierarchicalLODUtilitiesModule.h"
 #endif
 
 #define LOCTEXT_NAMESPACE "LODActor"
@@ -188,6 +190,7 @@ FString ALODActor::GetDetailedInfoInternal() const
 void ALODActor::PostLoad()
 {
 	Super::PostLoad();
+	StaticMeshComponent->MinDrawDistance = LODDrawDistance;
 	UpdateRegistrationToMatchMaximumLODLevel();
 
 #if WITH_EDITOR
@@ -550,7 +553,21 @@ const bool ALODActor::HasValidSubActors() const
 	{
 		TInlineComponentArray<UStaticMeshComponent*> Components;
 		SubActor->GetComponents(/*out*/ Components);
+
+#if WITH_EDITOR
+		FHierarchicalLODUtilitiesModule& Module = FModuleManager::LoadModuleChecked<FHierarchicalLODUtilitiesModule>("HierarchicalLODUtilities");
+		IHierarchicalLODUtilities* Utilities = Module.GetUtilities();
+
+		for (UStaticMeshComponent* Component : Components)
+		{
+			if (!Component->bHiddenInGame && Component->ShouldGenerateAutoLOD())
+			{
+				++NumMeshes;
+			}
+		}
+#else
 		NumMeshes += Components.Num();
+#endif
 
 		if (NumMeshes > 1)
 		{
diff --git a/Engine/Source/Runtime/Engine/Private/MapBuildData.cpp b/Engine/Source/Runtime/Engine/Private/MapBuildData.cpp
--- a/Engine/Source/Runtime/Engine/Private/MapBuildData.cpp
+++ b/Engine/Source/Runtime/Engine/Private/MapBuildData.cpp
@@ -19,6 +19,7 @@
 #include "Components/ModelComponent.h"
 #include "ComponentRecreateRenderStateContext.h"
 #include "UObject/RenderingObjectVersion.h"
+#include "ContentStreaming.h"
 
 FArchive& operator<<(FArchive& Ar, FMeshMapBuildData& MeshMapBuildData)
 {
@@ -70,6 +71,8 @@ void UWorld::PropagateLightingScenarioChange(bool bLevelWasMadeVisible)
 		}
 	}
 
+	IStreamingManager::Get().PropagateLightingScenarioChange();
+
 	// Skipping the reflection capture update if made invisible, in most cases another lighting scenario level will be made visible shortly after, 
 	// Or we're unloading all levels and then it doesn't matter if lighting is updated.
 	if (bLevelWasMadeVisible)
diff --git a/Engine/Source/Runtime/Engine/Private/Materials/Material.cpp b/Engine/Source/Runtime/Engine/Private/Materials/Material.cpp
--- a/Engine/Source/Runtime/Engine/Private/Materials/Material.cpp
+++ b/Engine/Source/Runtime/Engine/Private/Materials/Material.cpp
@@ -609,6 +609,12 @@ void ProcessSerializedInlineShaderMaps(UMaterialInterface* Owner, TArray<FMateri
 		check(DesiredQL < EMaterialQualityLevel::Num);
 		const int32 DesiredScore = QualityScores[DesiredQL];
 
+		FMaterialShaderMap* BestShaderMap[ERHIFeatureLevel::Num];
+		for (int32 FeatureIdx = 0; FeatureIdx < ERHIFeatureLevel::Num; ++FeatureIdx)
+		{
+			BestShaderMap[FeatureIdx] = nullptr;
+		}
+
 		for (int32 ResourceIndex = 0; ResourceIndex < LoadedResources.Num(); ResourceIndex++)
 		{
 			FMaterialResource& LoadedResource = LoadedResources[ResourceIndex];
@@ -637,17 +643,25 @@ void ProcessSerializedInlineShaderMaps(UMaterialInterface* Owner, TArray<FMateri
 				const int32 PotentialScore = FMath::Abs(QualityScores[LoadedQualityLevel] - DesiredScore);
 				if (PotentialScore < CurrentScore)
 				{
-					// replace existing shadermap with loadedshadermap.
-					for (int32 QualityLevelIndex = 0; QualityLevelIndex < EMaterialQualityLevel::Num; QualityLevelIndex++)
+					BestShaderMap[LoadedFeatureLevel] = LoadedShaderMap;
+				}
+			}
+		}
+
+		for (int32 FeatureIdx = 0; FeatureIdx < ERHIFeatureLevel::Num; ++FeatureIdx)
+		{
+			if (BestShaderMap[FeatureIdx])
+			{
+				// replace existing shadermap with loadedshadermap.
+				for (int32 QualityLevelIndex = 0; QualityLevelIndex < EMaterialQualityLevel::Num; QualityLevelIndex++)
+				{
+					if (!OutMaterialResourcesLoaded[QualityLevelIndex][FeatureIdx])
 					{
-						if (!OutMaterialResourcesLoaded[QualityLevelIndex][LoadedFeatureLevel])
-						{
-							OutMaterialResourcesLoaded[QualityLevelIndex][LoadedFeatureLevel] =
-								OwnerMaterialInstance ? OwnerMaterialInstance->AllocatePermutationResource() : OwnerMaterial->AllocateResource();
-						}
-						OutMaterialResourcesLoaded[QualityLevelIndex][LoadedFeatureLevel]->ReleaseShaderMap();
-						OutMaterialResourcesLoaded[QualityLevelIndex][LoadedFeatureLevel]->SetInlineShaderMap(LoadedShaderMap);
+						OutMaterialResourcesLoaded[QualityLevelIndex][FeatureIdx] =
+							OwnerMaterialInstance ? OwnerMaterialInstance->AllocatePermutationResource() : OwnerMaterial->AllocateResource();
 					}
+					OutMaterialResourcesLoaded[QualityLevelIndex][FeatureIdx]->ReleaseShaderMap();
+					OutMaterialResourcesLoaded[QualityLevelIndex][FeatureIdx]->SetInlineShaderMap(BestShaderMap[FeatureIdx]);
 				}
 			}
 		}
diff --git a/Engine/Source/Runtime/Engine/Private/Materials/MaterialInstance.cpp b/Engine/Source/Runtime/Engine/Private/Materials/MaterialInstance.cpp
--- a/Engine/Source/Runtime/Engine/Private/Materials/MaterialInstance.cpp
+++ b/Engine/Source/Runtime/Engine/Private/Materials/MaterialInstance.cpp
@@ -2343,12 +2343,6 @@ void UMaterialInstance::SetTextureParameterValueInternal(FName ParameterName, UT
 		// set as an ensure, because it is somehow possible to accidentally pass non-textures into here via blueprints...
 		if (Value && ensureMsgf(Value->IsA(UTexture::StaticClass()), TEXT("Expecting a UTexture! Value='%s' class='%s'"), *Value->GetName(), *Value->GetClass()->GetName()))
 		{
-			if ( FPlatformProperties::RequiresCookedData() )
-			{
-				// Daniel L: Temporary code to track down OR-40458 REMOVE ME
-				check(IsInGameThread());
-				check(!Value->IsPendingKill());
-			}
 			ParameterValue->ParameterValue = Value;
 			// Update the material instance data in the rendering thread.
 			GameThread_UpdateMIParameter(this, *ParameterValue);
diff --git a/Engine/Source/Runtime/Engine/Private/Materials/MaterialInterface.cpp b/Engine/Source/Runtime/Engine/Private/Materials/MaterialInterface.cpp
--- a/Engine/Source/Runtime/Engine/Private/Materials/MaterialInterface.cpp
+++ b/Engine/Source/Runtime/Engine/Private/Materials/MaterialInterface.cpp
@@ -18,6 +18,13 @@
 #include "Interfaces/ITargetPlatform.h"
 #include "Components.h"
 
+/**
+ * This is used to deprecate data that has been built with older versions.
+ * To regenerate the data, commands like "BUILDMATERIALTEXTURESTREAMINGDATA" can be used in the editor.
+ * Ideally the data would be stored the DDC instead of the asset, but this is not yet  possible because it requires the GPU.
+ */
+#define MATERIAL_TEXTURE_STREAMING_DATA_VERSION 1
+
 //////////////////////////////////////////////////////////////////////////
 
 UEnum* UMaterialInterface::SamplerTypeEnum = nullptr;
@@ -77,6 +84,13 @@ void UMaterialInterface::PostLoad()
 	{
 		PostLoadDefaultMaterials();
 	}
+
+#if WITH_EDITORONLY_DATA
+	if (TextureStreamingDataVersion != MATERIAL_TEXTURE_STREAMING_DATA_VERSION)
+	{
+		TextureStreamingData.Empty();
+	}
+#endif
 }
 
 void UMaterialInterface::GetUsedTexturesAndIndices(TArray<UTexture*>& OutTextures, TArray< TArray<int32> >& OutIndices, EMaterialQualityLevel::Type QualityLevel, ERHIFeatureLevel::Type FeatureLevel) const
@@ -515,6 +529,9 @@ bool UMaterialInterface::FindTextureStreamingDataIndexRange(FName TextureName, i
 void UMaterialInterface::SetTextureStreamingData(const TArray<FMaterialTextureInfo>& InTextureStreamingData)
 {
 	TextureStreamingData = InTextureStreamingData;
+#if WITH_EDITORONLY_DATA
+	TextureStreamingDataVersion = InTextureStreamingData.Num() ? MATERIAL_TEXTURE_STREAMING_DATA_VERSION : 0;
+#endif
 	SortTextureStreamingData(true, false);
 }
 
diff --git a/Engine/Source/Runtime/Engine/Private/NetConnection.cpp b/Engine/Source/Runtime/Engine/Private/NetConnection.cpp
--- a/Engine/Source/Runtime/Engine/Private/NetConnection.cpp
+++ b/Engine/Source/Runtime/Engine/Private/NetConnection.cpp
@@ -1667,7 +1667,7 @@ float UNetConnection::GetTimeoutValue()
 
 	float Timeout = Driver->InitialConnectTimeout;
 
-	if ( ( State != USOCK_Pending ) && ( bPendingDestroy || ( OwningActor && OwningActor->UseShortConnectTimeout() ) ) )
+	if ((State != USOCK_Pending) && (bPendingDestroy || (OwningActor && OwningActor->UseShortConnectTimeout())))
 	{
 		const float ConnectionTimeout = Driver->ConnectionTimeout;
 
@@ -1675,12 +1675,16 @@ float UNetConnection::GetTimeoutValue()
 		Timeout = bPendingDestroy ? 2.f : ConnectionTimeout;
 	}
 
-#if WITH_EDITOR || UE_BUILD_DEBUG
-	if (Driver->TimeoutMultiplierForUnoptimizedBuilds > 0)
+	// Longtimeouts allows a multiplier to be added to get correct disconnection behavior
+	// with with additional leniancy when required. Implicit in debug/editor builds
+	static bool LongTimeouts = FParse::Param(FCommandLine::Get(), TEXT("longtimeouts"));
+
+	if (Driver->TimeoutMultiplierForUnoptimizedBuilds > 0 
+		&& (LongTimeouts || WITH_EDITOR || UE_BUILD_DEBUG)
+		)
 	{
 		Timeout *= Driver->TimeoutMultiplierForUnoptimizedBuilds;
 	}
-#endif
 
 	return Timeout;
 }
diff --git a/Engine/Source/Runtime/Engine/Private/Particles/ParticleComponents.cpp b/Engine/Source/Runtime/Engine/Private/Particles/ParticleComponents.cpp
--- a/Engine/Source/Runtime/Engine/Private/Particles/ParticleComponents.cpp
+++ b/Engine/Source/Runtime/Engine/Private/Particles/ParticleComponents.cpp
@@ -1749,7 +1749,7 @@ void UParticleEmitter::CacheEmitterModuleInfo()
 
 float UParticleEmitter::GetQualityLevelSpawnRateMult()
 {
-	int32 EffectsQuality = Scalability::GetEffectsQualityDirect(true);
+	int32 EffectsQuality = Scalability::GetEffectsQualityDirect(IsInGameThread());
 	int32 ReferenceLevel = CVarQLSpawnRateReferenceLevel.GetValueOnAnyThread(true);
 	float Level = (ReferenceLevel - EffectsQuality);
 	float Q = FMath::Pow(QualityLevelSpawnRateScale, Level);
diff --git a/Engine/Source/Runtime/Engine/Private/Particles/ParticleCurveTexture.cpp b/Engine/Source/Runtime/Engine/Private/Particles/ParticleCurveTexture.cpp
--- a/Engine/Source/Runtime/Engine/Private/Particles/ParticleCurveTexture.cpp
+++ b/Engine/Source/Runtime/Engine/Private/Particles/ParticleCurveTexture.cpp
@@ -13,7 +13,7 @@
 #include "RHIStaticStates.h"
 #include "SceneUtils.h"
 #include "ParticleHelper.h"
-#include "Particles/ParticleResources.h"
+#include "ParticleResources.h"
 #include "ShaderParameterUtils.h"
 #include "GlobalShader.h"
 #include "FXSystem.h"
diff --git a/Engine/Source/Runtime/Engine/Private/Particles/ParticleGpuSimulation.cpp b/Engine/Source/Runtime/Engine/Private/Particles/ParticleGpuSimulation.cpp
--- a/Engine/Source/Runtime/Engine/Private/Particles/ParticleGpuSimulation.cpp
+++ b/Engine/Source/Runtime/Engine/Private/Particles/ParticleGpuSimulation.cpp
@@ -33,7 +33,7 @@
 #include "Particles/FXSystemPrivate.h"
 #include "Particles/ParticleSortingGPU.h"
 #include "Particles/ParticleCurveTexture.h"
-#include "Particles/ParticleResources.h"
+#include "ParticleResources.h"
 #include "ShaderParameterUtils.h"
 #include "GlobalShader.h"
 #include "VectorFieldVisualization.h"
diff --git a/Engine/Source/Runtime/Engine/Private/PhysicsEngine/BodyInstance.cpp b/Engine/Source/Runtime/Engine/Private/PhysicsEngine/BodyInstance.cpp
--- a/Engine/Source/Runtime/Engine/Private/PhysicsEngine/BodyInstance.cpp
+++ b/Engine/Source/Runtime/Engine/Private/PhysicsEngine/BodyInstance.cpp
@@ -3982,7 +3982,8 @@ static bool GetSquaredDistanceToBodyHelper(const PxRigidActor* RigidActor, const
 {
 	bool bFoundValidBody = false;
 	
-	PxGeometry& PGeom = PShape->getGeometry().any();
+	PxGeometryHolder Holder = PShape->getGeometry();	// getGeometry() result is stored on the stack, if we don't hold on to it it may be gone next statement.
+	PxGeometry& PGeom = Holder.any();
 	PxTransform PGlobalPose = GetPxTransform_AssumesLocked(PShape, RigidActor);
 	PxGeometryType::Enum GeomType = PShape->getGeometryType();
 	
@@ -4050,7 +4051,7 @@ bool FBodyInstance::GetSquaredDistanceToBody(const FVector& Point, float& OutDis
 			}
 			
 			PxGeometryHolder Holder = PShape->getGeometry();	// getGeometry() result is stored on the stack, if we don't hold on to it it may be gone next statement.
-			PxGeometry& PGeom = Holder.any();			
+			PxGeometry& PGeom = Holder.any();
 			
 			bool bFinish = false;
 			bool bSkip = false;
diff --git a/Engine/Source/Runtime/Engine/Private/SceneManagement.cpp b/Engine/Source/Runtime/Engine/Private/SceneManagement.cpp
--- a/Engine/Source/Runtime/Engine/Private/SceneManagement.cpp
+++ b/Engine/Source/Runtime/Engine/Private/SceneManagement.cpp
@@ -225,7 +225,7 @@ void FMeshElementCollector::ProcessTasks()
 
 void FMeshElementCollector::AddMesh(int32 ViewIndex, FMeshBatch& MeshBatch)
 {
-	checkSlow(MeshBatch.GetNumPrimitives() > 0);
+	//checkSlow(MeshBatch.GetNumPrimitives() > 0);
 	checkSlow(MeshBatch.VertexFactory && MeshBatch.MaterialRenderProxy);
 	checkSlow(PrimitiveSceneProxy);
 
diff --git a/Engine/Source/Runtime/Engine/Private/ShaderCompiler/ShaderCompiler.cpp b/Engine/Source/Runtime/Engine/Private/ShaderCompiler/ShaderCompiler.cpp
--- a/Engine/Source/Runtime/Engine/Private/ShaderCompiler/ShaderCompiler.cpp
+++ b/Engine/Source/Runtime/Engine/Private/ShaderCompiler/ShaderCompiler.cpp
@@ -1934,6 +1934,9 @@ void FShaderCompilingManager::ProcessCompiledShaderMaps(
 	}
 }
 
+
+
+
 void FShaderCompilingManager::PropagateMaterialChangesToPrimitives(const TMap<FMaterial*, FMaterialShaderMap*>& MaterialsToUpdate)
 {
 	TArray<UMaterialInterface*> UsedMaterials;
@@ -1985,6 +1988,7 @@ void FShaderCompilingManager::PropagateMaterialChangesToPrimitives(const TMap<FM
 	ComponentContexts.Empty();
 }
 
+
 /**
  * Shutdown the shader compile manager
  * this function should be used when ending the game to shutdown shader compile threads
@@ -2028,7 +2032,6 @@ bool FShaderCompilingManager::HandlePotentialRetryOnError(TMap<int32, FShaderMap
 				}
 			}
 
-			check(ShaderMap || It.Key() == GlobalShaderMapId);
 
 #if WITH_EDITORONLY_DATA
 
@@ -2351,14 +2354,16 @@ void FShaderCompilingManager::ProcessAsyncResults(bool bLimitExecutionTime, bool
 
 				TArray<int32> ShaderMapsToRemove;
 
+				// Get all material shader maps to finalize
+				//
 				for (TMap<int32, FShaderMapCompileResults>::TIterator It(ShaderMapJobs); It; ++It)
 				{
 					const FShaderMapCompileResults& Results = It.Value();
 
 					if (GetNumTotalJobs(Results.FinishedJobs) == Results.NumJobsQueued)
 					{
-						PendingFinalizeShaderMaps.Add(It.Key(), FShaderMapFinalizeResults(Results));
 						ShaderMapsToRemove.Add(It.Key());
+						PendingFinalizeShaderMaps.Add(It.Key(), FShaderMapFinalizeResults(Results));
 					}
 				}
 
@@ -2386,6 +2391,7 @@ void FShaderCompilingManager::ProcessAsyncResults(bool bLimitExecutionTime, bool
 				check(bLimitExecutionTime || PendingFinalizeShaderMaps.Num() == 0);
 			}
 
+
 			if (bBlockOnGlobalShaderCompletion)
 			{
 				check(PendingFinalizeShaderMaps.Num() == 0);
diff --git a/Engine/Source/Runtime/Engine/Private/StaticMesh.cpp b/Engine/Source/Runtime/Engine/Private/StaticMesh.cpp
--- a/Engine/Source/Runtime/Engine/Private/StaticMesh.cpp
+++ b/Engine/Source/Runtime/Engine/Private/StaticMesh.cpp
@@ -3087,6 +3087,37 @@ void UStaticMesh::SetVertexColorData(const TMap<FVector, FColor>& VertexColorDat
 #endif // #if WITH_EDITOR
 }
 
+ENGINE_API void UStaticMesh::RemoveVertexColors()
+{
+#if WITH_EDITOR
+	bool bRemovedVertexColors = false;
+
+	for (FStaticMeshSourceModel& SourceModel : SourceModels)
+	{
+		if (SourceModel.RawMeshBulkData && !SourceModel.RawMeshBulkData->IsEmpty())
+		{
+			FRawMesh RawMesh;
+			SourceModel.RawMeshBulkData->LoadRawMesh(RawMesh);
+
+			if (RawMesh.WedgeColors.Num() > 0)
+			{
+				RawMesh.WedgeColors.Empty();
+
+				SourceModel.RawMeshBulkData->SaveRawMesh(RawMesh);
+
+				bRemovedVertexColors = true;
+			}
+		}
+	}
+
+	if (bRemovedVertexColors)
+	{
+		Build();
+		MarkPackageDirty();
+	}
+#endif
+}
+
 void UStaticMesh::EnforceLightmapRestrictions()
 {
 	// Legacy content may contain a lightmap resolution of 0, which was valid when vertex lightmaps were supported, but not anymore with only texture lightmaps
diff --git a/Engine/Source/Runtime/Engine/Private/Streaming/LevelTextureManager.cpp b/Engine/Source/Runtime/Engine/Private/Streaming/LevelTextureManager.cpp
--- a/Engine/Source/Runtime/Engine/Private/Streaming/LevelTextureManager.cpp
+++ b/Engine/Source/Runtime/Engine/Private/Streaming/LevelTextureManager.cpp
@@ -9,7 +9,7 @@
 #include "Engine/Texture2D.h"
 #include "Engine/World.h"
 
-void FLevelTextureManager::Remove(FRemovedTextureArray& RemovedTextures)
+void FLevelTextureManager::Remove(FRemovedTextureArray* RemovedTextures)
 { 
 	TArray<const UPrimitiveComponent*> ReferencedComponents;
 	StaticInstances.GetReferencedComponents(ReferencedComponents);
@@ -28,9 +28,12 @@ void FLevelTextureManager::Remove(FRemovedTextureArray& RemovedTextures)
 	}
 
 	// Mark all static textures for removal.
-	for (auto It = StaticInstances.GetTextureIterator(); It; ++It)
+	if (RemovedTextures)
 	{
-		RemovedTextures.Push(*It);
+		for (auto It = StaticInstances.GetTextureIterator(); It; ++It)
+		{
+			RemovedTextures->Push(*It);
+		}
 	}
 
 	BuildStep = EStaticBuildStep::BuildTextureLookUpMap;
diff --git a/Engine/Source/Runtime/Engine/Private/Streaming/LevelTextureManager.h b/Engine/Source/Runtime/Engine/Private/Streaming/LevelTextureManager.h
--- a/Engine/Source/Runtime/Engine/Private/Streaming/LevelTextureManager.h
+++ b/Engine/Source/Runtime/Engine/Private/Streaming/LevelTextureManager.h
@@ -24,8 +24,8 @@ public:
 
 	ULevel* GetLevel() const { return Level; }
 
-	// Remove the whole level
-	void Remove(FRemovedTextureArray& RemovedTextures);
+	// Remove the whole level. Optional list of textures referenced
+	void Remove(FRemovedTextureArray* RemovedTextures);
 
 	// Invalidate a component reference.
 
diff --git a/Engine/Source/Runtime/Engine/Private/Streaming/StreamingManagerTexture.cpp b/Engine/Source/Runtime/Engine/Private/Streaming/StreamingManagerTexture.cpp
--- a/Engine/Source/Runtime/Engine/Private/Streaming/StreamingManagerTexture.cpp
+++ b/Engine/Source/Runtime/Engine/Private/Streaming/StreamingManagerTexture.cpp
@@ -140,7 +140,7 @@ void FStreamingManagerTexture::OnPreGarbageCollect()
 		if (LevelManager.GetLevel()->IsPendingKill())
 		{
 			FRemovedTextureArray RemovedTextures;
-			LevelManager.Remove(RemovedTextures);
+			LevelManager.Remove(&RemovedTextures);
 			SetTexturesRemovedTimestamp(RemovedTextures);
 
 			// Remove the level entry. The async task view will still be valid as it uses a shared ptr.
@@ -394,14 +394,13 @@ void FStreamingManagerTexture::ConditionalUpdateStaticData()
 			PreviousSettings.bUseNewMetrics != Settings.bUseNewMetrics ||
 			PreviousSettings.bUsePerTextureBias != Settings.bUsePerTextureBias)
 		{
-			TArray<ULevel*> Levels;
-			FRemovedTextureArray RemovedTextures;
+			TArray<ULevel*, TInlineAllocator<32> > Levels;
 
 			// RemoveLevel data
 			for (FLevelTextureManager& LevelManager : LevelTextureManagers)
 			{
 				Levels.Push(LevelManager.GetLevel());
-				LevelManager.Remove(RemovedTextures);
+				LevelManager.Remove(nullptr);
 			}
 			LevelTextureManagers.Empty();
 
@@ -549,7 +548,7 @@ void FStreamingManagerTexture::RemoveLevel( ULevel* Level )
 			if (LevelManager.GetLevel() == Level)
 			{
 				FRemovedTextureArray RemovedTextures;
-				LevelManager.Remove(RemovedTextures);
+				LevelManager.Remove(&RemovedTextures);
 				SetTexturesRemovedTimestamp(RemovedTextures);
 
 				// Remove the level entry. The async task view will still be valid as it uses a shared ptr.
@@ -1233,6 +1232,24 @@ void FStreamingManagerTexture::GetObjectReferenceBounds(const UObject* RefObject
 	}
 }
 
+void FStreamingManagerTexture::PropagateLightingScenarioChange()
+{
+	// Note that dynamic components don't need to be handled because their renderstates are updated, which triggers and update.
+	
+	TArray<ULevel*, TInlineAllocator<32> > Levels;
+	for (FLevelTextureManager& LevelManager : LevelTextureManagers)
+	{
+		Levels.Push(LevelManager.GetLevel());
+		LevelManager.Remove(nullptr);
+	}
+
+	LevelTextureManagers.Empty();
+
+	for (ULevel* Level : Levels)
+	{
+		AddLevel(Level);
+	}
+}
 
 #if STATS_FAST
 bool FStreamingManagerTexture::HandleDumpTextureStreamingStatsCommand( const TCHAR* Cmd, FOutputDevice& Ar )
@@ -1268,11 +1285,15 @@ bool FStreamingManagerTexture::HandleListStreamingTexturesCommand( const TCHAR*
 		SortedTextures.Add(StreamingTexture.Texture->GetFullName(), TextureIndex);
 	}
 
+	const bool bShouldOnlyListUnkownRef = FParse::Command(&Cmd, TEXT("UNKOWNREF"));
+
 	SortedTextures.KeySort(TLess<FString>());
 
 	for (TMap<FString, int32>::TConstIterator It(SortedTextures); It; ++It)
 	{
 		const FStreamingTexture& StreamingTexture = StreamingTextures[It.Value()];
+		if (bShouldOnlyListUnkownRef && !StreamingTexture.bUseUnkownRefHeuristic) continue;
+
 		const UTexture2D* Texture2D = StreamingTexture.Texture;
 		UE_LOG(LogContentStreaming, Log,  TEXT("Texture [%d] : %s"), It.Value(), *Texture2D->GetFullName() );
 
diff --git a/Engine/Source/Runtime/Engine/Private/Streaming/StreamingManagerTexture.h b/Engine/Source/Runtime/Engine/Private/Streaming/StreamingManagerTexture.h
--- a/Engine/Source/Runtime/Engine/Private/Streaming/StreamingManagerTexture.h
+++ b/Engine/Source/Runtime/Engine/Private/Streaming/StreamingManagerTexture.h
@@ -180,6 +180,9 @@ struct FStreamingManagerTexture : public ITextureStreamingManager
 	/** Return all bounds related to the ref object */
 	virtual void GetObjectReferenceBounds(const UObject* RefObject, TArray<FBox>& AssetBoxes) override;
 
+	/** Propagates a change to the active lighting scenario. */
+	void PropagateLightingScenarioChange() override;
+
 protected:
 //BEGIN: Thread-safe functions and data
 		friend class FAsyncTextureStreamingTask;
diff --git a/Engine/Source/Runtime/Engine/Private/Texture.cpp b/Engine/Source/Runtime/Engine/Private/Texture.cpp
--- a/Engine/Source/Runtime/Engine/Private/Texture.cpp
+++ b/Engine/Source/Runtime/Engine/Private/Texture.cpp
@@ -982,17 +982,6 @@ uint32 UTexture::GetMaximumDimension() const
 }
 #endif // #if WITH_EDITOR
 
-void UTexture::MarkPendingKill()
-{
-	Super::MarkPendingKill();
-	if ( FPlatformProperties::RequiresCookedData() )
-	{
-		// Daniel L: Temporary code to track down OR-40458 REMOVE ME
-		UE_LOG(LogTexture, Warning, TEXT("Texture is being marked as pending kill!!! %s"), *GetFullName());
-	}
-}
-
-
 FName GetDefaultTextureFormatName( const ITargetPlatform* TargetPlatform, const UTexture* Texture, const FConfigFile& EngineSettings, bool bSupportDX11TextureFormats )
 {
 	FName TextureFormatName = NAME_None;
diff --git a/Engine/Source/Runtime/Engine/Private/TickTaskManager.cpp b/Engine/Source/Runtime/Engine/Private/TickTaskManager.cpp
--- a/Engine/Source/Runtime/Engine/Private/TickTaskManager.cpp
+++ b/Engine/Source/Runtime/Engine/Private/TickTaskManager.cpp
@@ -267,7 +267,7 @@ public:
 		{
 #if DO_TIMEGUARD
 			FTimerNameDelegate NameFunction = FTimerNameDelegate::CreateLambda( [&]{ return FString::Printf(TEXT("Slowtick %s "), *Target->DiagnosticMessage()); } );
-			SCOPE_TIME_GUARD_DELEGATE_MS(NameFunction, 2);
+			SCOPE_TIME_GUARD_DELEGATE_MS(NameFunction, 4);
 #endif
 
 			Target->ExecuteTick(Target->CalculateDeltaTime(Context), Context.TickType, CurrentThread, MyCompletionGraphEvent);
diff --git a/Engine/Source/Runtime/Engine/Private/UnrealEngine.cpp b/Engine/Source/Runtime/Engine/Private/UnrealEngine.cpp
--- a/Engine/Source/Runtime/Engine/Private/UnrealEngine.cpp
+++ b/Engine/Source/Runtime/Engine/Private/UnrealEngine.cpp
@@ -1242,7 +1242,6 @@ void UEngine::Init(IEngineLoop* InEngineLoop)
 	// Dynamically load engine runtime modules
 	{
 		FModuleManager::Get().LoadModuleChecked(TEXT("StreamingPauseRendering"));
-		FModuleManager::Get().LoadModuleChecked(TEXT("Niagara"));
 		FModuleManager::Get().LoadModuleChecked(TEXT("GeometryCache"));
 		FModuleManager::Get().LoadModuleChecked(TEXT("MovieScene"));
 		FModuleManager::Get().LoadModuleChecked(TEXT("MovieSceneTracks"));
@@ -10312,8 +10311,6 @@ bool UEngine::LoadMap( FWorldContext& WorldContext, FURL URL, class UPendingNetG
 #if !(UE_BUILD_SHIPPING || UE_BUILD_TEST)
 	// Dump info
 
-	Exec( NULL, TEXT("MEM"));
-
 	VerifyLoadMapWorldCleanup();
 	
 #endif
diff --git a/Engine/Source/Runtime/Engine/Private/World.cpp b/Engine/Source/Runtime/Engine/Private/World.cpp
--- a/Engine/Source/Runtime/Engine/Private/World.cpp
+++ b/Engine/Source/Runtime/Engine/Private/World.cpp
@@ -290,6 +290,7 @@ FWorldDelegates::FOnLevelChanged FWorldDelegates::LevelRemovedFromWorld;
 FWorldDelegates::FLevelOffsetEvent FWorldDelegates::PostApplyLevelOffset;
 FWorldDelegates::FWorldGetAssetTags FWorldDelegates::GetAssetTags;
 FWorldDelegates::FOnWorldTickStart FWorldDelegates::OnWorldTickStart;
+FWorldDelegates::FOnWorldPostActorTick FWorldDelegates::OnWorldPostActorTick;
 #if WITH_EDITOR
 FWorldDelegates::FRefreshLevelScriptActionsEvent FWorldDelegates::RefreshLevelScriptActions;
 #endif // WITH_EDITOR
diff --git a/Engine/Source/Runtime/Engine/Public/Animation/AnimCurveTypes.h b/Engine/Source/Runtime/Engine/Public/Animation/AnimCurveTypes.h
--- a/Engine/Source/Runtime/Engine/Public/Animation/AnimCurveTypes.h
+++ b/Engine/Source/Runtime/Engine/Public/Animation/AnimCurveTypes.h
@@ -358,12 +358,12 @@ struct FBaseBlendedCurve
 	void Lerp(const FBaseBlendedCurve& A, const FBaseBlendedCurve& B, float Alpha)
 	{
 		check(A.Num() == B.Num());
-		if (FMath::Abs(Alpha) <= ZERO_ANIMWEIGHT_THRESH)
+		if (!FAnimWeight::IsRelevant(FMath::Abs(Alpha)))
 		{
 			// if blend is all the way for child1, then just copy its bone atoms
 			Override(A);
 		}
-		else if (FMath::Abs(Alpha - 1.0f) <= ZERO_ANIMWEIGHT_THRESH)
+		else if (!FAnimWeight::IsRelevant(FMath::Abs(Alpha - 1.0f)))
 		{
 			// if blend is all the way for child2, then just copy its bone atoms
 			Override(B);
@@ -385,11 +385,11 @@ struct FBaseBlendedCurve
 	void LerpTo(const FBaseBlendedCurve& Other, float Alpha)
 	{
 		check(Num() == Other.Num());
-		if (FMath::Abs(Alpha) <= ZERO_ANIMWEIGHT_THRESH)
+		if (!FAnimWeight::IsRelevant(FMath::Abs(Alpha)))
 		{
 			return;
 		}
-		else if (FMath::Abs(Alpha - 1.0f) <= ZERO_ANIMWEIGHT_THRESH)
+		else if (!FAnimWeight::IsRelevant(FMath::Abs(Alpha - 1.0f)))
 		{
 			// if blend is all the way for child2, then just copy its bone atoms
 			Override(Other);
@@ -425,7 +425,7 @@ struct FBaseBlendedCurve
 		check(bInitialized);
 		check(Num() == AdditiveCurve.Num());
 
-		if (Weight > ZERO_ANIMWEIGHT_THRESH)
+		if (FAnimWeight::IsRelevant(Weight))
 		{
 			for (int32 CurveId = 0; CurveId < Elements.Num(); ++CurveId)
 			{
diff --git a/Engine/Source/Runtime/Engine/Public/Animation/AnimTypes.h b/Engine/Source/Runtime/Engine/Public/Animation/AnimTypes.h
--- a/Engine/Source/Runtime/Engine/Public/Animation/AnimTypes.h
+++ b/Engine/Source/Runtime/Engine/Public/Animation/AnimTypes.h
@@ -137,6 +137,11 @@ public:
 		: InternalCounter(INDEX_NONE)
 	{}
 
+	bool HasEverBeenUpdated() const
+	{
+		return (InternalCounter != INDEX_NONE);
+	}
+
 	int16 Get() const
 	{
 		return InternalCounter;
diff --git a/Engine/Source/Runtime/Engine/Public/AnimationRuntime.h b/Engine/Source/Runtime/Engine/Public/AnimationRuntime.h
--- a/Engine/Source/Runtime/Engine/Public/AnimationRuntime.h
+++ b/Engine/Source/Runtime/Engine/Public/AnimationRuntime.h
@@ -301,6 +301,19 @@ public:
 	/** Accumulate a MeshSpaceRotation Additive pose to a local pose. Rotations are NOT normalized */
 	static void AccumulateMeshSpaceRotationAdditiveToLocalPose(FCompactPose& BasePose, const FCompactPose& MeshSpaceRotationAdditive, FBlendedCurve& BaseCurve, const FBlendedCurve& AdditiveCurve, float Weight);
 
+	/** Lerp for BoneTransforms. Stores results in A. Performs A = Lerp(A, B, Alpha);
+	* @param A : In/Out FCompactPose.
+	* @param B : B In FCompactPose.
+	* @param Alpha : Alpha.
+	*/
+	static void LerpPoses(FCompactPose& PoseA, const FCompactPose& PoseB, FBlendedCurve& CurveA, const FBlendedCurve& CurveB, float Alpha);
+
+	/** Lerp for BoneTransforms. Stores results in A. Performs A = Lerp(A, B, Alpha);
+	* @param A : In/Out FCompactPose.
+	* @param B : B In FCompactPose.
+	* @param Alpha : Alpha.
+	*/
+	static void LerpPosesPerBone(FCompactPose& PoseA, const FCompactPose& PoseB, FBlendedCurve& CurveA, const FBlendedCurve& CurveB, float Alpha, const TArray<float>& PerBoneWeights);
 
 	/** Lerp for BoneTransforms. Stores results in A. Performs A = Lerp(A, B, Alpha);
 	 * @param A : In/Out transform array.
@@ -389,8 +402,8 @@ public:
 #endif
 
 	/* Weight utility functions */
-	static bool IsFullWeight(float Weight) { return Weight > 1.f - ZERO_ANIMWEIGHT_THRESH; }
-	static bool HasWeight(float Weight) { return Weight > ZERO_ANIMWEIGHT_THRESH; }
+	static FORCEINLINE bool IsFullWeight(float Weight) { return FAnimWeight::IsFullWeight(Weight); }
+	static FORCEINLINE bool HasWeight(float Weight) { return FAnimWeight::IsRelevant(Weight); }
 	
 	/**
 	* Combine CurveKeys (that reference morph targets by name) and ActiveAnims (that reference morphs by reference) into the ActiveMorphTargets array.
diff --git a/Engine/Source/Runtime/Engine/Public/ContentStreaming.h b/Engine/Source/Runtime/Engine/Public/ContentStreaming.h
--- a/Engine/Source/Runtime/Engine/Public/ContentStreaming.h
+++ b/Engine/Source/Runtime/Engine/Public/ContentStreaming.h
@@ -269,6 +269,11 @@ struct IStreamingManager
 		return NumWantingResourcesCounter;
 	}
 
+	/** Propagates a change to the active lighting scenario. */
+	virtual void PropagateLightingScenarioChange() 
+	{
+	}
+
 protected:
 
 	/**
@@ -593,6 +598,9 @@ struct FStreamingManagerCollection : public IStreamingManager
 	 */
 	virtual void NotifyPrimitiveUpdated_Concurrent( const UPrimitiveComponent* Primitive ) override;
 
+	/** Propagates a change to the active lighting scenario. */
+	void PropagateLightingScenarioChange() override;
+
 protected:
 
 	virtual void AddOrRemoveTextureStreamingManagerIfNeeded(bool bIsInit=false);
diff --git a/Engine/Source/Runtime/Engine/Public/ParticleHelper.h b/Engine/Source/Runtime/Engine/Public/ParticleHelper.h
--- a/Engine/Source/Runtime/Engine/Public/ParticleHelper.h
+++ b/Engine/Source/Runtime/Engine/Public/ParticleHelper.h
@@ -2771,7 +2771,7 @@ public:
 
 #endif
 
-class FNullDynamicParameterVertexBuffer : public FVertexBuffer
+class ENGINE_API FNullDynamicParameterVertexBuffer : public FVertexBuffer
 {
 public:
 	/** 
@@ -2790,7 +2790,7 @@ public:
 };
 
 /** The global null color vertex buffer, which is set with a stride of 0 on meshes without a color component. */
-extern TGlobalResource<FNullDynamicParameterVertexBuffer> GNullDynamicParameterVertexBuffer;
+extern ENGINE_API TGlobalResource<FNullDynamicParameterVertexBuffer> GNullDynamicParameterVertexBuffer;
 
 FORCEINLINE FVector GetParticleBaseSize(const FBaseParticle& Particle, bool bKeepFlipScale = false)
 {
diff --git a/Engine/Source/Runtime/Engine/Public/Scalability.h b/Engine/Source/Runtime/Engine/Public/Scalability.h
--- a/Engine/Source/Runtime/Engine/Public/Scalability.h
+++ b/Engine/Source/Runtime/Engine/Public/Scalability.h
@@ -88,7 +88,7 @@ namespace Scalability
 	*
 	* @param bGameThread	If true, the game thread value for the CVar is returned, otherwise the render thread value is returned. Useful when accessing the CVar from a game task.
 	*/
-	int32 GetEffectsQualityDirect(bool bGameThread);
+	ENGINE_API int32 GetEffectsQualityDirect(bool bGameThread);
 
 	/**  */
 	ENGINE_API void InitScalabilitySystem();
diff --git a/Engine/Source/Runtime/Engine/Public/SceneManagement.h b/Engine/Source/Runtime/Engine/Public/SceneManagement.h
--- a/Engine/Source/Runtime/Engine/Public/SceneManagement.h
+++ b/Engine/Source/Runtime/Engine/Public/SceneManagement.h
@@ -918,7 +918,7 @@ public:
 	/** Accesses parameters needed for rendering the light. */
 	virtual void GetParameters(FLightParameters& LightParameters) const {}
 
-	virtual FVector2D GetDirectionalLightDistanceFadeParameters(ERHIFeatureLevel::Type InFeatureLevel, bool bPrecomputedLightingIsValid) const
+	virtual FVector2D GetDirectionalLightDistanceFadeParameters(ERHIFeatureLevel::Type InFeatureLevel, bool bPrecomputedLightingIsValid, int32 MaxNearCascades) const
 	{
 		return FVector2D(0, 0);
 	}
@@ -1005,7 +1005,7 @@ public:
 	{
 	}
 
-	virtual bool ShouldCreateRayTracedCascade(ERHIFeatureLevel::Type Type, bool bPrecomputedLightingIsValid) const { return false; }
+	virtual bool ShouldCreateRayTracedCascade(ERHIFeatureLevel::Type Type, bool bPrecomputedLightingIsValid, int32 MaxNearCascades) const { return false; }
 
 	// Accessors.
 	float GetUserShadowBias() const { return ShadowBias; }
diff --git a/Engine/Source/Runtime/Engine/Public/ShaderCompiler.h b/Engine/Source/Runtime/Engine/Public/ShaderCompiler.h
--- a/Engine/Source/Runtime/Engine/Public/ShaderCompiler.h
+++ b/Engine/Source/Runtime/Engine/Public/ShaderCompiler.h
@@ -356,6 +356,7 @@ private:
 	TArray<FShaderCommonCompileJob*> CompileQueue;
 	/** Map from shader map Id to the compile results for that map, used to gather compiled results. */
 	TMap<int32, FShaderMapCompileResults> ShaderMapJobs;
+
 	/** Number of jobs currently being compiled.  This includes CompileQueue and any jobs that have been assigned to workers but aren't complete yet. */
 	int32 NumOutstandingJobs;
 
@@ -426,6 +427,9 @@ private:
 	/** Finalizes the given shader map results and optionally assigns the affected shader maps to materials, while attempting to stay within an execution time budget. */
 	void ProcessCompiledShaderMaps(TMap<int32, FShaderMapFinalizeResults>& CompiledShaderMaps, float TimeBudget);
 
+	/** Finalizes the given Niagara shader map results and assigns the affected shader maps to Niagara scripts, while attempting to stay within an execution time budget. */
+	void ProcessCompiledNiagaraShaderMaps(TMap<int32, FShaderMapFinalizeResults>& CompiledShaderMaps, float TimeBudget);
+
 	/** Propagate the completed compile to primitives that might be using the materials compiled. */
 	void PropagateMaterialChangesToPrimitives(const TMap<FMaterial*, class FMaterialShaderMap*>& MaterialsToUpdate);
 
diff --git a/Engine/Source/Runtime/Engine/Public/WeightedRandomSampler.h b/Engine/Source/Runtime/Engine/Public/WeightedRandomSampler.h
--- a/Engine/Source/Runtime/Engine/Public/WeightedRandomSampler.h
+++ b/Engine/Source/Runtime/Engine/Public/WeightedRandomSampler.h
@@ -36,7 +36,7 @@ public:
 	/**
 	Takes two random values (0...1) and returns the corresponding element index.
 	*/
-	FORCEINLINE int32 GetEntryIndex(float R0, float R1)
+	FORCEINLINE int32 GetEntryIndex(float R0, float R1)const
 	{
 		int32 Idx = R0 * Prob.Num();
 		return R1 < Prob[Idx] ? Idx : Alias[Idx];
diff --git a/Engine/Source/Runtime/Landscape/Private/LandscapeCollision.cpp b/Engine/Source/Runtime/Landscape/Private/LandscapeCollision.cpp
--- a/Engine/Source/Runtime/Landscape/Private/LandscapeCollision.cpp
+++ b/Engine/Source/Runtime/Landscape/Private/LandscapeCollision.cpp
@@ -2144,4 +2144,7 @@ ULandscapeHeightfieldCollisionComponent::ULandscapeHeightfieldCollisionComponent
 
 	HeightfieldRowsCount = -1;
 	HeightfieldColumnsCount = -1;
+
+	// landscape collision components should be deterministically created and therefor are addressable over the network
+	SetNetAddressable();
 }
diff --git a/Engine/Source/Runtime/Landscape/Private/LandscapeRender.cpp b/Engine/Source/Runtime/Landscape/Private/LandscapeRender.cpp
--- a/Engine/Source/Runtime/Landscape/Private/LandscapeRender.cpp
+++ b/Engine/Source/Runtime/Landscape/Private/LandscapeRender.cpp
@@ -2867,12 +2867,10 @@ void ULandscapeComponent::GetStreamingTextureInfo(FStreamingTextureLevelContext&
 			const FVector2D& Scale = Lightmap->GetCoordinateScale();
 			if (Scale.X > SMALL_NUMBER && Scale.Y > SMALL_NUMBER)
 			{
-				float LightmapFactorX = TexelFactor / Scale.X;
-				float LightmapFactorY = TexelFactor / Scale.Y;
-				FStreamingTexturePrimitiveInfo& StreamingTexture = *new(OutStreamingTextures)FStreamingTexturePrimitiveInfo;
-				StreamingTexture.Bounds = BoundingSphere;
-				StreamingTexture.TexelFactor = FMath::Max(LightmapFactorX, LightmapFactorY);
-				StreamingTexture.Texture = Lightmap->GetTexture(LightmapIndex);
+				const float LightmapTexelFactor = TexelFactor / FMath::Min(Scale.X, Scale.Y);
+				new (OutStreamingTextures) FStreamingTexturePrimitiveInfo(Lightmap->GetTexture(LightmapIndex), Bounds, LightmapTexelFactor);
+				new (OutStreamingTextures) FStreamingTexturePrimitiveInfo(Lightmap->GetAOMaterialMaskTexture(), Bounds, LightmapTexelFactor);
+				new (OutStreamingTextures) FStreamingTexturePrimitiveInfo(Lightmap->GetSkyOcclusionTexture(), Bounds, LightmapTexelFactor);
 			}
 		}
 
@@ -2883,12 +2881,8 @@ void ULandscapeComponent::GetStreamingTextureInfo(FStreamingTextureLevelContext&
 			const FVector2D& Scale = Shadowmap->GetCoordinateScale();
 			if (Scale.X > SMALL_NUMBER && Scale.Y > SMALL_NUMBER)
 			{
-				float ShadowmapFactorX = TexelFactor / Scale.X;
-				float ShadowmapFactorY = TexelFactor / Scale.Y;
-				FStreamingTexturePrimitiveInfo& StreamingTexture = *new(OutStreamingTextures)FStreamingTexturePrimitiveInfo;
-				StreamingTexture.Bounds = BoundingSphere;
-				StreamingTexture.TexelFactor = FMath::Max(ShadowmapFactorX, ShadowmapFactorY);
-				StreamingTexture.Texture = Shadowmap->GetTexture();
+				const float ShadowmapTexelFactor = TexelFactor / FMath::Min(Scale.X, Scale.Y);
+				new (OutStreamingTextures) FStreamingTexturePrimitiveInfo(Shadowmap->GetTexture(), Bounds, ShadowmapTexelFactor);
 			}
 		}
 	}
diff --git a/Engine/Source/Runtime/RHI/Public/RHIUtilities.h b/Engine/Source/Runtime/RHI/Public/RHIUtilities.h
--- a/Engine/Source/Runtime/RHI/Public/RHIUtilities.h
+++ b/Engine/Source/Runtime/RHI/Public/RHIUtilities.h
@@ -59,13 +59,13 @@ struct FRWBuffer
 	}
 
 	// @param AdditionalUsage passed down to RHICreateVertexBuffer(), get combined with "BUF_UnorderedAccess | BUF_ShaderResource" e.g. BUF_Static
-	void Initialize(uint32 BytesPerElement, uint32 NumElements, EPixelFormat Format, uint32 AdditionalUsage = 0, const TCHAR* InDebugName = NULL)
-	{
+	void Initialize(uint32 BytesPerElement, uint32 NumElements, EPixelFormat Format, uint32 AdditionalUsage = 0,  const TCHAR* InDebugName = NULL, FResourceArrayInterface *InResourceArray = nullptr)	{
 		check(GMaxRHIFeatureLevel == ERHIFeatureLevel::SM5);
 		// Provide a debug name if using Fast VRAM so the allocators diagnostics will work
 		ensure(!((AdditionalUsage & BUF_FastVRAM) && !InDebugName));
 		NumBytes = BytesPerElement * NumElements;
 		FRHIResourceCreateInfo CreateInfo;
+		CreateInfo.ResourceArray = InResourceArray;
 		CreateInfo.DebugName = InDebugName;
 		Buffer = RHICreateVertexBuffer(NumBytes, BUF_UnorderedAccess | BUF_ShaderResource | AdditionalUsage, CreateInfo);
 		UAV = RHICreateUnorderedAccessView(Buffer, Format);
diff --git a/Engine/Source/Runtime/RenderCore/Public/RendererInterface.h b/Engine/Source/Runtime/RenderCore/Public/RendererInterface.h
--- a/Engine/Source/Runtime/RenderCore/Public/RendererInterface.h
+++ b/Engine/Source/Runtime/RenderCore/Public/RendererInterface.h
@@ -519,6 +519,14 @@ class FPostOpaqueRenderParameters
 };
 DECLARE_DELEGATE_OneParam(FPostOpaqueRenderDelegate, class FPostOpaqueRenderParameters&);
 
+
+class FComputeDispatcher
+{
+public:
+	virtual void Execute(FRHICommandList &RHICmdList) = 0;
+};
+
+
 class ICustomVisibilityQuery: public IRefCountedObject
 {
 public:
@@ -699,6 +707,10 @@ public:
 	virtual void RenderOverlayExtensions(const FSceneView& View, FRHICommandListImmediate& RHICmdList, FSceneRenderTargets& SceneContext) = 0;
 	virtual bool HasPostOpaqueExtentions() const = 0;
 
+	virtual void RegisterPostOpaqueComputeDispatcher(FComputeDispatcher *Dispatcher) = 0;
+	virtual void UnRegisterPostOpaqueComputeDispatcher(FComputeDispatcher *Dispatcher) = 0;
+	virtual void DispatchPostOpaqueCompute(FRHICommandList &CmdList) = 0;
+
 	/** Delegate that is called upon resolving scene color. */
 	DECLARE_MULTICAST_DELEGATE_TwoParams(FOnResolvedSceneColor, FRHICommandListImmediate& /*RHICmdList*/, class FSceneRenderTargets& /*SceneContext*/);
 
diff --git a/Engine/Source/Runtime/Renderer/Private/DeferredShadingRenderer.cpp b/Engine/Source/Runtime/Renderer/Private/DeferredShadingRenderer.cpp
--- a/Engine/Source/Runtime/Renderer/Private/DeferredShadingRenderer.cpp
+++ b/Engine/Source/Runtime/Renderer/Private/DeferredShadingRenderer.cpp
@@ -1195,6 +1195,8 @@ void FDeferredShadingSceneRenderer::Render(FRHICommandListImmediate& RHICmdList)
 		}
 	}
 
+	RendererModule.DispatchPostOpaqueCompute(RHICmdList);
+
 	// No longer needed, release
 	LightShaftOutput.LightShaftOcclusion = NULL;
 
diff --git a/Engine/Source/Runtime/Renderer/Private/DrawingPolicy.cpp b/Engine/Source/Runtime/Renderer/Private/DrawingPolicy.cpp
--- a/Engine/Source/Runtime/Renderer/Private/DrawingPolicy.cpp
+++ b/Engine/Source/Runtime/Renderer/Private/DrawingPolicy.cpp
@@ -205,6 +205,7 @@ const FVertexDeclarationRHIRef& FMeshDrawingPolicy::GetVertexDeclaration() const
 {
 	check(VertexFactory && VertexFactory->IsInitialized());
 	const FVertexDeclarationRHIRef& VertexDeclaration = VertexFactory->GetDeclaration();
-	check(IsValidRef(VertexDeclaration));
+	check(VertexFactory->NeedsDeclaration()==false || IsValidRef(VertexDeclaration));
+
 	return VertexDeclaration;
 }
diff --git a/Engine/Source/Runtime/Renderer/Private/LightGridInjection.cpp b/Engine/Source/Runtime/Renderer/Private/LightGridInjection.cpp
--- a/Engine/Source/Runtime/Renderer/Private/LightGridInjection.cpp
+++ b/Engine/Source/Runtime/Renderer/Private/LightGridInjection.cpp
@@ -454,7 +454,7 @@ void FDeferredShadingSceneRenderer::ComputeLightGrid(FRHICommandListImmediate& R
 							GlobalLightData.DirectionalLightDirection = LightParameters.NormalizedLightDirection;
 							GlobalLightData.DirectionalLightShadowMapChannelMask = ShadowMapChannelMaskPacked;
 
-							const FVector2D FadeParams = LightProxy->GetDirectionalLightDistanceFadeParameters(View.GetFeatureLevel(), LightSceneInfo->IsPrecomputedLightingValid());
+							const FVector2D FadeParams = LightProxy->GetDirectionalLightDistanceFadeParameters(View.GetFeatureLevel(), LightSceneInfo->IsPrecomputedLightingValid(), View.MaxShadowCascades);
 
 							GlobalLightData.DirectionalLightDistanceFadeMAD = FVector2D(FadeParams.Y, -FadeParams.X * FadeParams.Y);
 
diff --git a/Engine/Source/Runtime/Renderer/Private/ReflectionEnvironmentCapture.cpp b/Engine/Source/Runtime/Renderer/Private/ReflectionEnvironmentCapture.cpp
--- a/Engine/Source/Runtime/Renderer/Private/ReflectionEnvironmentCapture.cpp
+++ b/Engine/Source/Runtime/Renderer/Private/ReflectionEnvironmentCapture.cpp
@@ -1411,6 +1411,12 @@ void FScene::UpdateReflectionCaptureContents(UReflectionCaptureComponent* Captur
 				return;
 			}
 
+			if (FPlatformProperties::RequiresCookedData())
+			{
+				UE_LOG(LogEngine, Warning, TEXT("No built data for %s, skipping generation in cooked build."), *CaptureComponent->GetPathName() )
+				return;
+			}
+
 			ENQUEUE_UNIQUE_RENDER_COMMAND_ONEPARAMETER( 
 				ClearCommand,
 				int32, ReflectionCaptureSize, ReflectionCaptureSize,
diff --git a/Engine/Source/Runtime/Renderer/Private/RendererModule.h b/Engine/Source/Runtime/Renderer/Private/RendererModule.h
--- a/Engine/Source/Runtime/Renderer/Private/RendererModule.h
+++ b/Engine/Source/Runtime/Renderer/Private/RendererModule.h
@@ -76,6 +76,10 @@ public:
 	virtual void RegisterCustomCullingImpl(ICustomCulling* impl) override;
 	virtual void UnregisterCustomCullingImpl(ICustomCulling* impl) override;
 
+	virtual void RegisterPostOpaqueComputeDispatcher(FComputeDispatcher *Dispatcher) override;
+	virtual void UnRegisterPostOpaqueComputeDispatcher(FComputeDispatcher *Dispatcher) override;
+	virtual void DispatchPostOpaqueCompute(FRHICommandList &CmdList) override;
+
 	virtual void RegisterPostOpaqueRenderDelegate(const FPostOpaqueRenderDelegate& PostOpaqueRenderDelegate) override;
 	virtual void RegisterOverlayRenderDelegate(const FPostOpaqueRenderDelegate& OverlayRenderDelegate) override;
 	virtual void RenderPostOpaqueExtensions(const FSceneView& View, FRHICommandListImmediate& RHICmdList, FSceneRenderTargets& SceneContext) override;
@@ -101,6 +105,7 @@ private:
 	FPostOpaqueRenderDelegate PostOpaqueRenderDelegate;
 	FPostOpaqueRenderDelegate OverlayRenderDelegate;
 	FOnResolvedSceneColor PostResolvedSceneColorCallbacks;
+	TArray<FComputeDispatcher*>PostOpaqueDispatchers;
 };
 
 extern ICustomCulling* GCustomCullingImpl;
diff --git a/Engine/Source/Runtime/Renderer/Private/SceneRendering.cpp b/Engine/Source/Runtime/Renderer/Private/SceneRendering.cpp
--- a/Engine/Source/Runtime/Renderer/Private/SceneRendering.cpp
+++ b/Engine/Source/Runtime/Renderer/Private/SceneRendering.cpp
@@ -605,10 +605,10 @@ void FViewInfo::Init()
 		TranslucencyLightingVolumeSize[CascadeIndex] = FVector(0);
 	}
 
-	const int32 MaxMobileShadowCascadeCount = FMath::Clamp(CVarMaxMobileShadowCascades.GetValueOnAnyThread(), 1, MAX_MOBILE_SHADOWCASCADES);
+	const int32 MaxMobileShadowCascadeCount = FMath::Clamp(CVarMaxMobileShadowCascades.GetValueOnAnyThread(), 0, MAX_MOBILE_SHADOWCASCADES);
 	const int32 MaxShadowCascadeCountUpperBound = GetFeatureLevel() >= ERHIFeatureLevel::SM4 ? 10 : MaxMobileShadowCascadeCount;
 
-	MaxShadowCascades = FMath::Clamp<int32>(CVarMaxShadowCascades.GetValueOnAnyThread(), 1, MaxShadowCascadeCountUpperBound);
+	MaxShadowCascades = FMath::Clamp<int32>(CVarMaxShadowCascades.GetValueOnAnyThread(), 0, MaxShadowCascadeCountUpperBound);
 
 	ShaderMap = GetGlobalShaderMap(FeatureLevel);
 
@@ -2258,6 +2258,24 @@ TGlobalResource<FFilterVertexDeclaration>& FRendererModule::GetFilterVertexDecla
 	return GFilterVertexDeclaration;
 }
 
+void FRendererModule::RegisterPostOpaqueComputeDispatcher(FComputeDispatcher *Dispatcher)
+{
+	PostOpaqueDispatchers.AddUnique(Dispatcher);
+}
+
+void FRendererModule::UnRegisterPostOpaqueComputeDispatcher(FComputeDispatcher *Dispatcher)
+{
+	PostOpaqueDispatchers.Remove(Dispatcher);
+}
+
+void FRendererModule::DispatchPostOpaqueCompute(FRHICommandList &RHICmdList)
+{
+	for (FComputeDispatcher *Dispatcher : PostOpaqueDispatchers)
+	{
+		Dispatcher->Execute(RHICmdList);
+	}
+}
+
 void FRendererModule::RegisterPostOpaqueRenderDelegate(const FPostOpaqueRenderDelegate& InPostOpaqueRenderDelegate)
 {
 	this->PostOpaqueRenderDelegate = InPostOpaqueRenderDelegate;
diff --git a/Engine/Source/Runtime/Renderer/Private/ShadowRendering.h b/Engine/Source/Runtime/Renderer/Private/ShadowRendering.h
--- a/Engine/Source/Runtime/Renderer/Private/ShadowRendering.h
+++ b/Engine/Source/Runtime/Renderer/Private/ShadowRendering.h
@@ -89,7 +89,7 @@ void SetDeferredLightParameters(
 	DeferredLightUniformsValue.SourceLength = LightParameters.LightSourceLength;
 	DeferredLightUniformsValue.MinRoughness = LightParameters.LightMinRoughness;
 
-	const FVector2D FadeParams = LightSceneInfo->Proxy->GetDirectionalLightDistanceFadeParameters(View.GetFeatureLevel(), LightSceneInfo->IsPrecomputedLightingValid());
+	const FVector2D FadeParams = LightSceneInfo->Proxy->GetDirectionalLightDistanceFadeParameters(View.GetFeatureLevel(), LightSceneInfo->IsPrecomputedLightingValid(), View.MaxShadowCascades);
 
 	// use MAD for efficiency in the shader
 	DeferredLightUniformsValue.DistanceFadeMAD = FVector2D(FadeParams.Y, -FadeParams.X * FadeParams.Y);
diff --git a/Engine/Source/Runtime/Renderer/Private/ShadowSetup.cpp b/Engine/Source/Runtime/Renderer/Private/ShadowSetup.cpp
--- a/Engine/Source/Runtime/Renderer/Private/ShadowSetup.cpp
+++ b/Engine/Source/Runtime/Renderer/Private/ShadowSetup.cpp
@@ -1365,7 +1365,7 @@ void FSceneRenderer::UpdatePreshadowCache(FSceneRenderTargets& SceneContext)
 bool ShouldCreateObjectShadowForStationaryLight(const FLightSceneInfo* LightSceneInfo, const FPrimitiveSceneProxy* PrimitiveSceneProxy, bool bInteractionShadowMapped) 
 {
 	const bool bCreateObjectShadowForStationaryLight = 
-		LightSceneInfo->bCreatePerObjectShadowsForDynamicObjects 
+		LightSceneInfo->bCreatePerObjectShadowsForDynamicObjects
 		&& LightSceneInfo->IsPrecomputedLightingValid()
 		&& LightSceneInfo->Proxy->GetShadowMapChannel() != INDEX_NONE
 		// Create a per-object shadow if the object does not want static lighting and needs to integrate with the static shadowing of a stationary light
@@ -1746,7 +1746,7 @@ void FSceneRenderer::CreatePerObjectProjectedShadow(
 			for (int32 i = 0; i < ViewDependentWholeSceneShadows.Num(); i++)
 			{
 				const FProjectedShadowInfo* WholeSceneShadow = ViewDependentWholeSceneShadows[i];
-				const FVector2D DistanceFadeValues = WholeSceneShadow->GetLightSceneInfo().Proxy->GetDirectionalLightDistanceFadeParameters(Scene->GetFeatureLevel(), WholeSceneShadow->GetLightSceneInfo().IsPrecomputedLightingValid());
+				const FVector2D DistanceFadeValues = WholeSceneShadow->GetLightSceneInfo().Proxy->GetDirectionalLightDistanceFadeParameters(Scene->GetFeatureLevel(), WholeSceneShadow->GetLightSceneInfo().IsPrecomputedLightingValid(), WholeSceneShadow->DependentView->MaxShadowCascades);
 				const float DistanceFromShadowCenterSquared = (WholeSceneShadow->ShadowBounds.Center - Bounds.Origin).SizeSquared();
 				//@todo - if view dependent whole scene shadows are ever supported in splitscreen, 
 				// We can only disable the preshadow at this point if it is inside a whole scene shadow for all views
@@ -2507,7 +2507,7 @@ struct FGatherShadowPrimitivesPacket
 
 				if (PrimitiveFlagsCompact.bCastDynamicShadow)
 				{
-					FilterPrimitiveForShadows(Scene->PrimitiveBounds[PrimitiveIndex].BoxSphereBounds, PrimitiveFlagsCompact, Scene->Primitives[PrimitiveIndex], Scene->PrimitiveSceneProxies[PrimitiveIndex]); 
+					FilterPrimitiveForShadows(Scene->PrimitiveBounds[PrimitiveIndex].BoxSphereBounds, PrimitiveFlagsCompact, Scene->Primitives[PrimitiveIndex], Scene->PrimitiveSceneProxies[PrimitiveIndex]);
 				}
 			}
 		}
@@ -2809,7 +2809,7 @@ void FSceneRenderer::AddViewDependentWholeSceneShadowsForView(
 		// If rendering in stereo mode we render shadow depths only for the left eye, but project for both eyes!
 		if (View.StereoPass != eSSP_RIGHT_EYE)
 		{
-			const bool bExtraDistanceFieldCascade = LightSceneInfo.Proxy->ShouldCreateRayTracedCascade(View.GetFeatureLevel(), LightSceneInfo.IsPrecomputedLightingValid());
+			const bool bExtraDistanceFieldCascade = LightSceneInfo.Proxy->ShouldCreateRayTracedCascade(View.GetFeatureLevel(), LightSceneInfo.IsPrecomputedLightingValid(), View.MaxShadowCascades);
 
 			const int32 ProjectionCount = LightSceneInfo.Proxy->GetNumViewDependentWholeSceneShadows(View, LightSceneInfo.IsPrecomputedLightingValid()) + (bExtraDistanceFieldCascade?1:0);
 
diff --git a/Engine/Source/Runtime/Renderer/Private/VolumeLighting.h b/Engine/Source/Runtime/Renderer/Private/VolumeLighting.h
--- a/Engine/Source/Runtime/Renderer/Private/VolumeLighting.h
+++ b/Engine/Source/Runtime/Renderer/Private/VolumeLighting.h
@@ -78,7 +78,7 @@ public:
 				Planes[1] = FVector4((FVector)(ShadowCascadeSettings.FarFrustumPlane), -ShadowCascadeSettings.FarFrustumPlane.W);
 			}
 
-			const FVector2D FadeParams = LightSceneInfo->Proxy->GetDirectionalLightDistanceFadeParameters(View.GetFeatureLevel(), LightSceneInfo->IsPrecomputedLightingValid());
+			const FVector2D FadeParams = LightSceneInfo->Proxy->GetDirectionalLightDistanceFadeParameters(View.GetFeatureLevel(), LightSceneInfo->IsPrecomputedLightingValid(), View.MaxShadowCascades);
 
 			// setup constants for the MAD in shader
 			ShadowInjectParamValue.Z = FadeParams.Y;
diff --git a/Engine/Source/Runtime/Renderer/Renderer.Build.cs b/Engine/Source/Runtime/Renderer/Renderer.Build.cs
--- a/Engine/Source/Runtime/Renderer/Renderer.Build.cs
+++ b/Engine/Source/Runtime/Renderer/Renderer.Build.cs
@@ -10,8 +10,8 @@ public class Renderer : ModuleRules
 			new string[] {
 				"Runtime/Renderer/Private",
 				"Runtime/Renderer/Private/CompositionLighting",
-				"Runtime/Renderer/Private/PostProcess",
-			}
+				"Runtime/Renderer/Private/PostProcess"
+            }
 		);
 
 		PublicDependencyModuleNames.Add("Core");
diff --git a/Engine/Source/Runtime/ShaderCore/Private/ShaderCore.cpp b/Engine/Source/Runtime/ShaderCore/Private/ShaderCore.cpp
--- a/Engine/Source/Runtime/ShaderCore/Private/ShaderCore.cpp
+++ b/Engine/Source/Runtime/ShaderCore/Private/ShaderCore.cpp
@@ -71,6 +71,8 @@ IMPLEMENT_MODULE( FShaderCoreModule, ShaderCore );
 // Shader stats
 //
 
+DEFINE_STAT(STAT_ShaderCompiling_NiagaraShaders);
+DEFINE_STAT(STAT_ShaderCompiling_NumTotalNiagaraShaders);
 
 DEFINE_STAT(STAT_ShaderCompiling_MaterialShaders);
 DEFINE_STAT(STAT_ShaderCompiling_GlobalShaders);
@@ -552,6 +554,10 @@ static void GetShaderIncludes(const TCHAR* EntryPointVirtualFilePath, const TCHA
 					// Check virtual.
 					bIgnoreInclude |= !CheckVirtualShaderFilePath(ExtractedIncludeFilename);
 			
+					// Ignore Niagara generated shader code
+					bIgnoreInclude |= (ExtractedIncludeFilename == TEXT("/Engine/Private/NiagaraEmitterInstance.usf"));
+					bIgnoreInclude |= (ExtractedIncludeFilename == TEXT("/Engine/Private/NiagaraSimulation.usf"));
+
 					// Some headers aren't required to be found (platforms that the user doesn't have access to)
 					// @todo: Is there some way to generalize this"
 					const bool bIsOptionalInclude = (ExtractedIncludeFilename == TEXT("/Engine/Public/PS4/PS4Common.ush") 
diff --git a/Engine/Source/Runtime/ShaderCore/Private/VertexFactory.cpp b/Engine/Source/Runtime/ShaderCore/Private/VertexFactory.cpp
--- a/Engine/Source/Runtime/ShaderCore/Private/VertexFactory.cpp
+++ b/Engine/Source/Runtime/ShaderCore/Private/VertexFactory.cpp
@@ -181,8 +181,15 @@ void FVertexFactory::Set(FRHICommandList& RHICmdList) const
 		const FVertexStream& Stream = Streams[StreamIndex];
 		if (!Stream.bSetByVertexFactoryInSetMesh)
 		{
-			checkf(Stream.VertexBuffer->IsInitialized(), TEXT("Vertex buffer was not initialized! Stream %u, Stride %u, Name %s"), StreamIndex, Stream.Stride, *Stream.VertexBuffer->GetFriendlyName());
-			RHICmdList.SetStreamSource(StreamIndex, Stream.VertexBuffer->VertexBufferRHI, Stream.Offset);
+			if (!Stream.VertexBuffer)
+			{
+				RHICmdList.SetStreamSource(StreamIndex, nullptr, 0);
+			}
+			else
+			{
+				checkf(Stream.VertexBuffer->IsInitialized(), TEXT("Vertex buffer was not initialized! Stream %u, Stride %u, Name %s"), StreamIndex, Stream.Stride, *Stream.VertexBuffer->GetFriendlyName());
+				RHICmdList.SetStreamSource(StreamIndex, Stream.VertexBuffer->VertexBufferRHI, Stream.Offset);
+			}
 		}
 	}
 }
diff --git a/Engine/Source/Runtime/ShaderCore/Public/Shader.h b/Engine/Source/Runtime/ShaderCore/Public/Shader.h
--- a/Engine/Source/Runtime/ShaderCore/Public/Shader.h
+++ b/Engine/Source/Runtime/ShaderCore/Public/Shader.h
@@ -23,6 +23,7 @@
 
 class FGlobalShaderType;
 class FMaterialShaderType;
+class FNiagaraShaderType;
 class FMeshMaterialShaderType;
 class FShaderPipelineType;
 class FShaderType;
@@ -793,7 +794,8 @@ public:
 	{
 		Global,
 		Material,
-		MeshMaterial
+		MeshMaterial,
+		Niagara
 	};
 
 	typedef class FShader* (*ConstructSerializedType)();
@@ -878,6 +880,14 @@ public:
 	{
 		return (ShaderTypeForDynamicCast == EShaderTypeForDynamicCast::MeshMaterial) ? reinterpret_cast<const FMeshMaterialShaderType*>(this) : nullptr;
 	}
+	FORCEINLINE const FNiagaraShaderType* GetNiagaraShaderType() const
+	{
+		return (ShaderTypeForDynamicCast == EShaderTypeForDynamicCast::Niagara) ? reinterpret_cast<const FNiagaraShaderType*>(this) : nullptr;
+	}
+	FORCEINLINE FNiagaraShaderType* GetNiagaraShaderType()
+	{
+		return (ShaderTypeForDynamicCast == EShaderTypeForDynamicCast::Niagara) ? reinterpret_cast<FNiagaraShaderType*>(this) : nullptr;
+	}
 
 	// Accessors.
 	inline EShaderFrequency GetFrequency() const
diff --git a/Engine/Source/Runtime/ShaderCore/Public/ShaderCore.h b/Engine/Source/Runtime/ShaderCore/Public/ShaderCore.h
--- a/Engine/Source/Runtime/ShaderCore/Public/ShaderCore.h
+++ b/Engine/Source/Runtime/ShaderCore/Public/ShaderCore.h
@@ -26,6 +26,9 @@ SHADERCORE_API DECLARE_LOG_CATEGORY_EXTERN(LogShaders, Log, All);
 SHADERCORE_API DECLARE_LOG_CATEGORY_EXTERN(LogShaders, Error, All);
 #endif
 
+DECLARE_DWORD_ACCUMULATOR_STAT_EXTERN(TEXT("Num Total Niagara Shaders"), STAT_ShaderCompiling_NumTotalNiagaraShaders, STATGROUP_ShaderCompiling, SHADERCORE_API);
+DECLARE_FLOAT_ACCUMULATOR_STAT_EXTERN(TEXT("Total Niagara Shader Compiling Time"), STAT_ShaderCompiling_NiagaraShaders, STATGROUP_ShaderCompiling, SHADERCORE_API);
+
 DECLARE_FLOAT_ACCUMULATOR_STAT_EXTERN(TEXT("Total Material Shader Compiling Time"),STAT_ShaderCompiling_MaterialShaders,STATGROUP_ShaderCompiling, SHADERCORE_API);
 DECLARE_FLOAT_ACCUMULATOR_STAT_EXTERN(TEXT("Total Global Shader Compiling Time"),STAT_ShaderCompiling_GlobalShaders,STATGROUP_ShaderCompiling, SHADERCORE_API);
 DECLARE_FLOAT_ACCUMULATOR_STAT_EXTERN(TEXT("RHI Compile Time"),STAT_ShaderCompiling_RHI,STATGROUP_ShaderCompiling, SHADERCORE_API);
@@ -186,6 +189,8 @@ private:
 	};
 
 	TMap<FString,FParameterAllocation> ParameterMap;
+public:
+
 };
 
 /** Container for shader compiler definitions. */
diff --git a/Engine/Source/Runtime/ShaderCore/Public/VertexFactory.h b/Engine/Source/Runtime/ShaderCore/Public/VertexFactory.h
--- a/Engine/Source/Runtime/ShaderCore/Public/VertexFactory.h
+++ b/Engine/Source/Runtime/ShaderCore/Public/VertexFactory.h
@@ -338,11 +338,13 @@ class SHADERCORE_API FVertexFactory : public FRenderResource
 {
 public:
 	FVertexFactory() 
+		: bNeedsDeclaration(true)
 	{
 	}
 
 	FVertexFactory(ERHIFeatureLevel::Type InFeatureLevel) 
 		: FRenderResource(InFeatureLevel)
+		, bNeedsDeclaration(true)
 	{
 	}
 
@@ -413,6 +415,7 @@ public:
 	 */
 	virtual uint64 GetStaticBatchElementVisibility(const class FSceneView& View, const struct FMeshBatch* Batch) const { return 1; }
 
+	bool NeedsDeclaration() const { return bNeedsDeclaration; }
 protected:
 
 	/**
@@ -473,6 +476,9 @@ protected:
 	/** The vertex streams used to render the factory. */
 	TArray<FVertexStream,TFixedAllocator<MaxVertexElementCount> > Streams;
 
+	/* VF can explicitly set this to false to avoid errors without decls; this is for VFs that fetch from buffers directly (e.g. Niagara) */
+	bool bNeedsDeclaration;
+
 private:
 
 	/** The position only vertex stream used to render the factory during depth only passes. */
diff --git a/Engine/Source/Runtime/VectorVM/Private/VectorVM.cpp b/Engine/Source/Runtime/VectorVM/Private/VectorVM.cpp
--- a/Engine/Source/Runtime/VectorVM/Private/VectorVM.cpp
+++ b/Engine/Source/Runtime/VectorVM/Private/VectorVM.cpp
@@ -134,6 +134,8 @@ struct TUnaryScalarKernel : public TUnaryKernel<Kernel, FRegisterHandler<float>,
 template<typename Kernel>
 struct TUnaryVectorKernel : public TUnaryKernel<Kernel, FRegisterDestHandler<VectorRegister>, FConstantHandler<VectorRegister>, FRegisterHandler<VectorRegister>, VECTOR_WIDTH_FLOATS> {};
 template<typename Kernel>
+struct TUnaryScalarIntKernel : public TUnaryKernel<Kernel, FRegisterHandler<int32>, FConstantHandler<int32>, FRegisterHandler<int32>, 1> {};
+template<typename Kernel>
 struct TUnaryVectorIntKernel : public TUnaryKernel<Kernel, FRegisterDestHandler<VectorRegisterInt>, FConstantHandler<VectorRegisterInt>, FRegisterHandler<VectorRegisterInt>, VECTOR_WIDTH_FLOATS> {};
 
 /** Base class of Vector kernels with 2 operands. */
@@ -536,6 +538,7 @@ struct FVectorKernelRandom : public TUnaryVectorKernel<FVectorKernelRandom>
 	static void VM_FORCEINLINE DoKernel(VectorRegister* RESTRICT Dst, VectorRegister Src0)
 	{
 		const float rm = RAND_MAX;
+		//EEK!. Improve this. Implement GPU style seeded rand instead of this.
 		VectorRegister Result = MakeVectorRegister(static_cast<float>(FMath::Rand()) / rm,
 			static_cast<float>(FMath::Rand()) / rm,
 			static_cast<float>(FMath::Rand()) / rm,
@@ -544,7 +547,6 @@ struct FVectorKernelRandom : public TUnaryVectorKernel<FVectorKernelRandom>
 	}
 };
 
-
 /* gaussian distribution random number (not working yet) */
 struct FVectorKernelRandomGauss : public TBinaryVectorKernel<FVectorKernelRandomGauss>
 {
@@ -1063,6 +1065,18 @@ struct FVectorIntKernelSign : TUnaryVectorIntKernel<FVectorIntKernelSign>
 	}
 };
 
+//randomi,
+//No good way to do this with SSE atm so just do it scalar.
+struct FScalarIntKernelRandom : public TUnaryScalarIntKernel<FScalarIntKernelRandom>
+{
+	static void VM_FORCEINLINE DoKernel(int32* RESTRICT Dst, int32 Src0)
+	{
+		const float rm = RAND_MAX;
+		//EEK!. Improve this. Implement GPU style seeded rand instead of this.
+		*Dst = FMath::Rand() % (Src0 + 1);
+	}
+};
+
 //cmplti,
 struct FVectorIntKernelCompareLT : TBinaryVectorIntKernel<FVectorIntKernelCompareLT>
 {
@@ -1336,18 +1350,17 @@ void VectorVM::Init()
 void VectorVM::Exec(
 	uint8 const* Code,
 	uint8** InputRegisters,
-	uint8* InputRegisterSizes,
 	int32 NumInputRegisters,
 	uint8** OutputRegisters,
-	uint8* OutputRegisterSizes,
 	int32 NumOutputRegisters,
 	uint8 const* ConstantTable,
 	TArray<FDataSetMeta> &DataSetMetaTable,
 	FVMExternalFunction* ExternalFunctionTable,
+	void** UserPtrTable,
 	int32 NumInstances
 
 #if STATS
-	, TArray<TStatId>& StatScopes
+	, const TArray<TStatId>& StatScopes
 #endif
 	)
 {
@@ -1400,7 +1413,7 @@ void VectorVM::Exec(
 	{
 		// Setup execution context.
 		FVectorVMContext Context(Code, RegisterTable, ConstantTable, DataSetIndexTable.GetData(), DataSetOffsetTable.GetData(), 
-			ExternalFunctionTable, FMath::Min(InstancesLeft, (int32)InstancesPerChunk), InstancesPerChunk * ChunkIdx
+			ExternalFunctionTable, UserPtrTable, FMath::Min(InstancesLeft, (int32)InstancesPerChunk), InstancesPerChunk * ChunkIdx
 #if STATS
 			, StatScopes
 #endif
@@ -1471,6 +1484,7 @@ void VectorVM::Exec(
 			case EVectorVMOp::absi: FVectorIntKernelAbs::Exec(Context); break;
 			case EVectorVMOp::negi: FVectorIntKernelNegate::Exec(Context); break;
 			case EVectorVMOp::signi: FVectorIntKernelSign::Exec(Context); break;
+			case EVectorVMOp::randomi: FScalarIntKernelRandom::Exec(Context); break;
 			case EVectorVMOp::cmplti: FVectorIntKernelCompareLT::Exec(Context); break;
 			case EVectorVMOp::cmplei: FVectorIntKernelCompareLE::Exec(Context); break;
 			case EVectorVMOp::cmpgti: FVectorIntKernelCompareGT::Exec(Context); break;
diff --git a/Engine/Source/Runtime/VectorVM/Public/VectorVM.h b/Engine/Source/Runtime/VectorVM/Public/VectorVM.h
--- a/Engine/Source/Runtime/VectorVM/Public/VectorVM.h
+++ b/Engine/Source/Runtime/VectorVM/Public/VectorVM.h
@@ -64,11 +64,11 @@ enum class EVectorVMOp : uint8
 	min,
 	max,
 	pow,
+	round,
 	sign,
 	step,
 	random,
 	noise,
-	output,
 
 	//Comparison ops.
 	cmplt,
@@ -93,6 +93,7 @@ enum class EVectorVMOp : uint8
 	absi,
 	negi,
 	signi,
+	randomi,
 	cmplti,
 	cmplei,
 	cmpgti,
@@ -136,10 +137,6 @@ enum class EVectorVMOp : uint8
 	enter_stat_scope,
 	exit_stat_scope,
 
-	//TODO: Move back to the float ops when we can auto recompile existing data.
-	round,
-
-
 	NumOpcodes
 };
 
@@ -192,18 +189,17 @@ namespace VectorVM
 	VECTORVM_API void Exec(
 		uint8 const* Code,
 		uint8** InputRegisters,
-		uint8* InputRegisterSizes,
 		int32 NumInputRegisters,
 		uint8** OutputRegisters,
-		uint8* OutputRegisterSizes,
 		int32 NumOutputRegisters,
 		uint8 const* ConstantTable,
 		TArray<FDataSetMeta> &DataSetMetaTable,
 		FVMExternalFunction* ExternalFunctionTable,
+		void** UserPtrTable,
 		int32 NumInstances
 
 #if STATS
-		, TArray<TStatId>& StatScopes
+		, const TArray<TStatId>& StatScopes
 #endif
 		);
 
@@ -230,6 +226,8 @@ struct FVectorVMContext
 	int32 NumSecondaryDataSets;
 	/** Pointer to the shared data table. */
 	FVMExternalFunction* RESTRICT ExternalFunctionTable;
+	/** Table of user pointers.*/
+	void** UserPtrTable;
 	/** Number of instances to process. */
 	int32 NumInstances;
 	/** Start instance of current chunk. */
@@ -237,7 +235,7 @@ struct FVectorVMContext
 
 #if STATS
 	TArray<FCycleCounter> StatCounterStack;
-	TArray<TStatId>& StatScopes;
+	const TArray<TStatId>& StatScopes;
 #endif
 
 	/** Initialization constructor. */
@@ -248,10 +246,11 @@ struct FVectorVMContext
 		int32 *InDataSetIndexTable,
 		int32 *InDataSetOffsetTable,
 		FVMExternalFunction* InExternalFunctionTable,
+		void** InUserPtrTable,
 		int32 InNumInstances,
 		int32 InStartInstance
 #if STATS
-		, TArray<TStatId>& InStatScopes
+		, const TArray<TStatId>& InStatScopes
 #endif
 	)
 		: Code(InCode)
@@ -260,6 +259,7 @@ struct FVectorVMContext
 		, DataSetIndexTable(InDataSetIndexTable)
 		, DataSetOffsetTable(InDataSetOffsetTable)
 		, ExternalFunctionTable(InExternalFunctionTable)
+		, UserPtrTable(InUserPtrTable)
 		, NumInstances(InNumInstances)
 		, StartInstance(InStartInstance)
 #if STATS
@@ -370,6 +370,22 @@ struct FRegisterHandlerBase
 	{}
 };
 
+template<typename T>
+struct FUserPtrHandler
+{
+	int32 UserPtrIdx;
+	T* Ptr;
+	FUserPtrHandler(FVectorVMContext& Context)
+		: UserPtrIdx(*(int32*)(Context.ConstantTable + DecodeU16(Context)))
+		, Ptr((T*)Context.UserPtrTable[UserPtrIdx])
+	{
+		check(UserPtrIdx != INDEX_NONE);
+	}
+	FORCEINLINE T* Get() { return Ptr; }
+	FORCEINLINE T* operator->() { return Ptr; }
+	FORCEINLINE operator T*() { return Ptr; }
+};
+
 template<typename T>
 struct FRegisterHandler : public FRegisterHandlerBase
 {
diff --git a/Engine/Source/Runtime/Windows/D3D11RHI/Private/D3D11Shaders.cpp b/Engine/Source/Runtime/Windows/D3D11RHI/Private/D3D11Shaders.cpp
--- a/Engine/Source/Runtime/Windows/D3D11RHI/Private/D3D11Shaders.cpp
+++ b/Engine/Source/Runtime/Windows/D3D11RHI/Private/D3D11Shaders.cpp
@@ -235,13 +235,19 @@ FD3D11BoundShaderState::FD3D11BoundShaderState(
 	D3D11_INPUT_ELEMENT_DESC NullInputElement;
 	FMemory::Memzero(&NullInputElement,sizeof(D3D11_INPUT_ELEMENT_DESC));
 
-	FMemory::Memcpy(StreamStrides, InVertexDeclaration->StreamStrides, sizeof(StreamStrides));
-
 	FShaderCodeReader VertexShaderCode(InVertexShader->Code);
 
-	VERIFYD3D11RESULT_EX(
+	if (InVertexDeclaration == nullptr)
+	{
+		InputLayout = nullptr;
+	}
+	else
+	{
+		FMemory::Memcpy(StreamStrides, InVertexDeclaration->StreamStrides, sizeof(StreamStrides));
+
+		VERIFYD3D11RESULT_EX(
 		Direct3DDevice->CreateInputLayout(
-			InVertexDeclaration ? InVertexDeclaration->VertexElements.GetData() : &NullInputElement,
+			InVertexDeclaration && InVertexDeclaration->VertexElements.Num() ? InVertexDeclaration->VertexElements.GetData() : &NullInputElement,
 			InVertexDeclaration ? InVertexDeclaration->VertexElements.Num() : 0,
 			&InVertexShader->Code[ InVertexShader->Offset ],			// TEMP ugly
 			VertexShaderCode.GetActualShaderCodeSize() - InVertexShader->Offset,
@@ -249,6 +255,7 @@ FD3D11BoundShaderState::FD3D11BoundShaderState(
 			),
 		Direct3DDevice
 		);
+	}
 
 	VertexShader = InVertexShader->Resource;
 	PixelShader = InPixelShader ? InPixelShader->Resource : NULL;
