Since Niagara adds its own shader types and isn't necessarily loaded, the check is no longer valid for initialized modules. Instead, we log that certain types will not be supported.

#tests Mason reran the QAGame packaging with the fix and it worked successfully

diff --git a/Engine/Source/Runtime/Engine/Private/ShaderCompiler/ShaderCompiler.cpp b/Engine/Source/Runtime/Engine/Private/ShaderCompiler/ShaderCompiler.cpp
--- a/Engine/Source/Runtime/Engine/Private/ShaderCompiler/ShaderCompiler.cpp
+++ b/Engine/Source/Runtime/Engine/Private/ShaderCompiler/ShaderCompiler.cpp
@@ -354,11 +354,19 @@ static const TArray<const IShaderFormat*>& GetShaderFormats()
 
 		for (int32 Index = 0; Index < Modules.Num(); Index++)
 		{
-			IShaderFormat* Format = FModuleManager::GetModuleChecked<IShaderFormatModule>(Modules[Index]).GetShaderFormat();
+			IShaderFormatModule* Module = FModuleManager::GetModulePtr<IShaderFormatModule>(Modules[Index]);
+			if (Module != nullptr)
+			{
+				IShaderFormat* Format = Module->GetShaderFormat();
 
-			if (Format != nullptr)
+				if (Format != nullptr)
+				{
+					Results.Add(Format);
+				}
+			}
+			else
 			{
-				Results.Add(Format);
+				UE_LOG(LogShaders, Display, TEXT("Unable to load module %s, skipping its shader formats."), *Modules[Index].ToString());
 			}
 		}
 	}
@@ -2575,7 +2583,7 @@ void GlobalBeginCompileShader(
 
 		checkf(FPaths::GetExtension(Input.VirtualSourceFilePath) == TEXT("usf"),
 			TEXT("Incorrect virtual shader path extension for shader file to compile '%s': Only .usf files should be "
-			     "compiled. .ush file are meant to be included only."),
+				 "compiled. .ush file are meant to be included only."),
 			*Input.VirtualSourceFilePath);
 
 		for (const auto& Entry : Input.Environment.IncludeVirtualPathToContentsMap)
@@ -2586,12 +2594,12 @@ void GlobalBeginCompileShader(
 
 			checkf(VirtualShaderFilePath.Contains(TEXT("/Generated/")),
 				TEXT("Incorrect virtual shader path for generated file '%s': Generated files must be located under an "
-				     "non existing 'Generated' directory, for instance: /Engine/Generated/ or /Plugin/FooBar/Generated/."),
+					 "non existing 'Generated' directory, for instance: /Engine/Generated/ or /Plugin/FooBar/Generated/."),
 				*VirtualShaderFilePath);
 
 			checkf(VirtualShaderFilePath == Input.VirtualSourceFilePath || FPaths::GetExtension(VirtualShaderFilePath) == TEXT("ush"),
 				TEXT("Incorrect virtual shader path extension for generated file '%s': Generated file must either be the "
-				     "USF to compile, or a USH file to be included."),
+					 "USF to compile, or a USH file to be included."),
 				*VirtualShaderFilePath);
 		}
 	#endif
