Fix 1mb/s memory leak in editor

diff --git a/Engine/Source/Runtime/Engine/Private/Slate/DebugCanvas.cpp b/Engine/Source/Runtime/Engine/Private/Slate/DebugCanvas.cpp
--- a/Engine/Source/Runtime/Engine/Private/Slate/DebugCanvas.cpp
+++ b/Engine/Source/Runtime/Engine/Private/Slate/DebugCanvas.cpp
@@ -124,12 +124,10 @@ void FDebugCanvasDrawer::InitDebugCanvas(UWorld* InWorld)
 		// the same canvas
 	if (FSlateApplication::Get().IsNormalExecution())
 	{
-		FCanvas* Canvas = new FCanvas(RenderTarget, nullptr, InWorld, InWorld ? InWorld->FeatureLevel : GMaxRHIFeatureLevel);
+		GameThreadCanvas = MakeShared<FCanvas, ESPMode::ThreadSafe>(RenderTarget, nullptr, InWorld, InWorld ? InWorld->FeatureLevel : GMaxRHIFeatureLevel);
 
 		// Do not allow the canvas to be flushed outside of our debug rendering path
-		Canvas->SetAllowedModes(FCanvas::Allow_DeleteOnRender);
-
-		GameThreadCanvas = MakeShared<FCanvas, ESPMode::ThreadSafe>(*Canvas);
+		GameThreadCanvas->SetAllowedModes(FCanvas::Allow_DeleteOnRender);
 	}
 }
 
