Apply component instance data starting from the root component so that a component's attach parent has always had its component data applied before it

diff --git a/Engine/Source/Runtime/Engine/Private/ComponentInstanceDataCache.cpp b/Engine/Source/Runtime/Engine/Private/ComponentInstanceDataCache.cpp
--- a/Engine/Source/Runtime/Engine/Private/ComponentInstanceDataCache.cpp
+++ b/Engine/Source/Runtime/Engine/Private/ComponentInstanceDataCache.cpp
@@ -348,7 +348,66 @@ void FComponentInstanceDataCache::ApplyToActor(AActor* Actor, const ECacheApplyP
 	{
 		const bool bIsChildActor = Actor->IsChildActor();
 
-		TInlineComponentArray<UActorComponent*> Components(Actor);
+		// We want to apply instance data from the root node down to ensure changes such as transforms 
+		// propagate correctly so we will build the components list in a breadth-first manner.
+		TInlineComponentArray<UActorComponent*> Components;
+		Components.Reserve(Actor->GetComponents().Num());
+
+		auto AddComponentHierarchy = [&Components](USceneComponent* Component)
+		{
+			int32 FirstProcessIndex = Components.Num();
+
+			// Add this to our list and make it our starting node
+			Components.Add(Component);
+
+			int32 CompsToProcess = 1;
+
+			while (CompsToProcess)
+			{
+				// track how many elements were here
+				const int32 StartingProcessedCount = Components.Num();
+
+				// process the currently unprocessed elements
+				for (int32 ProcessIndex = 0; ProcessIndex < CompsToProcess; ++ProcessIndex)
+				{
+					USceneComponent* SceneComponent = CastChecked<USceneComponent>(Components[FirstProcessIndex + ProcessIndex]);
+
+					// add all children to the end of the array
+					for (int32 ChildIndex = 0; ChildIndex < SceneComponent->GetNumChildrenComponents(); ++ChildIndex)
+					{
+						if (USceneComponent* ChildComponent = SceneComponent->GetChildComponent(ChildIndex))
+						{
+							Components.Add(ChildComponent);
+						}
+					}
+				}
+
+				// next loop start with the nodes we just added
+				FirstProcessIndex = StartingProcessedCount;
+				CompsToProcess = Components.Num() - StartingProcessedCount;
+			}
+		};
+
+		if (USceneComponent* RootComponent = Actor->GetRootComponent())
+		{
+			AddComponentHierarchy(RootComponent);
+		}
+
+		for (UActorComponent* Component : Actor->GetComponents())
+		{
+			if (USceneComponent* SceneComponent = Cast<USceneComponent>(Component))
+			{
+				USceneComponent* ParentComponent = SceneComponent->GetAttachParent();
+				if ((ParentComponent == nullptr && SceneComponent != Actor->GetRootComponent()) || (ParentComponent && ParentComponent->GetOwner() != Actor))
+				{
+					AddComponentHierarchy(SceneComponent);
+				}
+			}
+			else if (Component)
+			{
+				Components.Add(Component);
+			}
+		}
 
 		// Cache all archetype objects
 		TMap<UActorComponent*, const UObject*> ComponentToArchetypeMap;
