Restore removed null poionter check, to fix crash after running map check.

diff --git a/Engine/Source/Runtime/Engine/Private/ErrorChecking.cpp b/Engine/Source/Runtime/Engine/Private/ErrorChecking.cpp
--- a/Engine/Source/Runtime/Engine/Private/ErrorChecking.cpp
+++ b/Engine/Source/Runtime/Engine/Private/ErrorChecking.cpp
@@ -128,7 +128,7 @@ void ABrush::CheckForErrors()
 	else
 	{		
 		// NOTE : don't report NULL texture references on the builder brush - it doesn't matter there
-		if( !FActorEditorUtils::IsABuilderBrush(this) && !IsBrushShape() )
+		if( Brush && !FActorEditorUtils::IsABuilderBrush(this) && !IsBrushShape() )
 		{
 			// A brush without any polygons in it isn't useful.  Should be deleted.
 			if( Brush->Polys->Element.Num() == 0 )
