HTML5 startup crash workaround

diff --git a/Engine/Source/Runtime/Engine/Private/EngineSessionManager.cpp b/Engine/Source/Runtime/Engine/Private/EngineSessionManager.cpp
--- a/Engine/Source/Runtime/Engine/Private/EngineSessionManager.cpp
+++ b/Engine/Source/Runtime/Engine/Private/EngineSessionManager.cpp
@@ -29,7 +29,7 @@ DEFINE_LOG_CATEGORY(LogEngineSessionManager);
 
 namespace SessionManagerDefs
 {
-	static const FTimespan SessionRecordExpiration = FTimespan::FromDays(30.0);
+	static const FTimespan SessionRecordExpiration = FTimespan(25920000000000); // workaround for HTML5, was FTimespan::FromDays(30.0);
 	static const FTimespan SessionRecordTimeout = FTimespan::FromMinutes(3.0);
 	static const FTimespan GlobalLockWaitTimeout = FTimespan::FromSeconds(0.5);
 	static const int HeartbeatPeriodSeconds(60);
