Copying //UE4/Dev-AnimPhys to //UE4/Dev-Main (Source: //UE4/Dev-AnimPhys @ 3624379)

diff --git a/Engine/Source/Runtime/Engine/Classes/Animation/SmartName.h b/Engine/Source/Runtime/Engine/Classes/Animation/SmartName.h
--- a/Engine/Source/Runtime/Engine/Classes/Animation/SmartName.h
+++ b/Engine/Source/Runtime/Engine/Classes/Animation/SmartName.h
@@ -7,6 +7,7 @@
 #include "Misc/Guid.h"
 #include "UObject/Class.h"
 #include "BoneContainer.h"
+#include "AnimPhysObjectVersion.h"
 #include "SmartName.generated.h"
 
 struct FSmartName;
@@ -39,14 +40,24 @@ struct FCurveMetaData
 	GENERATED_USTRUCT_BODY()
 
 	struct FAnimCurveType Type;
+	/** connected bones to this meta data */
 	TArray<struct FBoneReference> LinkedBones;
+	/* max LOD (lowest LOD) to evaluate this. -1 means it will evaluate all the time. */
+	uint8 MaxLOD;
 
 	friend FArchive& operator<<(FArchive& Ar, FCurveMetaData& B)
 	{
+		Ar.UsingCustomVersion(FAnimPhysObjectVersion::GUID);
+
 		Ar << B.Type.bMaterial;
 		Ar << B.Type.bMorphtarget;
 		Ar << B.LinkedBones;
 
+		if (Ar.CustomVer(FAnimPhysObjectVersion::GUID) >= FAnimPhysObjectVersion::AddLODToCurveMetaData)
+		{
+			Ar << B.MaxLOD;
+		}
+
 		return Ar;
 	}
 
@@ -57,6 +68,7 @@ struct FCurveMetaData
 	}
 
 	FCurveMetaData()
+		: MaxLOD(0xFF)
 	{
 	}
 };
@@ -68,25 +80,17 @@ struct ENGINE_API FSmartNameMapping
 
 	FSmartNameMapping();
 
-	// Add a name to the mapping, if it exists, get it.
-	// @param Name - The name to add/get
-	// @param OUT OutUid - The UID of the newly created or retrieved name
-	// @param OUT OutGuid - The GUID of the newly created or retrieved name
-	// @return bool - true if the name was added, false if it existed (OutUid will be correctly set still in this case)
-	bool AddOrFindName(FName Name, SmartName::UID_Type& OutUid, FGuid& OutGuid);
-	// Add a name to the mapping if it doesn't exist
-	// @param Name - The name to add/get
-	// @param OUT OutUid - The UID of the newly created or retrieved name
-	// @param InGuid - the guid of the name
-	// @return bool - true if the name was added, false if it existed (OutUid will be correctly set still in this case)
-	bool AddName(FName Name, SmartName::UID_Type& OutUid, const FGuid& InGuid);
+	// Add a name to the mapping, fails if it already exists
+	// @param InName - The name to add
+	// @return FSmartName - populated smart name
+	FSmartName AddName(FName InName);
 
 	// Get a name from the mapping
 	// @param Uid - SmartName::UID_Type of the name to retrieve
 	// @param OUT OutName - Retrieved name
 	// @return bool - true if name existed and OutName is valid
 	bool GetName(const SmartName::UID_Type& Uid, FName& OutName) const;
-	bool GetNameByGuid(const FGuid& Guid, FName& OutName) const;
+	
 	// Fill an array with all used UIDs
 	// @param Array - Array to fill
 	void FillUidArray(TArray<SmartName::UID_Type>& Array) const;
@@ -95,6 +99,7 @@ struct ENGINE_API FSmartNameMapping
 	// @param Array - Array to fill
 	void FillNameArray(TArray<FName>& Array) const;
 
+#if WITH_EDITOR
 	// Change a name
 	// @param Uid - SmartName::UID_Type of the name to change
 	// @param NewName - New name to set 
@@ -106,6 +111,12 @@ struct ENGINE_API FSmartNameMapping
 	// @return bool - true if the name was found and removed, false if the name wasn't present in the mapping
 	bool Remove(const SmartName::UID_Type& Uid);
 
+	// Remove a name from the mapping
+	// @param Uid - SmartName::UID_Type of the name to remove
+	// @return bool - true if the name was found and removed, false if the name wasn't present in the mapping
+	bool Remove(const FName& Name);
+#endif
+
 	// Return SmartName::UID_Type if it finds it
 	// @param Name - Name of curve to find UID for
 	// @return SmartName::UID_Type - MaxUID if it doesn't find, actual UID if it finds. 
@@ -121,17 +132,7 @@ struct ENGINE_API FSmartNameMapping
 	// @return bool - whether the name was found
 	bool Exists(const FName& Name) const;
 
-	// Get the number of names currently stored in this container
-	int32 GetNumNames() const;
-
-	// Find Or Add Smart Names
-#if WITH_EDITOR
-	bool FindOrAddSmartName(FName Name, FSmartName& OutName);
-	bool AddSmartName(FSmartName& OutName);
-#else
-	// in cooked build, you don't have GUID, so register without GUID
-	bool FindOrAddSmartName(FName Name, SmartName::UID_Type& OutUid);
-#endif // WITH_EDITOR
+	// Find Smart Names
 	bool FindSmartName(FName Name, FSmartName& OutName) const;
 	bool FindSmartNameByUID(SmartName::UID_Type UID, FSmartName& OutName) const;
 
@@ -156,9 +157,8 @@ struct ENGINE_API FSmartNameMapping
 	/* initialize curve meta data for the container */
 	void InitializeCurveMetaData(class USkeleton* Skeleton);
 private:
-	SmartName::UID_Type NextUid;				// The next SmartName::UID_Type to use 
-	TMap<SmartName::UID_Type, FName> UidMap;	// Mapping of UIDs and names. This is transient and built upon serialize - this is run-time data, and searching from SmartName::UID_Type to FName is important
-	TMap<FName, FGuid> GuidMap;	// Mapping of GUIDs and names. This is the one serializing
+	
+	TArray<FName> CurveNameList; // List of curve names, order gives UID
 	TMap<FName, FCurveMetaData> CurveMetaDataMap;
 };
 
@@ -184,9 +184,6 @@ protected:
 	FSmartNameMapping* GetContainerInternal(const FName& ContainerName);
 	const FSmartNameMapping* GetContainerInternal(const FName& ContainerName) const;
 
-	/* initialize curve meta data for the container */
-	void InitializeCurveMetaData(class USkeleton* Skeleton);
-
 private:
 	TMap<FName, FSmartNameMapping> NameMappings;	// List of smartname mappings
 };
@@ -194,7 +191,6 @@ private:
 USTRUCT()
 struct ENGINE_API FSmartName
 {
-
 	GENERATED_USTRUCT_BODY();
 
 	// name 
@@ -204,36 +200,19 @@ struct ENGINE_API FSmartName
 	// SmartName::UID_Type - for faster access
 	SmartName::UID_Type	UID;
 
-#if WITH_EDITORONLY_DATA
-	UPROPERTY()
-	FGuid	Guid;
-#endif // WITH_EDITORONLY_DATA
-
 	FSmartName()
 		: DisplayName(NAME_None)
 		, UID(SmartName::MaxUID)
 	{}
 
-#if WITH_EDITORONLY_DATA
-	FSmartName(const FName& InDisplayName, SmartName::UID_Type InUID, const FGuid& InGuid)
-		: DisplayName(InDisplayName)
-		, UID(InUID)
-		, Guid(InGuid)
-	{}
-#else
 	FSmartName(const FName& InDisplayName, SmartName::UID_Type InUID)
 		: DisplayName(InDisplayName)
 		, UID(InUID)
 	{}
-#endif // WITH_EDITORONLY_DATA
 
 	bool operator==(FSmartName const& Other) const
 	{
-#if WITH_EDITORONLY_DATA
-		return (DisplayName == Other.DisplayName && UID == Other.UID && Guid == Other.Guid);
-#else
 		return (DisplayName == Other.DisplayName && UID == Other.UID);
-#endif // WITH_EDITORONLY_DATA
 	}
 	bool operator!=(const FSmartName& Other) const
 	{
@@ -256,11 +235,7 @@ struct ENGINE_API FSmartName
 
 	bool IsValid() const
 	{
-#if WITH_EDITORONLY_DATA
-		return UID != SmartName::MaxUID && Guid.IsValid(); 
-#else
 		return UID != SmartName::MaxUID;
-#endif  // WITH_EDITORONLY_DATA
 	}
 };
 
diff --git a/Engine/Source/Runtime/Engine/Classes/Components/SkeletalMeshComponent.h b/Engine/Source/Runtime/Engine/Classes/Components/SkeletalMeshComponent.h
--- a/Engine/Source/Runtime/Engine/Classes/Components/SkeletalMeshComponent.h
+++ b/Engine/Source/Runtime/Engine/Classes/Components/SkeletalMeshComponent.h
@@ -372,9 +372,19 @@ public:
 	uint32 bDisableClothSimulation:1;
 
 private:
-	/** Disable animation curves for this component. */
+	/** Disable animation curves for this component. If this is set true, no curves will be processed */
 	UPROPERTY(EditAnywhere, AdvancedDisplay, Category = SkeletalMesh)
-	uint32 bDisableAnimCurves : 1;
+	uint32 bAllowAnimCurveEvaluation : 1;
+
+	/** DEPRECATED. Use bAllowAnimCurveEvaluation instead */
+	DEPRECATED(4.18, "This property is deprecated. Please use bAllowAnimCurveEvaluatiuon instead. Note that the meaning is reversed.")	
+	UPROPERTY()
+	uint32 bDisableAnimCurves_DEPRECATED : 1;
+
+	/** You can choose to disable certain curves if you prefer. 
+	 * This is transient curves that will be ignored by animation system if you choose this */
+	UPROPERTY(transient)
+	TArray<FName> DisallowedAnimCurves;
 
 public:
 	/** can't collide with part of environment if total collision volumes exceed 16 capsules or 32 planes per convex */
@@ -742,11 +752,32 @@ public:
 	}
 #endif 
 
+	DEPRECATED(4.18, "This function is deprecated. Please use SetAllowAnimCurveEvaluation instead. Note that the meaning is reversed.")
 	UFUNCTION(BlueprintCallable, Category = "Components|SkeletalMesh")
 	void SetDisableAnimCurves(bool bInDisableAnimCurves);
 
+	DEPRECATED(4.18, "This function is deprecated. Please use GetAllowedAnimCurveEvaluate instead. Note that the meaning is reversed.")
+	UFUNCTION(BlueprintCallable, Category = "Components|SkeletalMesh")
+	bool GetDisableAnimCurves() const { return !bAllowAnimCurveEvaluation; }
+
+	UFUNCTION(BlueprintCallable, Category = "Components|SkeletalMesh")
+	void SetAllowAnimCurveEvaluation(bool bInAllow);
+
+	UFUNCTION(BlueprintCallable, Category = "Components|SkeletalMesh")
+	bool GetAllowedAnimCurveEvaluate() const { return bAllowAnimCurveEvaluation; }
+
+	UFUNCTION(BlueprintCallable, Category = "Components|SkeletalMesh")
+	void AllowAnimCurveEvaluation(FName NameOfCurve, bool bAllow);
+
+	/** By reset, it will allow all the curves to be evaluated */
+	UFUNCTION(BlueprintCallable, Category = "Components|SkeletalMesh")
+	void ResetAllowedAnimCurveEvaluation();
+
+	/** resets, and then only allow the following list to be allowed/disallowed */
 	UFUNCTION(BlueprintCallable, Category = "Components|SkeletalMesh")
-	bool GetDisableAnimCurves() const { return bDisableAnimCurves; }
+	void SetAllowedAnimCurvesEvaluation(const TArray<FName>& List, bool bAllow);
+
+	const TArray<FName>& GetDisallowedAnimCurvesEvaluation() const { return DisallowedAnimCurves; }
 
 	/** We detach the Component once we are done playing it.
 	 *
@@ -1577,6 +1608,8 @@ protected:
 	bool NeedToSpawnAnimScriptInstance() const;
 	bool NeedToSpawnPostPhysicsInstance() const;
 
+	bool ShouldBlendPhysicsBones() const;
+
 private:
 
 	FSkeletalMeshComponentEndPhysicsTickFunction EndPhysicsTickFunction;
@@ -1595,7 +1628,6 @@ private:
 	*/
 	void FillComponentSpaceTransforms(const USkeletalMesh* InSkeletalMesh, const TArray<FTransform>& InBoneSpaceTransforms, TArray<FTransform>& OutComponentSpaceTransforms) const;
 
-	bool ShouldBlendPhysicsBones() const;
 	bool DoAnyPhysicsBodiesHaveWeight() const;
 
 	void ClearAnimScriptInstance();
diff --git a/Engine/Source/Runtime/Engine/Classes/Components/SkinnedMeshComponent.h b/Engine/Source/Runtime/Engine/Classes/Components/SkinnedMeshComponent.h
--- a/Engine/Source/Runtime/Engine/Classes/Components/SkinnedMeshComponent.h
+++ b/Engine/Source/Runtime/Engine/Classes/Components/SkinnedMeshComponent.h
@@ -430,7 +430,7 @@ public:
 	UPROPERTY(EditAnywhere, AdvancedDisplay, BlueprintReadOnly, Category=Lighting, meta=(UIMin = "0", UIMax = "1", EditCondition="bCastCapsuleIndirectShadow", DisplayName = "Capsule Indirect Shadow Min Visibility"))
 	float CapsuleIndirectShadowMinVisibility;
 
-	/** CPU skinning rendering - only for previewing in Persona and conversion tools */
+	/** Whether or not to CPU skin this component, requires render data refresh after changing */
 	UPROPERTY(transient)
 	uint32 bCPUSkinning : 1;
 
diff --git a/Engine/Source/Runtime/Engine/Classes/Engine/LODActor.h b/Engine/Source/Runtime/Engine/Classes/Engine/LODActor.h
--- a/Engine/Source/Runtime/Engine/Classes/Engine/LODActor.h
+++ b/Engine/Source/Runtime/Engine/Classes/Engine/LODActor.h
@@ -46,6 +46,12 @@ public:
 	UPROPERTY(Category=LODActor, VisibleAnywhere)
 	TArray<UObject*> SubObjects;
 
+#if WITH_EDITORONLY_DATA
+	/** Assets which have become stale and might get deleted (depends on undo/redo behaviour) */
+	UPROPERTY(transient)
+	TArray<UObject*> PreviousSubObjects;
+#endif // WITH_EDITORONLY_DATA
+
 	//~ Begin AActor Interface
 #if WITH_EDITOR
 	virtual void CheckForErrors() override;
@@ -136,6 +142,9 @@ public:
 	virtual void PreEditChange(UProperty* PropertyThatWillChange) override;
 	virtual void PostEditChangeProperty(FPropertyChangedEvent& PropertyChangedEvent) override;
 	virtual void Serialize(FArchive& Ar) override;
+
+	virtual void PreSave(const class ITargetPlatform* TargetPlatform) override;
+	virtual void BeginDestroy() override;
 #endif // WITH_EDITOR	
 	//~ End UObject Interface.		
 protected:
diff --git a/Engine/Source/Runtime/Engine/Classes/Engine/MeshMerging.h b/Engine/Source/Runtime/Engine/Classes/Engine/MeshMerging.h
--- a/Engine/Source/Runtime/Engine/Classes/Engine/MeshMerging.h
+++ b/Engine/Source/Runtime/Engine/Classes/Engine/MeshMerging.h
@@ -208,9 +208,13 @@ struct FMeshProxySettings
 	float HardAngleThreshold;
 
 	/** Lightmap resolution */
-	UPROPERTY(EditAnywhere, Category = ProxySettings)
+	UPROPERTY(EditAnywhere, Category = ProxySettings, meta = (EditCondition = "!bComputeLightMapResolution"))
 	int32 LightMapResolution;
 
+	/** If ticked will compute the lightmap resolution by summing the dimensions for each mesh included for merging */
+	UPROPERTY(EditAnywhere, Category = ProxySettings)
+	bool bComputeLightMapResolution;
+
 	/** Whether Simplygon should recalculate normals, otherwise the normals channel will be sampled from the original mesh */
 	UPROPERTY(EditAnywhere, Category = ProxySettings)
 	bool bRecalculateNormals;
@@ -296,9 +300,13 @@ struct FMeshMergingSettings
 	bool bGenerateLightMapUV;
 
 	/** Target lightmap resolution */
-	UPROPERTY(EditAnywhere, AdvancedDisplay, BlueprintReadWrite, Category = MeshSettings, meta=(ClampMax = 4096))
+	UPROPERTY(EditAnywhere, AdvancedDisplay, BlueprintReadWrite, Category = MeshSettings, meta=(ClampMax = 4096, editcondition = "!bComputedLightMapResolution"))
 	int32 TargetLightMapResolution;
 
+	/** Whether or not the lightmap resolution should be computed by summing the lightmap resolutions for the input Mesh Components */
+	UPROPERTY(EditAnywhere, AdvancedDisplay, BlueprintReadWrite, Category = MeshSettings)
+	bool bComputedLightMapResolution;
+
 	/** Whether we should import vertex colors into merged mesh */
 	UPROPERTY()
 	bool bImportVertexColors_DEPRECATED;
diff --git a/Engine/Source/Runtime/Engine/Classes/Engine/SkeletalMesh.h b/Engine/Source/Runtime/Engine/Classes/Engine/SkeletalMesh.h
--- a/Engine/Source/Runtime/Engine/Classes/Engine/SkeletalMesh.h
+++ b/Engine/Source/Runtime/Engine/Classes/Engine/SkeletalMesh.h
@@ -633,7 +633,7 @@ public:
 	class UBodySetup* BodySetup;
 
 	/**
-	 *	Physics and collision information used for this USkeletalMesh, set up in PhAT.
+	 *	Physics and collision information used for this USkeletalMesh, set up in Physics Asset Editor.
 	 *	This is used for per-bone hit detection, accurate bounding box calculation and ragdoll physics for example.
 	 */
 	UPROPERTY(EditAnywhere, AssetRegistrySearchable, BlueprintReadOnly, Category=Physics)
diff --git a/Engine/Source/Runtime/Engine/Classes/Kismet/GameplayStatics.h b/Engine/Source/Runtime/Engine/Classes/Kismet/GameplayStatics.h
--- a/Engine/Source/Runtime/Engine/Classes/Kismet/GameplayStatics.h
+++ b/Engine/Source/Runtime/Engine/Classes/Kismet/GameplayStatics.h
@@ -386,9 +386,11 @@ class ENGINE_API UGameplayStatics : public UBlueprintFunctionLibrary
 	 * @param PitchMultiplier - Multiplies the pitch.
 	 * @param ConcurrencySettings - Override concurrency settings package to play sound with
 	 * @param StartTime - How far in to the sound to begin playback at
+	 * @param ConcurrencySettings - Override concurrency settings package to play sound with
+	 * @param OwningActor - The actor to use as the "owner" for concurrency settings purposes. Allows PlaySound calls to do a concurrency limit per owner.
 	 */
 	UFUNCTION(BlueprintCallable, BlueprintCosmetic, Category="Audio", meta=( WorldContext="WorldContextObject", AdvancedDisplay = "2", UnsafeDuringActorConstruction = "true" ))
-	static void PlaySound2D(const UObject* WorldContextObject, USoundBase* Sound, float VolumeMultiplier = 1.f, float PitchMultiplier = 1.f, float StartTime = 0.f, USoundConcurrency* ConcurrencySettings = nullptr);
+	static void PlaySound2D(const UObject* WorldContextObject, USoundBase* Sound, float VolumeMultiplier = 1.f, float PitchMultiplier = 1.f, float StartTime = 0.f, USoundConcurrency* ConcurrencySettings = nullptr, AActor* OwningActor = nullptr);
 
 	/**
 	 * Spawns a sound with no attenuation, perfect for UI sounds.
@@ -432,9 +434,10 @@ class ENGINE_API UGameplayStatics : public UBlueprintFunctionLibrary
 	 * @param StartTime - How far in to the sound to begin playback at
 	 * @param AttenuationSettings - Override attenuation settings package to play sound with
 	 * @param ConcurrencySettings - Override concurrency settings package to play sound with
+	 * @param OwningActor - The actor to use as the "owner" for concurrency settings purposes. Allows PlaySound calls to do a concurrency limit per owner.
 	 */
 	UFUNCTION(BlueprintCallable, Category="Audio", meta=(WorldContext="WorldContextObject", AdvancedDisplay = "3", UnsafeDuringActorConstruction = "true"))
-	static void PlaySoundAtLocation(const UObject* WorldContextObject, USoundBase* Sound, FVector Location, FRotator Rotation, float VolumeMultiplier = 1.f, float PitchMultiplier = 1.f, float StartTime = 0.f, class USoundAttenuation* AttenuationSettings = nullptr, USoundConcurrency* ConcurrencySettings = nullptr);
+	static void PlaySoundAtLocation(const UObject* WorldContextObject, USoundBase* Sound, FVector Location, FRotator Rotation, float VolumeMultiplier = 1.f, float PitchMultiplier = 1.f, float StartTime = 0.f, class USoundAttenuation* AttenuationSettings = nullptr, USoundConcurrency* ConcurrencySettings = nullptr, AActor* OwningActor = nullptr);
 
 	static void PlaySoundAtLocation(const UObject* WorldContextObject, USoundBase* Sound, FVector Location, float VolumeMultiplier = 1.f, float PitchMultiplier = 1.f, float StartTime = 0.f, class USoundAttenuation* AttenuationSettings = nullptr, USoundConcurrency* ConcurrencySettings = nullptr)
 	{
diff --git a/Engine/Source/Runtime/Engine/Classes/PhysicsEngine/AggregateGeom.h b/Engine/Source/Runtime/Engine/Classes/PhysicsEngine/AggregateGeom.h
--- a/Engine/Source/Runtime/Engine/Classes/PhysicsEngine/AggregateGeom.h
+++ b/Engine/Source/Runtime/Engine/Classes/PhysicsEngine/AggregateGeom.h
@@ -55,7 +55,7 @@ struct ENGINE_API FKAggregateGeom
 		return SphereElems.Num() + SphylElems.Num() + BoxElems.Num() + ConvexElems.Num();
 	}
 
-	int32 GetElementCount(int32 Type) const;
+	int32 GetElementCount(EAggCollisionShape::Type Type) const;
 
 	void EmptyElements()
 	{
diff --git a/Engine/Source/Runtime/Engine/Classes/PhysicsEngine/BodyInstance.h b/Engine/Source/Runtime/Engine/Classes/PhysicsEngine/BodyInstance.h
--- a/Engine/Source/Runtime/Engine/Classes/PhysicsEngine/BodyInstance.h
+++ b/Engine/Source/Runtime/Engine/Classes/PhysicsEngine/BodyInstance.h
@@ -197,7 +197,7 @@ struct ENGINE_API FBodyInstance
 	 *	Index of this BodyInstance within the SkeletalMeshComponent/PhysicsAsset. 
 	 *	Is INDEX_NONE if a single body component
 	 */
-	int16 InstanceBodyIndex;
+	int32 InstanceBodyIndex;
 
 	/** When we are a body within a SkeletalMeshComponent, we cache the index of the bone we represent, to speed up sync'ing physics to anim. */
 	int16 InstanceBoneIndex;
@@ -805,10 +805,10 @@ public:
 	bool IsValidBodyInstance() const;
 
 	/** Get current transform in world space from physics body. */
-	FTransform GetUnrealWorldTransform(bool bWithProjection = true) const;
+	FTransform GetUnrealWorldTransform(bool bWithProjection = true, bool bForceGlobalPose = false) const;
 
 	/** Get current transform in world space from physics body. */
-	FTransform GetUnrealWorldTransform_AssumesLocked(bool bWithProjection = true) const;
+	FTransform GetUnrealWorldTransform_AssumesLocked(bool bWithProjection = true, bool bForceGlobalPose = false) const;
 
 	/**
 	 *	Move the physics body to a new pose.
diff --git a/Engine/Source/Runtime/Engine/Classes/PhysicsEngine/BodySetup.h b/Engine/Source/Runtime/Engine/Classes/PhysicsEngine/BodySetup.h
--- a/Engine/Source/Runtime/Engine/Classes/PhysicsEngine/BodySetup.h
+++ b/Engine/Source/Runtime/Engine/Classes/PhysicsEngine/BodySetup.h
@@ -138,7 +138,7 @@ class UBodySetup : public UObject
 	FName BoneName;
 
 	/** 
-	 *	If Unfixed it will use physics. If fixed, it will use kinematic. Default will inherit from OwnerComponent's behavior.
+	 *	If simulated it will use physics, if kinematic it will not be affected by physics, but can interact with physically simulated bodies. Default will inherit from OwnerComponent's behavior.
 	 */
 	UPROPERTY(EditAnywhere, Category=Physics)
 	TEnumAsByte<EPhysicsType> PhysicsType;
diff --git a/Engine/Source/Runtime/Engine/Classes/PhysicsEngine/ConstraintTypes.h b/Engine/Source/Runtime/Engine/Classes/PhysicsEngine/ConstraintTypes.h
--- a/Engine/Source/Runtime/Engine/Classes/PhysicsEngine/ConstraintTypes.h
+++ b/Engine/Source/Runtime/Engine/Classes/PhysicsEngine/ConstraintTypes.h
@@ -67,15 +67,15 @@ struct ENGINE_API FLinearConstraint : public FConstraintBaseParams
 	float Limit;
 
 	/** Indicates the linear constraint applied along the X-axis. Free implies no constraint at all. Locked implies no movement along X is allowed. Limited implies the distance in the joint along all active axes must be less than the Distance provided. */
-	UPROPERTY(EditAnywhere, Category = Linear)
+	UPROPERTY(EditAnywhere, Category = Linear, meta = (DisplayName= "X Motion"))
 	TEnumAsByte<enum ELinearConstraintMotion> XMotion;
 
-	/** Indicates the linear constraint applied along the X-axis. Free implies no constraint at all. Locked implies no movement along X is allowed. Limited implies the distance in the joint along all active axes must be less than the Distance provided. */
-	UPROPERTY(EditAnywhere, Category = Linear)
+	/** Indicates the linear constraint applied along the Y-axis. Free implies no constraint at all. Locked implies no movement along Y is allowed. Limited implies the distance in the joint along all active axes must be less than the Distance provided. */
+	UPROPERTY(EditAnywhere, Category = Linear, meta = (DisplayName= "Y Motion"))
 	TEnumAsByte<enum ELinearConstraintMotion> YMotion;
 
-	/** Indicates the linear constraint applied along the X-axis. Free implies no constraint at all. Locked implies no movement along X is allowed. Limited implies the distance in the joint along all active axes must be less than the Distance provided. */
-	UPROPERTY(EditAnywhere, Category = Linear)
+	/** Indicates the linear constraint applied along theZX-axis. Free implies no constraint at all. Locked implies no movement along Z is allowed. Limited implies the distance in the joint along all active axes must be less than the Distance provided. */
+	UPROPERTY(EditAnywhere, Category = Linear, meta = (DisplayName= "Z Motion"))
 	TEnumAsByte<enum ELinearConstraintMotion> ZMotion;
 
 	FLinearConstraint();
@@ -102,11 +102,11 @@ struct ENGINE_API FConeConstraint : public FConstraintBaseParams
 	float Swing2LimitDegrees;
 
 	/** Indicates whether the Swing1 limit is used.*/
-	UPROPERTY(EditAnywhere, Category = Angular)
+	UPROPERTY(EditAnywhere, Category = Angular, meta = (DisplayName= "Swing 1 Motion"))
 	TEnumAsByte<enum EAngularConstraintMotion> Swing1Motion;
 
 	/** Indicates whether the Swing2 limit is used.*/
-	UPROPERTY(EditAnywhere, Category = Angular)
+	UPROPERTY(EditAnywhere, Category = Angular, meta = (DisplayName= "Swing 2 Motion"))
 	TEnumAsByte<enum EAngularConstraintMotion> Swing2Motion;
 
 	FConeConstraint();
diff --git a/Engine/Source/Runtime/Engine/Classes/PhysicsEngine/PhysicsAsset.h b/Engine/Source/Runtime/Engine/Classes/PhysicsEngine/PhysicsAsset.h
--- a/Engine/Source/Runtime/Engine/Classes/PhysicsEngine/PhysicsAsset.h
+++ b/Engine/Source/Runtime/Engine/Classes/PhysicsEngine/PhysicsAsset.h
@@ -18,7 +18,7 @@ class USkeletalBodySetup;
  * PhysicsAsset contains a set of rigid bodies and constraints that make up a single ragdoll.
  * The asset is not limited to human ragdolls, and can be used for any physical simulation using bodies and constraints.
  * A SkeletalMesh has a single PhysicsAsset, which allows for easily turning ragdoll physics on or off for many SkeletalMeshComponents
- * The asset can be configured inside the Physics Asset Editor (PhAT).
+ * The asset can be configured inside the Physics Asset Editor.
  *
  * @see https://docs.unrealengine.com/latest/INT/Engine/Physics/PhAT/Reference/index.html
  * @see USkeletalMesh
@@ -37,7 +37,7 @@ class UPhysicsAsset : public UObject
 	UPROPERTY()
 	class USkeletalMesh * DefaultSkelMesh_DEPRECATED;
 
-	UPROPERTY()
+	UPROPERTY(AssetRegistrySearchable)
 	TSoftObjectPtr<class USkeletalMesh> PreviewSkeletalMesh;
 
 	UPROPERTY(EditAnywhere, Category = Profiles, meta=(DisableCopyPaste))
@@ -89,6 +89,10 @@ public:
 	 */
 	TMap<FRigidBodyIndexPair,bool>		CollisionDisableTable;
 
+	/** Information for thumbnail rendering */
+	UPROPERTY(VisibleAnywhere, Instanced, Category = Thumbnail)
+	class UThumbnailInfo* ThumbnailInfo;
+
 	//~ Begin UObject Interface
 	virtual void Serialize(FArchive& Ar) override;
 	virtual void PostLoad() override;
@@ -150,12 +154,15 @@ public:
 
 	void GetUsedMaterials(TArray<UMaterialInterface*>& Materials);
 
-	// @todo document
+	// Disable collsion between the bodies specified by index
 	ENGINE_API void DisableCollision(int32 BodyIndexA, int32 BodyIndexB);
 
-	// @todo document
+	// Enable collsion between the bodies specified by index
 	ENGINE_API void EnableCollision(int32 BodyIndexA, int32 BodyIndexB);
 
+	// Check whether the two bodies specified are enabled for collision
+	ENGINE_API bool IsCollisionEnabled(int32 BodyIndexA, int32 BodyIndexB) const;
+
 	/** Update the BoundsBodies array and cache the indices of bodies marked with bConsiderForBounds to BoundsBodies array. */
 	ENGINE_API void UpdateBoundsBodiesArray();
 
@@ -176,6 +183,12 @@ public:
 #if WITH_EDITOR
 	/** Update skeletal meshes when physics asset changes*/
 	ENGINE_API void RefreshPhysicsAssetChange() const;
+
+	/** Set the preview mesh for this physics asset */
+	ENGINE_API void SetPreviewMesh(USkeletalMesh* PreviewMesh);
+
+	/** Get the preview mesh for this physics asset */
+	ENGINE_API USkeletalMesh* GetPreviewMesh() const;
 #endif
 
 private:
diff --git a/Engine/Source/Runtime/Engine/Classes/Sound/AudioSettings.h b/Engine/Source/Runtime/Engine/Classes/Sound/AudioSettings.h
--- a/Engine/Source/Runtime/Engine/Classes/Sound/AudioSettings.h
+++ b/Engine/Source/Runtime/Engine/Classes/Sound/AudioSettings.h
@@ -86,6 +86,10 @@ class ENGINE_API UAudioSettings : public UDeveloperSettings
 	UPROPERTY(config, EditAnywhere, Category="Audio", meta=(AllowedClasses="SoundClass"))
 	FSoftObjectPath VoiPSoundClass;
 
+	/** The amount of audio to send to reverb submixes if no reverb send is setup for the source through attenuation settings. Only used in audio mixer. */
+	UPROPERTY(config, EditAnywhere, Category = "Audio", AdvancedDisplay, meta = (ClampMin = 0.0, ClampMax = 1.0))
+	float DefaultReverbSendLevel;
+
 	UPROPERTY(config, EditAnywhere, Category="Audio", AdvancedDisplay, meta=(ClampMin=0.1,ClampMax=1.5))
 	float LowPassFilterResonance;
 
diff --git a/Engine/Source/Runtime/Engine/Classes/Sound/SoundBase.h b/Engine/Source/Runtime/Engine/Classes/Sound/SoundBase.h
--- a/Engine/Source/Runtime/Engine/Classes/Sound/SoundBase.h
+++ b/Engine/Source/Runtime/Engine/Classes/Sound/SoundBase.h
@@ -12,15 +12,19 @@
 #include "Audio.h"
 #include "Sound/SoundConcurrency.h"
 #include "SoundSubmix.h"
+#include "Sound/SoundSourceBusSend.h"
 #include "SoundBase.generated.h"
 
 class USoundEffectSourcePreset;
 class USoundSubmix;
+class USoundSourceBus;
+class USoundEffectSourcePresetChain;
+
 struct FSoundSubmixSendInfo;
+struct FSoundSourceBusSendInfo;
 struct FActiveSound;
 struct FSoundParseParameters;
 
-
 UCLASS(config=Engine, hidecategories=Object, abstract, editinlinenew, BlueprintType)
 class ENGINE_API USoundBase : public UObject
 {
@@ -45,6 +49,10 @@ public:
 	UPROPERTY(EditAnywhere, BlueprintReadWrite, Category = Concurrency)
 	uint32 bOverrideConcurrency:1;
 
+	/** Whether or not to only send this audio's output to a bus. If true, will not be this sound won't be audible except through bus sends. */
+	UPROPERTY(EditAnywhere, BlueprintReadWrite, Category = Effects)
+	uint32 bOutputToBusOnly : 1;
+
 	UPROPERTY()
 	uint32 bIgnoreFocus_DEPRECATED:1;
 
@@ -90,6 +98,10 @@ public:
 	UPROPERTY(EditAnywhere, BlueprintReadWrite, Category = Effects)
 	USoundEffectSourcePresetChain* SourceEffectChain;
 
+	/** This sound will send it's audio output to this list of buses if there are bus instances playing.  */
+	UPROPERTY(EditAnywhere, BlueprintReadWrite, Category = Effects)
+	TArray<FSoundSourceBusSendInfo> BusSends;
+
 public:	
 	/** Number of times this cue is currently being played. */
 	int32 CurrentPlayCount;
@@ -139,6 +151,9 @@ public:
 	/** Returns the sound submix sends for this sound. */
 	void GetSoundSubmixSends(TArray<FSoundSubmixSendInfo>& OutSends) const;
 
+	/** Returns the sound source sends for this sound. */
+	void GetSoundSourceBusSends(TArray<FSoundSourceBusSendInfo>& OutSends) const;
+
 	/** Returns the FSoundConcurrencySettings struct to use. */
 	const FSoundConcurrencySettings* GetSoundConcurrencySettingsToApply();
 
diff --git a/Engine/Source/Runtime/Engine/Classes/Sound/SoundEffectBase.h b/Engine/Source/Runtime/Engine/Classes/Sound/SoundEffectBase.h
--- a/Engine/Source/Runtime/Engine/Classes/Sound/SoundEffectBase.h
+++ b/Engine/Source/Runtime/Engine/Classes/Sound/SoundEffectBase.h
@@ -104,8 +104,13 @@ public:
 	/** Queries if the given preset object is the parent preset, i.e. the preset which spawned this effect instance. */
 	bool IsParentPreset(USoundEffectPreset* InPreset) const;
 
+	/** Enqueues a lambda command on a thread safe queue which is pumped from the audio render thread. */
+	void EffectCommand(TFunction<void()> Command);
+
 protected:
 
+	void PumpPendingMessages();
+
 	FCriticalSection SettingsCritSect;
 	TArray<uint8> CurrentAudioThreadSettingsData;
 
@@ -116,6 +121,9 @@ protected:
 	FThreadSafeBool bIsRunning;
 	FThreadSafeBool bIsActive;
 
+	// Effect commmand queue
+	TQueue<TFunction<void()>> CommandQueue;
+
 	// Allow FAudioMixerSubmix to call ProcessAudio
 	friend class FMixerSubmix;
 
diff --git a/Engine/Source/Runtime/Engine/Classes/Sound/SoundEffectPreset.h b/Engine/Source/Runtime/Engine/Classes/Sound/SoundEffectPreset.h
--- a/Engine/Source/Runtime/Engine/Classes/Sound/SoundEffectPreset.h
+++ b/Engine/Source/Runtime/Engine/Classes/Sound/SoundEffectPreset.h
@@ -25,8 +25,10 @@ public:
 	virtual FSoundEffectBase* CreateNewEffect() const PURE_VIRTUAL(USoundEffectPreset::CreateNewEffect, return nullptr;);
 	virtual bool HasAssetActions() const { return true; }
 	virtual void Init() PURE_VIRTUAL(USoundEffectPreset::Init, ); 
+	virtual void OnInit() {};
 	virtual FColor GetPresetColor() const { return FColor(200.0f, 100.0f, 100.0f); }
 
+	void EffectCommand(TFunction<void()> Command);
 	void Update();
 	void AddEffectInstance(FSoundEffectBase* InSource);
 	void RemoveEffectInstance(FSoundEffectBase* InSource);
diff --git a/Engine/Source/Runtime/Engine/Classes/Sound/SoundEffectSubmix.h b/Engine/Source/Runtime/Engine/Classes/Sound/SoundEffectSubmix.h
--- a/Engine/Source/Runtime/Engine/Classes/Sound/SoundEffectSubmix.h
+++ b/Engine/Source/Runtime/Engine/Classes/Sound/SoundEffectSubmix.h
@@ -73,6 +73,6 @@ public:
 	/** Process the input block of audio. Called on audio thread. */
 	virtual void OnProcessAudio(const FSoundEffectSubmixInputData& InData, FSoundEffectSubmixOutputData& OutData) {};
 
-	/** Processes audio in the source effect. */
+	/** Processes audio in the submix effect. */
 	void ProcessAudio(FSoundEffectSubmixInputData& InData, FSoundEffectSubmixOutputData& OutData);
 };
diff --git a/Engine/Source/Runtime/Engine/Classes/Sound/SoundSourceBus.h b/Engine/Source/Runtime/Engine/Classes/Sound/SoundSourceBus.h
new file mode 100644
--- /dev/null
+++ b/Engine/Source/Runtime/Engine/Classes/Sound/SoundSourceBus.h
@@ -0,0 +1,60 @@
+// Copyright 1998-2017 Epic Games, Inc. All Rights Reserved.
+#pragma once
+
+#include "CoreMinimal.h"
+#include "UObject/ObjectMacros.h"
+#include "UObject/Object.h"
+#include "Sound/SoundWaveProcedural.h"
+#include "Sound/SoundSourceBusSend.h"
+#include "SoundSourceBus.generated.h"
+
+class USoundSourceBus;
+class USoundWaveProcedural;
+class FAudioDevice;
+
+// The number of channels to mix audio into the source bus
+UENUM(BlueprintType)
+enum class ESourceBusChannels : uint8
+{
+	Mono,
+	Stereo,
+};
+
+// A source bus is a type of USoundBase and can be "played" like any sound.
+UCLASS(hidecategories= (Compression, SoundWave, Streaming, Subtitles, Sound, Info, ImportSettings), ClassGroup = Sound, meta = (BlueprintSpawnableComponent))
+class ENGINE_API USoundSourceBus : public USoundWave
+{
+	GENERATED_BODY()
+
+public:
+
+	/** How many channels to use for the source bus. */
+	UPROPERTY(EditAnywhere, Category = BusProperties)
+	ESourceBusChannels SourceBusChannels;
+
+	/** The duration (in seconds) to use for the source bus. A duration of 0.0 indicates to play the source bus indefinitely. */
+	UPROPERTY(EditAnywhere, Category = BusProperties, meta = (UIMin = 0.0, ClampMin = 0.0))
+	float SourceBusDuration;
+
+	/** Stop the source bus when the volume goes to zero. */
+	UPROPERTY(EditAnywhere, Category = BusProperties, meta = (UIMin = 0.0, ClampMin = 0.0))
+	uint32 bAutoDeactivateWhenSilent:1;
+
+	//~ Begin UObject Interface. 
+	virtual void PostLoad() override;
+	virtual void BeginDestroy() override;
+#if WITH_EDITOR
+	virtual void PostEditChangeProperty(FPropertyChangedEvent& PropertyChangedEvent) override;
+#endif // WITH_EDITOR
+	//~ End UObject Interface. 
+
+	//~ Begin USoundBase Interface.
+	virtual bool IsPlayable() const override;
+	virtual float GetDuration() override;
+	//~ End USoundBase Interface.
+
+protected:
+
+	uint32 DurationSamples;
+	bool bInitialized;
+};
diff --git a/Engine/Source/Runtime/Engine/Classes/Sound/SoundSourceBusSend.h b/Engine/Source/Runtime/Engine/Classes/Sound/SoundSourceBusSend.h
new file mode 100644
--- /dev/null
+++ b/Engine/Source/Runtime/Engine/Classes/Sound/SoundSourceBusSend.h
@@ -0,0 +1,22 @@
+// Copyright 1998-2017 Epic Games, Inc. All Rights Reserved.
+
+#pragma once
+
+#include "CoreMinimal.h"
+#include "SoundSourceBusSend.generated.h"
+
+class USoundSourceBus;
+
+USTRUCT(BlueprintType)
+struct ENGINE_API FSoundSourceBusSendInfo
+{
+	GENERATED_USTRUCT_BODY()
+
+	// The amount of audio to send to the source bus
+	UPROPERTY(EditAnywhere, Category = SourceBusSend)
+	float SendLevel;
+
+	// The source bus to send the audio to
+	UPROPERTY(EditAnywhere, Category = SourceBusSend)
+	USoundSourceBus* SoundSourceBus;
+};
diff --git a/Engine/Source/Runtime/Engine/Classes/Sound/SoundWave.h b/Engine/Source/Runtime/Engine/Classes/Sound/SoundWave.h
--- a/Engine/Source/Runtime/Engine/Classes/Sound/SoundWave.h
+++ b/Engine/Source/Runtime/Engine/Classes/Sound/SoundWave.h
@@ -14,6 +14,7 @@
 #include "Sound/SoundBase.h"
 #include "Serialization/BulkData.h"
 #include "Sound/SoundGroups.h"
+#include "AudioMixerTypes.h"
 
 #include "SoundWave.generated.h"
 
@@ -140,6 +141,9 @@ class ENGINE_API USoundWave : public USoundBase
 	/** Set to true for programmatically-generated, streamed audio. */
 	uint32 bProcedural:1;
 
+	/** Set to true of this is a bus sound source. This will result in the sound wave not generating audio for itself, but generate audio through instances. Used only in audio mixer. */
+	uint32 bIsBus:1;
+
 	/** Set to true for procedural waves that can be processed asynchronously. */
 	uint32 bCanProcessAsync:1;
 
@@ -350,16 +354,21 @@ public:
 	FWaveInstance* HandleStart( FActiveSound& ActiveSound, const UPTRINT WaveInstanceHash ) const;
 
 	/** 
-	 * This is only for DTYPE_Procedural audio. Override this function.
-	 *  Put SamplesNeeded PCM samples into Buffer. If put less,
-	 *  silence will be filled in at the end of the buffer. If you
-	 *  put more, data will be truncated.
-	 *  Please note that "samples" means individual channels, not
-	 *  sample frames! If you have stereo data and SamplesNeeded
-	 *  is 1, you're writing two SWORDs, not four!
+	 * This is only used for DTYPE_Procedural audio. It's recommended to use USynthComponent base class
+	 * for procedurally generated sound vs overriding this function. If a new component is not feasible,
+	 * consider using USoundWaveProcedural base class vs USoundWave base class since as it implements
+	 * GeneratePCMData for you and you only need to return PCM data.
 	 */
 	virtual int32 GeneratePCMData(uint8* PCMData, const int32 SamplesNeeded) { ensure(false); return 0; }
 
+	/** 
+	* Return the format of the generated PCM data type. Used in audio mixer to allow generating float buffers and avoid unnecessary format conversions. 
+	* This feature is only supported in audio mixer. If your procedural sound wave needs to be used in both audio mixer and old audio engine,
+	* it's best to generate int16 data as old audio engine only supports int16 formats. Or check at runtime if the audio mixer is enabled.
+	* Audio mixer will convert from int16 to float internally.
+	*/
+	virtual Audio::EAudioMixerStreamDataFormat::Type GetGeneratedPCMDataFormat() const { return Audio::EAudioMixerStreamDataFormat::Int16; }
+
 	/** 
 	 * Gets the compressed data size from derived data cache for the specified format
 	 * 
diff --git a/Engine/Source/Runtime/Engine/Classes/Sound/SoundWaveProcedural.h b/Engine/Source/Runtime/Engine/Classes/Sound/SoundWaveProcedural.h
--- a/Engine/Source/Runtime/Engine/Classes/Sound/SoundWaveProcedural.h
+++ b/Engine/Source/Runtime/Engine/Classes/Sound/SoundWaveProcedural.h
@@ -76,4 +76,7 @@ public:
 
 	/** Called when GeneratePCMData is called but not enough data is available. Allows more data to be added, and will try again */
 	FOnSoundWaveProceduralUnderflow OnSoundWaveProceduralUnderflow;
+
+	/** Size in bytes of a single sample of audio in the procedural audio buffer. */
+	int32 SampleByteSize;
 };
diff --git a/Engine/Source/Runtime/Engine/Private/ActiveSound.cpp b/Engine/Source/Runtime/Engine/Private/ActiveSound.cpp
--- a/Engine/Source/Runtime/Engine/Private/ActiveSound.cpp
+++ b/Engine/Source/Runtime/Engine/Private/ActiveSound.cpp
@@ -173,18 +173,22 @@ void FActiveSound::SetAudioComponent(UAudioComponent* Component)
 	AudioComponentUserID = Component->GetAudioComponentUserID();
 	AudioComponentName = Component->GetFName();
 
-	if (Owner)
+	SetOwner(Owner);
+}
+
+void FActiveSound::SetOwner(AActor* Actor)
+{
+	if (Actor)
 	{
-		OwnerID = Owner->GetUniqueID();
-		OwnerName = Owner->GetFName();
+		OwnerID = Actor->GetUniqueID();
+		OwnerName = Actor->GetFName();
 	}
 	else
 	{
 		OwnerID = 0;
 		OwnerName = NAME_None;
 	}
-
-	}
+}
 
 FString FActiveSound::GetAudioComponentName() const
 {
@@ -225,18 +229,36 @@ USoundSubmix* FActiveSound::GetSoundSubmix() const
 
 void FActiveSound::SetSubmixSend(const FSoundSubmixSendInfo& SubmixSendInfo)
 {
-	for (int32 i = 0; i < SoundSubmixSendsOverride.Num(); ++i)
+	// Override send level if the submix send already included in active sound
+	for (FSoundSubmixSendInfo& Info : SoundSubmixSendsOverride)
 	{
-		if (SoundSubmixSendsOverride[i].SoundSubmix == SubmixSendInfo.SoundSubmix)
+		if (Info.SoundSubmix == SubmixSendInfo.SoundSubmix)
 		{
-			SoundSubmixSendsOverride[i].SendLevel = SubmixSendInfo.SendLevel;
+			Info.SendLevel = SubmixSendInfo.SendLevel;
 			return;
 		}
 	}
 
+	// Otherwise, add it to the submix send overrides
 	SoundSubmixSendsOverride.Add(SubmixSendInfo);
 }
 
+void FActiveSound::SetSourceBusSend(const FSoundSourceBusSendInfo& SourceBusSendInfo)
+{
+	// Override send level if the source bus send is already included in active sound
+	for (FSoundSourceBusSendInfo& Info : SoundSourceBusSendsOverride)
+	{
+		if (Info.SoundSourceBus == SourceBusSendInfo.SoundSourceBus)
+		{
+			Info.SendLevel = SourceBusSendInfo.SendLevel;
+			return;
+		}
+	}
+
+	// Otherwise, add it to the source bus send overrides
+	SoundSourceBusSendsOverride.Add(SourceBusSendInfo);
+}
+
 void FActiveSound::GetSoundSubmixSends(TArray<FSoundSubmixSendInfo>& OutSends) const
 {
 	if (Sound)
@@ -245,14 +267,14 @@ void FActiveSound::GetSoundSubmixSends(TArray<FSoundSubmixSendInfo>& OutSends) c
 		Sound->GetSoundSubmixSends(OutSends);
 
 		// Loop through the overrides, which may append or override the existing send
-		for (int32 i = 0; i < SoundSubmixSendsOverride.Num(); ++i)
+		for (const FSoundSubmixSendInfo& SendInfo : SoundSubmixSendsOverride)
 		{
 			bool bOverridden = false;
-			for (int32 j = 0; j < OutSends.Num(); ++j)
+			for (FSoundSubmixSendInfo& OutSendInfo : OutSends)
 			{
-				if (OutSends[j].SoundSubmix == SoundSubmixSendsOverride[j].SoundSubmix)
+				if (OutSendInfo.SoundSubmix == SendInfo.SoundSubmix)
 				{
-					OutSends[j].SendLevel = SoundSubmixSendsOverride[j].SendLevel;
+					OutSendInfo.SendLevel = SendInfo.SendLevel;
 					bOverridden = true;
 					break;
 				}
@@ -260,8 +282,36 @@ void FActiveSound::GetSoundSubmixSends(TArray<FSoundSubmixSendInfo>& OutSends) c
 
 			if (!bOverridden)
 			{
-				// Append
-				OutSends.Add(SoundSubmixSendsOverride[i]);
+				OutSends.Add(SendInfo);
+			}
+		}
+	}
+}
+
+void FActiveSound::GetSoundSourceBusSends(TArray<FSoundSourceBusSendInfo>& OutSends) const
+{
+	if (Sound)
+	{
+		// Get the base sends
+		Sound->GetSoundSourceBusSends(OutSends);
+
+		// Loop through the overrides, which may append or override the existing send
+		for (const FSoundSourceBusSendInfo& SendInfo : SoundSourceBusSendsOverride)
+		{
+			bool bOverridden = false;
+			for (FSoundSourceBusSendInfo& OutSendInfo : OutSends)
+			{
+				if (OutSendInfo.SoundSourceBus == SendInfo.SoundSourceBus)
+				{
+					OutSendInfo.SendLevel = SendInfo.SendLevel;
+					bOverridden = true;
+					break;
+				}
+			}
+
+			if (!bOverridden)
+			{
+				OutSends.Add(SendInfo);
 			}
 		}
 	}
@@ -356,6 +406,9 @@ void FActiveSound::UpdateWaveInstances( TArray<FWaveInstance*> &InWaveInstances,
 	ParseParams.SoundSubmix = GetSoundSubmix();
 	GetSoundSubmixSends(ParseParams.SoundSubmixSends);
 
+	ParseParams.bOutputToBusOnly = Sound->bOutputToBusOnly;
+	GetSoundSourceBusSends(ParseParams.SoundSourceBusSends);
+
 	// Set up the base source effect chain. 
 	ParseParams.SourceEffectChain = Sound->SourceEffectChain;
 
@@ -383,6 +436,12 @@ void FActiveSound::UpdateWaveInstances( TArray<FWaveInstance*> &InWaveInstances,
 		{
 			ApplyAttenuation(ParseParams, *ClosestListenerPtr);
 		}
+		else
+		{
+			// In the case of no attenuation settings, we still want to setup a default send reverb level
+			ParseParams.ReverbSendMethod = EReverbSendMethod::Manual;
+			ParseParams.ManualReverbSendLevel = AudioDevice->GetDefaultReverbSendLevel();
+		}
 
 		// if the closest listener is not the primary one, transform the sound transform so it's panned relative to primary listener position
 		if (ClosestListenerIndex != 0)
@@ -409,7 +468,7 @@ void FActiveSound::UpdateWaveInstances( TArray<FWaveInstance*> &InWaveInstances,
 			VolumeConcurrency = 0.0f;
 			for (const FWaveInstance* WaveInstance : ThisSoundsWaveInstances)
 			{
-				const float WaveInstanceVolume = WaveInstance->GetVolume();
+				const float WaveInstanceVolume = WaveInstance->GetVolumeWithDistanceAttenuation();
 				if (WaveInstanceVolume > VolumeConcurrency)
 				{
 					VolumeConcurrency = WaveInstanceVolume;
@@ -1032,18 +1091,20 @@ void FActiveSound::ApplyAttenuation(FSoundParseParameters& ParseParams, const FL
 		}
 	}
 
-	// Attenuate the volume based on the model
+	// Attenuate the volume based on the model. Note we don't apply the distance attenuation immediately to the sound.
+	// The audio mixer applies distance-based attenuation as a separate stage to feed source audio through source effects and buses.
+	// The old audio engine will scale this together when the wave instance is queried for GetActualVolume.
 	if (Settings->bAttenuate)
 	{
 		if (Settings->AttenuationShape == EAttenuationShape::Sphere)
 		{
 			// Update attenuation data in-case it hasn't been updated
 			AudioDevice->GetAttenuationListenerData(ListenerData, SoundTransform, *Settings, &Listener.Transform);
-			Volume *= Settings->AttenuationEval(ListenerData.AttenuationDistance, Settings->FalloffDistance, FocusDistanceScale);
+			ParseParams.DistanceAttenuation = Settings->AttenuationEval(ListenerData.AttenuationDistance, Settings->FalloffDistance, FocusDistanceScale);
 		}
 		else
 		{
-			Volume *= Settings->Evaluate(SoundTransform, ListenerLocation, FocusDistanceScale);
+			ParseParams.DistanceAttenuation = Settings->Evaluate(SoundTransform, ListenerLocation, FocusDistanceScale);
 		}
 	}
 
diff --git a/Engine/Source/Runtime/Engine/Private/Animation/AnimCompositeBase.cpp b/Engine/Source/Runtime/Engine/Private/Animation/AnimCompositeBase.cpp
--- a/Engine/Source/Runtime/Engine/Private/Animation/AnimCompositeBase.cpp
+++ b/Engine/Source/Runtime/Engine/Private/Animation/AnimCompositeBase.cpp
@@ -163,11 +163,14 @@ void FAnimSegment::GetRootMotionExtractionStepsForTrackRange(TArray<FRootMotionE
 		const float ValidPlayRate = GetValidPlayRate();
 		const float AbsValidPlayRate = FMath::Abs(ValidPlayRate);
 
+		const float StartTrackPositionForSegment = bTrackPlayingBackwards ? FMath::Min(StartTrackPosition, SegmentEndPos) : FMath::Max(StartTrackPosition, SegmentStartPos);
+		const float EndTrackPositionForSegment = bTrackPlayingBackwards ? FMath::Max(EndTrackPosition, SegmentStartPos) : FMath::Min(EndTrackPosition, SegmentEndPos);
+
 		// Get starting position, closest overlap.
-		float AnimStartPosition = ConvertTrackPosToAnimPos(bTrackPlayingBackwards ? FMath::Min(StartTrackPosition, SegmentEndPos) : FMath::Max(StartTrackPosition, SegmentStartPos));
+		float AnimStartPosition = ConvertTrackPosToAnimPos(StartTrackPositionForSegment);
 		AnimStartPosition = FMath::Clamp(AnimStartPosition, AnimStartTime, AnimEndTime);
 		//check( (AnimStartPosition >= AnimStartTime) && (AnimStartPosition <= AnimEndTime) );
-		float TrackTimeToGo = FMath::Abs(EndTrackPosition - StartTrackPosition);
+		float TrackTimeToGo = FMath::Abs(EndTrackPositionForSegment - StartTrackPositionForSegment);
 
 		// The track can be playing backwards and the animation can be playing backwards, so we
 		// need to combine to work out what direction we are traveling through the animation
diff --git a/Engine/Source/Runtime/Engine/Private/Animation/AnimCompressionDerivedData.h b/Engine/Source/Runtime/Engine/Private/Animation/AnimCompressionDerivedData.h
--- a/Engine/Source/Runtime/Engine/Private/Animation/AnimCompressionDerivedData.h
+++ b/Engine/Source/Runtime/Engine/Private/Animation/AnimCompressionDerivedData.h
@@ -44,7 +44,7 @@ public:
 		// This is a version string that mimics the old versioning scheme. If you
 		// want to bump this version, generate a new guid using VS->Tools->Create GUID and
 		// return it here. Ex.
-		return TEXT("486652B60D9C4989BFD91BDF463DB28D");
+		return TEXT("44870622F1374957B3252BC65DDCFA9F");
 	}
 
 	virtual FString GetPluginSpecificCacheKeySuffix() const override;
diff --git a/Engine/Source/Runtime/Engine/Private/Animation/AnimInstance.cpp b/Engine/Source/Runtime/Engine/Private/Animation/AnimInstance.cpp
--- a/Engine/Source/Runtime/Engine/Private/Animation/AnimInstance.cpp
+++ b/Engine/Source/Runtime/Engine/Private/Animation/AnimInstance.cpp
@@ -473,7 +473,8 @@ void UAnimInstance::PostUpdateAnimation()
 
 	Proxy.PostUpdate(this);
 
-	FRootMotionMovementParams& ExtractedRootMotion = Proxy.GetExtractedRootMotion();
+	ExtractedRootMotion.Accumulate(Proxy.GetExtractedRootMotion());
+	Proxy.GetExtractedRootMotion().Clear();
 
 	// blend in any montage-blended root motion that we now have correct weights for
 	for(const FQueuedRootMotionBlend& RootMotionBlend : RootMotionBlendQueue)
@@ -1011,9 +1012,9 @@ void UAnimInstance::RecalcRequiredBones()
 	}
 }
 
-void UAnimInstance::RecalcRequiredCurves(bool bDisableAnimCurves)
+void UAnimInstance::RecalcRequiredCurves(const FCurveEvaluationOption& CurveEvalOption)
 {
-	GetProxyOnGameThread<FAnimInstanceProxy>().RecalcRequiredCurves(bDisableAnimCurves);
+	GetProxyOnGameThread<FAnimInstanceProxy>().RecalcRequiredCurves(CurveEvalOption);
 }
 
 void UAnimInstance::Serialize(FArchive& Ar)
@@ -1162,12 +1163,15 @@ void UAnimInstance::UpdateCurves(const FBlendedHeapCurve& InCurve)
 	SCOPE_CYCLE_COUNTER(STAT_UpdateCurves);
 
 	FAnimInstanceProxy& Proxy = GetProxyOnGameThread<FAnimInstanceProxy>();
+	USkeletalMeshComponent* SkelMeshComp = GetSkelMeshComponent();
 
 	//Track material params we set last time round so we can clear them if they aren't set again.
 	MaterialParamatersToClear.Reset();
 	for(auto Iter = AnimationCurves[(uint8)EAnimCurveType::MaterialCurve].CreateConstIterator(); Iter; ++Iter)
 	{
-		if(Iter.Value() != 0.0f)
+		// when reset, we go back to default value
+		float DefaultValue = SkelMeshComp->GetScalarParameterDefaultValue(Iter.Key());
+		if(Iter.Value() != DefaultValue)
 		{
 			MaterialParamatersToClear.Add(Iter.Key());
 		}
@@ -1193,7 +1197,6 @@ void UAnimInstance::UpdateCurves(const FBlendedHeapCurve& InCurve)
 	//   - Make a copy of MaterialParametersToClear as it will be modified by AddCurveValue
 	//	 - When clear, we have to make sure to add directly to the material curve list because 
 	//   - sometimes they don't have the flag anymore, so we can't just call AddCurveValue
-	USkeletalMeshComponent* SkelMeshComp = GetSkelMeshComponent();
 	TArray<FName> ParamsToClearCopy = MaterialParamatersToClear;
 	for(int i = 0; i < ParamsToClearCopy.Num(); ++i)
 	{
@@ -1437,7 +1440,7 @@ void UAnimInstance::Montage_Advance(float DeltaSeconds)
 			FRootMotionMovementParams* RootMotionParams = nullptr;
 			if (bExtractRootMotion)
 			{
-				RootMotionParams = (RootMotionMode != ERootMotionMode::IgnoreRootMotion) ? &GetProxyOnGameThread<FAnimInstanceProxy>().GetExtractedRootMotion() : &LocalExtractedRootMotion;
+				RootMotionParams = (RootMotionMode != ERootMotionMode::IgnoreRootMotion) ? &ExtractedRootMotion : &LocalExtractedRootMotion;
 			}
 
 			MontageInstance->MontageSync_PreUpdate();
@@ -2391,13 +2394,13 @@ FRootMotionMovementParams UAnimInstance::ConsumeExtractedRootMotion(float Alpha)
 	}
 	else if (Alpha > (1.f - ZERO_ANIMWEIGHT_THRESH))
 	{
-		FRootMotionMovementParams RootMotion = GetProxyOnGameThread<FAnimInstanceProxy>().GetExtractedRootMotion();
-		GetProxyOnGameThread<FAnimInstanceProxy>().GetExtractedRootMotion().Clear();
+		FRootMotionMovementParams RootMotion = ExtractedRootMotion;
+		ExtractedRootMotion.Clear();
 		return RootMotion;
 	}
 	else
 	{
-		return GetProxyOnGameThread<FAnimInstanceProxy>().GetExtractedRootMotion().ConsumeRootMotion(Alpha);
+		return ExtractedRootMotion.ConsumeRootMotion(Alpha);
 	}
 }
 
diff --git a/Engine/Source/Runtime/Engine/Private/Animation/AnimInstanceProxy.cpp b/Engine/Source/Runtime/Engine/Private/Animation/AnimInstanceProxy.cpp
--- a/Engine/Source/Runtime/Engine/Private/Animation/AnimInstanceProxy.cpp
+++ b/Engine/Source/Runtime/Engine/Private/Animation/AnimInstanceProxy.cpp
@@ -836,7 +836,7 @@ FAnimNode_Base* FAnimInstanceProxy::GetNodeFromIndexUntyped(int32 NodeIdx, UScri
 
 void FAnimInstanceProxy::RecalcRequiredBones(USkeletalMeshComponent* Component, UObject* Asset)
 {
-	RequiredBones.InitializeTo(Component->RequiredBones, Component->GetDisableAnimCurves(), *Asset);
+	RequiredBones.InitializeTo(Component->RequiredBones, FCurveEvaluationOption(Component->GetAllowedAnimCurveEvaluate(), &Component->GetDisallowedAnimCurvesEvaluation(), Component->PredictedLODLevel), *Asset);
 
 	// If there is a ref pose override, we want to replace ref pose in RequiredBones
 	const FSkelMeshRefPoseOverride* RefPoseOverride = Component->GetRefPoseOverride();
@@ -880,9 +880,9 @@ void FAnimInstanceProxy::RecalcRequiredBones(USkeletalMeshComponent* Component,
 	bBoneCachesInvalidated = true;
 }
 
-void FAnimInstanceProxy::RecalcRequiredCurves(bool bDisableAnimCurves)
+void FAnimInstanceProxy::RecalcRequiredCurves(const FCurveEvaluationOption& CurveEvalOption)
 {
-	RequiredBones.CacheRequiredAnimCurveUids(bDisableAnimCurves);
+	RequiredBones.CacheRequiredAnimCurveUids(CurveEvalOption);
 }
 
 void FAnimInstanceProxy::UpdateAnimation()
diff --git a/Engine/Source/Runtime/Engine/Private/Animation/AnimNode_StateMachine.cpp b/Engine/Source/Runtime/Engine/Private/Animation/AnimNode_StateMachine.cpp
--- a/Engine/Source/Runtime/Engine/Private/Animation/AnimNode_StateMachine.cpp
+++ b/Engine/Source/Runtime/Engine/Private/Animation/AnimNode_StateMachine.cpp
@@ -339,7 +339,7 @@ void FAnimNode_StateMachine::Update_AnyThread(const FAnimationUpdateContext& Con
 	Context.AnimInstanceProxy->RecordMachineWeight(StateMachineIndexInClass, Context.GetFinalBlendWeight());
 
 	// If we just became relevant and haven't been initialized yet, then reinitialize state machine.
-	if (!bFirstUpdate && (UpdateCounter.Get() != INDEX_NONE) && !UpdateCounter.WasSynchronizedInTheLastFrame(Context.AnimInstanceProxy->GetUpdateCounter()) && (CVarAnimStateMachineRelevancyReset.GetValueOnAnyThread() == 1))
+	if (!bFirstUpdate && bReinitializeOnBecomingRelevant &&(UpdateCounter.Get() != INDEX_NONE) && !UpdateCounter.WasSynchronizedInTheLastFrame(Context.AnimInstanceProxy->GetUpdateCounter()) && (CVarAnimStateMachineRelevancyReset.GetValueOnAnyThread() == 1))
 	{
 		FAnimationInitializeContext InitializationContext(Context.AnimInstanceProxy);
 		Initialize_AnyThread(InitializationContext);
diff --git a/Engine/Source/Runtime/Engine/Private/Animation/AnimSequence.cpp b/Engine/Source/Runtime/Engine/Private/Animation/AnimSequence.cpp
--- a/Engine/Source/Runtime/Engine/Private/Animation/AnimSequence.cpp
+++ b/Engine/Source/Runtime/Engine/Private/Animation/AnimSequence.cpp
@@ -543,23 +543,17 @@ void UAnimSequence::PostLoad()
 		}
 	}
 
-	USkeleton* CurrentSkeleton = GetSkeleton();
-
-	if (CurrentSkeleton)
+	if (USkeleton* CurrentSkeleton = GetSkeleton())
 	{
-		VerifyCurveNames<FFloatCurve>(CurrentSkeleton, USkeleton::AnimCurveMappingName, CompressedCurveData.FloatCurves);
-	}
+		VerifyCurveNames<FFloatCurve>(*CurrentSkeleton, USkeleton::AnimCurveMappingName, CompressedCurveData.FloatCurves);
 
 #if WITH_EDITOR
-	if (CurrentSkeleton)
-	{
-		// Get the name mapping object for curves
-		for(FTransformCurve& Curve : RawCurveData.TransformCurves)
-		{
-			CurrentSkeleton->VerifySmartName(USkeleton::AnimTrackCurveMappingName, Curve.Name);
-		}
+		VerifyCurveNames<FTransformCurve>(*CurrentSkeleton, USkeleton::AnimTrackCurveMappingName, RawCurveData.TransformCurves);
+#endif
 	}
 
+#if WITH_EDITOR
+
 	// Compressed curve flags are not authoritative (they come from the DDC). Keep them up to date with
 	// actual anim flags 
 	for (FFloatCurve& Curve : RawCurveData.FloatCurves)
@@ -2069,7 +2063,8 @@ void UAnimSequence::RequestAnimCompression(bool bAsyncCompression, bool bAllowAl
 void UAnimSequence::RequestAnimCompression(bool bAsyncCompression, TSharedPtr<FAnimCompressContext> CompressContext)
 {
 #if WITH_EDITOR
-	if (GetSkeleton() == nullptr)
+	USkeleton* CurrentSkeleton = GetSkeleton();
+	if (CurrentSkeleton == nullptr)
 	{
 		bUseRawDataOnly = true;
 		return;
@@ -2098,7 +2093,7 @@ void UAnimSequence::RequestAnimCompression(bool bAsyncCompression, TSharedPtr<FA
 	const bool bDoCompressionInPlace = FUObjectThreadContext::Get().IsRoutingPostLoad;
 
 	// Need to make sure this is up to date.
-	VerifyCurveNames<FFloatCurve>(GetSkeleton(), USkeleton::AnimCurveMappingName, RawCurveData.FloatCurves);
+	VerifyCurveNames<FFloatCurve>(*CurrentSkeleton, USkeleton::AnimCurveMappingName, RawCurveData.FloatCurves);
 
 	if (bAsyncCompression)
 	{
@@ -2122,6 +2117,12 @@ void UAnimSequence::RequestAnimCompression(bool bAsyncCompression, TSharedPtr<FA
 			{
 				GetDerivedDataCacheRef().GetSynchronous(AnimCompressor, OutData);
 			}
+			else
+			{
+				// If we dont perform compression we need to clean this up
+				delete AnimCompressor;
+				AnimCompressor = nullptr;
+			}
 		}
 
 		if (bUseRawDataOnly && OutData.Num() > 0)
@@ -2231,15 +2232,18 @@ void UAnimSequence::SerializeCompressedData(FArchive& Ar, bool bDDCData)
 
 		if (Ar.IsLoading())
 		{
-			VerifyCurveNames<FFloatCurve>(GetSkeleton(), USkeleton::AnimCurveMappingName, CompressedCurveData.FloatCurves);
-			bUseRawDataOnly = !IsCompressedDataValid();
-			ensureMsgf(!bUseRawDataOnly, TEXT("Anim Compression failed for Sequence '%s' Guid:%s CompressedDebugData:\n\tOriginal Anim:%s\n\tAdditiveSetting:%i\n\tCompression Scheme:%s\n\tRawDataGuid:%s"),
-				*GetFullName(),
-				*RawDataGuid.ToString(),
-				*DebugData.FullName,
-				DebugData.AdditiveSetting,
-				*DebugData.CompressionSchemeName,
-				*DebugData.RawDataGuid.ToString());
+			if(USkeleton* CurrentSkeleton = GetSkeleton())
+			{
+				VerifyCurveNames<FFloatCurve>(*CurrentSkeleton, USkeleton::AnimCurveMappingName, CompressedCurveData.FloatCurves);
+				bUseRawDataOnly = !IsCompressedDataValid();
+				ensureMsgf(!bUseRawDataOnly, TEXT("Anim Compression failed for Sequence '%s' Guid:%s CompressedDebugData:\n\tOriginal Anim:%s\n\tAdditiveSetting:%i\n\tCompression Scheme:%s\n\tRawDataGuid:%s"),
+					*GetFullName(),
+					*RawDataGuid.ToString(),
+					*DebugData.FullName,
+					DebugData.AdditiveSetting,
+					*DebugData.CompressionSchemeName,
+					*DebugData.RawDataGuid.ToString());
+			}
 		}
 	}
 #endif
@@ -2367,7 +2371,7 @@ public:
 			RequiredBoneIndexArray[BoneIndex] = BoneIndex;
 		}
 
-		RequiredBones.InitializeTo(RequiredBoneIndexArray, false, *MySkeleton);
+		RequiredBones.InitializeTo(RequiredBoneIndexArray, FCurveEvaluationOption(true), *MySkeleton);
 	}
 
 };
@@ -2461,19 +2465,18 @@ void UAnimSequence::BakeOutAdditiveIntoRawData()
 		return; // Nothing to do
 	}
 
+	USkeleton* MySkeleton = GetSkeleton();
+	check(MySkeleton);
+
 	if (RefPoseSeq && RefPoseSeq->HasAnyFlags(EObjectFlags::RF_NeedPostLoad))
 	{
-		RefPoseSeq->VerifyCurveNames<FFloatCurve>(GetSkeleton(), USkeleton::AnimCurveMappingName, RefPoseSeq->RawCurveData.FloatCurves);
+		RefPoseSeq->VerifyCurveNames<FFloatCurve>(*MySkeleton, USkeleton::AnimCurveMappingName, RefPoseSeq->RawCurveData.FloatCurves);
 	}
 
 	FMemMark Mark(FMemStack::Get());
 
 	FByFramePoseEvalContext EvalContext(this);
 
-	USkeleton* MySkeleton = GetSkeleton();
-	check(MySkeleton);
-
-
 	//New raw data
 	FRawCurveTracks NewCurveTracks;
 
@@ -4201,7 +4204,7 @@ void UAnimSequence::BakeTrackCurvesToRawAnimation()
 		if(!NameMapping)
 		{
 			// if no name mapping is found but curve exists, we should verify curve namex
-			VerifyCurveNames<FTransformCurve>(CurSkeleton, USkeleton::AnimTrackCurveMappingName, RawCurveData.TransformCurves);
+			VerifyCurveNames<FTransformCurve>(*CurSkeleton, USkeleton::AnimTrackCurveMappingName, RawCurveData.TransformCurves);
 			NameMapping = CurSkeleton->GetSmartNameContainer(USkeleton::AnimTrackCurveMappingName);
 		}
 		
diff --git a/Engine/Source/Runtime/Engine/Private/Animation/AnimSequenceBase.cpp b/Engine/Source/Runtime/Engine/Private/Animation/AnimSequenceBase.cpp
--- a/Engine/Source/Runtime/Engine/Private/Animation/AnimSequenceBase.cpp
+++ b/Engine/Source/Runtime/Engine/Private/Animation/AnimSequenceBase.cpp
@@ -48,21 +48,10 @@ void UAnimSequenceBase::PostLoad()
 
 	if(USkeleton* MySkeleton = GetSkeleton())
 	{
-		// Fix up the existing curves to work with smartnames
-		if(GetLinkerUE4Version() < VER_UE4_SKELETON_ADD_SMARTNAMES)
-		{
-			for(FFloatCurve& Curve : RawCurveData.FloatCurves)
-			{
-				MySkeleton->VerifySmartName(USkeleton::AnimCurveMappingName, Curve.Name);
-			}
-		}
-		else
-		{
-			VerifyCurveNames<FFloatCurve>(MySkeleton, USkeleton::AnimCurveMappingName, RawCurveData.FloatCurves);
-		}
+		VerifyCurveNames<FFloatCurve>(*MySkeleton, USkeleton::AnimCurveMappingName, RawCurveData.FloatCurves);
 
 #if WITH_EDITOR
-		VerifyCurveNames<FTransformCurve>(MySkeleton, USkeleton::AnimTrackCurveMappingName, RawCurveData.TransformCurves);
+		VerifyCurveNames<FTransformCurve>(*MySkeleton, USkeleton::AnimTrackCurveMappingName, RawCurveData.TransformCurves);
 #endif
 
 		// this should continue to add if skeleton hasn't been saved either 
@@ -561,13 +550,13 @@ void UAnimSequenceBase::GetAssetRegistryTags(TArray<FAssetRegistryTag>& OutTags)
 	// The blank list is necessary when we attempt to delete a curve so
 	// an old asset can be detected from its asset data so we load as few
 	// as possible.
-	FString CurveIdList;
+	FString CurveNameList;
 
 	for(const FFloatCurve& Curve : RawCurveData.FloatCurves)
 	{
-		CurveIdList += FString::Printf(TEXT("%u%s"), Curve.Name.UID, *USkeleton::CurveTagDelimiter);
+		CurveNameList += FString::Printf(TEXT("%s%s"), *Curve.Name.DisplayName.ToString(), *USkeleton::CurveTagDelimiter);
 	}
-	OutTags.Add(FAssetRegistryTag(USkeleton::CurveTag, CurveIdList, FAssetRegistryTag::TT_Hidden));
+	OutTags.Add(FAssetRegistryTag(USkeleton::CurveNameTag, CurveNameList, FAssetRegistryTag::TT_Hidden));
 }
 
 uint8* UAnimSequenceBase::FindNotifyPropertyData(int32 NotifyIndex, UArrayProperty*& ArrayProperty)
diff --git a/Engine/Source/Runtime/Engine/Private/Animation/AnimSingleNodeInstance.cpp b/Engine/Source/Runtime/Engine/Private/Animation/AnimSingleNodeInstance.cpp
--- a/Engine/Source/Runtime/Engine/Private/Animation/AnimSingleNodeInstance.cpp
+++ b/Engine/Source/Runtime/Engine/Private/Animation/AnimSingleNodeInstance.cpp
@@ -89,6 +89,11 @@ void UAnimSingleNodeInstance::SetAnimationAsset(class UAnimationAsset* NewAsset,
 	}
 }
 
+UAnimationAsset* UAnimSingleNodeInstance::GetAnimationAsset() const
+{
+	return CurrentAsset;
+}
+
 void UAnimSingleNodeInstance::SetPreviewCurveOverride(const FName& PoseName, float Value, bool bRemoveIfZero)
 {
 	GetProxyOnGameThread<FAnimSingleNodeInstanceProxy>().SetPreviewCurveOverride(PoseName, Value, bRemoveIfZero);
diff --git a/Engine/Source/Runtime/Engine/Private/Animation/AnimationRuntime.cpp b/Engine/Source/Runtime/Engine/Private/Animation/AnimationRuntime.cpp
--- a/Engine/Source/Runtime/Engine/Private/Animation/AnimationRuntime.cpp
+++ b/Engine/Source/Runtime/Engine/Private/Animation/AnimationRuntime.cpp
@@ -658,11 +658,11 @@ void FAnimationRuntime::AccumulateAdditivePose(FCompactPose& BasePose, const FCo
 {
 	if (AdditiveType == AAT_RotationOffsetMeshSpace)
 	{
-		AccumulateMeshSpaceRotationAdditiveToLocalPose(BasePose, AdditivePose, BaseCurve, AdditiveCurve, Weight);
+		AccumulateMeshSpaceRotationAdditiveToLocalPoseInternal(BasePose, AdditivePose, Weight);
 	}
 	else
 	{
-		AccumulateLocalSpaceAdditivePose(BasePose, AdditivePose, BaseCurve, AdditiveCurve, Weight);
+		AccumulateLocalSpaceAdditivePoseInternal(BasePose, AdditivePose, Weight);
 	}
 
 	// if curve exists, accumulate with the weight, 
@@ -671,7 +671,7 @@ void FAnimationRuntime::AccumulateAdditivePose(FCompactPose& BasePose, const FCo
 	BasePose.NormalizeRotations();
 }
 
-void FAnimationRuntime::AccumulateLocalSpaceAdditivePose(FCompactPose& BasePose, const FCompactPose& AdditivePose, FBlendedCurve& BaseCurve, const FBlendedCurve& AdditiveCurve, float Weight)
+void FAnimationRuntime::AccumulateLocalSpaceAdditivePoseInternal(FCompactPose& BasePose, const FCompactPose& AdditivePose, float Weight)
 {
 	if (FAnimWeight::IsRelevant(Weight))
 	{
@@ -697,7 +697,7 @@ void FAnimationRuntime::AccumulateLocalSpaceAdditivePose(FCompactPose& BasePose,
 	}
 }
 
-void FAnimationRuntime::AccumulateMeshSpaceRotationAdditiveToLocalPose(FCompactPose& BasePose, const FCompactPose& MeshSpaceRotationAdditive, FBlendedCurve& BaseCurve, const FBlendedCurve& AdditiveCurve, float Weight)
+void FAnimationRuntime::AccumulateMeshSpaceRotationAdditiveToLocalPoseInternal(FCompactPose& BasePose, const FCompactPose& MeshSpaceRotationAdditive, float Weight)
 {
 	SCOPE_CYCLE_COUNTER(STAT_AccumulateMeshSpaceRotAdditiveToLocalPose);
 
@@ -707,7 +707,7 @@ void FAnimationRuntime::AccumulateMeshSpaceRotationAdditiveToLocalPose(FCompactP
 		FAnimationRuntime::ConvertPoseToMeshRotation(BasePose);
 
 		// Add MeshSpaceRotAdditive to it
-		FAnimationRuntime::AccumulateLocalSpaceAdditivePose(BasePose, MeshSpaceRotationAdditive, BaseCurve, AdditiveCurve, Weight);
+		FAnimationRuntime::AccumulateLocalSpaceAdditivePoseInternal(BasePose, MeshSpaceRotationAdditive, Weight);
 
 		// Convert back to local space
 		FAnimationRuntime::ConvertMeshRotationPoseToLocalSpace(BasePose);
diff --git a/Engine/Source/Runtime/Engine/Private/Animation/PoseAsset.cpp b/Engine/Source/Runtime/Engine/Private/Animation/PoseAsset.cpp
--- a/Engine/Source/Runtime/Engine/Private/Animation/PoseAsset.cpp
+++ b/Engine/Source/Runtime/Engine/Private/Animation/PoseAsset.cpp
@@ -626,6 +626,27 @@ void UPoseAsset::GetAssetRegistryTags(TArray<FAssetRegistryTag>& OutTags) const
 
 	// Number of poses
 	OutTags.Add(FAssetRegistryTag("Poses", FString::FromInt(GetNumPoses()), FAssetRegistryTag::TT_Numerical));
+#if WITH_EDITOR
+	TArray<FName> Names;
+	Names.Reserve(PoseContainer.PoseNames.Num() + PoseContainer.Curves.Num());
+
+	for (const FSmartName& SmartName : PoseContainer.PoseNames)
+	{
+		Names.Add(SmartName.DisplayName);
+	}
+
+	for (const FAnimCurveBase& Curve : PoseContainer.Curves)
+	{
+		Names.AddUnique(Curve.Name.DisplayName);
+	}
+
+	FString PoseNameList;
+	for(const FName& Name : Names)
+	{
+		PoseNameList += FString::Printf(TEXT("%s%s"), *Name.ToString(), *USkeleton::CurveTagDelimiter);
+	}
+	OutTags.Add(FAssetRegistryTag(USkeleton::CurveNameTag, PoseNameList, FAssetRegistryTag::TT_Hidden)); //write pose names as curve tag as they use 
+#endif
 }
 
 int32 UPoseAsset::GetNumPoses() const
@@ -890,6 +911,33 @@ void UPoseAsset::Reinitialize()
 	BasePoseIndex = INDEX_NONE;
 }
 
+void UPoseAsset::RenameSmartName(const FName& InOriginalName, const FName& InNewName)
+{
+	for (FSmartName SmartName : PoseContainer.PoseNames)
+	{
+		if (SmartName.DisplayName == InOriginalName)
+		{
+			SmartName.DisplayName = InNewName;
+			break;
+		}
+	}
+
+	for (FAnimCurveBase& Curve : PoseContainer.Curves)
+	{
+		if (Curve.Name.DisplayName == InOriginalName)
+		{
+			Curve.Name.DisplayName = InNewName;
+			break;
+		}
+	}
+}
+
+void UPoseAsset::RemoveSmartNames(const TArray<FName>& InNamesToRemove)
+{
+	DeletePoses(InNamesToRemove);
+	DeleteCurves(InNamesToRemove);
+}
+
 void UPoseAsset::CreatePoseFromAnimation(class UAnimSequence* AnimSequence, const TArray<FSmartName>* InPoseNames/*== nullptr*/)
 {
 	if (AnimSequence)
@@ -1013,7 +1061,6 @@ void UPoseAsset::UpdatePoseFromAnimation(class UAnimSequence* AnimSequence)
 		OnPoseListChanged.Broadcast();
 	}
 }
-#endif // WITH_EDITOR
 
 bool UPoseAsset::ModifyPoseName(FName OldPoseName, FName NewPoseName, const SmartName::UID_Type* NewUID)
 {
@@ -1028,8 +1075,7 @@ bool UPoseAsset::ModifyPoseName(FName OldPoseName, FName NewPoseName, const Smar
 	FSmartName OldPoseSmartName;
 	ensureAlways(MySkeleton->GetSmartNameByName(USkeleton::AnimCurveMappingName, OldPoseName, OldPoseSmartName));
 
-	FPoseData* PoseData = PoseContainer.FindPoseData(OldPoseSmartName);
-	if (PoseData)
+	if (FPoseData* PoseData = PoseContainer.FindPoseData(OldPoseSmartName))
 	{
 		FSmartName NewPoseSmartName;
 		if (NewUID)
@@ -1038,11 +1084,7 @@ bool UPoseAsset::ModifyPoseName(FName OldPoseName, FName NewPoseName, const Smar
 		}
 		else
 		{
-			// we're renaming current one
-			MySkeleton->Modify();
-			MySkeleton->RenameSmartName(USkeleton::AnimCurveMappingName, OldPoseSmartName.UID, NewPoseName);
-			NewPoseSmartName = OldPoseSmartName;
-			NewPoseSmartName.DisplayName = NewPoseName;
+			MySkeleton->AddSmartNameAndModify(USkeleton::AnimCurveMappingName, NewPoseName, NewPoseSmartName);
 		}
 
 		PoseContainer.RenamePose(OldPoseSmartName, NewPoseSmartName);
@@ -1053,6 +1095,7 @@ bool UPoseAsset::ModifyPoseName(FName OldPoseName, FName NewPoseName, const Smar
 
 	return false;
 }
+#endif // WITH_EDITOR
 
 int32 UPoseAsset::DeletePoses(TArray<FName> PoseNamesToDelete)
 {
diff --git a/Engine/Source/Runtime/Engine/Private/Animation/Skeleton.cpp b/Engine/Source/Runtime/Engine/Private/Animation/Skeleton.cpp
--- a/Engine/Source/Runtime/Engine/Private/Animation/Skeleton.cpp
+++ b/Engine/Source/Runtime/Engine/Private/Animation/Skeleton.cpp
@@ -35,7 +35,7 @@
 const FName USkeleton::AnimNotifyTag = FName(TEXT("AnimNotifyList"));
 const FString USkeleton::AnimNotifyTagDelimiter = TEXT(";");
 
-const FName USkeleton::CurveTag = FName(TEXT("CurveUIDList"));
+const FName USkeleton::CurveNameTag = FName(TEXT("CurveNameList"));
 const FString USkeleton::CurveTagDelimiter = TEXT(";");
 
 const FName USkeleton::RigTag = FName(TEXT("Rig"));
@@ -117,7 +117,11 @@ void USkeleton::PostLoad()
 	IncreaseAnimCurveUidVersion();
 
 	// refresh linked bone indices
-	SmartNames.InitializeCurveMetaData(this);
+	FSmartNameMapping* CurveMappingTable = SmartNames.GetContainerInternal(USkeleton::AnimCurveMappingName);
+	if (CurveMappingTable)
+	{
+		CurveMappingTable->InitializeCurveMetaData(this);
+	}
 }
 
 void USkeleton::PostDuplicate(bool bDuplicateForPIE)
@@ -1192,27 +1196,17 @@ void USkeleton::RenameSlotName(const FName& OldName, const FName& NewName)
 
 bool USkeleton::AddSmartNameAndModify(FName ContainerName, FName NewDisplayName, FSmartName& NewName)
 {
-	bool Successful = false;
-	FSmartNameMapping* RequestedMapping = GetOrAddSmartNameContainer(ContainerName);
-	if (RequestedMapping)
+	NewName.DisplayName = NewDisplayName;
+	const bool bAdded = VerifySmartNameInternal(ContainerName, NewName);
+
+	if(bAdded)
 	{
-		if (RequestedMapping->FindSmartName(NewDisplayName, NewName))
-		{
-			Successful = true;
-		}
-		else
-		{
-			// if it didn't find, mark modify
-			Modify(true);
-			if (RequestedMapping->FindOrAddSmartName(NewDisplayName, NewName))
-			{
-				Successful = true;
-				IncreaseAnimCurveUidVersion();
-			}
-		}
+		IncreaseAnimCurveUidVersion();
 	}
-	return Successful;
+
+	return bAdded;
 }
+
 bool USkeleton::RenameSmartnameAndModify(FName ContainerName, SmartName::UID_Type Uid, FName NewName)
 {
 	bool Successful = false;
@@ -1244,24 +1238,27 @@ void USkeleton::RemoveSmartnameAndModify(FName ContainerName, SmartName::UID_Typ
 	}
 }
 
-void USkeleton::RemoveSmartnamesAndModify(FName ContainerName, const TArray<SmartName::UID_Type>& Uids)
+void USkeleton::RemoveSmartnamesAndModify(FName ContainerName, const TArray<FName>& Names)
 {
 	FSmartNameMapping* RequestedMapping = SmartNames.GetContainerInternal(ContainerName);
 	if (RequestedMapping)
 	{
 		bool bModified = false;
-		for (const SmartName::UID_Type& UID : Uids)
+		for (const FName CurveName : Names)
 		{
-			if (RequestedMapping->Exists(UID))
+			if (RequestedMapping->Exists(CurveName))
 			{
-				RequestedMapping->Remove(UID);
+				if (!bModified)
+				{
+					Modify();
+				}
+				RequestedMapping->Remove(CurveName);
 				bModified = true;
 			}
 		}
 
 		if (bModified)
 		{
-			Modify();
 			IncreaseAnimCurveUidVersion();
 		}
 	}
@@ -1321,75 +1318,27 @@ bool USkeleton::FillSmartNameByDisplayName(FSmartNameMapping* Mapping, const FNa
 	{
 		OutSmartName.DisplayName = DisplayName;
 
-#if WITH_EDITORONLY_DATA
-		// if same guid, this is same
-		if (SkeletonName.Guid == OutSmartName.Guid)
-		{
-			OutSmartName.UID = SkeletonName.UID;
-			return true;
-		}
-		// else, take Skeleton Guid, we don't allow same name with different Guid
-		else
-		{
-			OutSmartName.Guid = SkeletonName.Guid;
-			OutSmartName.UID = SkeletonName.UID;
-			return true;
-		}
-#else
 		// if not editor, we assume name is always correct
 		OutSmartName.UID = SkeletonName.UID;
 		return true;
-#endif // WITH_EDITORONLY_DATA
 	}
 
 	return false;
 }
-//@todo: this does prioritize name over UID since UID doesn't work well
-// until we have GUID, we'll prioritize name over UID
+
 bool USkeleton::VerifySmartNameInternal(const FName&  ContainerName, FSmartName& InOutSmartName)
 {
 	FSmartNameMapping* Mapping = GetOrAddSmartNameContainer(ContainerName);
 	if (Mapping != nullptr)
 	{
-		// make a copy just in case we change by accident
-		FName DisplayName = InOutSmartName.DisplayName;
-
-#if WITH_EDITOR
-		// if I find the name, fill up the data
-		// look for same guid, the name might have been changed
-		if (Mapping->GetNameByGuid(InOutSmartName.Guid, DisplayName))
-		{
-			ensureAlways(FillSmartNameByDisplayName(Mapping, DisplayName, InOutSmartName));
-		}
-		else if (FillSmartNameByDisplayName(Mapping, DisplayName, InOutSmartName) == false)
+		if (!Mapping->FindSmartName(InOutSmartName.DisplayName, InOutSmartName))
 		{
-			// look for same guid, the name might have been changed
-			if (Mapping->GetNameByGuid(InOutSmartName.Guid, DisplayName))
-			{
-				ensureAlways(FillSmartNameByDisplayName(Mapping, DisplayName, InOutSmartName));
-			}
-			else if (InOutSmartName.IsValid())
-			{
-				// this is only case where we add new one
-				Modify();
-				ensureAlways(Mapping->AddSmartName(InOutSmartName));
-				return true;
-			}
-			else
-			{
-				// this is only case where we add new one
-				Modify();
-				Mapping->FindOrAddSmartName(InOutSmartName.DisplayName, InOutSmartName);
-				return true;
-			}
+#if WITH_EDITOR
+			Modify();
+#endif
+			InOutSmartName = Mapping->AddName(InOutSmartName.DisplayName);
+			return true;
 		}
-#else
-		// if cooking didn't save the skeleton package properly
-		// meaning it can be loaded by animations but not saved together
-		// then later on, it gets loaded, and it's saved with empty array
-		// so it can fail to load the name, so we'll have to add it manually
-		Mapping->FindOrAddSmartName(DisplayName, InOutSmartName.UID);
-#endif // WITH_EDITOR
 	}
 
 	return false;
@@ -1441,20 +1390,6 @@ void USkeleton::RegenerateVirtualBoneGuid()
 	check(VirtualBoneGuid.IsValid());
 }
 
-bool USkeleton::RenameSmartName(FName ContainerName, const SmartName::UID_Type& Uid, FName NewName)
-{
-	FSmartNameMapping* Mapping = SmartNames.GetContainerInternal(ContainerName);
-	if (Mapping->Exists(Uid))
-	{
-		Modify();
-		IncreaseAnimCurveUidVersion();
-		Mapping->Rename(Uid, NewName);
-		return true;
-	}
-
-	return false;
-}
-
 void USkeleton::IncreaseAnimCurveUidVersion()
 {
 	// this will be checked by skeletalmeshcomponent and if it's same, it won't care. If it's different, it will rebuild UID list
diff --git a/Engine/Source/Runtime/Engine/Private/Animation/SmartName.cpp b/Engine/Source/Runtime/Engine/Private/Animation/SmartName.cpp
--- a/Engine/Source/Runtime/Engine/Private/Animation/SmartName.cpp
+++ b/Engine/Source/Runtime/Engine/Private/Animation/SmartName.cpp
@@ -11,109 +11,49 @@
 //
 ///////////////////////////////////////////////////////////////////////
 FSmartNameMapping::FSmartNameMapping()
-: NextUid(0)
 {}
 
-bool FSmartNameMapping::AddOrFindName(FName Name, SmartName::UID_Type& OutUid, FGuid& OutGuid)
+FSmartName FSmartNameMapping::AddName(FName Name)
 {
-	check(Name.IsValid());
+	// Make sure we are not trying to do an invalid add
+	check(Name.IsValid() && !CurveNameList.Contains(Name));
 
-	// Check for SmartName::UID_Type overflow
-	const SmartName::UID_Type* ExistingUid = UidMap.FindKey(Name);
-	const FGuid* ExistingGuid = GuidMap.Find(Name);
+	// Make sure we didn't reach the UID limit
+	check(CurveNameList.Num() < (SmartName::MaxUID-1));
 
-	// make sure they both exists and same 
-	check(!!ExistingUid == !!ExistingGuid);
-	if(ExistingUid && ExistingGuid->IsValid())
-	{
-		// Already present in the list
-		OutUid = *ExistingUid;
-		OutGuid = *ExistingGuid;
-		return false;
-	}
-
-	// make sure we didn't reach till end
-	OutGuid = FGuid::NewGuid();
-	return AddName(Name, OutUid, OutGuid);
-}
-
-bool FSmartNameMapping::AddName(FName Name, SmartName::UID_Type& OutUid, const FGuid& InGuid)
-{
-	check(Name.IsValid());
-	check(InGuid.IsValid());
-
-	const SmartName::UID_Type* ExistingUid = UidMap.FindKey(Name);
-	const FGuid* ExistingGuid = GuidMap.Find(Name);
-
-	if (ExistingUid == nullptr && (ExistingGuid == nullptr || !ExistingGuid->IsValid()))
-	{
-		// make sure we didn't reach till end
-		check(NextUid != SmartName::MaxUID);
-
-		OutUid = NextUid;
-		UidMap.Add(OutUid, Name);
-		GuidMap.Add(Name, InGuid);
-
-		++NextUid;
-
-		return true;
-	}
-
-	return false;
+	FSmartName NewSmartName(Name, CurveNameList.Add(Name));
+	CurveMetaDataMap.Add(Name);
+	return NewSmartName;
 }
 
 bool FSmartNameMapping::GetName(const SmartName::UID_Type& Uid, FName& OutName) const
 {
-	const FName* FoundName = UidMap.Find(Uid);
-	if(FoundName)
-	{
-		OutName = *FoundName;
-		return true;
-	}
-	return false;
-}
-
-bool FSmartNameMapping::GetNameByGuid(const FGuid& Guid, FName& OutName) const
-{
-	if (Guid.IsValid())
+	if (CurveNameList.IsValidIndex(Uid))
 	{
-		const FName* FoundName = GuidMap.FindKey(Guid);
-		if (FoundName)
-		{
-			OutName = *FoundName;
-			return true;
-		}
+		OutName = CurveNameList[Uid];
+		return OutName != NAME_None; // Name may have been removed
 	}
-
 	return false;
 }
 
+#if WITH_EDITOR
 bool FSmartNameMapping::Rename(const SmartName::UID_Type& Uid, FName NewName)
 {
-	FName* ExistingName = UidMap.Find(Uid);
-	if(ExistingName)
+	FName ExistingName;
+	if(GetName(Uid, ExistingName))
 	{
-		FGuid* Guid = GuidMap.Find(*ExistingName);
-		check(Guid);
-
-		// re add new value
-		GuidMap.Remove(*ExistingName);
-		GuidMap.Add(NewName, *Guid);
-
 		// fix up meta data
-		FCurveMetaData* MetaDataToCopy = CurveMetaDataMap.Find(*ExistingName);
+		FCurveMetaData* MetaDataToCopy = CurveMetaDataMap.Find(ExistingName);
 		if (MetaDataToCopy)
 		{
 			FCurveMetaData& NewMetaData = CurveMetaDataMap.Add(NewName);
 			NewMetaData = *MetaDataToCopy;
 			
 			// remove old one
-			CurveMetaDataMap.Remove(*ExistingName);			
+			CurveMetaDataMap.Remove(ExistingName);			
 		}
 
-		check(GuidMap.Num() == UidMap.Num());
-
-		*ExistingName = NewName;
+		CurveNameList[Uid] = NewName;
 		return true;
 	}
 	return false;
@@ -121,110 +61,110 @@ bool FSmartNameMapping::Rename(const SmartName::UID_Type& Uid, FName NewName)
 
 bool FSmartNameMapping::Remove(const SmartName::UID_Type& Uid)
 {
-	FName* ExistingName = UidMap.Find(Uid);
-	if (ExistingName)
+	FName ExistingName;
+	if (GetName(Uid, ExistingName))
 	{
-		FGuid* Guid = GuidMap.Find(*ExistingName);
-		check(Guid);
-
-		// re add new value
-		GuidMap.Remove(*ExistingName);
-		CurveMetaDataMap.Remove(*ExistingName);
-		UidMap.Remove(Uid);
-
-		check(GuidMap.Num() == UidMap.Num());
+		CurveMetaDataMap.Remove(ExistingName);
+		CurveNameList[Uid] = NAME_None;
 
 		return true;
 	}
+	return false;
+}
 
+bool FSmartNameMapping::Remove(const FName& Name)
+{
+	const SmartName::UID_Type Uid = FindUID(Name);
+	if (Uid != SmartName::MaxUID)
+	{
+		CurveMetaDataMap.Remove(Name);
+		CurveNameList[Uid] = NAME_None;
+		return true;
+	}
 	return false;
 }
+#endif
 
 void FSmartNameMapping::Serialize(FArchive& Ar)
 {
 	Ar.UsingCustomVersion(FFrameworkObjectVersion::GUID);
+	Ar.UsingCustomVersion(FAnimPhysObjectVersion::GUID);
 	if (Ar.CustomVer(FFrameworkObjectVersion::GUID) >= FFrameworkObjectVersion::SmartNameRefactor)
 	{
-		if (Ar.IsSaving() && Ar.IsCooking())
-		{
-			// stript out guid from the map
-			for (TPair<FName, FGuid>& GuidPair: GuidMap)
-			{
-				// clear guid, so that we don't use it once cooked
-				GuidPair.Value = FGuid();
-			}
-		}
-
-		Ar << GuidMap;
-
-		if (Ar.ArIsLoading)
+		if (Ar.CustomVer(FAnimPhysObjectVersion::GUID) < FAnimPhysObjectVersion::SmartNameRefactorForDeterministicCooking)
 		{
-			NextUid = 0;
-			//fill up uidmap
-			UidMap.Empty(GuidMap.Num());
-			for (const TPair<FName, FGuid>& GuidPair : GuidMap)
-			{
-				UidMap.Add(NextUid, GuidPair.Key);
-				++NextUid;
-			}
+			TMap<FName, FGuid> TempGuidMap;
+			Ar << TempGuidMap;
 		}
 	}
 	else if(Ar.UE4Ver() >= VER_UE4_SKELETON_ADD_SMARTNAMES)
 	{
-		Ar << NextUid;
-		Ar << UidMap;
-
-		TMap<SmartName::UID_Type, FName> NewUidMap;
-		// convert to GuidMap
-		NextUid = 0;
-		GuidMap.Empty(UidMap.Num());
-		for (TPair<SmartName::UID_Type, FName>& UidPair : UidMap)
-		{
-			GuidMap.Add(UidPair.Value, FGuid::NewGuid());
-			NewUidMap.Add(NextUid++, UidPair.Value);
-		}
+		SmartName::UID_Type NextUidTemp;
+		Ar << NextUidTemp;
 
-		UidMap = NewUidMap;
+		TMap<SmartName::UID_Type, FName> TempUidMap;
+		Ar << TempUidMap;
 	}
 
 	if (Ar.CustomVer(FFrameworkObjectVersion::GUID) >= FFrameworkObjectVersion::MoveCurveTypesToSkeleton)
 	{
 		Ar << CurveMetaDataMap;
 	}
-}
 
-int32 FSmartNameMapping::GetNumNames() const
-{
-	return UidMap.Num();
+	if (Ar.IsLoading())
+	{
+		CurveMetaDataMap.GenerateKeyArray(CurveNameList);
+	}
 }
 
 void FSmartNameMapping::FillUidArray(TArray<SmartName::UID_Type>& Array) const
 {
-	UidMap.GenerateKeyArray(Array);
+	Array.Reset(CurveNameList.Num());
+	
+	for (int32 NameIndex = 0; NameIndex < CurveNameList.Num(); ++NameIndex)
+	{
+		//In editor names can be removed and so have to deal with empty slots
+#if WITH_EDITOR
+		if (CurveNameList[NameIndex] != NAME_None)
+		{
+			Array.Add(NameIndex);
+		}
+#else
+		Array.Add(NameIndex);
+#endif
+	}
 }
 
 void FSmartNameMapping::FillNameArray(TArray<FName>& Array) const
 {
-	UidMap.GenerateValueArray(Array);
+	//In editor names can be removed and so have to deal with empty slots
+#if WITH_EDITOR
+	Array.Reset(CurveNameList.Num());
+	for (const FName& Name : CurveNameList)
+	{
+		if (Name != NAME_None)
+		{
+			Array.Add(Name);
+		}
+	}
+#else
+	Array = CurveNameList;
+#endif
 }
 
 bool FSmartNameMapping::Exists(const SmartName::UID_Type& Uid) const
 {
-	return UidMap.Find(Uid) != nullptr;
+	return CurveNameList.IsValidIndex(Uid) && CurveNameList[Uid] != NAME_None;
 }
 
 bool FSmartNameMapping::Exists(const FName& Name) const
 {
-	return UidMap.FindKey(Name) != nullptr;
+	return CurveNameList.Contains(Name);
 }
 
 SmartName::UID_Type FSmartNameMapping::FindUID(const FName& Name) const
 {
-	if (const SmartName::UID_Type* Uid = UidMap.FindKey(Name))
-	{
-		return *Uid;
-	}
-	return SmartName::MaxUID;
+	return CurveNameList.IndexOfByKey(Name);
 }
 
 FArchive& operator<<(FArchive& Ar, FSmartNameMapping& Elem)
@@ -234,53 +174,12 @@ FArchive& operator<<(FArchive& Ar, FSmartNameMapping& Elem)
 	return Ar;
 }
 
-#if WITH_EDITOR
-bool FSmartNameMapping::FindOrAddSmartName(FName Name, FSmartName& OutName)
-{
-	SmartName::UID_Type NewUID;
-	FGuid NewGuid;
-	bool bNewlyAdded = AddOrFindName(Name, NewUID, NewGuid);
-	OutName = FSmartName(Name, NewUID, NewGuid);
-
-	return bNewlyAdded;
-}
-
-bool FSmartNameMapping::AddSmartName(FSmartName& OutName)
-{
-	return AddName(OutName.DisplayName, OutName.UID, OutName.Guid);
-}
-
-#else 
-
-// in cooked build, we don't have Guid, so just register with empty guid. 
-// you only should come here if it it hasn't found yet. 
-bool FSmartNameMapping::FindOrAddSmartName(FName Name, SmartName::UID_Type& OutUid)
-{
-	FSmartName FoundName;
-	if (FindSmartName(Name, FoundName))
-	{
-		OutUid = FoundName.UID;
-		return true;
-	}
-	else
-	{
-		// the guid is discarded, just we want to make sure it's not same
-		return AddName(Name, OutUid, FGuid::NewGuid());
-	}
-}
-#endif // WITH_EDITOR
-
 bool FSmartNameMapping::FindSmartName(FName Name, FSmartName& OutName) const
 {
 	SmartName::UID_Type ExistingUID = FindUID(Name);
 	if (ExistingUID != SmartName::MaxUID)
 	{
-#if WITH_EDITOR
-		const FGuid* ExistingGuid = GuidMap.Find(Name);
-		OutName = FSmartName(Name, ExistingUID, *ExistingGuid);
-#else
 		OutName = FSmartName(Name, ExistingUID);
-#endif // WITH_EDITOR
 		return true;
 	}
 
@@ -293,9 +192,6 @@ bool FSmartNameMapping::FindSmartNameByUID(SmartName::UID_Type UID, FSmartName&
 	if (GetName(UID, ExistingName))
 	{
 		OutName.DisplayName = ExistingName;
-#if WITH_EDITORONLY_DATA
-		OutName.Guid = *GuidMap.Find(ExistingName);
-#endif // WITH_EDITORONLY_DATA
 		OutName.UID = UID;
 		return true;
 	}
@@ -349,16 +245,6 @@ const FSmartNameMapping* FSmartNameContainer::GetContainerInternal(const FName&
 	return NameMappings.Find(ContainerName);
 }
 
-/* initialize curve meta data for the container */
-void FSmartNameContainer::InitializeCurveMetaData(class USkeleton* Skeleton)
-{
-	FSmartNameMapping* CurveMappingTable = GetContainerInternal(USkeleton::AnimCurveMappingName);
-	if (CurveMappingTable)
-	{
-		CurveMappingTable->InitializeCurveMetaData(Skeleton);
-	}
-}
-
 ////////////////////////////////////////////////////////////////////////
 //
 // FSmartName
@@ -373,22 +259,19 @@ bool FSmartName::Serialize(FArchive& Ar)
 		SmartName::UID_Type TempUID;
 		Ar << TempUID;
 	}
-	else
-	{
 #if WITH_EDITOR
-		if (Ar.IsTransacting())
-		{
-			Ar << UID;
-		}
-#endif // WITH_EDITOR
+	else if (Ar.IsTransacting())
+	{
+		Ar << UID;
 	}
+#endif
+
 	// only save if it's editor build and not cooking
-#if WITH_EDITORONLY_DATA
-	if (!Ar.IsSaving() || !Ar.IsCooking())
+	if (Ar.CustomVer(FAnimPhysObjectVersion::GUID) < FAnimPhysObjectVersion::SmartNameRefactorForDeterministicCooking)
 	{
-		Ar << Guid;
+		FGuid TempGUID;
+		Ar << TempGUID;
 	}
-#endif // WITH_EDITORONLY_DATA
 
 	return true;
 }
diff --git a/Engine/Source/Runtime/Engine/Private/Audio.cpp b/Engine/Source/Runtime/Engine/Private/Audio.cpp
--- a/Engine/Source/Runtime/Engine/Private/Audio.cpp
+++ b/Engine/Source/Runtime/Engine/Private/Audio.cpp
@@ -57,48 +57,47 @@ DEFINE_STAT(STAT_AudioGatherWaveInstances);
 DEFINE_STAT(STAT_AudioFindNearestLocation);
 
 
-bool IsAudioPluginEnabled(EAudioPlugin::Type PluginType)
+
+bool IsAudioPluginEnabled(EAudioPlugin PluginType)
 {
-	TArray<IAudioPlugin *> AudioPlugin = IModularFeatures::Get().GetModularFeatureImplementations<IAudioPlugin>(IAudioPlugin::GetModularFeatureName());
-	if (AudioPlugin.Num() > 0)
+	switch (PluginType)
 	{
-		if (PluginType == EAudioPlugin::SPATIALIZATION)
-		{
-			return AudioPlugin[0]->ImplementsSpatialization();
-		}
-		else if (PluginType == EAudioPlugin::REVERB)
-		{
-			return AudioPlugin[0]->ImplementsReverb();
-		}
-		else if (PluginType == EAudioPlugin::OCCLUSION)
-		{
-			return AudioPlugin[0]->ImplementsOcclusion();
-		}
+	case EAudioPlugin::SPATIALIZATION:
+		return AudioPluginUtilities::GetDesiredSpatializationPlugin(AudioPluginUtilities::CurrentPlatform) != nullptr;
+		break;
+	case EAudioPlugin::REVERB:
+		return AudioPluginUtilities::GetDesiredReverbPlugin(AudioPluginUtilities::CurrentPlatform) != nullptr;
+		break;
+	case EAudioPlugin::OCCLUSION:
+		return AudioPluginUtilities::GetDesiredOcclusionPlugin(AudioPluginUtilities::CurrentPlatform) != nullptr;
+		break;
+	default:
+		return false;
+		break;
 	}
-
-	return false;
 }
 
-bool DoesAudioPluginHaveCustomSettings(EAudioPlugin::Type PluginType)
+bool DoesAudioPluginHaveCustomSettings(EAudioPlugin PluginType)
 {
-	TArray<IAudioPlugin *> AudioPlugin = IModularFeatures::Get().GetModularFeatureImplementations<IAudioPlugin>(IAudioPlugin::GetModularFeatureName());
-	if (AudioPlugin.Num() > 0)
+	if (PluginType == EAudioPlugin::SPATIALIZATION)
 	{
-		if (PluginType == EAudioPlugin::SPATIALIZATION)
-		{
-			return AudioPlugin[0]->HasCustomSpatializationSetting();
-		}
-		else if (PluginType == EAudioPlugin::REVERB)
-		{
-			return AudioPlugin[0]->HasCustomReverbSetting();
-		}
-		else if (PluginType == EAudioPlugin::OCCLUSION)
-		{
-			return AudioPlugin[0]->HasCustomOcclusionSetting();
-		}
+		IAudioSpatializationFactory* Factory = AudioPluginUtilities::GetDesiredSpatializationPlugin(AudioPluginUtilities::CurrentPlatform);
+		return Factory && Factory->HasCustomSpatializationSetting();
+	}
+	else if (PluginType == EAudioPlugin::REVERB)
+	{
+		IAudioReverbFactory* Factory = AudioPluginUtilities::GetDesiredReverbPlugin(AudioPluginUtilities::CurrentPlatform);
+		return Factory && Factory->HasCustomReverbSetting();
+	}
+	else if (PluginType == EAudioPlugin::OCCLUSION)
+	{
+		IAudioOcclusionFactory* Factory = AudioPluginUtilities::GetDesiredOcclusionPlugin(AudioPluginUtilities::CurrentPlatform);
+		return Factory && Factory->HasCustomOcclusionSetting();
+	}
+	else
+	{
+		return false;
 	}
-
-	return false;
 }
 
 
@@ -694,8 +693,10 @@ uint32 FWaveInstance::TypeHashCounter = 0;
 FWaveInstance::FWaveInstance( FActiveSound* InActiveSound )
 	: WaveData(nullptr)
 	, SoundClass(nullptr)
+	, SoundSubmix(nullptr)
 	, ActiveSound(InActiveSound)
 	, Volume(0.0f)
+	, DistanceAttenuation(1.0f)
 	, VolumeMultiplier(1.0f)
 	, VolumeApp(1.0f)
 	, Priority(1.0f)
@@ -706,6 +707,7 @@ FWaveInstance::FWaveInstance( FActiveSound* InActiveSound )
 	, LFEBleed(0.0f)
 	, LoopingMode(LOOP_Never)
 	, StartTime(-1.f)
+	, bOutputToBusOnly(false)
 	, bApplyRadioFilter(false)
 	, bIsStarted(false)
 	, bIsFinished(false)
@@ -804,12 +806,23 @@ void FWaveInstance::AddReferencedObjects( FReferenceCollector& Collector )
 float FWaveInstance::GetActualVolume() const
 {
 	// Include all volumes 
-	return GetVolume() * VolumeApp;
+	return GetVolume() * VolumeApp * DistanceAttenuation;
+}
+
+float FWaveInstance::GetDistanceAttenuation() const
+{
+	// Only includes volume attenuation due do distance
+	return DistanceAttenuation;
+}
+
+float FWaveInstance::GetVolumeWithDistanceAttenuation() const
+{
+	return GetVolume() * DistanceAttenuation;
 }
 
 float FWaveInstance::GetVolume() const
 {
-	// Include all volumes 
+	// Only includes non-attenuation and non-app volumes
 	return Volume * VolumeMultiplier;
 }
 
@@ -821,7 +834,7 @@ bool FWaveInstance::ShouldStopDueToMaxConcurrency() const
 float FWaveInstance::GetVolumeWeightedPriority() const
 {
 	// This will result in zero-volume sounds still able to be sorted due to priority but give non-zero volumes higher priority than 0 volumes
-	float ActualVolume = GetVolume();
+	float ActualVolume = GetVolumeWithDistanceAttenuation();
 	if (ActualVolume > 0.0f)
 	{
 		return ActualVolume * Priority;
