Copying //UE4/Dev-AnimPhys to //UE4/Dev-Main (Source: //UE4/Dev-AnimPhys @ 3624379)

diff --git a/Engine/Source/Runtime/Engine/Private/AudioDevice.cpp b/Engine/Source/Runtime/Engine/Private/AudioDevice.cpp
--- a/Engine/Source/Runtime/Engine/Private/AudioDevice.cpp
+++ b/Engine/Source/Runtime/Engine/Private/AudioDevice.cpp
@@ -21,12 +21,14 @@
 #include "Sound/SoundWave.h"
 #include "AudioDecompress.h"
 #include "AudioEffect.h"
+#include "AudioPluginUtilities.h"
 #include "Sound/AudioSettings.h"
 #include "Sound/SoundCue.h"
 #include "Sound/SoundNode.h"
 #include "Sound/SoundNodeWavePlayer.h"
 #include "GameFramework/GameUserSettings.h"
 #include "GameFramework/WorldSettings.h"
+#include "GeneralProjectSettings.h"
 
 #if WITH_EDITOR
 #include "AssetRegistryModule.h"
@@ -108,11 +110,10 @@ FAudioDevice::FAudioDevice()
 	: CommonAudioPool(nullptr)
 	, CommonAudioPoolFreeBytes(0)
 	, DeviceHandle(INDEX_NONE)
-	, AudioPlugin(nullptr)
 	, SpatializationPluginInterface(nullptr)
 	, ReverbPluginInterface(nullptr)
 	, OcclusionInterface(nullptr)
-	, ListenerObserver(nullptr)
+	, PluginListeners()
 	, CurrentTick(0)
 	, TestAudioComponent(nullptr)
 	, DebugState(DEBUGSTATE_None)
@@ -125,16 +126,20 @@ FAudioDevice::FAudioDevice()
 	, Effects(nullptr)
 	, CurrentReverbEffect(nullptr)
 	, PlatformAudioHeadroom(1.0f)
+	, DefaultReverbSendLevel(0.2f)
 	, bHRTFEnabledForAll_OnGameThread(false)
 	, bGameWasTicking(true)
 	, bDisableAudioCaching(false)
 	, bIsAudioDeviceHardwareInitialized(false)
 	, bAudioMixerModuleLoaded(false)
+	, bSpatializationIsExternalSend(false)
+	, bOcclusionIsExternalSend(false)
+	, bReverbIsExternalSend(false)
 	, bStartupSoundsPreCached(false)
 	, bSpatializationInterfaceEnabled(false)
 	, bOcclusionInterfaceEnabled(false)
 	, bReverbInterfaceEnabled(false)
-	, bListenerObserverInterfaceEnabled(false)
+	, bPluginListenersInitialized(false)
 	, bHRTFEnabledForAll(false)
 	, bIsDeviceMuted(false)
 	, bIsInitialized(false)
@@ -170,6 +175,8 @@ bool FAudioDevice::Init(int32 InMaxChannels)
 	}
 
 	bool bDeferStartupPrecache = false;
+	
+	PluginListeners.Reset();
 
 	// initialize max channels taking into account platform configurations
 	// Get a copy of the platform-specific settings (overriden by platforms)
@@ -196,9 +203,11 @@ bool FAudioDevice::Init(int32 InMaxChannels)
 		PlatformAudioHeadroom = FMath::Pow(10.0f, Headroom / 20.0f);
 	}
 
-	bAllowCenterChannel3DPanning = GetDefault<UAudioSettings>()->bAllowCenterChannel3DPanning;
+	const UAudioSettings* AudioSettings = GetDefault<UAudioSettings>();
 
-	bAllowVirtualizedSounds = GetDefault<UAudioSettings>()->bAllowVirtualizedSounds;
+	bAllowCenterChannel3DPanning = AudioSettings->bAllowCenterChannel3DPanning;
+	bAllowVirtualizedSounds = AudioSettings->bAllowVirtualizedSounds;
+	DefaultReverbSendLevel = AudioSettings->DefaultReverbSendLevel;
 
 	const FSoftObjectPath DefaultBaseSoundMixName = GetDefault<UAudioSettings>()->DefaultBaseSoundMix;
 	if (DefaultBaseSoundMixName.IsValid())
@@ -219,45 +228,61 @@ bool FAudioDevice::Init(int32 InMaxChannels)
 		Effects = CreateEffectsManager();
 	}
 
-	// Get the audio spatialization plugin
-	// Note: there is only *one* spatialization plugin currently, but the GetModularFeatureImplementation only returns a TArray
-	// Therefore, we are just grabbing the first one that is returned (if one is returned).
+	
 
-	TArray<IAudioPlugin *> AudioPlugins = IModularFeatures::Get().GetModularFeatureImplementations<IAudioPlugin>(IAudioPlugin::GetModularFeatureName());
-	for (IAudioPlugin* Plugin : AudioPlugins)
+	//Get the requested spatialization plugin and set it up.
+	IAudioSpatializationFactory* SpatializationPluginFactory = AudioPluginUtilities::GetDesiredSpatializationPlugin(AudioPluginUtilities::CurrentPlatform);
+	if (SpatializationPluginFactory != nullptr)
 	{
-#if WITH_EDITOR
-		// If the plugin doesn't support multiple instances, and this is the main audio device, then break. We'll only use it for PIE.
-		if (!Plugin->SupportsMultipleAudioDevices() && IsMainAudioDevice())
+		SpatializationPluginInterface = SpatializationPluginFactory->CreateNewSpatializationPlugin(this);
+		if (!IsAudioMixerEnabled())
 		{
-			break;
+			//Set up initialization parameters for system level effect plugins:
+			FAudioPluginInitializationParams PluginInitializationParams;
+			PluginInitializationParams.SampleRate = SampleRate;
+			PluginInitializationParams.NumSources = MaxChannels;
+			PluginInitializationParams.BufferLength = PlatformSettings.CallbackBufferFrameSize;
+			PluginInitializationParams.AudioDevicePtr = this;
+			
+			SpatializationPluginInterface->Initialize(PluginInitializationParams);
 		}
-#endif
-		// Store the ptr to the audio plugin we're using
-		AudioPlugin = Plugin;
-
-		// Initialize the plugin
-		Plugin->Initialize();
-
-		// Create feature override interfaces. Note these can be null if the audio plugin is not creating an override for the feature.
-		SpatializationPluginInterface = Plugin->CreateSpatializationInterface(this);
 
-		bSpatializationInterfaceEnabled = SpatializationPluginInterface.IsValid();
+		bSpatializationInterfaceEnabled = true;
+		bSpatializationIsExternalSend = SpatializationPluginFactory->IsExternalSend();
+		UE_LOG(LogAudio, Log, TEXT("Using Audio Spatialization Plugin: %s is external send: %d"), *(SpatializationPluginFactory->GetDisplayName()), bSpatializationIsExternalSend);
+	}
+	else
+	{
+		UE_LOG(LogAudio, Log, TEXT("Using built-in audio spatialization."));
+	}
 
-		// Only audio mixer supports these features
-		if (IsAudioMixerEnabled())
-		{
-			ReverbPluginInterface = Plugin->CreateReverbInterface(this);
-			OcclusionInterface = Plugin->CreateOcclusionInterface(this);
-			ListenerObserver = Plugin->CreateListenerObserverInterface(this);
+	//Get the requested reverb plugin and set it up:
+	IAudioReverbFactory* ReverbPluginFactory = AudioPluginUtilities::GetDesiredReverbPlugin(AudioPluginUtilities::CurrentPlatform);
+	if (ReverbPluginFactory != nullptr)
+	{
+		ReverbPluginInterface = ReverbPluginFactory->CreateNewReverbPlugin(this);
+		bReverbInterfaceEnabled = true;
+		bReverbIsExternalSend = ReverbPluginFactory->IsExternalSend();
+		UE_LOG(LogAudio, Log, TEXT("Audio Reverb Plugin: %s"), *(ReverbPluginFactory->GetDisplayName()));
 
-			bReverbInterfaceEnabled = ReverbPluginInterface.IsValid();
-			bOcclusionInterfaceEnabled = OcclusionInterface.IsValid();
-			bListenerObserverInterfaceEnabled = ListenerObserver.IsValid();
-		}
+	}
+	else
+	{
+		UE_LOG(LogAudio, Log, TEXT("Using built-in audio reverb."));
+	}
 
-		// We are only using the first one enabled since this is a singleton system.
-		break;
+	//Get the requested occlusion plugin and set it up.
+	IAudioOcclusionFactory* OcclusionPluginFactory = AudioPluginUtilities::GetDesiredOcclusionPlugin(AudioPluginUtilities::CurrentPlatform);
+	if (OcclusionPluginFactory != nullptr)
+	{
+		OcclusionInterface = OcclusionPluginFactory->CreateNewOcclusionPlugin(this);
+		bOcclusionInterfaceEnabled = true;
+		bOcclusionIsExternalSend = OcclusionPluginFactory->IsExternalSend();
+		UE_LOG(LogAudio, Display, TEXT("Audio Occlusion Plugin: %s"), *(OcclusionPluginFactory->GetDisplayName()));
+	}
+	else
+	{
+		UE_LOG(LogAudio, Display, TEXT("Using built-in audio occlusion."));
 	}
 
 	// allow the platform to startup
@@ -351,12 +376,12 @@ void FAudioDevice::Teardown()
 		Effects = nullptr;
 	}
 
-	if (bListenerObserverInterfaceEnabled && ListenerObserver.IsValid())
+	for (TAudioPluginListenerPtr PluginListener : PluginListeners)
 	{
-		ListenerObserver->OnListenerShutdown(this);
-		ListenerObserver = nullptr;
+		PluginListener->OnListenerShutdown(this);
 	}
-
+	
+	
 	// let platform shutdown
 	TeardownHardware();
 
@@ -373,15 +398,16 @@ void FAudioDevice::Teardown()
 	Sources.Reset();
 	FreeSources.Reset();
 
-	if (AudioPlugin != nullptr)
-	{
-		SpatializationPluginInterface = nullptr;
-		ReverbPluginInterface = nullptr;
-		OcclusionInterface = nullptr;
+	SpatializationPluginInterface.Reset();
+	bSpatializationInterfaceEnabled = false;
 
-		AudioPlugin->Shutdown();
-		AudioPlugin = nullptr;
-	}
+	ReverbPluginInterface.Reset();
+	bReverbInterfaceEnabled = false;
+
+	OcclusionInterface.Reset();
+	bOcclusionInterfaceEnabled = false;
+
+	PluginListeners.Reset();
 }
 
 void FAudioDevice::Suspend(bool bGameTicking)
@@ -1765,7 +1791,7 @@ void FAudioDevice::UpdatePassiveSoundMixModifiers(TArray<FWaveInstance*>& WaveIn
 			USoundClass* SoundClass = WaveInstance->SoundClass;
 			if (SoundClass) 
 			{
-				const float WaveInstanceActualVolume = WaveInstance->GetVolume();
+				const float WaveInstanceActualVolume = WaveInstance->GetVolumeWithDistanceAttenuation();
 				// Check each SoundMix individually for volume levels
 				for (const FPassiveSoundMixModifier& PassiveSoundMixModifier : SoundClass->PassiveSoundMixModifiers)
 				{
@@ -2266,18 +2292,32 @@ void FAudioDevice::SetListener(UWorld* World, const int32 InViewportIndex, const
 
 	ListenerTransforms[InViewportIndex] = ListenerTransformCopy;
 
-	// Broadcast to a 3rd party plugin listener observer if enabled
-	if (bListenerObserverInterfaceEnabled && ListenerObserver.IsValid())
+	DECLARE_CYCLE_STAT(TEXT("FAudioThreadTask.SetListener"), STAT_AudioSetListener, STATGROUP_AudioThreadCommands);
+
+	uint32 WorldID = INDEX_NONE;
+
+	if (World != nullptr)
 	{
-		ListenerObserver->OnListenerUpdated(this, World, InViewportIndex, ListenerTransform, InDeltaSeconds);
+		WorldID = World->GetUniqueID();
 	}
 
-	DECLARE_CYCLE_STAT(TEXT("FAudioThreadTask.SetListener"), STAT_AudioSetListener, STATGROUP_AudioThreadCommands);
+	// Initialize the plugin listeners if we haven't already. This needs to be done here since this is when we're 
+	// guaranteed to have a world ptr and we've already initialized the audio device. 
+	if (World && !bPluginListenersInitialized)
+	{
+		InitializePluginListeners(World);
+		bPluginListenersInitialized = true;
+	}
 
 	FAudioDevice* AudioDevice = this;
-	const uint32 WorldID = World->GetUniqueID();
 	FAudioThread::RunCommandOnAudioThread([AudioDevice, WorldID, InViewportIndex, ListenerTransformCopy, InDeltaSeconds]()
 	{
+		// Broadcast to a 3rd party plugin listener observer if enabled
+		for (TAudioPluginListenerPtr PluginManager : AudioDevice->PluginListeners)
+		{
+			PluginManager->OnListenerUpdated(AudioDevice, InViewportIndex, ListenerTransformCopy, InDeltaSeconds);
+		}
+
 		TArray<FListener>& AudioThreadListeners = AudioDevice->Listeners;
 		if (InViewportIndex >= AudioThreadListeners.Num())
 		{
@@ -2918,7 +2958,7 @@ void FAudioDevice::StopSources(TArray<FWaveInstance*>& WaveInstances, int32 Firs
 			Source->LastUpdate = CurrentTick;
 
 			// If they are still audible, mark them as such
-			float VolumeWeightedPriority = WaveInstance->GetVolume();
+			float VolumeWeightedPriority = WaveInstance->GetVolumeWithDistanceAttenuation();
 			if (VolumeWeightedPriority > 0.0f)
 			{
 				Source->LastHeardUpdate = CurrentTick;
@@ -2971,7 +3011,7 @@ void FAudioDevice::StopSources(TArray<FWaveInstance*>& WaveInstances, int32 Firs
 	for (int32 InstanceIndex = 0; InstanceIndex < FirstActiveIndex; InstanceIndex++)
 	{
 		FWaveInstance* WaveInstance = WaveInstances[ InstanceIndex ];
-		if (WaveInstance->GetVolume() > 0.1f)
+		if (WaveInstance->GetVolumeWithDistanceAttenuation() > 0.1f)
 		{
 			AudibleInactiveSounds++;
 		}
@@ -3292,7 +3332,7 @@ void FAudioDevice::SendUpdateResultsToGameThread(const int32 FirstActiveIndex)
 				FAudioStats::FStatWaveInstanceInfo WaveInstanceInfo;
 				FSoundSource* Source = WaveInstanceSourceMap.FindRef(WaveInstance);
 				WaveInstanceInfo.Description = Source ? Source->Describe((RequestedAudioStats & ERequestedAudioStats::LongSoundNames) != 0) : FString(TEXT("No source"));
-				WaveInstanceInfo.ActualVolume = WaveInstance->GetVolume();
+				WaveInstanceInfo.ActualVolume = WaveInstance->GetVolumeWithDistanceAttenuation();
 				WaveInstanceInfo.InstanceIndex = InstanceIndex;
 				WaveInstanceInfo.WaveInstanceName = *WaveInstance->GetName();
 				StatSoundInfos[*SoundInfoIndex].WaveInstanceInfos.Add(MoveTemp(WaveInstanceInfo));
@@ -3393,6 +3433,17 @@ void FAudioDevice::StopAllSounds(bool bShouldStopUISounds)
 	ProcessingPendingActiveSoundStops();
 }
 
+void FAudioDevice::InitializePluginListeners(UWorld* World)
+{
+	check(IsInGameThread());
+	check(!bPluginListenersInitialized);
+
+	for (TAudioPluginListenerPtr PluginListener : PluginListeners)
+	{
+		PluginListener->OnListenerInitialize(this, World);
+	}
+}
+
 void FAudioDevice::AddNewActiveSound(const FActiveSound& NewActiveSound)
 {
 	if (NewActiveSound.Sound == nullptr)
@@ -3880,9 +3931,11 @@ FAudioDevice::FCreateComponentParams::FCreateComponentParams(FAudioDevice* InAud
 void FAudioDevice::FCreateComponentParams::CommonInit()
 {
 	bPlay = false;
+	bStopWhenOwnerDestroyed = true;
 	bLocationSet = false;
 	AttenuationSettings = nullptr;
 	ConcurrencySettings = nullptr;
+	Location = FVector::ZeroVector;
 }
 
 void FAudioDevice::FCreateComponentParams::SetLocation(const FVector InLocation)
@@ -4008,7 +4061,7 @@ UAudioComponent* FAudioDevice::CreateComponent(USoundBase* Sound, const FCreateC
 	return AudioComponent;
 }
 
-void FAudioDevice::PlaySoundAtLocation(USoundBase* Sound, UWorld* World, float VolumeMultiplier, float PitchMultiplier, float StartTime, const FVector& Location, const FRotator& Rotation, USoundAttenuation* AttenuationSettings, USoundConcurrency* ConcurrencySettings, const TArray<FAudioComponentParam>* Params)
+void FAudioDevice::PlaySoundAtLocation(USoundBase* Sound, UWorld* World, float VolumeMultiplier, float PitchMultiplier, float StartTime, const FVector& Location, const FRotator& Rotation, USoundAttenuation* AttenuationSettings, USoundConcurrency* ConcurrencySettings, const TArray<FAudioComponentParam>* Params, AActor* OwningActor)
 {
 	check(IsInGameThread());
 
@@ -4059,6 +4112,8 @@ void FAudioDevice::PlaySoundAtLocation(USoundBase* Sound, UWorld* World, float V
 		NewActiveSound.ConcurrencySettings = ConcurrencySettings;
 		NewActiveSound.Priority = Sound->Priority;
 
+		NewActiveSound.SetOwner(OwningActor);
+
 		// Apply any optional audio component instance params on the sound
 		if (Params)
 		{
@@ -4206,9 +4261,23 @@ void FAudioDevice::Precache(USoundWave* SoundWave, bool bSynchronous, bool bTrac
 	}
 	else if (SoundWave->bProcedural)
 	{
-		// Streaming data, created programmatically.
+		// Procedurally created audio
 		SoundWave->DecompressionType = DTYPE_Procedural;
 	}
+	else if (SoundWave->bIsBus)
+	{
+		// Audio data which will be generated by instanced objects, not from the sound wave asset
+		if (IsAudioMixerEnabled())
+		{
+			// Buses will initialize as procedural, but not actually become a procedural sound wave
+			SoundWave->DecompressionType = DTYPE_Procedural;
+		}
+		else
+		{
+			// Buses are only supported with audio mixer
+			SoundWave->DecompressionType = DTYPE_Invalid;
+		}
+	}
 	else if (HasCompressedAudioInfoClass(SoundWave))
 	{
 		const FSoundGroup& SoundGroup = GetDefault<USoundGroups>()->GetSoundGroup(SoundWave->SoundGroup);
@@ -4472,6 +4541,16 @@ void FAudioDevice::StopSoundsUsingResource(USoundWave* SoundWave, TArray<UAudioC
 	}
 }
 
+void FAudioDevice::RegisterPluginListener(const TAudioPluginListenerPtr PluginListener)
+{
+	PluginListeners.AddUnique(PluginListener);
+}
+
+void FAudioDevice::UnregisterPluginListener(const TAudioPluginListenerPtr PluginListener)
+{
+	PluginListeners.RemoveSingle(PluginListener);
+}
+
 bool FAudioDevice::IsAudioDeviceMuted() const
 {
 	check(IsInAudioThread());
@@ -4542,8 +4621,11 @@ bool FAudioDevice::CanUseVRAudioDevice()
 		UEditorEngine* EdEngine = Cast<UEditorEngine>(GEngine);
 		return EdEngine->bUseVRPreviewForPlayWorld;
 	}
+	else
 #endif
-	return true;
+	{
+		return FParse::Param(FCommandLine::Get(), TEXT("vr")) || GetDefault<UGeneralProjectSettings>()->bStartInVR;
+	}
 }
 
 void FAudioDevice::SetTransientMasterVolume(const float InTransientMasterVolume)
diff --git a/Engine/Source/Runtime/Engine/Private/AudioPluginUtilities.cpp b/Engine/Source/Runtime/Engine/Private/AudioPluginUtilities.cpp
new file mode 100644
--- /dev/null
+++ b/Engine/Source/Runtime/Engine/Private/AudioPluginUtilities.cpp
@@ -0,0 +1,139 @@
+// Copyright 1998-2017 Epic Games, Inc. All Rights Reserved.
+
+#include "AudioPluginUtilities.h"
+
+#include "Misc/ConfigCacheIni.h"
+#include "CoreGlobals.h"
+
+
+/** Platform config section for each platform's target settings. */
+FORCEINLINE const TCHAR* GetPlatformConfigSection(EAudioPlatform AudioPlatform)
+{
+	switch (AudioPlatform)
+	{
+		case EAudioPlatform::Windows:
+			return TEXT("/Script/WindowsTargetPlatform.WindowsTargetSettings");
+
+		case EAudioPlatform::Mac:
+			return TEXT("/Script/MacTargetPlatform.MacTargetSettings");
+
+		case EAudioPlatform::Linux:
+			return TEXT("/Script/LinuxTargetPlatform.LinuxTargetSettings");
+
+		case EAudioPlatform::IOS:
+			return TEXT("/Script/IOSRuntimeSettings.IOSRuntimeSettings");
+
+		case EAudioPlatform::Android:
+			return TEXT("/Script/AndroidRuntimeSettings.AndroidRuntimeSettings");
+
+		case EAudioPlatform::XboxOne:
+			return TEXT("/Script/XboxOnePlatformEditor.XboxOneTargetSettings");
+
+		case EAudioPlatform::Playstation4:
+			return TEXT("/Script/PS4PlatformEditor.PS4TargetSettings");
+
+		case EAudioPlatform::Switch:
+			return TEXT("/Script/SwitchRuntimeSettings.SwitchRuntimeSettings");
+
+		case EAudioPlatform::HTML5:
+			return TEXT("/Script/HTML5PlatformEditor.HTML5TargetSettings");
+
+		case EAudioPlatform::Unknown:
+			return TEXT("");
+
+		default:
+			checkf(false, TEXT("Undefined audio platform."));
+			break;
+	}
+	return TEXT("");
+}
+
+/** Get the target setting name for each platform type. */
+FORCEINLINE const TCHAR* GetPluginConfigName(EAudioPlugin PluginType)
+{
+	switch (PluginType)
+	{
+		case EAudioPlugin::SPATIALIZATION:
+			return TEXT("SpatializationPlugin");
+
+		case EAudioPlugin::REVERB:
+			return TEXT("ReverbPlugin");
+
+		case EAudioPlugin::OCCLUSION:
+			return TEXT("OcclusionPlugin");
+
+		default:
+			checkf(false, TEXT("Undefined audio plugin type."));
+			return TEXT("");
+	}
+}
+
+/************************************************************************/
+/* Plugin Utilities                                                     */
+/************************************************************************/
+IAudioSpatializationFactory* AudioPluginUtilities::GetDesiredSpatializationPlugin(EAudioPlatform AudioPlatform)
+{
+	//Get the name of the desired spatialization plugin:
+	FString DesiredSpatializationPlugin = GetDesiredPluginName(EAudioPlugin::SPATIALIZATION, AudioPlatform);
+
+
+	TArray<IAudioSpatializationFactory *> SpatializationPluginFactories = IModularFeatures::Get().GetModularFeatureImplementations<IAudioSpatializationFactory>(IAudioSpatializationFactory::GetModularFeatureName());
+
+	//Iterate through all of the plugins we've discovered:
+	for (IAudioSpatializationFactory* PluginFactory : SpatializationPluginFactories)
+	{
+		//if this plugin's name matches the name found in the platform settings, use it:
+		if (PluginFactory->GetDisplayName().Equals(DesiredSpatializationPlugin))
+		{
+			return PluginFactory;
+		}
+	}
+
+	return nullptr;
+}
+
+IAudioReverbFactory* AudioPluginUtilities::GetDesiredReverbPlugin(EAudioPlatform AudioPlatform)
+{
+	//Get the name of the desired Reverb plugin:
+	FString DesiredReverbPlugin = GetDesiredPluginName(EAudioPlugin::REVERB, AudioPlatform);
+
+	TArray<IAudioReverbFactory *> ReverbPluginFactories = IModularFeatures::Get().GetModularFeatureImplementations<IAudioReverbFactory>(IAudioReverbFactory::GetModularFeatureName());
+
+	//Iterate through all of the plugins we've discovered:
+	for (IAudioReverbFactory* PluginFactory : ReverbPluginFactories)
+	{
+		//if this plugin's name matches the name found in the platform settings, use it:
+		if (PluginFactory->GetDisplayName().Equals(DesiredReverbPlugin))
+		{
+			return PluginFactory;
+		}
+	}
+
+	return nullptr;
+}
+
+IAudioOcclusionFactory* AudioPluginUtilities::GetDesiredOcclusionPlugin(EAudioPlatform AudioPlatform)
+{
+	FString DesiredOcclusionPlugin = GetDesiredPluginName(EAudioPlugin::OCCLUSION, AudioPlatform);
+
+	TArray<IAudioOcclusionFactory *> OcclusionPluginFactories = IModularFeatures::Get().GetModularFeatureImplementations<IAudioOcclusionFactory>(IAudioOcclusionFactory::GetModularFeatureName());
+
+	//Iterate through all of the plugins we've discovered:
+	for (IAudioOcclusionFactory* PluginFactory : OcclusionPluginFactories)
+	{
+		//if this plugin's name matches the name found in the platform settings, use it:
+		if (PluginFactory->GetDisplayName().Equals(DesiredOcclusionPlugin))
+		{
+			return PluginFactory;
+		}
+	}
+
+	return nullptr;
+}
+
+FString AudioPluginUtilities::GetDesiredPluginName(EAudioPlugin PluginType, EAudioPlatform AudioPlatform)
+{
+	FString PluginName;
+	GConfig->GetString(GetPlatformConfigSection(AudioPlatform), GetPluginConfigName(PluginType), PluginName, GEngineIni);
+	return PluginName;
+}
diff --git a/Engine/Source/Runtime/Engine/Private/AudioSettings.cpp b/Engine/Source/Runtime/Engine/Private/AudioSettings.cpp
--- a/Engine/Source/Runtime/Engine/Private/AudioSettings.cpp
+++ b/Engine/Source/Runtime/Engine/Private/AudioSettings.cpp
@@ -65,6 +65,7 @@ void UAudioSettings::AddDefaultSettings()
 	GConfig->GetInt(TEXT("Audio"), TEXT("MaxChannels"), DefaultSettings.MaxChannels, GEngineIni); // for backwards compatibility
 	QualityLevels.Add(DefaultSettings);
 	bAllowVirtualizedSounds = true;
+	DefaultReverbSendLevel = 0.2f;
 }
 
 #if WITH_EDITOR
diff --git a/Engine/Source/Runtime/Engine/Private/BoneContainer.cpp b/Engine/Source/Runtime/Engine/Private/BoneContainer.cpp
--- a/Engine/Source/Runtime/Engine/Private/BoneContainer.cpp
+++ b/Engine/Source/Runtime/Engine/Private/BoneContainer.cpp
@@ -25,7 +25,7 @@ FBoneContainer::FBoneContainer()
 	PoseToSkeletonBoneIndexArray.Empty();
 }
 
-FBoneContainer::FBoneContainer(const TArray<FBoneIndexType>& InRequiredBoneIndexArray, bool bDisableAnimCurves, UObject& InAsset)
+FBoneContainer::FBoneContainer(const TArray<FBoneIndexType>& InRequiredBoneIndexArray, const FCurveEvaluationOption& CurveEvalOption, UObject& InAsset)
 : BoneIndicesArray(InRequiredBoneIndexArray)
 , Asset(&InAsset)
 , AssetSkeletalMesh(NULL)
@@ -35,15 +35,15 @@ FBoneContainer::FBoneContainer(const TArray<FBoneIndexType>& InRequiredBoneIndex
 , bUseRAWData(false)
 , bUseSourceData(false)
 {
-	Initialize(bDisableAnimCurves);
+	Initialize(CurveEvalOption);
 }
 
-void FBoneContainer::InitializeTo(const TArray<FBoneIndexType>& InRequiredBoneIndexArray, bool bDisableAnimCurves, UObject& InAsset)
+void FBoneContainer::InitializeTo(const TArray<FBoneIndexType>& InRequiredBoneIndexArray, const FCurveEvaluationOption& CurveEvalOption, UObject& InAsset)
 {
 	BoneIndicesArray = InRequiredBoneIndexArray;
 	Asset = &InAsset;
 
-	Initialize(bDisableAnimCurves);
+	Initialize(CurveEvalOption);
 }
 
 struct FBoneContainerScratchArea : public TThreadSingleton<FBoneContainerScratchArea>
@@ -51,7 +51,7 @@ struct FBoneContainerScratchArea : public TThreadSingleton<FBoneContainerScratch
 	TArray<int32> MeshIndexToCompactPoseIndex;
 };
 
-void FBoneContainer::Initialize(bool bDisableAnimCurves)
+void FBoneContainer::Initialize(const FCurveEvaluationOption& CurveEvalOption)
 {
 	RefSkeleton = NULL;
 	UObject* AssetObj = Asset.Get();
@@ -182,12 +182,12 @@ void FBoneContainer::Initialize(bool bDisableAnimCurves)
 	}
 
 	// cache required curve UID list according to new bone sets
-	CacheRequiredAnimCurveUids(bDisableAnimCurves);
+	CacheRequiredAnimCurveUids(CurveEvalOption);
 }
 
-void FBoneContainer::CacheRequiredAnimCurveUids(bool bDisableAnimCurves)
+void FBoneContainer::CacheRequiredAnimCurveUids(const FCurveEvaluationOption& CurveEvalOption)
 {
-	if (!bDisableAnimCurves && AssetSkeleton.IsValid())
+	if (CurveEvalOption.bAllowCurveEvaluation && AssetSkeleton.IsValid())
 	{
 		// this is placeholder. In the future, this will change to work with linked joint of curve meta data
 		// anim curve name Uids; For now it adds all of them
@@ -205,26 +205,42 @@ void FBoneContainer::CacheRequiredAnimCurveUids(bool bDisableAnimCurves)
 			{
 				for (int32 CurveNameIndex = CurveNames.Num() - 1; CurveNameIndex >=0 ; --CurveNameIndex)
 				{
-					const FCurveMetaData* CurveMetaData = Mapping->GetCurveMetaData(CurveNames[CurveNameIndex]);
-					if (CurveMetaData && CurveMetaData->LinkedBones.Num() > 0)
+					const FName& CurveName = CurveNames[CurveNameIndex];
+					if (CurveEvalOption.DisallowedList && CurveEvalOption.DisallowedList->Contains(CurveName))
 					{
-						bool bRemove = true;
-						for (int32 LinkedBoneIndex = 0; LinkedBoneIndex < CurveMetaData->LinkedBones.Num(); ++LinkedBoneIndex)
+						//remove the UID
+						AnimCurveNameUids.RemoveAt(CurveNameIndex);
+					}
+					else
+					{
+						const FCurveMetaData* CurveMetaData = Mapping->GetCurveMetaData(CurveNames[CurveNameIndex]);
+						if (CurveMetaData)
 						{
-							const FBoneReference& BoneReference = CurveMetaData->LinkedBones[LinkedBoneIndex];
-							// we want to make sure all the joints are removed from RequiredBones before removing this UID
-							if (BoneReference.GetCompactPoseIndex(*this) != INDEX_NONE)
+							if (CurveMetaData->MaxLOD < CurveEvalOption.LODIndex)
 							{
-								// still has some joint that matters, do not remove
-								bRemove = false;
-								break;
+								AnimCurveNameUids.RemoveAt(CurveNameIndex);
+							}
+							else if (CurveMetaData->LinkedBones.Num() > 0)
+							{
+								bool bRemove = true;
+								for (int32 LinkedBoneIndex = 0; LinkedBoneIndex < CurveMetaData->LinkedBones.Num(); ++LinkedBoneIndex)
+								{
+									const FBoneReference& BoneReference = CurveMetaData->LinkedBones[LinkedBoneIndex];
+									// we want to make sure all the joints are removed from RequiredBones before removing this UID
+									if (BoneReference.GetCompactPoseIndex(*this) != INDEX_NONE)
+									{
+										// still has some joint that matters, do not remove
+										bRemove = false;
+										break;
+									}
+								}
+
+								if (bRemove)
+								{
+									//remove the UID
+									AnimCurveNameUids.RemoveAt(CurveNameIndex);
+								}
 							}
-						}
-
-						if (bRemove)
-						{
-							//remove the UID
-							AnimCurveNameUids.RemoveAt(CurveNameIndex);
 						}
 					}
 				}
diff --git a/Engine/Source/Runtime/Engine/Private/Components/PoseableMeshComponent.cpp b/Engine/Source/Runtime/Engine/Private/Components/PoseableMeshComponent.cpp
--- a/Engine/Source/Runtime/Engine/Private/Components/PoseableMeshComponent.cpp
+++ b/Engine/Source/Runtime/Engine/Private/Components/PoseableMeshComponent.cpp
@@ -31,7 +31,7 @@ bool UPoseableMeshComponent::AllocateTransformData()
 				RequiredBoneIndexArray[BoneIndex] = BoneIndex;
 			}
 
-			RequiredBones.InitializeTo(RequiredBoneIndexArray, true, *SkeletalMesh);
+			RequiredBones.InitializeTo(RequiredBoneIndexArray, FCurveEvaluationOption(false), *SkeletalMesh);
 		}
 
 		FillComponentSpaceTransforms();
diff --git a/Engine/Source/Runtime/Engine/Private/Components/SkeletalMeshComponent.cpp b/Engine/Source/Runtime/Engine/Private/Components/SkeletalMeshComponent.cpp
--- a/Engine/Source/Runtime/Engine/Private/Components/SkeletalMeshComponent.cpp
+++ b/Engine/Source/Runtime/Engine/Private/Components/SkeletalMeshComponent.cpp
@@ -32,6 +32,7 @@
 #include "ClothingSimulationInterface.h"
 #include "Features/IModularFeatures.h"
 #include "Misc/RuntimeErrors.h"
+#include "AnimPhysObjectVersion.h"
 
 #define LOCTEXT_NAMESPACE "SkeletalMeshComponent"
 
@@ -219,6 +220,7 @@ USkeletalMeshComponent::USkeletalMeshComponent(const FObjectInitializer& ObjectI
 	ClothingSimulationContext = nullptr;
 
 	bPostEvaluatingAnimation = false;
+	bAllowAnimCurveEvaluation = true;
 }
 
 void USkeletalMeshComponent::Serialize(FArchive& Ar)
@@ -288,6 +290,13 @@ void USkeletalMeshComponent::Serialize(FArchive& Ar)
 	{
 		BodyInstance.bAutoWeld = false;
 	}
+
+	Ar.UsingCustomVersion(FAnimPhysObjectVersion::GUID);
+	if (Ar.IsLoading() && Ar.CustomVer(FAnimPhysObjectVersion::GUID) < FAnimPhysObjectVersion::RenameDisableAnimCurvesToAllowAnimCurveEvaluation)
+	{
+		bAllowAnimCurveEvaluation = !bDisableAnimCurves_DEPRECATED;
+	}
+
 	PRAGMA_ENABLE_DEPRECATION_WARNINGS
 }
 
@@ -1336,20 +1345,22 @@ void USkeletalMeshComponent::RecalcRequiredCurves()
 		return;
 	}
 
+	const FCurveEvaluationOption CurveEvalOption(bAllowAnimCurveEvaluation, &DisallowedAnimCurves, PredictedLODLevel);
+
 	// make sure animation requiredcurve to mark as dirty
 	if (AnimScriptInstance)
 	{
-		AnimScriptInstance->RecalcRequiredCurves(bDisableAnimCurves);
+		AnimScriptInstance->RecalcRequiredCurves(CurveEvalOption);
 	}
 
 	for(UAnimInstance* SubInstance : SubInstances)
 	{
-		SubInstance->RecalcRequiredCurves(bDisableAnimCurves);
+		SubInstance->RecalcRequiredCurves(CurveEvalOption);
 	}
 
 	if(PostProcessAnimInstance)
 	{
-		PostProcessAnimInstance->RecalcRequiredCurves(bDisableAnimCurves);
+		PostProcessAnimInstance->RecalcRequiredCurves(CurveEvalOption);
 	}
 
 	MarkRequiredCurveUpToDate();
@@ -2062,7 +2073,7 @@ void USkeletalMeshComponent::PostAnimEvaluation(FAnimationEvaluationContext& Eva
 
 	if(PostProcessAnimInstance)
 	{
-		PostProcessAnimInstance->UpdateCurves(EvaluationContext.Curve);
+		PostProcessAnimInstance->UpdateCurves(AnimCurves);
 	}
 
 	bNeedToFlipSpaceBaseBuffers = true;
@@ -2437,7 +2448,8 @@ void USkeletalMeshComponent::SetAnimationMode(EAnimationMode::Type InAnimationMo
 	}
 
 	// when mode is swapped, make sure to reinitialize
-	// even if it was same mode
+	// even if it was same mode, this was due to users who wants to use BP construction script to do this
+	// if you use it in the construction script, it gets serialized, but it never instantiate. 
 	if(SkeletalMesh != nullptr)
 	{
 		if (InitializeAnimScriptInstance(true))
@@ -3255,12 +3267,83 @@ void USkeletalMeshComponent::ComputeTeleportDistanceThresholdInRadians()
 
 void USkeletalMeshComponent::SetDisableAnimCurves(bool bInDisableAnimCurves)
 {
-	if (bDisableAnimCurves != bInDisableAnimCurves)
+	SetAllowAnimCurveEvaluation(!bInDisableAnimCurves);
+}
+
+void USkeletalMeshComponent::SetAllowAnimCurveEvaluation(bool bInAllow)
+{
+	if (bAllowAnimCurveEvaluation != bInAllow)
 	{
-		bDisableAnimCurves = bInDisableAnimCurves;
+		bAllowAnimCurveEvaluation = bInAllow;
 		// clear cache uid version, so it will update required curves
 		CachedAnimCurveUidVersion = 0;
 	}
 }
 
+void USkeletalMeshComponent::AllowAnimCurveEvaluation(FName NameOfCurve, bool bAllow)
+{
+	// if allow is same as disallowed curve, which means it mismatches
+	if (bAllow == DisallowedAnimCurves.Contains(NameOfCurve))
+	{
+		if (bAllow)
+		{
+			DisallowedAnimCurves.Remove(NameOfCurve);
+			CachedAnimCurveUidVersion = 0;
+		}
+		else
+		{
+			DisallowedAnimCurves.Add(NameOfCurve);
+			CachedAnimCurveUidVersion = 0;
+
+		}
+	}
+}
+
+void USkeletalMeshComponent::ResetAllowedAnimCurveEvaluation()
+{
+	DisallowedAnimCurves.Reset();
+	CachedAnimCurveUidVersion = 0;
+}
+
+void USkeletalMeshComponent::SetAllowedAnimCurvesEvaluation(const TArray<FName>& List, bool bAllow)
+{
+	// Reset already clears the version - CachedAnimCurveUidVersion = 0;
+	ResetAllowedAnimCurveEvaluation();
+	if (bAllow)
+	{
+		struct FFilterDisallowedList
+		{
+			FFilterDisallowedList(const TArray<FName>& InAllowedList) : AllowedList(InAllowedList) {}
+
+			FORCEINLINE bool operator()(const FName& Name) const
+			{
+				return AllowedList.Contains(Name);
+			}
+			const TArray<FName>& AllowedList;
+		};
+
+		if (SkeletalMesh)
+		{
+			USkeleton* Skeleton = SkeletalMesh->Skeleton;
+			if (Skeleton)
+			{
+				const FSmartNameMapping* Mapping = Skeleton->GetSmartNameContainer(USkeleton::AnimCurveMappingName);
+				if (Mapping != nullptr)
+				{
+					TArray<FName> CurveNames;
+					Mapping->FillNameArray(CurveNames);
+
+					DisallowedAnimCurves = CurveNames;
+					DisallowedAnimCurves.RemoveAllSwap(FFilterDisallowedList(List));
+				}
+			}
+
+		}
+	}
+	else
+	{
+		DisallowedAnimCurves = List;
+	}
+}
+
 #undef LOCTEXT_NAMESPACE
diff --git a/Engine/Source/Runtime/Engine/Private/Components/SkinnedMeshComponent.cpp b/Engine/Source/Runtime/Engine/Private/Components/SkinnedMeshComponent.cpp
--- a/Engine/Source/Runtime/Engine/Private/Components/SkinnedMeshComponent.cpp
+++ b/Engine/Source/Runtime/Engine/Private/Components/SkinnedMeshComponent.cpp
@@ -441,7 +441,7 @@ void USkinnedMeshComponent::CreateRenderState_Concurrent()
 			FSkeletalMeshResource* SkelMeshResource = SkeletalMesh->GetResourceForRendering();
 
 			// Also check if skeletal mesh has too many bones/chunk for GPU skinning.
-			const bool bIsCPUSkinned = SkelMeshResource->RequiresCPUSkinning(SceneFeatureLevel) || (GIsEditor && ShouldCPUSkin());
+			const bool bIsCPUSkinned = SkelMeshResource->RequiresCPUSkinning(SceneFeatureLevel) || ShouldCPUSkin();
 			if(bIsCPUSkinned)
 			{
 				MeshObject = ::new FSkeletalMeshObjectCPUSkin(this, SkelMeshResource, SceneFeatureLevel);
@@ -2445,10 +2445,11 @@ void USkinnedMeshComponent::GetCPUSkinnedVertices(TArray<FFinalSkinVertex>& OutV
 	FlushRenderingCommands();
 
 	check(MeshObject);
-
+	check(MeshObject->IsCPUSkinned());
+		
 	// Copy our vertices out. We know we are using CPU skinning now, so this cast is safe
 	OutVertices = static_cast<FSkeletalMeshObjectCPUSkin*>(MeshObject)->GetCachedFinalVertices();
-
+	
 	// switch skinning mode, LOD etc. back
 	bCPUSkinning = bCachedCPUSkinning;
 	ForcedLodModel = 0;
diff --git a/Engine/Source/Runtime/Engine/Private/GameplayStatics.cpp b/Engine/Source/Runtime/Engine/Private/GameplayStatics.cpp
--- a/Engine/Source/Runtime/Engine/Private/GameplayStatics.cpp
+++ b/Engine/Source/Runtime/Engine/Private/GameplayStatics.cpp
@@ -964,7 +964,7 @@ void UGameplayStatics::SetGlobalListenerFocusParameters(const UObject* WorldCont
 	}
 }
 
-void UGameplayStatics::PlaySound2D(const UObject* WorldContextObject, class USoundBase* Sound, float VolumeMultiplier, float PitchMultiplier, float StartTime, class USoundConcurrency* ConcurrencySettings)
+void UGameplayStatics::PlaySound2D(const UObject* WorldContextObject, class USoundBase* Sound, float VolumeMultiplier, float PitchMultiplier, float StartTime, class USoundConcurrency* ConcurrencySettings, AActor* OwningActor)
 {
 	if (!Sound || !GEngine || !GEngine->UseSound())
 	{
@@ -994,6 +994,8 @@ void UGameplayStatics::PlaySound2D(const UObject* WorldContextObject, class USou
 		NewActiveSound.Priority = Sound->Priority;
 		NewActiveSound.SubtitlePriority = Sound->GetSubtitlePriority();
 
+		NewActiveSound.SetOwner(OwningActor);
+
 		AudioDevice->AddNewActiveSound(NewActiveSound);
 	}
 }
@@ -1051,7 +1053,7 @@ UAudioComponent* UGameplayStatics::SpawnSound2D(const UObject* WorldContextObjec
 	return AudioComponent;
 }
 
-void UGameplayStatics::PlaySoundAtLocation(const UObject* WorldContextObject, class USoundBase* Sound, FVector Location, FRotator Rotation, float VolumeMultiplier, float PitchMultiplier, float StartTime, class USoundAttenuation* AttenuationSettings, class USoundConcurrency* ConcurrencySettings)
+void UGameplayStatics::PlaySoundAtLocation(const UObject* WorldContextObject, class USoundBase* Sound, FVector Location, FRotator Rotation, float VolumeMultiplier, float PitchMultiplier, float StartTime, class USoundAttenuation* AttenuationSettings, class USoundConcurrency* ConcurrencySettings, AActor* OwningActor)
 {
 	if (!Sound || !GEngine || !GEngine->UseSound())
 	{
@@ -1066,7 +1068,7 @@ void UGameplayStatics::PlaySoundAtLocation(const UObject* WorldContextObject, cl
 
 	if (FAudioDevice* AudioDevice = ThisWorld->GetAudioDevice())
 	{
-		AudioDevice->PlaySoundAtLocation(Sound, ThisWorld, VolumeMultiplier, PitchMultiplier, StartTime, Location, Rotation, AttenuationSettings, ConcurrencySettings);
+		AudioDevice->PlaySoundAtLocation(Sound, ThisWorld, VolumeMultiplier, PitchMultiplier, StartTime, Location, Rotation, AttenuationSettings, ConcurrencySettings, nullptr, OwningActor);
 	}
 }
 
diff --git a/Engine/Source/Runtime/Engine/Private/LODActor.cpp b/Engine/Source/Runtime/Engine/Private/LODActor.cpp
--- a/Engine/Source/Runtime/Engine/Private/LODActor.cpp
+++ b/Engine/Source/Runtime/Engine/Private/LODActor.cpp
@@ -21,6 +21,7 @@
 #include "Editor.h"
 
 #include "HierarchicalLODUtilitiesModule.h"
+#include "ObjectTools.h"
 #endif
 
 #define LOCTEXT_NAMESPACE "LODActor"
@@ -536,6 +537,8 @@ void ALODActor::SetIsDirty(const bool bNewState)
 		{
 			GEditor->BroadcastHLODActorMarkedDirty(this);
 		}
+		PreviousSubObjects.Append(SubObjects);
+		SubObjects.Empty();
 #endif // WITH_EDITOR
 	}	
 	else
@@ -638,7 +641,7 @@ void ALODActor::SetStaticMesh(class UStaticMesh* InStaticMesh)
 		SetIsDirty(false);
 
 		ensure(StaticMeshComponent->GetStaticMesh() == InStaticMesh);
-		if (InStaticMesh->RenderData && InStaticMesh->RenderData->LODResources.Num() > 0)
+		if (InStaticMesh && InStaticMesh->RenderData && InStaticMesh->RenderData->LODResources.Num() > 0)
 		{
 			NumTrianglesInMergedMesh = InStaticMesh->RenderData->LODResources[0].GetNumTriangles();
 		}
@@ -768,10 +771,33 @@ void ALODActor::Serialize(FArchive& Ar)
 
 	bRequiresLODScreenSizeConversion = Ar.CustomVer(FFrameworkObjectVersion::GUID) < FFrameworkObjectVersion::LODsUseResolutionIndependentScreenSize;
 }
-#endif
 
-//////////////////////////////////////////////////////////////////////////
-// AHLODMeshCullingVolume
+void ALODActor::PreSave(const class ITargetPlatform* TargetPlatform)
+{
+	AActor::PreSave(TargetPlatform);
+	if (PreviousSubObjects.Num())
+	{
+		PreviousSubObjects.RemoveAll([](const UObject* Object) -> bool { return Object == nullptr; });
+		ObjectTools::DeleteObjectsUnchecked(PreviousSubObjects);
+		PreviousSubObjects.Empty();
+	}
+}
 
+void ALODActor::BeginDestroy()
+{
+	AActor::BeginDestroy();
+	if (PreviousSubObjects.Num())
+	{
+		for (UObject* Object : PreviousSubObjects)
+		{
+			if (Object)
+			{
+				Object->MarkPendingKill();
+			}
+		}
+		PreviousSubObjects.Empty();
+	}
+}
 
+#endif
 #undef LOCTEXT_NAMESPACE
diff --git a/Engine/Source/Runtime/Engine/Private/PhysicsEngine/BodyInstance.cpp b/Engine/Source/Runtime/Engine/Private/PhysicsEngine/BodyInstance.cpp
--- a/Engine/Source/Runtime/Engine/Private/PhysicsEngine/BodyInstance.cpp
+++ b/Engine/Source/Runtime/Engine/Private/PhysicsEngine/BodyInstance.cpp
@@ -70,18 +70,20 @@ int32 FillInlinePxShapeArray_AssumesLocked(FInlinePxShapeArray& Array, const phy
 	return NumShapes;
 }
 
-PxTransform GetPxTransform_AssumesLocked(const PxRigidActor* PRigidActor)
+PxTransform GetPxTransform_AssumesLocked(const PxRigidActor* PRigidActor, bool bForceGlobalPose = false)
 {
-
-	if (const PxRigidDynamic* PRigidDynamic = PRigidActor->is<PxRigidDynamic>())
+	if (!bForceGlobalPose)
 	{
-		PxTransform Transform;
-
-		// getKinematicTarget will do checks to ensure the Body is both kinematic and
-		// that the kinematic target has been set.
-		if (PRigidDynamic->getKinematicTarget(Transform))
+		if (const PxRigidDynamic* const PRigidDynamic = PRigidActor->is<PxRigidDynamic>())
 		{
-			return Transform;
+		    PxTransform Transform;
+    
+		    // getKinematicTarget will do checks to ensure the Body is both kinematic and
+		    // that the kinematic target has been set.
+		    if (PRigidDynamic->getKinematicTarget(Transform))
+		    {
+			    return Transform;
+		    }
 		}
 	}
 
@@ -2697,13 +2699,13 @@ bool FBodyInstance::IsValidBodyInstance() const
 }
 
 template <bool NeedsLock>
-FTransform GetUnrealWorldTransformImp(const FBodyInstance* BodyInstance, bool bWithProjection)
+FTransform GetUnrealWorldTransformImp(const FBodyInstance* BodyInstance, bool bWithProjection, bool bGlobalPose)
 {
 	FTransform WorldTM = FTransform::Identity;
 #if WITH_PHYSX
 	FPhysXSupport<NeedsLock>::ExecuteOnPxRigidActorReadOnly(BodyInstance, [&](const PxRigidActor* PActor)
 	{
-		PxTransform PTM = GetPxTransform_AssumesLocked(PActor);
+		PxTransform PTM = GetPxTransform_AssumesLocked(PActor, bGlobalPose);
 		WorldTM = P2UTransform(PTM);
 
 		if(bWithProjection)
@@ -2716,15 +2718,15 @@ FTransform GetUnrealWorldTransformImp(const FBodyInstance* BodyInstance, bool bW
 	return WorldTM;
 }
 
-FTransform FBodyInstance::GetUnrealWorldTransform(bool bWithProjection /* = true*/) const
+FTransform FBodyInstance::GetUnrealWorldTransform(bool bWithProjection /* = true*/, bool bForceGlobalPose /* = true*/) const
 {
-	return GetUnrealWorldTransformImp<true>(this, bWithProjection);
+	return GetUnrealWorldTransformImp<true>(this, bWithProjection, bForceGlobalPose);
 }
 
 
-FTransform FBodyInstance::GetUnrealWorldTransform_AssumesLocked(bool bWithProjection /* = true*/) const
+FTransform FBodyInstance::GetUnrealWorldTransform_AssumesLocked(bool bWithProjection /* = true*/, bool bForceGlobalPose /* = true*/) const
 {
-	return GetUnrealWorldTransformImp<false>(this, bWithProjection);
+	return GetUnrealWorldTransformImp<false>(this, bWithProjection, bForceGlobalPose);
 }
 
 void FBodyInstance::SetBodyTransform(const FTransform& NewTransform, ETeleportType Teleport)
diff --git a/Engine/Source/Runtime/Engine/Private/PhysicsEngine/BodySetup.cpp b/Engine/Source/Runtime/Engine/Private/PhysicsEngine/BodySetup.cpp
--- a/Engine/Source/Runtime/Engine/Private/PhysicsEngine/BodySetup.cpp
+++ b/Engine/Source/Runtime/Engine/Private/PhysicsEngine/BodySetup.cpp
@@ -1621,17 +1621,17 @@ float FKAggregateGeom::GetVolume(const FVector& Scale) const
 	return Volume;
 }
 
-int32 FKAggregateGeom::GetElementCount(int32 Type) const
+int32 FKAggregateGeom::GetElementCount(EAggCollisionShape::Type Type) const
 {
 	switch (Type)
 	{
-	case KPT_Box:
+	case EAggCollisionShape::Box:
 		return BoxElems.Num();
-	case KPT_Convex:
+	case EAggCollisionShape::Convex:
 		return ConvexElems.Num();
-	case KPT_Sphyl:
+	case EAggCollisionShape::Sphyl:
 		return SphylElems.Num();
-	case KPT_Sphere:
+	case EAggCollisionShape::Sphere:
 		return SphereElems.Num();
 	default:
 		return 0;
diff --git a/Engine/Source/Runtime/Engine/Private/PhysicsEngine/PhysAnim.cpp b/Engine/Source/Runtime/Engine/Private/PhysicsEngine/PhysAnim.cpp
--- a/Engine/Source/Runtime/Engine/Private/PhysicsEngine/PhysAnim.cpp
+++ b/Engine/Source/Runtime/Engine/Private/PhysicsEngine/PhysAnim.cpp
@@ -185,8 +185,8 @@ void USkeletalMeshComponent::PerformBlendPhysicsBones(const TArray<FBoneIndexTyp
 		FTransform TM;
 	};
 
-#define DEPERCATED_PHYSBLEND_UPDATES_PHYSX 0	//If you really need the old behavior of physics blending set this to 1. Note this is inefficient and doesn't lead to good results. Please use UPhysicalAnimationComponent
-#if DEPERCATED_PHYSBLEND_UPDATES_PHYSX
+#define DEPRECATED_PHYSBLEND_UPDATES_PHYSX 0	//If you really need the old behavior of physics blending set this to 1. Note this is inefficient and doesn't lead to good results. Please use UPhysicalAnimationComponent
+#if DEPRECATED_PHYSBLEND_UPDATES_PHYSX
 	TArray<FBodyTMPair, TMemStackAllocator<alignof(FBodyTMPair)>> PendingBodyTMs;
 #endif
 
@@ -208,7 +208,7 @@ void USkeletalMeshComponent::PerformBlendPhysicsBones(const TArray<FBoneIndexTyp
 			// See if this is a physics bone..
 			int32 BodyIndex = PhysicsAsset->FindBodyIndex(SkeletalMesh->RefSkeleton.GetBoneName(BoneIndex));
 			// need to update back to physX so that physX knows where it was after blending
-#if DEPERCATED_PHYSBLEND_UPDATES_PHYSX
+#if DEPRECATED_PHYSBLEND_UPDATES_PHYSX
 			bool bUpdatePhysics = false;
 #endif
 			FBodyInstance* PhysicsAssetBodyInstance = nullptr;
@@ -277,7 +277,7 @@ void USkeletalMeshComponent::PerformBlendPhysicsBones(const TArray<FBoneIndexTyp
 							bSetParentScale = true;
 						}
 
-						#if DEPERCATED_PHYSBLEND_UPDATES_PHYSX
+						#if DEPRECATED_PHYSBLEND_UPDATES_PHYSX
 						if (UsePhysWeight < 1.f)
 						{
 							bUpdatePhysics = true;
@@ -313,7 +313,7 @@ void USkeletalMeshComponent::PerformBlendPhysicsBones(const TArray<FBoneIndexTyp
 					EditableComponentSpaceTransforms[BoneIndex] = Bodies[BodyIndex]->GetUnrealWorldTransform_AssumesLocked().GetRelativeTransform(NewComponentToWorld);
 				}
 			}
-#if DEPERCATED_PHYSBLEND_UPDATES_PHYSX
+#if DEPRECATED_PHYSBLEND_UPDATES_PHYSX
 			if (bUpdatePhysics && PhysicsAssetBodyInstance)
 			{
 				//This is extremely inefficient. We need to obtain a write lock which will block other threads from blending
@@ -327,7 +327,7 @@ void USkeletalMeshComponent::PerformBlendPhysicsBones(const TArray<FBoneIndexTyp
 
 #if WITH_PHYSX
 	}	//end scope for read lock
-#if DEPERCATED_PHYSBLEND_UPDATES_PHYSX
+#if DEPRECATED_PHYSBLEND_UPDATES_PHYSX
 	if(PendingBodyTMs.Num() && CVarPhysicsAnimBlendUpdatesPhysX.GetValueOnAnyThread())
 	{
 		//This is extremely inefficient. We need to obtain a write lock which will block other threads from blending
diff --git a/Engine/Source/Runtime/Engine/Private/PhysicsEngine/PhysDerivedData.h b/Engine/Source/Runtime/Engine/Private/PhysicsEngine/PhysDerivedData.h
--- a/Engine/Source/Runtime/Engine/Private/PhysicsEngine/PhysDerivedData.h
+++ b/Engine/Source/Runtime/Engine/Private/PhysicsEngine/PhysDerivedData.h
@@ -50,7 +50,7 @@ public:
 		// This is a version string that mimics the old versioning scheme. If you
 		// want to bump this version, generate a new guid using VS->Tools->Create GUID and
 		// return it here. Ex.
-		return TEXT("48A34D3CEFD54C709592943289228CCF");	
+		return TEXT("3B0AB674B19F4A698B28F7025C031EA2");	
 	}
 
 	virtual FString GetPluginSpecificCacheKeySuffix() const override
diff --git a/Engine/Source/Runtime/Engine/Private/PhysicsEngine/PhysDrawing.cpp b/Engine/Source/Runtime/Engine/Private/PhysicsEngine/PhysDrawing.cpp
--- a/Engine/Source/Runtime/Engine/Private/PhysicsEngine/PhysDrawing.cpp
+++ b/Engine/Source/Runtime/Engine/Private/PhysicsEngine/PhysDrawing.cpp
@@ -30,9 +30,10 @@ static const int32 DrawConeLimitSides = 40;
 static const float DebugJointPosSize = 5.0f;
 static const float DebugJointAxisSize = 20.0f;
 
-static const float JointRenderThickness = 0.3f;
-static const float JointRenderSize = 10.f;
-static const float LimitRenderSize = 16.0f;
+static const float JointRenderThickness = 0.1f;
+static const float UnselectedJointRenderSize = 4.f;
+static const float SelectedJointRenderSize = 10.f;
+static const float LimitRenderSize = 0.16f;
 
 static const FColor JointUnselectedColor(255, 0, 255);
 static const FColor JointRed(FColor::Red);
@@ -899,7 +900,7 @@ void FConstraintInstance::DrawConstraintImp(const FPDIOrCollector& PDIOrCollecto
 	return;
 #endif
 
-	const ESceneDepthPriorityGroup Layer = ESceneDepthPriorityGroup::SDPG_Foreground;
+	const ESceneDepthPriorityGroup Layer = ESceneDepthPriorityGroup::SDPG_World;
 	FPrimitiveDrawInterface* PDI = PDIOrCollector.GetPDI();
 
 	check((GEngine->ConstraintLimitMaterialX != nullptr) && (GEngine->ConstraintLimitMaterialY != nullptr) && (GEngine->ConstraintLimitMaterialZ != nullptr));
@@ -914,22 +915,14 @@ void FConstraintInstance::DrawConstraintImp(const FPDIOrCollector& PDIOrCollecto
 	FVector Con1Pos = Con1Frame.GetTranslation();
 	FVector Con2Pos = Con2Frame.GetTranslation();
 
-	float Length = JointRenderSize;
-	float Thickness = LimitDrawScale * JointRenderThickness;
+	float Length = (bDrawSelected ? SelectedJointRenderSize : UnselectedJointRenderSize) * Scale;
+	float Thickness = JointRenderThickness;
 
 	// Special mode for drawing joints just as points..
-	if(bDrawAsPoint)
+	if(bDrawAsPoint && !bDrawSelected)
 	{
-		if(bDrawSelected)
-		{
-			PDI->DrawPoint( Con1Frame.GetTranslation(), JointRed, 3.f, Layer );
-			PDI->DrawPoint( Con2Frame.GetTranslation(), JointBlue, 3.f, Layer );
-		}
-		else
-		{
-			PDI->DrawPoint( Con1Frame.GetTranslation(), JointUnselectedColor, 3.f, Layer );
-			PDI->DrawPoint( Con2Frame.GetTranslation(), JointUnselectedColor, 3.f, Layer );
-		}
+		PDI->DrawPoint( Con1Frame.GetTranslation(), JointUnselectedColor, 4.f, ESceneDepthPriorityGroup::SDPG_Foreground );
+		PDI->DrawPoint( Con2Frame.GetTranslation(), JointUnselectedColor, 4.f, ESceneDepthPriorityGroup::SDPG_Foreground );
 
 		// do nothing else in this mode.
 		return;
diff --git a/Engine/Source/Runtime/Engine/Private/PhysicsEngine/PhysScene.cpp b/Engine/Source/Runtime/Engine/Private/PhysicsEngine/PhysScene.cpp
--- a/Engine/Source/Runtime/Engine/Private/PhysicsEngine/PhysScene.cpp
+++ b/Engine/Source/Runtime/Engine/Private/PhysicsEngine/PhysScene.cpp
@@ -422,7 +422,7 @@ static FAutoConsoleCommandWithWorldAndArgs GSetPhysXTreeRebuildRate(TEXT("p.Phys
 );
 
 
-/** Exposes creation of physics-engine scene outside Engine (for use with PhAT for example). */
+/** Exposes creation of physics-engine scene outside Engine (for use with Physics Asset Editor for example). */
 FPhysScene::FPhysScene()
 {
 	LineBatcher = NULL;
diff --git a/Engine/Source/Runtime/Engine/Private/PhysicsEngine/PhysSubstepTasks.h b/Engine/Source/Runtime/Engine/Private/PhysicsEngine/PhysSubstepTasks.h
--- a/Engine/Source/Runtime/Engine/Private/PhysicsEngine/PhysSubstepTasks.h
+++ b/Engine/Source/Runtime/Engine/Private/PhysicsEngine/PhysSubstepTasks.h
@@ -75,7 +75,7 @@ struct FKinematicTarget
 	FKinematicTarget(FBodyInstance* Body, const FTransform& TM) :
 		BodyInstance(Body),
 		TargetTM(TM),
-		OriginalTM(Body->GetUnrealWorldTransform_AssumesLocked())
+		OriginalTM(Body->GetUnrealWorldTransform(true, true))
 	{
 	}
 
diff --git a/Engine/Source/Runtime/Engine/Private/PhysicsEngine/PhysicsAsset.cpp b/Engine/Source/Runtime/Engine/Private/PhysicsEngine/PhysicsAsset.cpp
--- a/Engine/Source/Runtime/Engine/Private/PhysicsEngine/PhysicsAsset.cpp
+++ b/Engine/Source/Runtime/Engine/Private/PhysicsEngine/PhysicsAsset.cpp
@@ -15,6 +15,10 @@
 #include "Logging/MessageLog.h"
 #include "UObjectIterator.h"
 
+#if WITH_EDITOR
+#include "Misc/MessageDialog.h"
+#endif
+
 #define LOCTEXT_NAMESPACE "PhysicsAsset"
 
 ///////////////////////////////////////	
@@ -262,6 +266,21 @@ void UPhysicsAsset::DisableCollision(int32 BodyIndexA, int32 BodyIndexB)
 	CollisionDisableTable.Add(Key, 0);
 }
 
+bool UPhysicsAsset::IsCollisionEnabled(int32 BodyIndexA, int32 BodyIndexB) const
+{
+	if(BodyIndexA == BodyIndexB)
+	{
+		return false;
+	}
+
+	if(CollisionDisableTable.Find(FRigidBodyIndexPair(BodyIndexA, BodyIndexB)))
+	{
+		return false;
+	}
+
+	return true;
+}
+
 FBox UPhysicsAsset::CalcAABB(const USkinnedMeshComponent* MeshComp, const FTransform& LocalToWorld) const
 {
 	FBox Box(ForceInit);
@@ -757,6 +776,47 @@ void UPhysicsAsset::RefreshPhysicsAssetChange() const
 		}
 	}
 }
+
+USkeletalMesh* UPhysicsAsset::GetPreviewMesh() const
+{
+	USkeletalMesh* PreviewMesh = PreviewSkeletalMesh.Get();
+	if(!PreviewMesh)
+	{
+		// if preview mesh isn't loaded, see if we have set
+		FSoftObjectPath PreviewMeshStringRef = PreviewSkeletalMesh.ToSoftObjectPath();
+		// load it since now is the time to load
+		if(!PreviewMeshStringRef.ToString().IsEmpty())
+		{
+			PreviewMesh = Cast<USkeletalMesh>(StaticLoadObject(USkeletalMesh::StaticClass(), nullptr, *PreviewMeshStringRef.ToString(), nullptr, LOAD_None, nullptr));
+		}
+	}
+
+	return PreviewMesh;
+}
+
+void UPhysicsAsset::SetPreviewMesh(USkeletalMesh* PreviewMesh)
+{
+	if(PreviewMesh)
+	{
+		// See if any bones are missing from the skeletal mesh we are trying to use
+		// @todo Could do more here - check for bone lengths etc. Maybe modify asset?
+		for (int32 i = 0; i < SkeletalBodySetups.Num(); ++i)
+		{
+			FName BodyName = SkeletalBodySetups[i]->BoneName;
+			int32 BoneIndex = PreviewMesh->RefSkeleton.FindBoneIndex(BodyName);
+			if (BoneIndex == INDEX_NONE)
+			{
+				FMessageDialog::Open(EAppMsgType::Ok,
+					FText::Format( LOCTEXT("BoneMissingFromSkelMesh", "The SkeletalMesh is missing bone '{0}' needed by this PhysicsAsset."), FText::FromName(BodyName) ));
+				return;
+			}
+		}
+	}
+
+	Modify();
+	PreviewSkeletalMesh = PreviewMesh;
+}
+
 #endif
 
 void UPhysicsAsset::GetResourceSizeEx(FResourceSizeEx& CumulativeResourceSize)
diff --git a/Engine/Source/Runtime/Engine/Private/PreviewScene.cpp b/Engine/Source/Runtime/Engine/Private/PreviewScene.cpp
--- a/Engine/Source/Runtime/Engine/Private/PreviewScene.cpp
+++ b/Engine/Source/Runtime/Engine/Private/PreviewScene.cpp
@@ -214,3 +214,9 @@ void FPreviewScene::SaveSettings(const TCHAR* Section)
 {
 	GConfig->SetRotator( Section, TEXT("LightDir"), GetLightDirection(), GEditorPerProjectIni );
 }
+
+FLinearColor FPreviewScene::GetBackgroundColor() const
+{
+	FLinearColor BackgroundColor = FColor(55, 55, 55);
+	return BackgroundColor;
+}
\ No newline at end of file
diff --git a/Engine/Source/Runtime/Engine/Private/SkeletalMesh.cpp b/Engine/Source/Runtime/Engine/Private/SkeletalMesh.cpp
--- a/Engine/Source/Runtime/Engine/Private/SkeletalMesh.cpp
+++ b/Engine/Source/Runtime/Engine/Private/SkeletalMesh.cpp
@@ -5568,7 +5568,7 @@ void FSkeletalMeshSceneProxy::GetDynamicElementsSection(const TArray<const FScen
 			}
 			else
 			{
-				Mesh.bUseSelectionOutline = false;
+				Mesh.bUseSelectionOutline = !bCanHighlightSelectedSections && bIsSelected;
 			}
 		#endif
 
diff --git a/Engine/Source/Runtime/Engine/Private/SkeletalMeshComponentPhysics.cpp b/Engine/Source/Runtime/Engine/Private/SkeletalMeshComponentPhysics.cpp
--- a/Engine/Source/Runtime/Engine/Private/SkeletalMeshComponentPhysics.cpp
+++ b/Engine/Source/Runtime/Engine/Private/SkeletalMeshComponentPhysics.cpp
@@ -551,14 +551,20 @@ void USkeletalMeshComponent::InitArticulated(FPhysScene* PhysScene)
 	uint32 SkelMeshCompID = GetUniqueID();
 	PhysScene->DeferredAddCollisionDisableTable(SkelMeshCompID, &PhysicsAsset->CollisionDisableTable);
 
-	int32 NumBodies = PhysicsAsset->SkeletalBodySetups.Num();
-	if(Aggregate == NULL && NumBodies > RagdollAggregateThreshold && NumBodies <= AggregateMaxSize)
+	int32 NumShapes = 0;
+	const int32 NumBodies = PhysicsAsset->SkeletalBodySetups.Num();
+	for(int32 BodyIndex = 0; BodyIndex < NumBodies; ++BodyIndex)
+	{
+		NumShapes += PhysicsAsset->SkeletalBodySetups[BodyIndex]->AggGeom.GetElementCount();
+	}
+
+	if(Aggregate == NULL && NumShapes > RagdollAggregateThreshold && NumShapes <= AggregateMaxSize)
 	{
 		Aggregate = GPhysXSDK->createAggregate(PhysicsAsset->SkeletalBodySetups.Num(), true);
 	}
-	else if(Aggregate && NumBodies > AggregateMaxSize)
+	else if(Aggregate && NumShapes > AggregateMaxSize)
 	{
-		UE_LOG(LogSkeletalMesh, Log, TEXT("USkeletalMeshComponent::InitArticulated : Too many bodies to create aggregate, Max: %u, This: %d"), AggregateMaxSize, NumBodies);
+		UE_LOG(LogSkeletalMesh, Log, TEXT("USkeletalMeshComponent::InitArticulated : Too many shapes to create aggregate, Max: %u, This: %d"), AggregateMaxSize, NumShapes);
 	}
 #endif //WITH_PHYSX
 
@@ -592,6 +598,7 @@ void USkeletalMeshComponent::InstantiatePhysicsAsset(const UPhysicsAsset& PhysAs
 	// Create all the OutBodies.
 	check(OutBodies.Num() == 0);
 	OutBodies.AddZeroed(NumOutBodies);
+
 	for (int32 BodyIdx = 0; BodyIdx<NumOutBodies; BodyIdx++)
 	{
 		UBodySetup* PhysicsAssetBodySetup = PhysAsset.SkeletalBodySetups[BodyIdx];
@@ -1841,6 +1848,7 @@ FVector USkeletalMeshComponent::GetSkinnedVertexPosition(int32 VertexIndex) cons
 
 				const FSkelMeshSection& ClothSection = Model.Sections[Section.CorrespondClothSectionIndex];
 				ClothAssetIndex = ClothSection.CorrespondClothAssetIndex;
+				ClothAssetGuid = ClothSection.ClothingData.AssetGuid;
 
 				// the index can exceed the range because this vertex index is based on the corresponding original section
 				// the number of cloth chunk's vertices is not always same as the corresponding one 
@@ -2135,13 +2143,13 @@ void USkeletalMeshComponent::AddClothingBounds(FBoxSphereBounds& InOutBounds, co
 
 void USkeletalMeshComponent::RecreateClothingActors()
 {
+	ReleaseAllClothingResources();
+
 	if(SkeletalMesh == nullptr)
 	{
 		return;
 	}
 
-	ReleaseAllClothingResources();
-
 	if (bDisableClothSimulation)
 	{
 		return;
@@ -2786,7 +2794,7 @@ bool USkeletalMeshComponent::GetClothSimulatedPosition_GameThread(const FGuid& A
 
 		if(ActorData && ActorData->Positions.IsValidIndex(VertexIndex))
 		{
-			OutSimulPos = GetComponentTransform().TransformPosition(ActorData->Positions[VertexIndex]);
+			OutSimulPos = ActorData->Positions[VertexIndex];
 
 			bSucceed = true;
 		}
diff --git a/Engine/Source/Runtime/Engine/Private/SoundBase.cpp b/Engine/Source/Runtime/Engine/Private/SoundBase.cpp
--- a/Engine/Source/Runtime/Engine/Private/SoundBase.cpp
+++ b/Engine/Source/Runtime/Engine/Private/SoundBase.cpp
@@ -100,6 +100,11 @@ void USoundBase::GetSoundSubmixSends(TArray<FSoundSubmixSendInfo>& OutSends) con
 	OutSends = SoundSubmixSends;
 }
 
+void USoundBase::GetSoundSourceBusSends(TArray<FSoundSourceBusSendInfo>& OutSends) const
+{
+	OutSends = BusSends;
+}
+
 const FSoundConcurrencySettings* USoundBase::GetSoundConcurrencySettingsToApply()
 {
 	if (bOverrideConcurrency)
diff --git a/Engine/Source/Runtime/Engine/Private/SoundEffectBase.cpp b/Engine/Source/Runtime/Engine/Private/SoundEffectBase.cpp
--- a/Engine/Source/Runtime/Engine/Private/SoundEffectBase.cpp
+++ b/Engine/Source/Runtime/Engine/Private/SoundEffectBase.cpp
@@ -34,6 +34,8 @@ void FSoundEffectBase::SetPreset(USoundEffectPreset* Inpreset)
 
 void FSoundEffectBase::Update()
 {
+	PumpPendingMessages();
+
 	if (bChanged && Preset)
 	{
 		OnPresetChanged();
@@ -57,4 +59,16 @@ bool FSoundEffectBase::IsParentPreset(USoundEffectPreset* InPreset) const
 	return ParentPreset == InPreset;
 }
 
-
+void FSoundEffectBase::EffectCommand(TFunction<void()> Command)
+{
+	CommandQueue.Enqueue(MoveTemp(Command));
+}
+void FSoundEffectBase::PumpPendingMessages()
+{
+	// Pumps the commadn queue
+	TFunction<void()> Command;
+	while (CommandQueue.Dequeue(Command))
+	{
+		Command();
+	}
+}
diff --git a/Engine/Source/Runtime/Engine/Private/SoundEffectPreset.cpp b/Engine/Source/Runtime/Engine/Private/SoundEffectPreset.cpp
--- a/Engine/Source/Runtime/Engine/Private/SoundEffectPreset.cpp
+++ b/Engine/Source/Runtime/Engine/Private/SoundEffectPreset.cpp
@@ -13,6 +13,14 @@ USoundEffectPreset::USoundEffectPreset(const FObjectInitializer& ObjectInitializ
 	
 }
 
+void USoundEffectPreset::EffectCommand(TFunction<void()> Command)
+{
+	for (int32 i = 0; i < Instances.Num(); ++i)
+	{
+		Instances[i]->EffectCommand(Command);
+	}
+}
+
 void USoundEffectPreset::Update()
 {
 	for (int32 i = 0; i < Instances.Num(); ++i)
@@ -27,6 +35,9 @@ void USoundEffectPreset::AddEffectInstance(FSoundEffectBase* InSource)
 	{
 		bInitialized = true;
 		Init();
+
+		// Call the optional virtual function which subclasses can implement if they need initialization
+		OnInit();
 	}
 
 	Instances.AddUnique(InSource);
diff --git a/Engine/Source/Runtime/Engine/Private/SoundSourceBus.cpp b/Engine/Source/Runtime/Engine/Private/SoundSourceBus.cpp
new file mode 100644
--- /dev/null
+++ b/Engine/Source/Runtime/Engine/Private/SoundSourceBus.cpp
@@ -0,0 +1,65 @@
+// Copyright 1998-2017 Epic Games, Inc. All Rights Reserved.
+
+
+#include "Sound/SoundSourceBus.h"
+#include "AudioDeviceManager.h"
+#include "EngineGlobals.h"
+#include "Engine/Engine.h"
+#include "UObject/UObjectIterator.h"
+#include "ActiveSound.h"
+
+void USoundSourceBus::PostLoad()
+{
+	Super::PostLoad();
+
+	// This is a bus. This will result in the decompression type to be set as DTYPE_Bus. Audio won't be generated from this object but from instance data in audio mixer.
+	bIsBus = true;
+
+	// Allow users to manually set the source bus duration
+	Duration = GetDuration();
+
+	// This sound wave is looping if the source bus duration is 0.0f
+	bLooping = (SourceBusDuration == 0.0f);
+
+	// Keep playing this bus when the volume is 0
+	// Note source buses can't ever be truly virtual as they are procedurally generated.
+	bVirtualizeWhenSilent = !bAutoDeactivateWhenSilent;
+
+	// Use the main/default audio device for storing and retrieving sound class properties
+	FAudioDeviceManager* AudioDeviceManager = (GEngine ? GEngine->GetAudioDeviceManager() : nullptr);
+
+	// Set the channels equal to the users channel count choice
+	switch (SourceBusChannels)
+	{
+		case ESourceBusChannels::Mono:
+			NumChannels = 1;
+			break;
+
+		case ESourceBusChannels::Stereo:
+			NumChannels = 2;
+			break;
+	}
+}
+
+void USoundSourceBus::BeginDestroy()
+{
+	Super::BeginDestroy();
+}
+
+#if WITH_EDITOR
+void USoundSourceBus::PostEditChangeProperty(FPropertyChangedEvent& PropertyChangedEvent)
+{
+	// stub
+}
+#endif
+
+bool USoundSourceBus::IsPlayable() const
+{
+	return true;
+}
+
+float USoundSourceBus::GetDuration()
+{
+	return (SourceBusDuration > 0.0f) ? SourceBusDuration : INDEFINITELY_LOOPING_DURATION;
+}
+
diff --git a/Engine/Source/Runtime/Engine/Private/SoundWave.cpp b/Engine/Source/Runtime/Engine/Private/SoundWave.cpp
--- a/Engine/Source/Runtime/Engine/Private/SoundWave.cpp
+++ b/Engine/Source/Runtime/Engine/Private/SoundWave.cpp
@@ -650,9 +650,10 @@ void USoundWave::Parse( FAudioDevice* AudioDevice, const UPTRINT NodeWaveInstanc
 	if( !WaveInstance->bIsFinished )
 	{
 		// Propagate properties and add WaveInstance to outgoing array of FWaveInstances.
-		WaveInstance->Volume = ParseParams.Volume * Volume;
-		WaveInstance->VolumeMultiplier = ParseParams.VolumeMultiplier;
-		WaveInstance->VolumeApp = ParseParams.VolumeApp;
+		WaveInstance->SetVolume(ParseParams.Volume * Volume);
+		WaveInstance->SetVolumeMultiplier(ParseParams.VolumeMultiplier);
+		WaveInstance->SetDistanceAttenuation(ParseParams.DistanceAttenuation);
+		WaveInstance->SetVolumeApp(ParseParams.VolumeApp);
 		WaveInstance->Pitch = ParseParams.Pitch * Pitch;
 		WaveInstance->bEnableLowPassFilter = ParseParams.bEnableLowPassFilter;
 		WaveInstance->bIsOccluded = ParseParams.bIsOccluded;
@@ -683,7 +684,8 @@ void USoundWave::Parse( FAudioDevice* AudioDevice, const UPTRINT NodeWaveInstanc
 		{
 			FSoundClassProperties* SoundClassProperties = AudioDevice->GetSoundClassCurrentProperties(ParseParams.SoundClass);
 			// Use values from "parsed/ propagated" sound class properties
-			WaveInstance->VolumeMultiplier *= SoundClassProperties->Volume;
+			float VolumeMultiplier = WaveInstance->GetVolumeMultiplier();
+			WaveInstance->SetVolumeMultiplier(VolumeMultiplier* SoundClassProperties->Volume);
 			WaveInstance->Pitch *= SoundClassProperties->Pitch;
 			//TODO: Add in HighFrequencyGainMultiplier property to sound classes
 
@@ -702,7 +704,8 @@ void USoundWave::Parse( FAudioDevice* AudioDevice, const UPTRINT NodeWaveInstanc
 
 			if (SoundClassProperties->bApplyAmbientVolumes)
 			{
-				WaveInstance->VolumeMultiplier *= ParseParams.InteriorVolumeMultiplier;
+				VolumeMultiplier = WaveInstance->GetVolumeMultiplier();
+				WaveInstance->SetVolumeMultiplier(VolumeMultiplier * ParseParams.InteriorVolumeMultiplier);
 				WaveInstance->RadioFilterVolume *= ParseParams.InteriorVolumeMultiplier;
 				WaveInstance->RadioFilterVolumeThreshold *= ParseParams.InteriorVolumeMultiplier;
 			}
@@ -748,7 +751,8 @@ void USoundWave::Parse( FAudioDevice* AudioDevice, const UPTRINT NodeWaveInstanc
 		// If we're normalizing 3d stereo spatialized sounds, we need to scale by -6 dB
 		if (WaveInstance->bUseSpatialization && ParseParams.bApplyNormalizationToStereoSounds && NumChannels == 2)
 		{
-			WaveInstance->Volume *= 0.5f;
+			float WaveInstanceVolume = WaveInstance->GetVolume();
+			WaveInstance->SetVolume(WaveInstanceVolume * 0.5f);
 		}
 
 		// Copy reverb send settings
@@ -762,6 +766,14 @@ void USoundWave::Parse( FAudioDevice* AudioDevice, const UPTRINT NodeWaveInstanc
 		WaveInstance->SoundSubmix = ParseParams.SoundSubmix;
 		WaveInstance->SoundSubmixSends = ParseParams.SoundSubmixSends;
 
+		// Copy over the source bus send and data
+		if (!WaveInstance->ActiveSound->bIsPreviewSound)
+		{
+			WaveInstance->bOutputToBusOnly = ParseParams.bOutputToBusOnly;
+		}
+
+		WaveInstance->SoundSourceBusSends = ParseParams.SoundSourceBusSends;
+
 		if (AudioDevice->IsHRTFEnabledForAll() && ParseParams.SpatializationMethod == ESoundSpatializationAlgorithm::SPATIALIZATION_Default)
 		{
 			WaveInstance->SpatializationMethod = ESoundSpatializationAlgorithm::SPATIALIZATION_HRTF;
@@ -781,7 +793,7 @@ void USoundWave::Parse( FAudioDevice* AudioDevice, const UPTRINT NodeWaveInstanc
 		// That needs to change in the future, because there are still reasons a sound (and thus its subtitle) may not play.
 		// But for now at least that makes it possible handle virtualizing properly.
 		bool bHasSubtitles = ActiveSound.bHandleSubtitles && (ActiveSound.bHasExternalSubtitles || (Subtitles.Num() > 0));
-		if (WaveInstance->GetVolume() > KINDA_SMALL_NUMBER || ((bVirtualizeWhenSilent || bHasSubtitles) && AudioDevice->VirtualSoundsEnabled()))
+		if (WaveInstance->GetVolumeWithDistanceAttenuation() > KINDA_SMALL_NUMBER || ((bVirtualizeWhenSilent || bHasSubtitles) && AudioDevice->VirtualSoundsEnabled()))
 		{
 			bAddedWaveInstance = true;
 			WaveInstances.Add(WaveInstance);
diff --git a/Engine/Source/Runtime/Engine/Private/SoundWaveProcedural.cpp b/Engine/Source/Runtime/Engine/Private/SoundWaveProcedural.cpp
--- a/Engine/Source/Runtime/Engine/Private/SoundWaveProcedural.cpp
+++ b/Engine/Source/Runtime/Engine/Private/SoundWaveProcedural.cpp
@@ -9,12 +9,18 @@ USoundWaveProcedural::USoundWaveProcedural(const FObjectInitializer& ObjectIniti
 	bReset = false;
 	NumBufferUnderrunSamples = 512;
 	NumSamplesToGeneratePerCallback = 1024;
+
+	SampleByteSize = 2;
+
 	checkf(NumSamplesToGeneratePerCallback >= NumBufferUnderrunSamples, TEXT("Should generate more samples than this per callback."));
 }
 
 void USoundWaveProcedural::QueueAudio(const uint8* AudioData, const int32 BufferSize)
 {
-	if (BufferSize == 0 || !ensure((BufferSize % sizeof(int16)) == 0))
+	Audio::EAudioMixerStreamDataFormat::Type Format = GetGeneratedPCMDataFormat();
+	SampleByteSize = (Format == Audio::EAudioMixerStreamDataFormat::Int16) ? 2 : 4;
+
+	if (BufferSize == 0 || !ensure((BufferSize % SampleByteSize) == 0))
 	{
 		return;
 	}
@@ -47,7 +53,10 @@ int32 USoundWaveProcedural::GeneratePCMData(uint8* PCMData, const int32 SamplesN
 		AvailableByteCount.Reset();
 	}
 
-	int32 SamplesAvailable = AudioBuffer.Num() / sizeof(int16);
+	Audio::EAudioMixerStreamDataFormat::Type Format = GetGeneratedPCMDataFormat();
+	SampleByteSize = (Format == Audio::EAudioMixerStreamDataFormat::Int16) ? 2 : 4;
+
+	int32 SamplesAvailable = AudioBuffer.Num() / SampleByteSize;
 	int32 SamplesToGenerate = FMath::Min(NumSamplesToGeneratePerCallback, SamplesNeeded);
 
 	check(SamplesToGenerate >= NumBufferUnderrunSamples);
@@ -75,13 +84,13 @@ int32 USoundWaveProcedural::GeneratePCMData(uint8* PCMData, const int32 SamplesN
 		PumpQueuedAudio();
 	}
 
-	SamplesAvailable = AudioBuffer.Num() / sizeof(int16);
+	SamplesAvailable = AudioBuffer.Num() / SampleByteSize;
 
 	// Wait until we have enough samples that are requested before starting.
 	if (SamplesAvailable >= SamplesToGenerate)
 	{
 		const int32 SamplesToCopy = FMath::Min<int32>(SamplesToGenerate, SamplesAvailable);
-		const int32 BytesToCopy = SamplesToCopy * sizeof(int16);
+		const int32 BytesToCopy = SamplesToCopy * SampleByteSize;
 
 		FMemory::Memcpy((void*)PCMData, &AudioBuffer[0], BytesToCopy);
 		AudioBuffer.RemoveAt(0, BytesToCopy);
@@ -93,7 +102,7 @@ int32 USoundWaveProcedural::GeneratePCMData(uint8* PCMData, const int32 SamplesN
 	}
 
 	// There wasn't enough data ready, write out zeros
-	const int32 BytesCopied = NumBufferUnderrunSamples * sizeof(int16);
+	const int32 BytesCopied = NumBufferUnderrunSamples * SampleByteSize;
 	FMemory::Memzero(PCMData, BytesCopied);
 	return BytesCopied;
 }
diff --git a/Engine/Source/Runtime/Engine/Private/World.cpp b/Engine/Source/Runtime/Engine/Private/World.cpp
--- a/Engine/Source/Runtime/Engine/Private/World.cpp
+++ b/Engine/Source/Runtime/Engine/Private/World.cpp
@@ -90,6 +90,10 @@
 	#include "Editor/UnrealEdTypes.h"
 	#include "Settings/LevelEditorPlaySettings.h"
 	#include "HierarchicalLOD.h"
+	#include "IHierarchicalLODUtilities.h"
+	#include "HierarchicalLODUtilitiesModule.h"
+	#include "ObjectTools.h"
+	#include "Engine/LODActor.h"
 #endif
 
 
@@ -464,6 +468,15 @@ bool UWorld::Rename(const TCHAR* InName, UObject* NewOuter, ERenameFlags Flags)
 	// Make sure our legacy lightmap data is initialized so it can be renamed
 	PersistentLevel->HandleLegacyMapBuildData();
 
+	FHierarchicalLODUtilitiesModule& Module = FModuleManager::LoadModuleChecked<FHierarchicalLODUtilitiesModule>("HierarchicalLODUtilities");
+	IHierarchicalLODUtilities* Utilities = Module.GetUtilities();
+	UPackage* OldHLODPackage = nullptr;	
+	// See if any LODActors were found in the level, and if so retrieve the HLOD Package
+	if (PersistentLevel->Actors.ContainsByPredicate([](const AActor* Actor) { return Actor != nullptr && Actor->IsA<ALODActor>(); }))
+	{
+		OldHLODPackage = Utilities->CreateOrRetrieveLevelHLODPackage(PersistentLevel);
+	}
+
 	if (bShouldFail)
 	{
 		return false;
@@ -566,6 +579,30 @@ bool UWorld::Rename(const TCHAR* InName, UObject* NewOuter, ERenameFlags Flags)
 		NewPackage->ThisContainsMap();
 	}
 
+	// Move over HLOD assets to new _HLOD Package
+	if (OldHLODPackage)
+	{
+		UPackage* NewHLODPackage = Utilities->CreateOrRetrieveLevelHLODPackage(PersistentLevel);
+		TArray<UObject*> Objects;
+		// Retrieve all of the HLOD objects 
+		ForEachObjectWithOuter(OldHLODPackage, [&Objects](UObject* Obj)
+		{
+			if (ObjectTools::IsObjectBrowsable(Obj))
+			{
+				Objects.Add(Obj);
+			}
+		});
+		// Rename them 'into' the new HLOD package
+		for (UObject* Object : Objects)
+		{
+			Object->Rename(*Object->GetName(), NewHLODPackage);
+		}
+		// Delete the old HLOD package
+		TArray<UObject*> DeleteObjects = { Cast<UObject>(OldHLODPackage) };
+		ObjectTools::DeleteObjectsUnchecked(DeleteObjects);
+	}
+	
+
 	return true;
 }
 #endif
diff --git a/Engine/Source/Runtime/Engine/Public/ActiveSound.h b/Engine/Source/Runtime/Engine/Public/ActiveSound.h
--- a/Engine/Source/Runtime/Engine/Public/ActiveSound.h
+++ b/Engine/Source/Runtime/Engine/Public/ActiveSound.h
@@ -11,12 +11,15 @@
 #include "Components/AudioComponent.h"
 #include "Sound/AudioVolume.h"
 #include "Sound/SoundSubmix.h"
+#include "Sound/SoundSourceBus.h"
 #include "AudioDevice.h"
 
 class FAudioDevice;
 class USoundBase;
 class USoundSubmix;
+class USoundSourceBus;
 struct FSoundSubmixSendInfo;
+struct FSoundSourceBusSendInfo;
 class USoundWave;
 struct FListener;
 
@@ -40,6 +43,9 @@ struct FSoundParseParameters
 	// The volume product of the sound
 	float Volume;
 
+	// The attenuation of the sound due to distance attenuation
+	float DistanceAttenuation;
+
 	// A volume scalse on the sound specified by user
 	float VolumeMultiplier;
 	
@@ -76,6 +82,9 @@ struct FSoundParseParameters
 	// The submix sends to use
 	TArray<FSoundSubmixSendInfo> SoundSubmixSends;
 
+	// The source bus sends to use
+	TArray<FSoundSourceBusSendInfo> SoundSourceBusSends;
+
 	// Reverb wet-level parameters
 	EReverbSendMethod ReverbSendMethod;
 	FVector2D ReverbSendLevelRange;
@@ -116,6 +125,9 @@ struct FSoundParseParameters
 	// The lowpass filter to apply if the sound is inside an ambient zone
 	float AmbientZoneFilterFrequency;
 
+	// Whether or not to ouput this audio to buses only
+	uint32 bOutputToBusOnly:1;
+
 	// Whether the sound should be spatialized
 	uint32 bUseSpatialization:1;
 
@@ -138,6 +150,7 @@ struct FSoundParseParameters
 		: SoundClass(nullptr)
 		, Velocity(ForceInit)
 		, Volume(1.f)
+		, DistanceAttenuation(1.f)
 		, VolumeMultiplier(1.f)
 		, VolumeApp(1.f)
 		, InteriorVolumeMultiplier(1.f)
@@ -163,6 +176,7 @@ struct FSoundParseParameters
 		, AttenuationHighpassFilterFrequency(MIN_FILTER_FREQUENCY)
 		, OcclusionFilterFrequency(MAX_FILTER_FREQUENCY)
 		, AmbientZoneFilterFrequency(MAX_FILTER_FREQUENCY)
+		, bOutputToBusOnly(false)
 		, bUseSpatialization(false)
 		, bLooping(false)
 		, bEnableLowPassFilter(false)
@@ -198,6 +212,7 @@ public:
 	uint64 GetAudioComponentID() const { return AudioComponentID; }
 	FName GetAudioComponentUserID() const { return AudioComponentUserID; }
 	void SetAudioComponent(UAudioComponent* Component);
+	void SetOwner(AActor* Owner);
 	FString GetAudioComponentName() const;
 	FString GetOwnerName() const;
 
@@ -244,6 +259,9 @@ private:
 	/** Optional override the submix sends for the sound. */
 	TArray<FSoundSubmixSendInfo> SoundSubmixSendsOverride;
 
+	/** Optional override for the source bus sends for the sound. */
+	TArray<FSoundSourceBusSendInfo> SoundSourceBusSendsOverride;
+
 public:
 	/** Whether or not the sound has checked if it was occluded already. Used to initialize a sound as occluded and bypassing occlusion interpolation. */
 	uint8 bHasCheckedOcclusion:1;
@@ -492,6 +510,9 @@ public:
 	/** Gets the sound submix sends to use for this sound instance. */
 	void GetSoundSubmixSends(TArray<FSoundSubmixSendInfo>& OutSends) const;
 
+	/** Gets the sound source bus sends to use for this sound instance. */
+	void GetSoundSourceBusSends(TArray<FSoundSourceBusSendInfo>& OutSends) const;
+
 	/* Determines which listener is the closest to the sound */
 	int32 FindClosestListener( const TArray<struct FListener>& InListeners ) const;
 	
@@ -516,6 +537,9 @@ public:
 	/** Sets the amount of audio from this active sound to send to the submix. */
 	void SetSubmixSend(const FSoundSubmixSendInfo& SubmixSendInfo);
 
+	/** Sets the amount of audio from this active sound to send to the source bus. */
+	void SetSourceBusSend(const FSoundSourceBusSendInfo& SourceBusSendInfo);
+
 private:
 	
 	struct FAsyncTraceDetails
diff --git a/Engine/Source/Runtime/Engine/Public/Animation/AnimInstanceProxy.h b/Engine/Source/Runtime/Engine/Public/Animation/AnimInstanceProxy.h
--- a/Engine/Source/Runtime/Engine/Public/Animation/AnimInstanceProxy.h
+++ b/Engine/Source/Runtime/Engine/Public/Animation/AnimInstanceProxy.h
@@ -538,7 +538,7 @@ protected:
 	/** 
 	 * Recalculate required curve list for animation - if you call RecalcRequiredBones, this should already happen
 	 */
-	void RecalcRequiredCurves(bool bDisableAnimCurves);
+	void RecalcRequiredCurves(const FCurveEvaluationOption& CurveEvalOption);
 
 	/** Update the material parameters of the supplied component from this instance */
 	void UpdateCurvesToComponents(USkeletalMeshComponent* Component);
diff --git a/Engine/Source/Runtime/Engine/Public/Animation/AnimTypes.h b/Engine/Source/Runtime/Engine/Public/Animation/AnimTypes.h
--- a/Engine/Source/Runtime/Engine/Public/Animation/AnimTypes.h
+++ b/Engine/Source/Runtime/Engine/Public/Animation/AnimTypes.h
@@ -497,6 +497,12 @@ struct FAnimWeight
 	{
 		return (InWeight >= (1.f - ZERO_ANIMWEIGHT_THRESH));
 	}
+
+	/** Get a small relevant weight for ticking */
+	static FORCEINLINE float GetSmallestRelevantWeight()
+	{
+		return 2.f * ZERO_ANIMWEIGHT_THRESH;
+	}
 };
 
 /** 
diff --git a/Engine/Source/Runtime/Engine/Public/AnimationRuntime.h b/Engine/Source/Runtime/Engine/Public/AnimationRuntime.h
--- a/Engine/Source/Runtime/Engine/Public/AnimationRuntime.h
+++ b/Engine/Source/Runtime/Engine/Public/AnimationRuntime.h
@@ -295,11 +295,20 @@ public:
 	/** Accumulate Additive Pose based on AdditiveType*/
 	static void AccumulateAdditivePose(FCompactPose& BasePose, const FCompactPose& AdditivePose, FBlendedCurve& BaseCurve, const FBlendedCurve& AdditiveCurve, float Weight, enum EAdditiveAnimationType AdditiveType);
 
+private:
 	/** Accumulates weighted AdditivePose to BasePose. Rotations are NOT normalized. */
-	static void AccumulateLocalSpaceAdditivePose(FCompactPose& BasePose, const FCompactPose& AdditivePose, FBlendedCurve& BaseCurve, const FBlendedCurve& AdditiveCurve, float Weight);
+	static void AccumulateLocalSpaceAdditivePoseInternal(FCompactPose& BasePose, const FCompactPose& AdditivePose, float Weight);
 
 	/** Accumulate a MeshSpaceRotation Additive pose to a local pose. Rotations are NOT normalized */
-	static void AccumulateMeshSpaceRotationAdditiveToLocalPose(FCompactPose& BasePose, const FCompactPose& MeshSpaceRotationAdditive, FBlendedCurve& BaseCurve, const FBlendedCurve& AdditiveCurve, float Weight);
+	static void AccumulateMeshSpaceRotationAdditiveToLocalPoseInternal(FCompactPose& BasePose, const FCompactPose& MeshSpaceRotationAdditive, float Weight);
+public:
+
+	/** Accumulates weighted AdditivePose to BasePose. Rotations are NOT normalized. */
+	static void AccumulateLocalSpaceAdditivePose(FCompactPose& BasePose, const FCompactPose& AdditivePose, FBlendedCurve& BaseCurve, const FBlendedCurve& AdditiveCurve, float Weight) { AccumulateAdditivePose(BasePose, AdditivePose, BaseCurve, AdditiveCurve, Weight, EAdditiveAnimationType::AAT_LocalSpaceBase); }
+
+	/** Accumulate a MeshSpaceRotation Additive pose to a local pose. Rotations are NOT normalized */
+	static void AccumulateMeshSpaceRotationAdditiveToLocalPose(FCompactPose& BasePose, const FCompactPose& MeshSpaceRotationAdditive, FBlendedCurve& BaseCurve, const FBlendedCurve& AdditiveCurve, float Weight) { AccumulateAdditivePose(BasePose, MeshSpaceRotationAdditive, BaseCurve, AdditiveCurve, Weight, EAdditiveAnimationType::AAT_RotationOffsetMeshSpace); }
+
 
 	/** Lerp for BoneTransforms. Stores results in A. Performs A = Lerp(A, B, Alpha);
 	* @param A : In/Out FCompactPose.
diff --git a/Engine/Source/Runtime/Engine/Public/Audio.h b/Engine/Source/Runtime/Engine/Public/Audio.h
--- a/Engine/Source/Runtime/Engine/Public/Audio.h
+++ b/Engine/Source/Runtime/Engine/Public/Audio.h
@@ -13,13 +13,19 @@
 #include "Sound/SoundSubmix.h"
 #include "Sound/SoundAttenuation.h"
 #include "Sound/SoundEffectSource.h"
+#include "Sound/SoundSourceBusSend.h"
 #include "IAudioExtensionPlugin.h"
+#include "AudioPluginUtilities.h"
 
 class FAudioDevice;
 class USoundNode;
 class USoundWave;
+class USoundClass;
+class USoundSubmix;
+class USoundSourceBus;
 struct FActiveSound;
 struct FWaveInstance;
+struct FSoundSourceBusSendInfo;
 
 //#include "Sound/SoundConcurrency.h"
 
@@ -227,24 +233,9 @@ private:
 	TArray<FNotifyBufferDetails> Notifies;
 };
 
-
-/**
-* Enumeration of audio plugin types
-*
-*/
-struct EAudioPlugin
-{
-	enum Type
-	{
-		SPATIALIZATION,
-		REVERB,
-		OCCLUSION
-	};
-};
-
 /** Queries if a plugin of the given type is enabled. */
-ENGINE_API bool IsAudioPluginEnabled(EAudioPlugin::Type PluginType);
-ENGINE_API bool DoesAudioPluginHaveCustomSettings(EAudioPlugin::Type PluginType);
+ENGINE_API bool IsAudioPluginEnabled(EAudioPlugin PluginType);
+ENGINE_API bool DoesAudioPluginHaveCustomSettings(EAudioPlugin PluginType);
 
 /**
  * Structure encapsulating all information required to play a USoundWave on a channel/source. This is required
@@ -256,17 +247,20 @@ struct ENGINE_API FWaveInstance
 	static uint32 TypeHashCounter;
 
 	/** Wave data */
-	class USoundWave* WaveData;
+	USoundWave* WaveData;
 
 	/** Sound class */
-	class USoundClass* SoundClass;
+	USoundClass* SoundClass;
 
-	/** Sound submix */
-	class USoundSubmix* SoundSubmix;
+	/** Sound submix object to send audio to for mixing in audio mixer.  */
+	USoundSubmix* SoundSubmix;
 
 	/** Sound submix sends */
 	TArray<FSoundSubmixSendInfo> SoundSubmixSends;
 
+	/** The sound sourcebus sends. */
+	TArray<FSoundSourceBusSendInfo> SoundSourceBusSends;
+
 	/** Sound effect chain */
 	USoundEffectSourcePresetChain* SourceEffectChain;
 
@@ -276,15 +270,22 @@ struct ENGINE_API FWaveInstance
 	/** Active Sound this wave instance belongs to */
 	struct FActiveSound* ActiveSound;
 
+private:
+
 	/** Current volume */
 	float Volume;
 
+	/** Volume attenuation due to distance. */
+	float DistanceAttenuation;
+
 	/** Current volume multiplier - used to zero the volume without stopping the source */
 	float VolumeMultiplier;
 
 	/** The volume of the wave instance due to application volume or tab-state */
 	float VolumeApp;
 
+public:
+
 	/** An audio component priority value that scales with volume (post all gain stages) and is used to determine voice playback priority. */
 	float Priority;
 
@@ -309,6 +310,9 @@ struct ENGINE_API FWaveInstance
 	/** An offset/seek time to play this wave instance. */
 	float StartTime;
 
+	/** Whether or not to only output this wave instances audio to buses only. Audio mixer only. */
+	uint32 bOutputToBusOnly:1;
+
 	/** Set to true if the sound nodes state that the radio filter should be applied */
 	uint32 bApplyRadioFilter:1;
 
@@ -444,9 +448,24 @@ struct ENGINE_API FWaveInstance
 	/** Returns the actual volume the wave instance will play at */
 	bool ShouldStopDueToMaxConcurrency() const;
 
+	/** Setters for various volume values on wave instances. */
+	void SetVolume(const float InVolume) { Volume = InVolume; }
+	void SetDistanceAttenuation(const float InDistanceAttenuation) { DistanceAttenuation = InDistanceAttenuation; }
+	void SetVolumeApp(const float InVolumeApp) { VolumeApp = InVolumeApp; }
+	void SetVolumeMultiplier(const float InVolumeMultiplier) { VolumeMultiplier = InVolumeMultiplier; }
+
+	/** Returns the volume multiplier on the wave instance. */
+	float GetVolumeMultiplier() const { return VolumeMultiplier; }
+
 	/** Returns the actual volume the wave instance will play at, including all gain stages. */
 	float GetActualVolume() const;
 
+	/** Returns the volume of the sound including distance attenuation. */
+	float GetVolumeWithDistanceAttenuation() const;
+
+	/** Returns the distance attenuation of the source voice. */
+	float GetDistanceAttenuation() const;
+
 	/** Returns the volume of the wave instance (ignoring application muting) */
 	float GetVolume() const;
 
diff --git a/Engine/Source/Runtime/Engine/Public/AudioDevice.h b/Engine/Source/Runtime/Engine/Public/AudioDevice.h
--- a/Engine/Source/Runtime/Engine/Public/AudioDevice.h
+++ b/Engine/Source/Runtime/Engine/Public/AudioDevice.h
@@ -12,6 +12,7 @@
 #include "Sound/AudioVolume.h"
 #include "Sound/SoundConcurrency.h"
 #include "Sound/SoundMix.h"
+#include "Sound/SoundSourceBus.h"
 #include "Sound/AudioSettings.h"
 #include "AudioDeviceManager.h"
 #include "EngineGlobals.h"
@@ -26,6 +27,7 @@ class USoundBase;
 class USoundEffectSourcePreset;
 class USoundEffectSubmixPreset;
 class USoundSubmix;
+class USoundSourceBus;
 class USoundWave;
 struct FActiveSound;
 struct FAudioQualitySettings;
@@ -50,7 +52,6 @@ class USoundMix;
 class FAudioEffectsManager;
 class FViewportClient;
 class ICompressedAudioInfo;
-class IAudioPlugin;
 class IAudioSpatialization;
 class UReverbEffect;
 class USoundConcurrency;
@@ -415,6 +416,7 @@ class IDeviceChangedListener
 {
 public:
 	virtual void OnDeviceRemoved(FString DeviceID) = 0;
+	virtual void OnDefaultDeviceChanged() = 0;
 };
 
 class ENGINE_API FAudioDevice : public FExec
@@ -612,6 +614,7 @@ public:
 
 	/**
 	 * Sets the details about the listener
+	 * @param	World				The world the listener is being set on.
 	 * @param   ListenerIndex		The index of the listener
 	 * @param   ListenerTransform   The listener's world transform
 	 * @param   DeltaSeconds		The amount of time over which velocity should be calculated.  If 0, then velocity will not be calculated.
@@ -688,7 +691,7 @@ public:
 	 * @param	USoundConcurrency	The sound's sound concurrency settings to use (optional). Will use the USoundBase's USoundConcurrency if not specified.
 	 * @param	Params				An optional list of audio component params to immediately apply to a sound.
 	 */
-	void PlaySoundAtLocation(USoundBase* Sound, UWorld* World, float VolumeMultiplier, float PitchMultiplier, float StartTime, const FVector& Location, const FRotator& Rotation, USoundAttenuation* AttenuationSettings = nullptr, USoundConcurrency* ConcurrencySettings = nullptr, const TArray<FAudioComponentParam>* Params = nullptr);
+	void PlaySoundAtLocation(USoundBase* Sound, UWorld* World, float VolumeMultiplier, float PitchMultiplier, float StartTime, const FVector& Location, const FRotator& Rotation, USoundAttenuation* AttenuationSettings = nullptr, USoundConcurrency* ConcurrencySettings = nullptr, const TArray<FAudioComponentParam>* Params = nullptr, AActor* OwningActor = nullptr);
 
 	/**
 	 * Adds an active sound to the audio device
@@ -944,6 +947,12 @@ public:
 		FAudioThread::ResumeAudioThread();
 	}
 
+	/** Registers a third party listener-observer to this audio device. */
+	void RegisterPluginListener(const TAudioPluginListenerPtr PluginListener);
+
+	/* Unregisters a third party listener-observer to this audio device. */
+	void UnregisterPluginListener(const TAudioPluginListenerPtr PluginListener);
+
 	bool IsAudioDeviceMuted() const;
 
 	void SetDeviceMuted(bool bMuted);
@@ -1024,7 +1033,7 @@ public:
 	/** Whether or not there's a spatialization plugin enabled. */
 	bool IsSpatializationPluginEnabled() const
 	{
-		return bSpatializationInterfaceEnabled && SpatializationPluginInterface.IsValid();
+		return bSpatializationInterfaceEnabled;
 	}
 
 	/** Return the spatialization plugin interface. */
@@ -1036,14 +1045,14 @@ public:
 	/** Whether or not there's an occlusion plugin enabled. */
 	bool IsOcclusionPluginEnabled() const
 	{
-		return bOcclusionInterfaceEnabled && OcclusionInterface.IsValid();
+		return bOcclusionInterfaceEnabled;
 	}
 
 	static bool IsOcclusionPluginLoaded()
 	{
 		if (FAudioDevice* MainAudioDevice = GEngine->GetMainAudioDevice())
 		{
-			return MainAudioDevice->OcclusionInterface.IsValid();
+			return MainAudioDevice->bOcclusionInterfaceEnabled;
 		}
 		return false;
 	}
@@ -1051,14 +1060,14 @@ public:
 	/** Whether or not there's a reverb plugin enabled. */
 	bool IsReverbPluginEnabled() const
 	{
-		return bReverbInterfaceEnabled && ReverbPluginInterface.IsValid();
+		return bReverbInterfaceEnabled;
 	}
 
 	static bool IsReverbPluginLoaded()
 	{
 		if (FAudioDevice* MainAudioDevice = GEngine->GetMainAudioDevice())
 		{
-			return MainAudioDevice->ReverbPluginInterface.IsValid();
+			return MainAudioDevice->bReverbInterfaceEnabled;
 		}
 		return false;
 	}
@@ -1096,6 +1105,15 @@ protected:
 
 private:
 
+	/**
+	* Initializes all plugin listeners belonging to this audio device.
+	* Called in the game thread.
+	*
+	* @param World: Pointer to the UWorld the listener is in.
+	* @param InViewportIndex: Viewport that the listener belongs to.
+	*/
+	void InitializePluginListeners(UWorld* World);
+
 	/**
 	 * Parses the sound classes and propagates multiplicative properties down the tree.
 	 */
@@ -1363,6 +1381,9 @@ public:
 		bUseAttenuationForNonGameWorlds = bInUseAttenuationForNonGameWorlds;
 	}
 
+	/** Returns the default reverb send level used for sources which have reverb applied but no attenuation settings. */
+	float GetDefaultReverbSendLevel() const { return DefaultReverbSendLevel; }
+
 public:
 
 	/** The maximum number of concurrent audible sounds */
@@ -1388,9 +1409,6 @@ public:
 	/** The handle for this audio device used in the audio device manager. */
 	uint32 DeviceHandle;
 
-	/** An audio plugin. */
-	IAudioPlugin* AudioPlugin;
-
 	/** 3rd party audio spatialization interface. */
 	TAudioSpatializationPtr SpatializationPluginInterface;
 
@@ -1400,8 +1418,8 @@ public:
 	/** 3rd party occlusion interface. */
 	TAudioOcclusionPtr OcclusionInterface;
 
-	/** 3rd party listener observer. */
-	TAudioListenerObserverPtr ListenerObserver;
+	/** 3rd party listener observers registered to this audio device. */
+	TArray<TAudioPluginListenerPtr> PluginListeners;
 
 private:
 	// Audio thread representation of listeners
@@ -1467,6 +1485,9 @@ private:
 	/** A volume headroom to apply to specific platforms to achieve better platform consistency. */
 	float PlatformAudioHeadroom;
 
+	/** The default reverb send level to use for sources which have reverb applied but don't have an attenuation settings. */
+	float DefaultReverbSendLevel;
+
 	/** Reverb Effects activated without volumes - Game Thread owned */
 	TMap<FName, FActivatedReverb> ActivatedReverbs;
 
@@ -1488,6 +1509,11 @@ public:
 	/** Whether or not the audio mixer module is being used by this device. */
 	uint8 bAudioMixerModuleLoaded : 1;
 
+	/** Whether of not various audio plugin interfaces are external sends. */
+	uint8 bSpatializationIsExternalSend:1;
+	uint8 bOcclusionIsExternalSend:1;
+	uint8 bReverbIsExternalSend:1;
+
 private:
 	/* True once the startup sounds have been precached */
 	uint8 bStartupSoundsPreCached:1;
@@ -1496,7 +1522,9 @@ private:
 	uint8 bSpatializationInterfaceEnabled:1;
 	uint8 bOcclusionInterfaceEnabled:1;
 	uint8 bReverbInterfaceEnabled:1;
-	uint8 bListenerObserverInterfaceEnabled:1;
+	
+	/** Whether or not we've initialized plugin listeners array. */
+	uint8 bPluginListenersInitialized:1;
 
 	/** Whether HRTF is enabled for all 3d sounds. */
 	uint8 bHRTFEnabledForAll:1;
diff --git a/Engine/Source/Runtime/Engine/Public/AudioPluginUtilities.h b/Engine/Source/Runtime/Engine/Public/AudioPluginUtilities.h
new file mode 100644
--- /dev/null
+++ b/Engine/Source/Runtime/Engine/Public/AudioPluginUtilities.h
@@ -0,0 +1,46 @@
+// Copyright 1998-2017 Epic Games, Inc. All Rights Reserved.
+#pragma once
+#include "IAudioExtensionPlugin.h"
+
+/************************************************************************/
+/* Plugin Utilities                                                     */
+/************************************************************************/
+struct ENGINE_API AudioPluginUtilities
+{
+
+#if PLATFORM_WINDOWS
+	static const EAudioPlatform CurrentPlatform = EAudioPlatform::Windows;
+#elif PLATFORM_MAC
+	static const EAudioPlatform CurrentPlatform = EAudioPlatform::Mac;
+#elif PLATFORM_LINUX
+	static const EAudioPlatform CurrentPlatform = EAudioPlatform::Linux;
+#elif PLATFORM_IOS
+	static const EAudioPlatform CurrentPlatform = EAudioPlatform::IOS;
+#elif PLATFORM_ANDROID
+	static const EAudioPlatform CurrentPlatform = EAudioPlatform::Android;
+#elif PLATFORM_XBOXONE
+	static const EAudioPlatform CurrentPlatform = EAudioPlatform::XboxOne;
+#elif PLATFORM_PS4
+	static const EAudioPlatform CurrentPlatform = EAudioPlatform::Playstation4;
+#elif PLATFORM_SWITCH
+	static const EAudioPlatform CurrentPlatform = EAudioPlatform::Switch;
+#elif PLATFORM_HTML5
+	static const EAudioPlatform CurrentPlatform = EAudioPlatform::HTML5;
+#else
+	static const EAudioPlatform CurrentPlatform = EAudioPlatform::Unknown;
+#endif
+
+	/*
+	 * These functions return a pointer to the plugin factory
+	 * that matches the plugin name specified in the target
+	 * platform's settings.
+	 * 
+	 * if no matching plugin is found, nullptr is returned.
+	 */
+	static IAudioSpatializationFactory* GetDesiredSpatializationPlugin(EAudioPlatform AudioPlatform);
+	static IAudioReverbFactory* GetDesiredReverbPlugin(EAudioPlatform AudioPlatform);
+	static IAudioOcclusionFactory* GetDesiredOcclusionPlugin(EAudioPlatform AudioPlatform);
+
+	/** This function returns the name of the plugin specified in the platform settings. */
+	static FString GetDesiredPluginName(EAudioPlugin PluginType, EAudioPlatform AudioPlatform);
+};
\ No newline at end of file
diff --git a/Engine/Source/Runtime/Engine/Public/BoneContainer.h b/Engine/Source/Runtime/Engine/Public/BoneContainer.h
--- a/Engine/Source/Runtime/Engine/Public/BoneContainer.h
+++ b/Engine/Source/Runtime/Engine/Public/BoneContainer.h
@@ -33,6 +33,22 @@ struct FVirtualBoneCompactPoseData
 	{}
 };
 
+/**
+ * This is curve evaluation options for bone container
+ */
+struct FCurveEvaluationOption
+{
+	bool bAllowCurveEvaluation;
+	const TArray<FName>* DisallowedList;
+	int32 LODIndex;
+	
+	FCurveEvaluationOption(bool bInAllowCurveEvaluation = true, const TArray<FName>* InDisallowedList = nullptr, int32 InLODIndex = 0)
+		: bAllowCurveEvaluation(bInAllowCurveEvaluation)
+		, DisallowedList(InDisallowedList)
+		, LODIndex(InLODIndex)
+	{
+	}
+};
 /**
 * This is a native transient structure.
 * Contains:
@@ -93,10 +109,10 @@ public:
 
 	FBoneContainer();
 
-	FBoneContainer(const TArray<FBoneIndexType>& InRequiredBoneIndexArray, bool bDisableAnimCurves, UObject& InAsset);
+	FBoneContainer(const TArray<FBoneIndexType>& InRequiredBoneIndexArray, const FCurveEvaluationOption& CurveEvalOption, UObject& InAsset);
 
 	/** Initialize BoneContainer to a new Asset, RequiredBonesArray and RefPoseArray. */
-	void InitializeTo(const TArray<FBoneIndexType>& InRequiredBoneIndexArray, bool bDisableAnimCurves, UObject& InAsset);
+	void InitializeTo(const TArray<FBoneIndexType>& InRequiredBoneIndexArray, const FCurveEvaluationOption& CurveEvalOption, UObject& InAsset);
 
 	/** Returns true if FBoneContainer is Valid. Needs an Asset, a RefPoseArray, and a RequiredBonesArray. */
 	const bool IsValid() const
@@ -309,11 +325,11 @@ public:
 	}
 
 	/** Cache required Anim Curve Uids */
-	void CacheRequiredAnimCurveUids(bool bDisableAnimCurves);
+	void CacheRequiredAnimCurveUids(const FCurveEvaluationOption& CurveEvalOption);
 
 private:
 	/** Initialize FBoneContainer. */
-	void Initialize(bool bDisableAnimCurves);
+	void Initialize(const FCurveEvaluationOption& CurveEvalOption);
 
 	/** Cache remapping data if current Asset is a SkeletalMesh, with all compatible Skeletons. */
 	void RemapFromSkelMesh(USkeletalMesh const & SourceSkeletalMesh, USkeleton& TargetSkeleton);
