Copying //UE4/Dev-Framework to //UE4/Dev-Main (Source: //UE4/Dev-Framework @ 3628051)

diff --git a/Engine/Source/Runtime/Engine/Classes/Components/PrimitiveComponent.h b/Engine/Source/Runtime/Engine/Classes/Components/PrimitiveComponent.h
--- a/Engine/Source/Runtime/Engine/Classes/Components/PrimitiveComponent.h
+++ b/Engine/Source/Runtime/Engine/Classes/Components/PrimitiveComponent.h
@@ -211,14 +211,18 @@ public:
 	UPROPERTY()
 	TEnumAsByte<enum ESceneDepthPriorityGroup> ViewOwnerDepthPriorityGroup;
 
+	/** Quality of indirect lighting for Movable primitives.  This has a large effect on Indirect Lighting Cache update time. */
+	UPROPERTY(EditAnywhere, AdvancedDisplay, BlueprintReadOnly, Category=Lighting)
+	TEnumAsByte<EIndirectLightingCacheQuality> IndirectLightingCacheQuality;
+
 public:
 
 	/** Whether this primitive is referenced by a FLevelTextureManager  */
-	mutable uint32 bAttachedToStreamingManagerAsStatic : 1;
+	mutable uint8 bAttachedToStreamingManagerAsStatic : 1;
 	/** Whether this primitive is referenced by a FDynamicTextureInstanceManager */
-	mutable uint32 bAttachedToStreamingManagerAsDynamic : 1;
+	mutable uint8 bAttachedToStreamingManagerAsDynamic : 1;
 	/** Whether this primitive is handled as dynamic, although it could have no references */
-	mutable uint32 bHandledByStreamingManagerAsDynamic : 1;
+	mutable uint8 bHandledByStreamingManagerAsDynamic : 1;
 
 	/** Whether this primitive is referenced by the streaming manager and should sent callbacks when detached or destroyed */
 	FORCEINLINE bool IsAttachedToStreamingManager() const { return !!(bAttachedToStreamingManagerAsStatic | bAttachedToStreamingManagerAsDynamic); }
@@ -229,7 +233,7 @@ public:
 	 * This can help performance if you'd like to avoid overhead of creating physics state when triggers 
 	 */
 	UPROPERTY(EditAnywhere, AdvancedDisplay, BlueprintReadOnly, Category=Collision)
-	uint32 bAlwaysCreatePhysicsState:1;
+	uint8 bAlwaysCreatePhysicsState:1;
 
 	/**
 	 * If true, this component will generate overlap events when it is overlapping other components (eg Begin Overlap).
@@ -239,7 +243,7 @@ public:
 	 * @see UpdateOverlaps(), BeginComponentOverlap(), EndComponentOverlap()
 	 */
 	UPROPERTY(EditAnywhere, BlueprintReadWrite, Category=Collision)
-	uint32 bGenerateOverlapEvents:1;
+	uint8 bGenerateOverlapEvents:1;
 
 	/**
 	 * If true, this component will generate individual overlaps for each overlapping physics body if it is a multi-body component. When false, this component will
@@ -247,7 +251,7 @@ public:
 	 * influence on single body components.
 	 */
 	UPROPERTY(EditAnywhere, AdvancedDisplay, BlueprintReadWrite, Category=Collision)
-	uint32 bMultiBodyOverlap:1;
+	uint8 bMultiBodyOverlap:1;
 
 	/**
 	 * If true, this component will look for collisions on both physic scenes during movement.
@@ -255,7 +259,7 @@ public:
 	 * @see MoveComponent()
 	 */
 	UPROPERTY(EditAnywhere, AdvancedDisplay, BlueprintReadWrite, Category=Collision)
-	uint32 bCheckAsyncSceneOnMove:1;
+	uint8 bCheckAsyncSceneOnMove:1;
 
 	/**
 	 * If true, component sweeps with this component should trace against complex collision during movement (for example, each triangle of a mesh).
@@ -263,54 +267,54 @@ public:
 	 * @see MoveComponent()
 	 */
 	UPROPERTY(EditAnywhere, AdvancedDisplay, BlueprintReadWrite, Category=Collision)
-	uint32 bTraceComplexOnMove:1;
+	uint8 bTraceComplexOnMove:1;
 
 	/**
 	 * If true, component sweeps will return the material in their hit result.
 	 * @see MoveComponent(), FHitResult
 	 */
 	UPROPERTY(EditAnywhere, AdvancedDisplay, BlueprintReadWrite, Category=Collision)
-	uint32 bReturnMaterialOnMove:1;
+	uint8 bReturnMaterialOnMove:1;
 
 	/** True if the primitive should be rendered using ViewOwnerDepthPriorityGroup if viewed by its owner. */
 	UPROPERTY()
-	uint32 bUseViewOwnerDepthPriorityGroup:1;
+	uint8 bUseViewOwnerDepthPriorityGroup:1;
 
 	/** Whether to accept cull distance volumes to modify cached cull distance. */
 	UPROPERTY(EditAnywhere, AdvancedDisplay, BlueprintReadOnly, Category=LOD)
-	uint32 bAllowCullDistanceVolume:1;
+	uint8 bAllowCullDistanceVolume:1;
 
 	/** true if the primitive has motion blur velocity meshes */
 	UPROPERTY()
-	uint32 bHasMotionBlurVelocityMeshes:1;
+	uint8 bHasMotionBlurVelocityMeshes:1;
 	
 	/** If true, this component will be visible in reflection captures. */
 	UPROPERTY(EditAnywhere, AdvancedDisplay, BlueprintReadOnly, Category = Rendering)
-	uint32 bVisibleInReflectionCaptures:1;
+	uint8 bVisibleInReflectionCaptures:1;
 
 	/** If true, this component will be rendered in the main pass (z prepass, basepass, transparency) */
 	UPROPERTY(EditAnywhere, AdvancedDisplay, BlueprintReadOnly, Category = Rendering)
-	uint32 bRenderInMainPass:1;
+	uint8 bRenderInMainPass:1;
 
 	/** If true, this component will be rendered in mono only if an HMD is connected and monoscopic far field rendering is activated. */
 	UPROPERTY(EditAnywhere, AdvancedDisplay, BlueprintReadOnly, Category = Rendering)
-	uint32 bRenderInMono:1;
+	uint8 bRenderInMono:1;
 
 	/** Whether the primitive receives decals. */
 	UPROPERTY(EditAnywhere, AdvancedDisplay, BlueprintReadOnly, Category=Rendering)
-	uint32 bReceivesDecals:1;
+	uint8 bReceivesDecals:1;
 
 	/** If this is True, this component won't be visible when the view actor is the component's owner, directly or indirectly. */
 	UPROPERTY(EditAnywhere, AdvancedDisplay, BlueprintReadOnly, Category = Rendering)
-	uint32 bOwnerNoSee:1;
+	uint8 bOwnerNoSee:1;
 
 	/** If this is True, this component will only be visible when the view actor is the component's owner, directly or indirectly. */
 	UPROPERTY(EditAnywhere, AdvancedDisplay, BlueprintReadOnly, Category = Rendering)
-	uint32 bOnlyOwnerSee:1;
+	uint8 bOnlyOwnerSee:1;
 
 	/** Treat this primitive as part of the background for occlusion purposes. This can be used as an optimization to reduce the cost of rendering skyboxes, large ground planes that are part of the vista, etc. */
 	UPROPERTY(EditAnywhere, AdvancedDisplay, BlueprintReadOnly, Category=Rendering)
-	uint32 bTreatAsBackgroundForOcclusion:1;
+	uint8 bTreatAsBackgroundForOcclusion:1;
 
 	/** 
 	 * Whether to render the primitive in the depth only pass.  
@@ -318,19 +322,19 @@ public:
 	 * @todo - if any rendering features rely on a complete depth only pass, this variable needs to go away.
 	 */
 	UPROPERTY(EditAnywhere, AdvancedDisplay, BlueprintReadOnly, Category=Rendering)
-	uint32 bUseAsOccluder:1;
+	uint8 bUseAsOccluder:1;
 
 	/** If this is True, this component can be selected in the editor. */
 	UPROPERTY()
-	uint32 bSelectable:1;
+	uint8 bSelectable:1;
 
 	/** If true, forces mips for textures used by this component to be resident when this component's level is loaded. */
 	UPROPERTY(EditAnywhere, AdvancedDisplay, BlueprintReadOnly, Category=TextureStreaming)
-	uint32 bForceMipStreaming:1;
+	uint8 bForceMipStreaming:1;
 
 	/** If true a hit-proxy will be generated for each instance of instanced static meshes */
 	UPROPERTY()
-	uint32 bHasPerInstanceHitProxies:1;
+	uint8 bHasPerInstanceHitProxies:1;
 
 	// Lighting flags
 	
@@ -340,23 +344,23 @@ public:
 	 * This flag is ignored (no shadows will be generated) if all materials on this component have an Unlit shading model.
 	 */
 	UPROPERTY(EditAnywhere, BlueprintReadOnly, Category=Lighting)
-	uint32 CastShadow:1;
+	uint8 CastShadow:1;
 
 	/** Controls whether the primitive should inject light into the Light Propagation Volume.  This flag is only used if CastShadow is true. **/
 	UPROPERTY(EditAnywhere, BlueprintReadOnly, Category=Lighting, AdvancedDisplay, meta=(EditCondition="CastShadow"))
-	uint32 bAffectDynamicIndirectLighting:1;
+	uint8 bAffectDynamicIndirectLighting:1;
 
 	/** Controls whether the primitive should affect dynamic distance field lighting methods.  This flag is only used if CastShadow is true. **/
 	UPROPERTY(EditAnywhere, BlueprintReadOnly, Category=Lighting, AdvancedDisplay, meta=(EditCondition="CastShadow"))
-	uint32 bAffectDistanceFieldLighting:1;
+	uint8 bAffectDistanceFieldLighting:1;
 
 	/** Controls whether the primitive should cast shadows in the case of non precomputed shadowing.  This flag is only used if CastShadow is true. **/
 	UPROPERTY(EditAnywhere, BlueprintReadOnly, Category=Lighting, AdvancedDisplay, meta=(EditCondition="CastShadow", DisplayName = "Dynamic Shadow"))
-	uint32 bCastDynamicShadow:1;
+	uint8 bCastDynamicShadow:1;
 
 	/** Whether the object should cast a static shadow from shadow casting lights.  This flag is only used if CastShadow is true. */
 	UPROPERTY(EditAnywhere, BlueprintReadOnly, Category=Lighting, AdvancedDisplay, meta=(EditCondition="CastShadow", DisplayName = "Static Shadow"))
-	uint32 bCastStaticShadow:1;
+	uint8 bCastStaticShadow:1;
 
 	/** 
 	 * Whether the object should cast a volumetric translucent shadow.
@@ -364,34 +368,34 @@ public:
 	 * But have artifacts when used on highly opaque surfaces.
 	 */
 	UPROPERTY(EditAnywhere, AdvancedDisplay, BlueprintReadOnly, Category=Lighting, meta=(EditCondition="CastShadow", DisplayName = "Volumetric Translucent Shadow"))
-	uint32 bCastVolumetricTranslucentShadow:1;
+	uint8 bCastVolumetricTranslucentShadow:1;
 
 	/** 
 	 * When enabled, the component will only cast a shadow on itself and not other components in the world.  
 	 * This is especially useful for first person weapons, and forces bCastInsetShadow to be enabled.
 	 */
 	UPROPERTY(EditAnywhere, AdvancedDisplay, BlueprintReadOnly, Category=Lighting, meta=(EditCondition="CastShadow"))
-	uint32 bSelfShadowOnly:1;
+	uint8 bSelfShadowOnly:1;
 
 	/** 
 	 * When enabled, the component will be rendering into the far shadow cascades (only for directional lights).
 	 */
 	UPROPERTY(EditAnywhere, AdvancedDisplay, BlueprintReadOnly, Category=Lighting, meta=(EditCondition="CastShadow", DisplayName = "Far Shadow"))
-	uint32 bCastFarShadow:1;
+	uint8 bCastFarShadow:1;
 
 	/** 
 	 * Whether this component should create a per-object shadow that gives higher effective shadow resolution. 
 	 * Useful for cinematic character shadowing. Assumed to be enabled if bSelfShadowOnly is enabled.
 	 */
 	UPROPERTY(EditAnywhere, AdvancedDisplay, BlueprintReadOnly, Category=Lighting, meta=(EditCondition="CastShadow", DisplayName = "Dynamic Inset Shadow"))
-	uint32 bCastInsetShadow:1;
+	uint8 bCastInsetShadow:1;
 
 	/** 
 	 * Whether this component should cast shadows from lights that have bCastShadowsFromCinematicObjectsOnly enabled.
 	 * This is useful for characters in a cinematic with special cinematic lights, where the cost of shadowmap rendering of the environment is undesired.
 	 */
 	UPROPERTY(EditAnywhere, AdvancedDisplay, BlueprintReadOnly, Category=Lighting, meta=(EditCondition="CastShadow"))
-	uint32 bCastCinematicShadow:1;
+	uint8 bCastCinematicShadow:1;
 
 	/** 
 	 *	If true, the primitive will cast shadows even if bHidden is true.
@@ -399,11 +403,11 @@ public:
 	 *	This flag is only used if CastShadow is true.
 	 */
 	UPROPERTY(EditAnywhere, AdvancedDisplay, BlueprintReadOnly, Category=Lighting, meta=(EditCondition="CastShadow", DisplayName = "Hidden Shadow"))
-	uint32 bCastHiddenShadow:1;
+	uint8 bCastHiddenShadow:1;
 
 	/** Whether this primitive should cast dynamic shadows as if it were a two sided material. */
 	UPROPERTY(EditAnywhere, AdvancedDisplay, BlueprintReadOnly, Category=Lighting, meta=(EditCondition="CastShadow", DisplayName = "Shadow Two Sided"))
-	uint32 bCastShadowAsTwoSided:1;
+	uint8 bCastShadowAsTwoSided:1;
 
 	/** 
 	 * Whether to light this primitive as if it were static, including generating lightmaps.  
@@ -411,7 +415,7 @@ public:
 	 * This is useful for moving meshes that don't change significantly.
 	 */
 	UPROPERTY(EditAnywhere, AdvancedDisplay, BlueprintReadOnly, Category=Lighting)
-	uint32 bLightAsIfStatic:1;
+	uint8 bLightAsIfStatic:1;
 
 	/** 
 	 * Whether to light this component and any attachments as a group.  This only has effect on the root component of an attachment tree.
@@ -419,11 +423,7 @@ public:
 	 * This is useful for improving performance when multiple movable components are attached together.
 	 */
 	UPROPERTY(EditAnywhere, AdvancedDisplay, BlueprintReadOnly, Category=Lighting)
-	uint32 bLightAttachmentsAsGroup:1;
-
-	/** Quality of indirect lighting for Movable primitives.  This has a large effect on Indirect Lighting Cache update time. */
-	UPROPERTY(EditAnywhere, AdvancedDisplay, BlueprintReadOnly, Category=Lighting)
-	TEnumAsByte<EIndirectLightingCacheQuality> IndirectLightingCacheQuality;
+	uint8 bLightAttachmentsAsGroup:1;
 
 	/** 
 	 * Mobile only:
@@ -431,7 +431,7 @@ public:
 	 * If disabled this component will only receive static shadows from stationary lights.
 	 */
 	UPROPERTY(EditAnywhere, AdvancedDisplay, BlueprintReadOnly, Category = Mobile, meta=(DisplayName ="Receive Combined Static and CSM Shadows from Stationary Lights"))
-	uint32 bReceiveCombinedCSMAndStaticShadowsFromStationaryLights : 1;
+	uint8 bReceiveCombinedCSMAndStaticShadowsFromStationaryLights : 1;
 
 	/** 
 	 * Whether the whole component should be shadowed as one from stationary lights, which makes shadow receiving much cheaper.
@@ -439,55 +439,79 @@ public:
 	 * This is currently only used on stationary directional lights.  
 	 */
 	UPROPERTY(EditAnywhere, AdvancedDisplay, BlueprintReadOnly, Category=Lighting)
-	uint32 bSingleSampleShadowFromStationaryLights:1;
-
-	/** 
-	 * Channels that this component should be in.  Lights with matching channels will affect the component.  
-	 * These channels only apply to opaque materials, direct lighting, and dynamic lighting and shadowing.
-	 */
-	UPROPERTY(EditAnywhere, AdvancedDisplay, BlueprintReadOnly, Category=Lighting)
-	FLightingChannels LightingChannels;
+	uint8 bSingleSampleShadowFromStationaryLights:1;
 
 	// Physics
 	
 	/** Will ignore radial impulses applied to this component. */
 	UPROPERTY(EditAnywhere, BlueprintReadWrite, Category = Physics)
-	uint32 bIgnoreRadialImpulse:1;
+	uint8 bIgnoreRadialImpulse:1;
 
 	/** Will ignore radial forces applied to this component. */
 	UPROPERTY(EditAnywhere, BlueprintReadWrite, Category = Physics)
-	uint32 bIgnoreRadialForce:1;
+	uint8 bIgnoreRadialForce:1;
 
 	/** True for damage to this component to apply physics impulse, false to opt out of these impulses. */
 	UPROPERTY(EditAnywhere, BlueprintReadWrite, Category = Physics)
-	uint32 bApplyImpulseOnDamage : 1;
+	uint8 bApplyImpulseOnDamage : 1;
 
 	// General flags.
 	
 	/** If this is True, this component must always be loaded on clients, even if Hidden and CollisionEnabled is NoCollision. */
 	UPROPERTY()
-	uint32 AlwaysLoadOnClient:1;
+	uint8 AlwaysLoadOnClient:1;
 
 	/** If this is True, this component must always be loaded on servers, even if Hidden and CollisionEnabled is NoCollision */
 	UPROPERTY()
-	uint32 AlwaysLoadOnServer:1;
+	uint8 AlwaysLoadOnServer:1;
 
 	/** Composite the drawing of this component onto the scene after post processing (only applies to editor drawing) */
 	UPROPERTY()
-	uint32 bUseEditorCompositing:1;
+	uint8 bUseEditorCompositing:1;
 
 	/** If true, this component will be rendered in the CustomDepth pass (usually used for outlines) */
 	UPROPERTY(EditAnywhere, AdvancedDisplay, BlueprintReadOnly, Category=Rendering, meta=(DisplayName = "Render CustomDepth Pass"))
-	uint32 bRenderCustomDepth:1;
+	uint8 bRenderCustomDepth:1;
 
-	/** Optionally write this 0-255 value to the stencil buffer in CustomDepth pass (Requires project setting or r.CustomDepth == 3) */
-	UPROPERTY(EditAnywhere, AdvancedDisplay, BlueprintReadOnly, Category=Rendering,  meta=(UIMin = "0", UIMax = "255", editcondition = "bRenderCustomDepth", DisplayName = "CustomDepth Stencil Value"))
-	int32 CustomDepthStencilValue;
+protected:
+	/** Result of last call to AreAllCollideableDescendantsRelative(). */
+	uint8 bCachedAllCollideableDescendantsRelative : 1;
+
+public:
+	/** If true then DoCustomNavigableGeometryExport will be called to collect navigable geometry of this component. */
+	UPROPERTY()
+	TEnumAsByte<EHasCustomNavigableGeometry::Type> bHasCustomNavigableGeometry;
+
+private:
+	UPROPERTY()
+	TEnumAsByte<enum ECanBeCharacterBase> CanBeCharacterBase_DEPRECATED;
+
+	FMaskFilter MoveIgnoreMask;
+
+public:
+	/**
+	 * Determine whether a Character can step up onto this component.
+	 * This controls whether they can try to step up on it when they bump in to it, not whether they can walk on it after landing on it.
+	 * @see FWalkableSlopeOverride
+	 */
+	UPROPERTY(EditAnywhere, BlueprintReadWrite, Category=Collision)
+	TEnumAsByte<enum ECanBeCharacterBase> CanCharacterStepUpOn;
+
+	/** 
+	 * Channels that this component should be in.  Lights with matching channels will affect the component.  
+	 * These channels only apply to opaque materials, direct lighting, and dynamic lighting and shadowing.
+	 */
+	UPROPERTY(EditAnywhere, AdvancedDisplay, BlueprintReadOnly, Category=Lighting)
+	FLightingChannels LightingChannels;
 
 	/** Mask used for stencil buffer writes. */
 	UPROPERTY(EditAnywhere, AdvancedDisplay, BlueprintReadOnly, Category = "Rendering", meta = (editcondition = "bRenderCustomDepth"))
 	ERendererStencilMask CustomDepthStencilWriteMask;
 
+	/** Optionally write this 0-255 value to the stencil buffer in CustomDepth pass (Requires project setting or r.CustomDepth == 3) */
+	UPROPERTY(EditAnywhere, AdvancedDisplay, BlueprintReadOnly, Category=Rendering,  meta=(UIMin = "0", UIMax = "255", editcondition = "bRenderCustomDepth", DisplayName = "CustomDepth Stencil Value"))
+	int32 CustomDepthStencilValue;
+
 	/**
 	 * Translucent objects with a lower sort priority draw behind objects with a higher priority.
 	 * Translucent objects with the same priority are rendered from back-to-front based on their bounds origin.
@@ -514,11 +538,27 @@ public:
 	UPROPERTY(EditAnywhere, BlueprintReadOnly, AdvancedDisplay, Category=Rendering, meta=(UIMin = "0.0", UIMax = "3.0"))
 	float LpvBiasMultiplier;
 
-	// Internal physics engine data.
-	
-	/** Physics scene information for this component, holds a single rigid body with multiple shapes. */
-	UPROPERTY(EditAnywhere, BlueprintReadOnly, Category=Collision, meta=(ShowOnlyInnerProperties))
-	FBodyInstance BodyInstance;
+	/**
+	* Incremented by the main thread before being attached to the scene, decremented
+	* by the rendering thread after removal. This counter exists to assert that 
+	* operations are safe in order to help avoid race conditions.
+	*
+	*           *** Runtime logic should NEVER rely on this value. ***
+	*
+	* The only safe assertions to make are:
+	*
+	*     AttachmentCounter == 0: The primitive is not exposed to the rendering
+	*                             thread, it is safe to modify shared members.
+	*                             This assertion is valid ONLY from the main thread.
+	*
+	*     AttachmentCounter >= 1: The primitive IS exposed to the rendering
+	*                             thread and therefore shared members must not
+	*                             be modified. This assertion may be made from
+	*                             any thread. Note that it is valid and expected
+	*                             for AttachmentCounter to be larger than 1, e.g.
+	*                             during reattachment.
+	*/
+	FThreadSafeCounter AttachmentCounter;
 
 	/** Used to detach physics objects before simulation begins. This is needed because at runtime we can't have simulated objects inside the attachment hierarchy */
 	virtual void BeginPlay() override;
@@ -528,16 +568,9 @@ protected:
 	/** Returns true if all descendant components that we can possibly overlap with use relative location and rotation. */
 	virtual bool AreAllCollideableDescendantsRelative(bool bAllowCachedValue = true) const;
 
-	/** Result of last call to AreAllCollideableDescendantsRelative(). */
-	uint32 bCachedAllCollideableDescendantsRelative:1;
-
 	/** Last time we checked AreAllCollideableDescendantsRelative(), so we can throttle those tests since it rarely changes once false. */
 	float LastCheckedAllCollideableDescendantsTime;
 
-	/** If true then DoCustomNavigableGeometryExport will be called to collect navigable geometry of this component. */
-	UPROPERTY()
-	TEnumAsByte<EHasCustomNavigableGeometry::Type> bHasCustomNavigableGeometry;
-
 	/** Next id to be used by a component. */
 	static FThreadSafeCounter NextComponentId;
 
@@ -565,20 +598,6 @@ public:
 
 	UPROPERTY(transient)
 	float LastRenderTimeOnScreen;
-private:
-	UPROPERTY()
-	TEnumAsByte<enum ECanBeCharacterBase> CanBeCharacterBase_DEPRECATED;
-
-	FMaskFilter MoveIgnoreMask;
-
-public:
-	/**
-	 * Determine whether a Character can step up onto this component.
-	 * This controls whether they can try to step up on it when they bump in to it, not whether they can walk on it after landing on it.
-	 * @see FWalkableSlopeOverride
-	 */
-	UPROPERTY(EditAnywhere, BlueprintReadWrite, Category=Collision)
-	TEnumAsByte<enum ECanBeCharacterBase> CanCharacterStepUpOn;
 
 	/**
 	 * Set of actors to ignore during component sweeps in MoveComponent().
@@ -779,6 +798,12 @@ protected:
 
 public:
 
+	// Internal physics engine data.
+	
+	/** Physics scene information for this component, holds a single rigid body with multiple shapes. */
+	UPROPERTY(EditAnywhere, BlueprintReadOnly, Category=Collision, meta=(ShowOnlyInnerProperties))
+	FBodyInstance BodyInstance;
+
 	/** 
 	 *	Event called when a component hits (or is hit by) something solid. This could happen due to things like Character movement, using Set Location with 'sweep' enabled, or physics simulation.
 	 *	For events when objects overlap (e.g. walking into a trigger) see the 'Overlap' event.
@@ -1360,28 +1385,6 @@ public:
 	/** A fence to track when the primitive is detached from the scene in the rendering thread. */
 	FRenderCommandFence DetachFence;
 
-	/**
-	 * Incremented by the main thread before being attached to the scene, decremented
-	 * by the rendering thread after removal. This counter exists to assert that 
-	 * operations are safe in order to help avoid race conditions.
-	 *
-	 *           *** Runtime logic should NEVER rely on this value. ***
-	 *
-	 * The only safe assertions to make are:
-	 *
-	 *     AttachmentCounter == 0: The primitive is not exposed to the rendering
-	 *                             thread, it is safe to modify shared members.
-	 *                             This assertion is valid ONLY from the main thread.
-	 *
-	 *     AttachmentCounter >= 1: The primitive IS exposed to the rendering
-	 *                             thread and therefore shared members must not
-	 *                             be modified. This assertion may be made from
-	 *                             any thread. Note that it is valid and expected
-	 *                             for AttachmentCounter to be larger than 1, e.g.
-	 *                             during reattachment.
-	 */
-	FThreadSafeCounter AttachmentCounter;
-
 	// Scene data
 private:
 	/** LOD parent primitive to draw instead of this one (multiple UPrim's will point to the same LODParent ) */
diff --git a/Engine/Source/Runtime/Engine/Classes/Components/SceneComponent.h b/Engine/Source/Runtime/Engine/Classes/Components/SceneComponent.h
--- a/Engine/Source/Runtime/Engine/Classes/Components/SceneComponent.h
+++ b/Engine/Source/Runtime/Engine/Classes/Components/SceneComponent.h
@@ -1,4 +1,4 @@
-// Copyright 1998-2017 Epic Games, Inc. All Rights Reserved.
+﻿// Copyright 1998-2017 Epic Games, Inc. All Rights Reserved.
 
 #pragma once
 
@@ -116,146 +116,184 @@
 
 	virtual void GetLifetimeReplicatedProps(TArray<FLifetimeProperty>& OutLifetimeProps) const override;
 
+	/** Cached level collection that contains the level this component is registered in, for fast access in IsVisible(). */
+	const FLevelCollection* CachedLevelCollection;
+
+private:
+	/** Physics Volume in which this SceneComponent is located **/
+	UPROPERTY(transient)
+	TWeakObjectPtr<class APhysicsVolume> PhysicsVolume;
+
 	/** What we are currently attached to. If valid, RelativeLocation etc. are used relative to this object */
-private:
 	UPROPERTY(ReplicatedUsing = OnRep_AttachParent)
 	USceneComponent* AttachParent;
 
+	/** Optional socket name on AttachParent that we are attached to. */
+	UPROPERTY(ReplicatedUsing = OnRep_AttachSocketName)
+	FName AttachSocketName;
+
 	/** List of child SceneComponents that are attached to us. */
 	UPROPERTY(ReplicatedUsing = OnRep_AttachChildren, Transient)
 	TArray<USceneComponent*> AttachChildren;
 
 	/** Set of attached SceneComponents that were attached by the client so we can fix up AttachChildren when it is replicated to us. */
 	UPROPERTY(Transient)
-	TSet<USceneComponent*> ClientAttachedChildren;
+	TArray<USceneComponent*> ClientAttachedChildren;
 
-	/** Optional socket name on AttachParent that we are attached to. */
-	UPROPERTY(ReplicatedUsing = OnRep_AttachSocketName)
-	FName AttachSocketName;
+	FName NetOldAttachSocketName;
+	USceneComponent* NetOldAttachParent;
 
+public:
+	/** Current bounds of the component */
+	FBoxSphereBounds Bounds;
+
+	/** Location of the component relative to its parent */
+	UPROPERTY(EditAnywhere, BlueprintReadOnly, ReplicatedUsing=OnRep_Transform, Category = Transform)
+	FVector RelativeLocation;
+
+	/** Rotation of the component relative to its parent */
+	UPROPERTY(EditAnywhere, BlueprintReadOnly, ReplicatedUsing=OnRep_Transform, Category=Transform)
+	FRotator RelativeRotation;
+
+	/**
+	*	Non-uniform scaling of the component relative to its parent.
+	*	Note that scaling is always applied in local space (no shearing etc)
+	*/
+	UPROPERTY(BlueprintReadOnly, ReplicatedUsing=OnRep_Transform, interp, Category=Transform)
+	FVector RelativeScale3D;
+
+private:
+
+	/** Current transform of the component, relative to the world */
+	FTransform ComponentToWorld;
+
+public:
+
+	/**
+	* Velocity of the component.
+	* @see GetComponentVelocity()
+	*/
+	UPROPERTY()
+	FVector ComponentVelocity;
+
+private:
 	/** True if we have ever updated ComponentToWorld based on RelativeLocation/Rotation/Scale. Used at startup to make sure it is initialized. */
 	UPROPERTY(Transient)
-	uint32 bWorldToComponentUpdated : 1;
+	uint8 bComponentToWorldUpdated : 1;
 
 public:
 	/** If RelativeLocation should be considered relative to the world, rather than the parent */
 	UPROPERTY(EditAnywhere, AdvancedDisplay, BlueprintReadWrite, ReplicatedUsing=OnRep_Transform, Category=Transform)
-	uint32 bAbsoluteLocation:1;
+	uint8 bAbsoluteLocation:1;
 
 	/** If RelativeRotation should be considered relative to the world, rather than the parent */
 	UPROPERTY(EditAnywhere, AdvancedDisplay, BlueprintReadWrite, ReplicatedUsing=OnRep_Transform, Category=Transform)
-	uint32 bAbsoluteRotation:1;
+	uint8 bAbsoluteRotation:1;
 
 	/** If RelativeScale3D should be considered relative to the world, rather than the parent */
 	UPROPERTY(EditAnywhere, AdvancedDisplay, BlueprintReadWrite, ReplicatedUsing=OnRep_Transform, Category=Transform)
-	uint32 bAbsoluteScale:1;
+	uint8 bAbsoluteScale:1;
 
 	/** Whether to completely draw the primitive; if false, the primitive is not drawn, does not cast a shadow. */
 	UPROPERTY(EditAnywhere, BlueprintReadOnly, ReplicatedUsing=OnRep_Visibility,  Category = Rendering)
-	uint32 bVisible:1;
+	uint8 bVisible:1;
 
 	/** Whether to hide the primitive in game, if the primitive is Visible. */
 	UPROPERTY(EditAnywhere, BlueprintReadOnly, Category=Rendering, meta=(SequencerTrackClass = "MovieSceneVisibilityTrack"))
-	uint32 bHiddenInGame:1;
+	uint8 bHiddenInGame:1;
 
 	/**
 	 * Whether or not the cached PhysicsVolume this component overlaps should be updated when the component is moved.
 	 * @see GetPhysicsVolume()
 	 */
 	UPROPERTY(EditAnywhere, BlueprintReadWrite, AdvancedDisplay, Category=Physics)
-	uint32 bShouldUpdatePhysicsVolume:1;
+	uint8 bShouldUpdatePhysicsVolume:1;
 
 	/** If true, a change in the bounds of the component will call trigger a streaming data rebuild */
 	UPROPERTY()
-	uint32 bBoundsChangeTriggersStreamingDataRebuild:1;
+	uint8 bBoundsChangeTriggersStreamingDataRebuild:1;
 
 	/** If true, this component uses its parents bounds when attached.
 	 *  This can be a significant optimization with many components attached together.
 	 */
 	UPROPERTY(EditAnywhere, BlueprintReadWrite, AdvancedDisplay, Category=Rendering)
-	uint32 bUseAttachParentBound:1;
+	uint8 bUseAttachParentBound:1;
 
 protected:
 
 	// Transient flag that temporarily disables UpdateOverlaps within DetachFromParent().
-	uint32 bDisableDetachmentUpdateOverlaps:1;
+	uint8 bDisableDetachmentUpdateOverlaps:1;
 
 	/** If true, OnUpdateTransform virtual will be called each time this component is moved. */
-	uint32 bWantsOnUpdateTransform:1;
+	uint8 bWantsOnUpdateTransform:1;
 
 private:
 
+	uint8 bNetUpdateTransform : 1;
+	uint8 bNetUpdateAttachment : 1;
+
 	// DEPRECATED
 	UPROPERTY()
-	uint32 bAbsoluteTranslation_DEPRECATED : 1;
+	uint8 bAbsoluteTranslation_DEPRECATED : 1;
 
 	// Appends all descendants (recursively) of this scene component to the list of Children.  NOTE: It does NOT clear the list first.
 	void AppendDescendants(TArray<USceneComponent*>& Children) const;
 
-	/** Physics Volume in which this SceneComponent is located **/
-	UPROPERTY(transient)
-	TWeakObjectPtr<class APhysicsVolume> PhysicsVolume;
+public:
 
-public:
-	/** Current bounds of the component */
-	FBoxSphereBounds Bounds;
+#if WITH_EDITORONLY_DATA
+	UPROPERTY()
+	uint8 bVisualizeComponent : 1;
+#endif
 
-	/** Current transform of the component, relative to the world */
-	DEPRECATED(4.17, "ComponentToWorld will be made private, use GetComponentTransform() instead.")
-	FTransform ComponentToWorld;
+	/** How often this component is allowed to move, used to make various optimizations. Only safe to set in constructor. */
+	UPROPERTY(Category = Mobility, EditAnywhere, BlueprintReadOnly)
+	TEnumAsByte<EComponentMobility::Type> Mobility;
+
+	/** If detail mode is >= system detail mode, primitive won't be rendered. */
+	UPROPERTY(EditAnywhere, AdvancedDisplay, BlueprintReadOnly, Category = LOD)
+	TEnumAsByte<enum EDetailMode> DetailMode;
 
 private:
 	/** Cache that avoids Quat<->Rotator conversions if possible. Only to be used with GetComponentTransform().GetRotation(). */
 	FRotationConversionCache WorldRotationCache;
 
+	/** Cache that avoids Quat<->Rotator conversions if possible. Only to be used with RelativeRotation. */
+	FRotationConversionCache RelativeRotationCache;
+
 public:
 
-	/** Location of the component relative to its parent */
-	UPROPERTY(EditAnywhere, BlueprintReadOnly, ReplicatedUsing=OnRep_Transform, Category = Transform)
-	FVector RelativeLocation;
-
-	/** Rotation of the component relative to its parent */
-	UPROPERTY(EditAnywhere, BlueprintReadOnly, ReplicatedUsing=OnRep_Transform, Category=Transform)
-	FRotator RelativeRotation;
-
 	/** Sets the RelativeRotationCache. Used to ensure component ends up with the same RelativeRotation after calling SetWorldTransform(). */
 	void SetRelativeRotationCache(const FRotationConversionCache& InCache);
 	/** Get the RelativeRotationCache.  */
 	FORCEINLINE const FRotationConversionCache& GetRelativeRotationCache() const { return RelativeRotationCache; }
 
+	/** Delegate that will be called when PhysicsVolume has been changed **/
+	UPROPERTY(BlueprintAssignable, Category=PhysicsVolume, meta=(DisplayName="Physics Volume Changed"))
+	FPhysicsVolumeChanged PhysicsVolumeChangedDelegate;
+
+	/** Returns the current scoped movement update, or NULL if there is none. @see FScopedMovementUpdate */
+	class FScopedMovementUpdate* GetCurrentScopedMovement() const;
+
 private:
-	/** Cache that avoids Quat<->Rotator conversions if possible. Only to be used with RelativeRotation. */
-	FRotationConversionCache RelativeRotationCache;
+
+	/** Stack of current movement scopes. */
+	TArray<class FScopedMovementUpdate*> ScopedMovementStack;
+
+	void BeginScopedMovementUpdate(class FScopedMovementUpdate& ScopedUpdate);
+	void EndScopedMovementUpdate(class FScopedMovementUpdate& ScopedUpdate);
+
+	friend class FScopedMovementUpdate;
+	friend class FScopedPreventAttachedComponentMove;
 
 public:
-	/**
-	*	Non-uniform scaling of the component relative to its parent.
-	*	Note that scaling is always applied in local space (no shearing etc)
-	*/
-	UPROPERTY(BlueprintReadOnly, ReplicatedUsing=OnRep_Transform, interp, Category=Transform)
-	FVector RelativeScale3D;
 
+#if WITH_EDITORONLY_DATA
 	UPROPERTY()
 	FVector RelativeTranslation_DEPRECATED;
-
-	/** How often this component is allowed to move, used to make various optimizations. Only safe to set in constructor. */
-	UPROPERTY(Category = Mobility, EditAnywhere, BlueprintReadOnly)
-	TEnumAsByte<EComponentMobility::Type> Mobility;
-
-	/** If detail mode is >= system detail mode, primitive won't be rendered. */
-	UPROPERTY(EditAnywhere, AdvancedDisplay, BlueprintReadOnly, Category=LOD)
-	TEnumAsByte<enum EDetailMode> DetailMode;
+#endif
 
 private:
-
-	uint8 bNetUpdateTransform:1;
-	uint8 bNetUpdateAttachment:1;
-	FName NetOldAttachSocketName;
-	USceneComponent *NetOldAttachParent;
-
-	/** Cached level collection that contains the level this component is registered in, for fast access in IsVisible(). */
-	const FLevelCollection* CachedLevelCollection;
-
 	UFUNCTION()
 	void OnRep_Transform();
 
@@ -278,14 +316,5 @@
 	virtual void PostRepNotifies() override;
 
 public:
-
 	/**
-	 * Velocity of the component.
-	 * @see GetComponentVelocity()
-	 */
-	UPROPERTY()
-	FVector ComponentVelocity;
-
-public:
-	/**
 	 * Set the location of the component relative to its parent
@@ -817,10 +846,6 @@
 	}
 
 public:
-	/** Delegate that will be called when PhysicsVolume has been changed **/
-	UPROPERTY(BlueprintAssignable, Category=PhysicsVolume, meta=(DisplayName="Physics Volume Changed"))
-	FPhysicsVolumeChanged PhysicsVolumeChangedDelegate;
-
 	FTransformUpdated TransformUpdated;
 
 	//~ Begin ActorComponent Interface
@@ -841,14 +866,14 @@
 	virtual class FActorComponentInstanceData* GetComponentInstanceData() const override;
 	//~ End ActorComponent Interface
 
-	// Call UpdateComponentToWorld if bWorldToComponentUpdated is false.
+	// Call UpdateComponentToWorld if bComponentToWorldUpdated is false.
 	void ConditionalUpdateComponentToWorld();
 
 	//~ Begin UObject Interface
-	virtual void Serialize(FArchive& Ar) override;
 	virtual void PostInterpChange(UProperty* PropertyThatChanged) override;
 	virtual void BeginDestroy() override;
 #if WITH_EDITORONLY_DATA
+	virtual void Serialize(FArchive& Ar) override;
 	static void AddReferencedObjects(UObject* InThis, FReferenceCollector& Collector);
 #endif
 
@@ -915,28 +940,10 @@
 	/** Returns true if movement is currently within the scope of an FScopedMovementUpdate. */
 	bool IsDeferringMovementUpdates() const;
 
-	/** Returns the current scoped movement update, or NULL if there is none. @see FScopedMovementUpdate */
-	class FScopedMovementUpdate* GetCurrentScopedMovement() const;
-
-private:
-
-	/** Stack of current movement scopes. */
-	TArray<class FScopedMovementUpdate*> ScopedMovementStack;
-
-	void BeginScopedMovementUpdate(class FScopedMovementUpdate& ScopedUpdate);
-	void EndScopedMovementUpdate(class FScopedMovementUpdate& ScopedUpdate);
-
-	friend class FScopedMovementUpdate;
-	friend class FScopedPreventAttachedComponentMove;
-
 #if WITH_EDITORONLY_DATA
 protected:
 	/** Editor only component used to display the sprite so as to be able to see the location of the Audio Component  */
 	class UBillboardComponent* SpriteComponent;
-
-public:
-	UPROPERTY()
-	uint32 bVisualizeComponent : 1;
 #endif
 
 public:
@@ -971,10 +978,8 @@
 	/** Sets the cached component to world directly. This should be used very rarely. */
 	FORCEINLINE void SetComponentToWorld(const FTransform& NewComponentToWorld)
 	{
-		bWorldToComponentUpdated = true;
-		PRAGMA_DISABLE_DEPRECATION_WARNINGS
+		bComponentToWorldUpdated = true;
 		ComponentToWorld = NewComponentToWorld;
-		PRAGMA_ENABLE_DEPRECATION_WARNINGS
 	}
 
 	/** 
@@ -983,17 +988,13 @@
 	 */
 	FORCEINLINE const FTransform& GetComponentToWorld() const 
 	{ 
-		PRAGMA_DISABLE_DEPRECATION_WARNINGS
 		return ComponentToWorld;
-		PRAGMA_ENABLE_DEPRECATION_WARNINGS
 	}
 
 	/** Get the current component-to-world transform for this component */
 	FORCEINLINE const FTransform& GetComponentTransform() const
 	{
-		PRAGMA_DISABLE_DEPRECATION_WARNINGS
 		return ComponentToWorld;
-		PRAGMA_ENABLE_DEPRECATION_WARNINGS
 	}
 
 	/** Update transforms of any components attached to this one. */
@@ -1298,7 +1299,7 @@
 
 FORCEINLINE_DEBUGGABLE void USceneComponent::ConditionalUpdateComponentToWorld()
 {
-	if (!bWorldToComponentUpdated)
+	if (!bComponentToWorldUpdated)
 	{
 		UpdateComponentToWorld();
 	}
diff --git a/Engine/Source/Runtime/Engine/Classes/Components/StaticMeshComponent.h b/Engine/Source/Runtime/Engine/Classes/Components/StaticMeshComponent.h
--- a/Engine/Source/Runtime/Engine/Classes/Components/StaticMeshComponent.h
+++ b/Engine/Source/Runtime/Engine/Classes/Components/StaticMeshComponent.h
@@ -162,10 +162,6 @@ class ENGINE_API UStaticMeshComponent : public UMeshComponent
 	UPROPERTY()
 	int32 PreviousLODLevel;
 
-	/** Whether to override the MinLOD setting of the static mesh asset with the MinLOD of this component. */
-	UPROPERTY(EditAnywhere, AdvancedDisplay, BlueprintReadOnly, Category=LOD)
-	bool bOverrideMinLOD;
-
 	/** 
 	 * Specifies the smallest LOD that will be used for this component.  
 	 * This is ignored if ForcedLodModel is enabled.
@@ -173,18 +169,23 @@ class ENGINE_API UStaticMeshComponent : public UMeshComponent
 	UPROPERTY(EditAnywhere, AdvancedDisplay, BlueprintReadOnly, Category=LOD, meta=(editcondition = "bOverrideMinLOD"))
 	int32 MinLOD;
 
+	/** Subdivision step size for static vertex lighting.				*/
+	UPROPERTY()
+	int32 SubDivisionStepSize;
+
 	/** The static mesh that this component uses to render */
-	DEPRECATED(4.14, "Direct access to StaticMesh member is deprecated and will be made private soon. Please use SetStaticMesh and GetStaticMesh.")
+private:
 	UPROPERTY(EditAnywhere, BlueprintReadOnly, Category=StaticMesh, ReplicatedUsing=OnRep_StaticMesh, meta=(AllowPrivateAccess="true"))
 	class UStaticMesh* StaticMesh;
 
+public:
+
+	/** Helper function to get the FName of the private static mesh member */
+	static const FName GetMemberNameChecked_StaticMesh() { return GET_MEMBER_NAME_CHECKED(UStaticMeshComponent, StaticMesh); }
+
 	UFUNCTION()
 	void OnRep_StaticMesh(class UStaticMesh *OldStaticMesh);
 
-	/** If true, WireframeColorOverride will be used. If false, color is determined based on mobility and physics simulation settings */
-	UPROPERTY(EditAnywhere, BlueprintReadOnly, Category=Rendering, meta=(InlineEditConditionToggle))
-	bool bOverrideWireframeColor;
-
 	/** Wireframe color to use if bOverrideWireframeColor is true */
 	UPROPERTY(EditAnywhere, AdvancedDisplay, BlueprintReadOnly, Category=Rendering, meta=(editcondition = "bOverrideWireframeColor"))
 	FColor WireframeColorOverride;
@@ -217,23 +218,31 @@ class ENGINE_API UStaticMeshComponent : public UMeshComponent
 	int32 StaticMeshImportVersion;
 #endif
 
+	/** If true, WireframeColorOverride will be used. If false, color is determined based on mobility and physics simulation settings */
+	UPROPERTY(EditAnywhere, BlueprintReadOnly, Category=Rendering, meta=(InlineEditConditionToggle))
+	uint8 bOverrideWireframeColor:1;
+
+	/** Whether to override the MinLOD setting of the static mesh asset with the MinLOD of this component. */
+	UPROPERTY(EditAnywhere, AdvancedDisplay, BlueprintReadOnly, Category=LOD)
+	uint8 bOverrideMinLOD:1;
+
 	/** If true, bForceNavigationObstacle flag will take priority over navigation data stored in StaticMesh */
 	UPROPERTY(transient)
-	uint32 bOverrideNavigationExport : 1;
+	uint8 bOverrideNavigationExport : 1;
 
 	/** Allows overriding navigation export behavior per component: full collisions or dynamic obstacle */
 	UPROPERTY(transient)
-	uint32 bForceNavigationObstacle : 1;
+	uint8 bForceNavigationObstacle : 1;
 
 	/** If true, mesh painting is disallowed on this instance. Set if vertex colors are overridden in a construction script. */
 	UPROPERTY(EditAnywhere, AdvancedDisplay, BlueprintReadWrite, Category=Rendering)
-	uint32 bDisallowMeshPaintPerInstance : 1;
+	uint8 bDisallowMeshPaintPerInstance : 1;
 
 #if !(UE_BUILD_SHIPPING)
 	/** Draw mesh collision if used for complex collision */
-	uint32 bDrawMeshCollisionIfComplex : 1;
+	uint8 bDrawMeshCollisionIfComplex : 1;
 	/** Draw mesh collision if used for simple collision */
-	uint32 bDrawMeshCollisionIfSimple : 1;
+	uint8 bDrawMeshCollisionIfSimple : 1;
 #endif
 
 	/**
@@ -242,15 +251,11 @@ class ENGINE_API UStaticMeshComponent : public UMeshComponent
 	 *	to avoid them using distance-based streaming.
 	 */
 	UPROPERTY(EditAnywhere, AdvancedDisplay, BlueprintReadWrite, Category=TextureStreaming)
-	uint32 bIgnoreInstanceForTextureStreaming:1;
+	uint8 bIgnoreInstanceForTextureStreaming:1;
 
 	/** Whether to override the lightmap resolution defined in the static mesh. */
 	UPROPERTY(EditAnywhere, BlueprintReadOnly, Category=Lighting, meta=(InlineEditConditionToggle))
-	uint32 bOverrideLightMapRes:1;
-
-	/** Light map resolution to use on this component, used if bOverrideLightMapRes is true and there is a valid StaticMesh. */
-	UPROPERTY(EditAnywhere, BlueprintReadOnly, Category=Lighting, meta=(ClampMax = 4096, editcondition="bOverrideLightMapRes") )
-	int32 OverriddenLightMapRes;
+	uint8 bOverrideLightMapRes:1;
 
 	/** 
 	 * Whether to use the mesh distance field representation (when present) for shadowing indirect lighting (from lightmaps or skylight) on Movable components.
@@ -258,7 +263,32 @@ class ENGINE_API UStaticMeshComponent : public UMeshComponent
 	 * The StaticMesh must have 'Generate Mesh Distance Field' enabled, or the project must have 'Generate Mesh Distance Fields' enabled for this feature to work.
 	 */
 	UPROPERTY(EditAnywhere, AdvancedDisplay, BlueprintReadOnly, Category=Lighting, meta=(DisplayName = "Distance Field Indirect Shadow"))
-	uint32 bCastDistanceFieldIndirectShadow:1;
+	uint8 bCastDistanceFieldIndirectShadow:1;
+
+	/** Whether to override the DistanceFieldSelfShadowBias setting of the static mesh asset with the DistanceFieldSelfShadowBias of this component. */
+	UPROPERTY(EditAnywhere, AdvancedDisplay, BlueprintReadOnly, Category=Lighting)
+	uint8 bOverrideDistanceFieldSelfShadowBias:1;
+
+	/** Whether to use subdivisions or just the triangle's vertices.	*/
+	UPROPERTY()
+	uint8 bUseSubDivisions:1;
+
+	/** Use the collision profile specified in the StaticMesh asset.*/
+	UPROPERTY(EditAnywhere, Category = Collision)
+	uint8 bUseDefaultCollision:1;
+
+#if WITH_EDITORONLY_DATA
+	/** The component has some custom painting on LODs or not. */
+	UPROPERTY()
+	uint8 bCustomOverrideVertexColorPerLOD:1;
+
+	UPROPERTY(transient)
+	uint8 bDisplayVertexColors:1;
+#endif
+
+	/** Light map resolution to use on this component, used if bOverrideLightMapRes is true and there is a valid StaticMesh. */
+	UPROPERTY(EditAnywhere, BlueprintReadOnly, Category=Lighting, meta=(ClampMax = 4096, editcondition="bOverrideLightMapRes") )
+	int32 OverriddenLightMapRes;
 
 	/** 
 	 * Controls how dark the dynamic indirect shadow can be.
@@ -266,10 +296,6 @@ class ENGINE_API UStaticMeshComponent : public UMeshComponent
 	UPROPERTY(EditAnywhere, AdvancedDisplay, BlueprintReadOnly, Category=Lighting, meta=(UIMin = "0", UIMax = "1", DisplayName = "Distance Field Indirect Shadow Min Visibility"))
 	float DistanceFieldIndirectShadowMinVisibility;
 
-	/** Whether to override the DistanceFieldSelfShadowBias setting of the static mesh asset with the DistanceFieldSelfShadowBias of this component. */
-	UPROPERTY(EditAnywhere, AdvancedDisplay, BlueprintReadOnly, Category=Lighting)
-	bool bOverrideDistanceFieldSelfShadowBias;
-
 	/** Useful for reducing self shadowing from distance field methods when using world position offset to animate the mesh's vertices. */
 	UPROPERTY(EditAnywhere, AdvancedDisplay, BlueprintReadOnly, Category=Lighting)
 	float DistanceFieldSelfShadowBias;
@@ -283,16 +309,10 @@ class ENGINE_API UStaticMeshComponent : public UMeshComponent
 	UPROPERTY(EditAnywhere, AdvancedDisplay, BlueprintReadWrite, Category=TextureStreaming, meta=(ClampMin = 0, ToolTip="Allows adjusting the desired resolution of streaming textures that uses UV 0.  1.0 is the default, whereas a higher value increases the streamed-in resolution."))
 	float StreamingDistanceMultiplier;
 
-	/** Subdivision step size for static vertex lighting.				*/
-	UPROPERTY()
-	int32 SubDivisionStepSize;
-
-	/** Whether to use subdivisions or just the triangle's vertices.	*/
-	UPROPERTY()
-	uint32 bUseSubDivisions:1;
-
+#if WITH_EDITORONLY_DATA
 	UPROPERTY()
 	TArray<FGuid> IrrelevantLights_DEPRECATED;
+#endif
 
 	/** Static mesh LOD data.  Contains static lighting data along with instanced mesh vertex colors. */
 	UPROPERTY(transient)
@@ -302,10 +322,6 @@ class ENGINE_API UStaticMeshComponent : public UMeshComponent
 	UPROPERTY(NonTransactional)
 	TArray<FStreamingTextureBuildInfo> StreamingTextureData;
 
-	/** Use the collision profile specified in the StaticMesh asset.*/
-	UPROPERTY(EditAnywhere, Category = Collision)
-	bool bUseDefaultCollision;
-
 #if WITH_EDITORONLY_DATA
 	/** Derived data key of the static mesh, used to determine if an update from the source static mesh is required. */
 	UPROPERTY()
@@ -314,13 +330,6 @@ class ENGINE_API UStaticMeshComponent : public UMeshComponent
 	/** Material Bounds used for texture streaming. */
 	UPROPERTY(NonTransactional)
 	TArray<uint32> MaterialStreamingRelativeBoxes;
-
-	/** The component has some custom painting on LODs or not. */
-	UPROPERTY()
-	bool bCustomOverrideVertexColorPerLOD;
-
-	UPROPERTY(transient)
-	bool bDisplayVertexColors;
 #endif
 
 	/** The Lightmass settings for this object. */
@@ -336,9 +345,7 @@ class ENGINE_API UStaticMeshComponent : public UMeshComponent
 	/** Get the StaticMesh used by this instance. */
 	UStaticMesh* GetStaticMesh() const 
 	{ 
-		PRAGMA_DISABLE_DEPRECATION_WARNINGS
 		return StaticMesh; 
-		PRAGMA_ENABLE_DEPRECATION_WARNINGS
 	}
 
 	UFUNCTION(BlueprintCallable, Category="Rendering|LOD")
diff --git a/Engine/Source/Runtime/Engine/Classes/EdGraph/EdGraphPin.h b/Engine/Source/Runtime/Engine/Classes/EdGraph/EdGraphPin.h
--- a/Engine/Source/Runtime/Engine/Classes/EdGraph/EdGraphPin.h
+++ b/Engine/Source/Runtime/Engine/Classes/EdGraph/EdGraphPin.h
@@ -108,14 +108,6 @@ struct FEdGraphPinType
 	EPinContainerType ContainerType;
 
 private:
-	/** DEPRECATED(4.17) Whether or not this pin represents a map of keys to values */
-	UPROPERTY()
-	uint8 bIsMap_DEPRECATED:1;
-
-	/** DEPRECATED(4.17) Whether or not this pin represents a set of (unique) values */
-	UPROPERTY()
-	uint8 bIsSet_DEPRECATED:1;
-
 	/** DEPRECATED(4.17) Whether or not this pin represents an array of values */
 	UPROPERTY()
 	uint8 bIsArray_DEPRECATED:1;
@@ -142,8 +134,6 @@ public:
 	FEdGraphPinType() 
 		: PinSubCategoryObject(nullptr)
 		, ContainerType(EPinContainerType::None)
-		, bIsMap_DEPRECATED(false)
-		, bIsSet_DEPRECATED(false)
 		, bIsArray_DEPRECATED(false)
 		, bIsReference(false)
 		, bIsConst(false)
@@ -158,8 +148,6 @@ public:
 		, PinSubCategoryObject(InPinSubCategoryObject)
 		, PinValueType(InValueTerminalType)
 		, ContainerType(ToPinContainerType(bInIsArray, bInIsSet, bInIsMap))
-		, bIsMap_DEPRECATED(false)
-		, bIsSet_DEPRECATED(false)
 		, bIsArray_DEPRECATED(false)
 		, bIsReference(bInIsReference)
 		, bIsConst(false)
@@ -173,8 +161,6 @@ public:
 		, PinSubCategoryObject(InPinSubCategoryObject)
 		, PinValueType(InValueTerminalType)
 		, ContainerType(InPinContainerType)
-		, bIsMap_DEPRECATED(false)
-		, bIsSet_DEPRECATED(false)
 		, bIsArray_DEPRECATED(false)
 		, bIsReference(bInIsReference)
 		, bIsConst(false)
@@ -212,6 +198,7 @@ public:
 	}
 
 	ENGINE_API bool Serialize(FArchive& Ar);
+	ENGINE_API void PostSerialize(const FArchive& Ar);
 
 	static ENGINE_API FEdGraphPinType GetPinTypeForTerminalType( const FEdGraphTerminalType& TerminalType );
 	static ENGINE_API FEdGraphPinType GetTerminalTypeForContainer( const FEdGraphPinType& ContainerType );
@@ -246,6 +233,7 @@ struct TStructOpsTypeTraits< FEdGraphPinType > : public TStructOpsTypeTraitsBase
 	enum 
 	{
 		WithSerializer = true,
+		WithPostSerialize = true
 	};
 };
 
diff --git a/Engine/Source/Runtime/Engine/Classes/Engine/BlueprintGeneratedClass.h b/Engine/Source/Runtime/Engine/Classes/Engine/BlueprintGeneratedClass.h
--- a/Engine/Source/Runtime/Engine/Classes/Engine/BlueprintGeneratedClass.h
+++ b/Engine/Source/Runtime/Engine/Classes/Engine/BlueprintGeneratedClass.h
@@ -376,10 +376,17 @@ public:
 		DebugNodeIndexLookup.MultiFind(Node, RecordIndices, true);
 		for(int i = 0; i < RecordIndices.Num(); ++i)
 		{
-			const FNodeToCodeAssociation& Record = DebugNodeLineNumbers[RecordIndices[i]];
-			if (UFunction* Scope = Record.Scope.Get())
+			int32 RecordIndex = RecordIndices[i];
+			if (DebugNodeLineNumbers.IsValidIndex(RecordIndex))
 			{
-				InstallSites.Add(&(Scope->Script[Record.Offset]));
+				const FNodeToCodeAssociation& Record = DebugNodeLineNumbers[RecordIndex];
+				if (UFunction* Scope = Record.Scope.Get())
+				{
+					if (Scope->Script.IsValidIndex(Record.Offset))
+					{
+						InstallSites.Add(&(Scope->Script[Record.Offset]));
+					}
+				}
 			}
 		}
 	}
@@ -616,6 +623,15 @@ public:
 	UPROPERTY(AssetRegistrySearchable)
 	int32	NumReplicatedProperties;
 
+	/** Flag used to indicate if this class has a nativized parent in a cooked build. */
+	UPROPERTY()
+	uint8 bHasNativizedParent:1;
+
+private:
+	/** Flag to make sure the custom property list has been initialized */
+	uint8 bCustomPropertyListForPostConstructionInitialized:1;
+
+public:
 	/** Array of objects containing information for dynamically binding delegates to functions in this blueprint */
 	UPROPERTY()
 	TArray<class UDynamicBlueprintBinding*> DynamicBindingObjects;
@@ -642,10 +658,12 @@ public:
 	UPROPERTY()
 	UFunction* UberGraphFunction;
 
+#if WITH_EDITORONLY_DATA
 	// This is a list of event graph call function nodes that are simple (no argument) thunks into the event graph (typically used for animation delegates, etc...)
 	// It is a deprecated list only used for backwards compatibility prior to VER_UE4_SERIALIZE_BLUEPRINT_EVENTGRAPH_FASTCALLS_IN_UFUNCTION.
 	UPROPERTY()
 	TArray<FEventGraphFastCallPair> FastCallPairs_DEPRECATED;
+#endif
 
 #if WITH_EDITORONLY_DATA
 	UPROPERTY(Transient)
@@ -661,10 +679,6 @@ public:
 	UPROPERTY()
 	TMap<FName, struct FBlueprintCookedComponentInstancingData> CookedComponentInstancingData;
 
-	/** Flag used to indicate if this class has a nativized parent in a cooked build. */
-	UPROPERTY()
-	bool bHasNativizedParent;
-
 	/** 
 	 * Gets an array of all BPGeneratedClasses (including InClass as 0th element) parents of given generated class 
 	 *
@@ -813,6 +827,4 @@ private:
 	TIndirectArray<FCustomPropertyListNode> CustomPropertyListForPostConstruction;
 	/** In some cases UObject::ConditionalPostLoad() code calls PostLoadDefaultObject() on a class that's still being serialized. */
 	FCriticalSection SerializeAndPostLoadCritical;
-	/** Flag to make sure the custom property list has been initialized */
-	bool bCustomPropertyListForPostConstructionInitialized;
 };
diff --git a/Engine/Source/Runtime/Engine/Classes/Engine/EngineTypes.h b/Engine/Source/Runtime/Engine/Classes/Engine/EngineTypes.h
--- a/Engine/Source/Runtime/Engine/Classes/Engine/EngineTypes.h
+++ b/Engine/Source/Runtime/Engine/Classes/Engine/EngineTypes.h
@@ -370,13 +370,13 @@ struct FLightingChannels
 
 	/** Default channel for all primitives and lights. */
 	UPROPERTY(EditAnywhere, BlueprintReadOnly, Category=Channels)
-	bool bChannel0;
+	uint8 bChannel0:1;
 
 	UPROPERTY(EditAnywhere, BlueprintReadOnly, Category=Channels)
-	bool bChannel1;
+	uint8 bChannel1:1;
 
 	UPROPERTY(EditAnywhere, BlueprintReadOnly, Category=Channels)
-	bool bChannel2;
+	uint8 bChannel2:1;
 };
 
 inline uint8 GetLightingChannelMaskForStruct(FLightingChannels Value)
diff --git a/Engine/Source/Runtime/Engine/Classes/Engine/Scene.h b/Engine/Source/Runtime/Engine/Classes/Engine/Scene.h
--- a/Engine/Source/Runtime/Engine/Classes/Engine/Scene.h
+++ b/Engine/Source/Runtime/Engine/Classes/Engine/Scene.h
@@ -517,11 +517,11 @@ struct FWeightedBlendable
 	GENERATED_USTRUCT_BODY()
 
 	/** 0:no effect .. 1:full effect */
-	UPROPERTY(interp, Category=FWeightedBlendable, meta=(ClampMin = "0.0", ClampMax = "1.0", Delta = "0.01"))
+	UPROPERTY(interp, BlueprintReadWrite, Category=FWeightedBlendable, meta=(ClampMin = "0.0", ClampMax = "1.0", Delta = "0.01"))
 	float Weight;
 
 	/** should be of the IBlendableInterface* type but UProperties cannot express that */
-	UPROPERTY(EditAnywhere, Category=FWeightedBlendable, meta=( AllowedClasses="BlendableInterface", Keywords="PostProcess" ))
+	UPROPERTY(EditAnywhere, BlueprintReadWrite, Category=FWeightedBlendable, meta=( AllowedClasses="BlendableInterface", Keywords="PostProcess" ))
 	UObject* Object;
 
 	// default constructor
@@ -541,12 +541,12 @@ struct FWeightedBlendable
 };
 
 // for easier detail customization, needed?
-USTRUCT()
+USTRUCT(BlueprintType)
 struct FWeightedBlendables
 {
 	GENERATED_USTRUCT_BODY()
 
-	UPROPERTY(EditAnywhere, Category="PostProcessSettings", meta=( Keywords="PostProcess" ))
+	UPROPERTY(EditAnywhere, BlueprintReadWrite, Category="PostProcessSettings", meta=( Keywords="PostProcess" ))
 	TArray<FWeightedBlendable> Array;
 };
 
@@ -1592,7 +1592,7 @@ struct FPostProcessSettings
 	 * Allows custom post process materials to be defined, using a MaterialInstance with the same Material as its parent to allow blending.
 	 * For materials this needs to be the "PostProcess" domain type. This can be used for any UObject object implementing the IBlendableInterface (e.g. could be used to fade weather settings).
 	 */
-	UPROPERTY(EditAnywhere, Category="Rendering Features", meta=( Keywords="PostProcess", DisplayName = "Post Process Materials" ))
+	UPROPERTY(EditAnywhere, BlueprintReadWrite, Category="Rendering Features", meta=( Keywords="PostProcess", DisplayName = "Post Process Materials" ))
 	FWeightedBlendables WeightedBlendables;
 
 	// for backwards compatibility
diff --git a/Engine/Source/Runtime/Engine/Classes/Engine/StreamableManager.h b/Engine/Source/Runtime/Engine/Classes/Engine/StreamableManager.h
--- a/Engine/Source/Runtime/Engine/Classes/Engine/StreamableManager.h
+++ b/Engine/Source/Runtime/Engine/Classes/Engine/StreamableManager.h
@@ -130,6 +130,9 @@ private:
 	/** Call to call the update delegate if bound, will propagate to parents */
 	void CallUpdateDelegate();
 
+	/** Unbind all of our delegates, called after a cancel/load */
+	void UnbindDelegates();
+
 	/** True if this request has finished loading. It may still be active, or it may have been released */
 	bool bLoadCompleted;
 
diff --git a/Engine/Source/Runtime/Engine/Classes/Engine/Texture2D.h b/Engine/Source/Runtime/Engine/Classes/Engine/Texture2D.h
--- a/Engine/Source/Runtime/Engine/Classes/Engine/Texture2D.h
+++ b/Engine/Source/Runtime/Engine/Classes/Engine/Texture2D.h
@@ -161,7 +161,6 @@ public:
 #if WITH_EDITOR
 	virtual void PostLinkerChange() override;
 	virtual void PostEditChangeProperty(FPropertyChangedEvent& PropertyChangedEvent) override;
-	virtual void WillNeverCacheCookedPlatformDataAgain() override;
 #endif // WITH_EDITOR
 	virtual void BeginDestroy() override;
 	virtual void PostLoad() override;
diff --git a/Engine/Source/Runtime/Engine/Classes/GameFramework/Actor.h b/Engine/Source/Runtime/Engine/Classes/GameFramework/Actor.h
--- a/Engine/Source/Runtime/Engine/Classes/GameFramework/Actor.h
+++ b/Engine/Source/Runtime/Engine/Classes/GameFramework/Actor.h
@@ -245,10 +245,6 @@
 	UFUNCTION()
 	virtual void OnRep_Owner();
 
-	/** Used to specify the net driver to replicate on (NAME_None || NAME_GameNetDriver is the default net driver) */
-	UPROPERTY()
-	FName NetDriverName;
-
 private:
 	/**
 	 * Describes how much control the remote machine has over the actor.
@@ -265,6 +261,10 @@
 
 public:
 
+	/** Used to specify the net driver to replicate on (NAME_None || NAME_GameNetDriver is the default net driver) */
+	UPROPERTY()
+	FName NetDriverName;
+
 	/**
 	 * Set whether this actor replicates to network clients. When this actor is spawned on the server it will be sent to clients as well.
 	 * Properties flagged for replication will update on clients if they change on the server.
@@ -296,6 +296,10 @@
 	UPROPERTY(EditDefaultsOnly, ReplicatedUsing=OnRep_ReplicatedMovement, Category=Replication, AdvancedDisplay)
 	struct FRepMovement ReplicatedMovement;
 
+	/** How long this Actor lives before dying, 0=forever. Note this is the INITIAL value and should not be modified once play has begun. */
+	UPROPERTY(EditAnywhere, BlueprintReadOnly, Category = Actor)
+	float InitialLifeSpan;
+
 private:
 	/**
 	 * Used for replicating attachment of this actor's RootComponent to another actor.
@@ -330,6 +334,11 @@
 	UPROPERTY(EditAnywhere, Category=Input)
 	TEnumAsByte<EAutoReceiveInput::Type> AutoReceiveInput;
 
+#if WITH_EDITORONLY_DATA
+	UPROPERTY()
+	TEnumAsByte<enum EInputConsumeOptions> InputConsumeOption_DEPRECATED;
+#endif
+
 	/** The priority of this input component when pushed in to the stack. */
 	UPROPERTY(EditAnywhere, Category=Input)
 	int32 InputPriority;
@@ -338,9 +347,6 @@
 	UPROPERTY()
 	class UInputComponent* InputComponent;
 
-	UPROPERTY()
-	TEnumAsByte<enum EInputConsumeOptions> InputConsumeOption_DEPRECATED;
-
 	/** Square of the max distance from the client's viewpoint that this actor is relevant and will be replicated. */
 	UPROPERTY(BlueprintReadOnly, EditDefaultsOnly, Category=Replication)
 	float NetCullDistanceSquared;   
@@ -443,6 +449,49 @@
 	UPROPERTY(Category = Actor, EditAnywhere, AdvancedDisplay)
 	uint8 bCanBeInCluster:1;
 
+protected:
+
+	/**
+	 * If false, the Blueprint ReceiveTick() event will be disabled on dedicated servers.
+	 * @see AllowReceiveTickEventOnDedicatedServer()
+	 */
+	UPROPERTY()
+	uint8 bAllowReceiveTickEventOnDedicatedServer:1;
+
+private:
+
+	/** 
+	 *	Indicates that PreInitializeComponents/PostInitializeComponents have been called on this Actor 
+	 *	Prevents re-initializing of actors spawned during level startup
+	 */
+	uint8 bActorInitialized:1;
+	
+	enum class EActorBeginPlayState : uint8
+	{
+		HasNotBegunPlay,
+		BeginningPlay,
+		HasBegunPlay,
+	};
+
+	/** 
+	 *	Indicates that BeginPlay has been called for this Actor.
+	 *  Set back to false once EndPlay has been called.
+	 */
+	EActorBeginPlayState ActorHasBegunPlay:2;
+
+public:
+	/** Indicates the actor was pulled through a seamless travel.  */
+	UPROPERTY()
+	uint8 bActorSeamlessTraveled:1;
+
+	/** Whether this actor should not be affected by world origin shifting. */
+	UPROPERTY(EditAnywhere, AdvancedDisplay, Category=Actor)
+	uint8 bIgnoresOriginShifting:1;
+	
+	/** If true, and if World setting has bEnableHierarchicalLOD equal to true, then it will generate LODActor from groups of clustered Actor */
+	UPROPERTY(EditAnywhere, AdvancedDisplay, Category = HLOD, meta = (DisplayName = "Include Actor for HLOD Mesh generation"))
+	uint8 bEnableAutoLODGeneration:1;
+
 public:
 
 	/** Controls how to handle spawning this actor in a situation where it's colliding with something else. "Default" means AlwaysSpawn here. */
@@ -483,10 +532,6 @@
 	UPROPERTY(transient)
 	TArray<class AMatineeActor*> ControllingMatineeActors;
 
-	/** How long this Actor lives before dying, 0=forever. Note this is the INITIAL value and should not be modified once play has begun. */
-	UPROPERTY(EditAnywhere, BlueprintReadOnly, Category=Actor)
-	float InitialLifeSpan;
-
 	/** Handle for efficient management of LifeSpanExpired timer */
 	FTimerHandle TimerHandle_LifeSpanExpired;
 
@@ -500,9 +545,11 @@
 	TArray< FName > Layers;
 
 private:
+#if WITH_EDITORONLY_DATA
 	/** The Actor that owns the UChildActorComponent that owns this Actor. */
 	UPROPERTY()
 	TWeakObjectPtr<AActor> ParentComponentActor_DEPRECATED;	
+#endif
 
 	/** The UChildActorComponent that owns this Actor. */
 	UPROPERTY()
@@ -582,49 +629,6 @@
 	int32 GetNumUncachedLights();
 #endif // WITH_EDITORONLY_DATA
 
-protected:
-
-	/**
-	 * If false, the Blueprint ReceiveTick() event will be disabled on dedicated servers.
-	 * @see AllowReceiveTickEventOnDedicatedServer()
-	 */
-	UPROPERTY()
-	uint8 bAllowReceiveTickEventOnDedicatedServer:1;
-
 private:
 
-	/** 
-	 *	Indicates that PreInitializeComponents/PostInitializeComponents have been called on this Actor 
-	 *	Prevents re-initializing of actors spawned during level startup
-	 */
-	uint8 bActorInitialized:1;
-	
-	enum class EActorBeginPlayState : uint8
-	{
-		HasNotBegunPlay,
-		BeginningPlay,
-		HasBegunPlay,
-	};
-
-	/** 
-	 *	Indicates that BeginPlay has been called for this Actor.
-	 *  Set back to false once EndPlay has been called.
-	 */
-	EActorBeginPlayState ActorHasBegunPlay:2;
-
-public:
-	/** Indicates the actor was pulled through a seamless travel.  */
-	UPROPERTY()
-	uint8 bActorSeamlessTraveled:1;
-
-	/** Whether this actor should not be affected by world origin shifting. */
-	UPROPERTY(EditAnywhere, AdvancedDisplay, Category=Actor)
-	uint8 bIgnoresOriginShifting:1;
-	
-	/** If true, and if World setting has bEnableHierarchicalLOD equal to true, then it will generate LODActor from groups of clustered Actor */
-	UPROPERTY(EditAnywhere, AdvancedDisplay, Category = HLOD, meta = (DisplayName = "Include Actor for HLOD Mesh generation"))
-	uint8 bEnableAutoLODGeneration:1;
-
-private:
-
 	static uint32 BeginPlayCallDepth;
@@ -2764,7 +2768,7 @@
 	 * It's recommended to use TArrays with a TInlineAllocator to potentially avoid memory allocation costs.
 	 * TInlineComponentArray is defined to make this easier, for example:
 	 * {
-	 * 	   TInlineComponentArray<UPrimitiveComponent*> PrimComponents;
+	 * 	   TInlineComponentArray<UActorComponent*> PrimComponents;
 	 *     Actor->GetComponents(PrimComponents);
 	 * }
 	 *
@@ -2957,6 +2961,9 @@
 	/** Private version without inlining that does *not* check Dedicated server build flags (which should already have been done). */
 	ENetMode InternalGetNetMode() const;
 
+	/** Unified implementation function to be called from the two implementations of PostEditUndo for the AActor specific elements that need to happen. */
+	bool InternalPostEditUndo();
+
 	friend struct FMarkActorIsBeingDestroyed;
 	friend struct FActorParentComponentSetter;
 
diff --git a/Engine/Source/Runtime/Engine/Classes/GameFramework/ForceFeedbackEffect.h b/Engine/Source/Runtime/Engine/Classes/GameFramework/ForceFeedbackEffect.h
--- a/Engine/Source/Runtime/Engine/Classes/GameFramework/ForceFeedbackEffect.h
+++ b/Engine/Source/Runtime/Engine/Classes/GameFramework/ForceFeedbackEffect.h
@@ -1,4 +1,4 @@
-// Copyright 1998-2017 Epic Games, Inc. All Rights Reserved.
+﻿// Copyright 1998-2017 Epic Games, Inc. All Rights Reserved.
 #pragma once
 
 #include "CoreMinimal.h"
@@ -49,29 +49,32 @@ struct ENGINE_API FActiveForceFeedbackEffect
 
 	FName Tag;
 	uint32 bLooping:1;
+	uint32 bIgnoreTimeDilation:1;
 	float PlayTime;
 
 	FActiveForceFeedbackEffect()
-		: ForceFeedbackEffect(NULL)
+		: ForceFeedbackEffect(nullptr)
 		, Tag(NAME_None)
 		, bLooping(false)
+		, bIgnoreTimeDilation(false)
 		, PlayTime(0.f)
 	{
 	}
 
-	FActiveForceFeedbackEffect(UForceFeedbackEffect* InEffect, const bool bInLooping, const FName InTag)
+	FActiveForceFeedbackEffect(UForceFeedbackEffect* InEffect, const bool bInLooping, const bool bInIgnoreTimeDilation, const FName InTag)
 		: ForceFeedbackEffect(InEffect)
 		, Tag(InTag)
 		, bLooping(bInLooping)
+		, bIgnoreTimeDilation(bInIgnoreTimeDilation)
 		, PlayTime(0.f)
 	{
 	}
 
 	// Updates the final force feedback values based on this effect.  Returns true if the effect should continue playing, false if it is finished.
-	bool Update(const float DeltaTime, struct FForceFeedbackValues& Values);
+	bool Update(float DeltaTime, FForceFeedbackValues& Values);
 
 	// Gets the current values at the stored play time
-	void GetValues(struct FForceFeedbackValues& Values) const;
+	void GetValues(FForceFeedbackValues& Values) const;
 };
 
 /**
diff --git a/Engine/Source/Runtime/Engine/Classes/GameFramework/HUD.h b/Engine/Source/Runtime/Engine/Classes/GameFramework/HUD.h
--- a/Engine/Source/Runtime/Engine/Classes/GameFramework/HUD.h
+++ b/Engine/Source/Runtime/Engine/Classes/GameFramework/HUD.h
@@ -404,7 +404,6 @@ public:
 	UFUNCTION(BlueprintCallable, Category=HUD)
 	void AddHitBox(FVector2D Position, FVector2D Size, FName InName, bool bConsumesInput, int32 Priority = 0);
 
-protected:
 	/** Returns the PlayerController for this HUD's player.	 */
 	UFUNCTION(BlueprintCallable, Category=HUD)
 	APlayerController* GetOwningPlayerController() const;
@@ -413,7 +412,6 @@ protected:
 	UFUNCTION(BlueprintCallable, Category=HUD)
 	APawn* GetOwningPawn() const;
 
-public:
 	/**
 	 * Draws a colored line between two points
 	 * @param Start - start of line
diff --git a/Engine/Source/Runtime/Engine/Classes/GameFramework/PlayerController.h b/Engine/Source/Runtime/Engine/Classes/GameFramework/PlayerController.h
--- a/Engine/Source/Runtime/Engine/Classes/GameFramework/PlayerController.h
+++ b/Engine/Source/Runtime/Engine/Classes/GameFramework/PlayerController.h
@@ -1,4 +1,4 @@
-// Copyright 1998-2017 Epic Games, Inc. All Rights Reserved.
+﻿// Copyright 1998-2017 Epic Games, Inc. All Rights Reserved.
 
 #pragma once
 
@@ -907,10 +907,17 @@ public:
 	 * Play a force feedback pattern on the player's controller
 	 * @param	ForceFeedbackEffect		The force feedback pattern to play
 	 * @param	bLooping				Whether the pattern should be played repeatedly or be a single one shot
+	 * @param	bIgnoreTimeDilation		Whether the pattern should ignore time dilation
 	 * @param	Tag						A tag that allows stopping of an effect.  If another effect with this Tag is playing, it will be stopped and replaced
 	 */
 	UFUNCTION(unreliable, client, BlueprintCallable, Category="Game|Feedback")
-	void ClientPlayForceFeedback(class UForceFeedbackEffect* ForceFeedbackEffect, bool bLooping, FName Tag);
+	void ClientPlayForceFeedback(class UForceFeedbackEffect* ForceFeedbackEffect, bool bLooping, bool bIgnoreTimeDilation, FName Tag);
+
+	DEPRECATED(4.18, "Use version that specifies whether to ignore time dilation or not")
+	void ClientPlayForceFeedback(class UForceFeedbackEffect* ForceFeedbackEffect, bool bLooping, FName Tag)
+	{
+		ClientPlayForceFeedback(ForceFeedbackEffect, bLooping, false, Tag);
+	}
 
 	/** 
 	 * Stops a playing force feedback pattern
diff --git a/Engine/Source/Runtime/Engine/Classes/GameFramework/SpringArmComponent.h b/Engine/Source/Runtime/Engine/Classes/GameFramework/SpringArmComponent.h
--- a/Engine/Source/Runtime/Engine/Classes/GameFramework/SpringArmComponent.h
+++ b/Engine/Source/Runtime/Engine/Classes/GameFramework/SpringArmComponent.h
@@ -47,7 +47,11 @@ class ENGINE_API USpringArmComponent : public USceneComponent
 	/**
 	 * If this component is placed on a pawn, should it use the view/control rotation of the pawn where possible?
 	 * When disabled, the component will revert to using the stored RelativeRotation of the component.
-	 * @see APawn::GetViewRotation()
+	 * Note that this component itself does not rotate, but instead maintains its relative rotation to its parent as normal,
+	 * and just repositions and rotates its children as desired by the inherited rotation settings. Use GetTargetRotation()
+	 * if you want the rotation target based on all the settings (UsePawnControlRotation, InheritPitch, etc).
+	 *
+	 * @see GetTargetRotation(), APawn::GetViewRotation()
 	 */
 	UPROPERTY(EditAnywhere, BlueprintReadWrite, Category=CameraSettings)
 	uint32 bUsePawnControlRotation:1;
@@ -108,6 +112,13 @@ class ENGINE_API USpringArmComponent : public USceneComponent
 	UPROPERTY(EditAnywhere, BlueprintReadWrite, Category=Lag, meta=(editcondition="bEnableCameraLag", ClampMin="0.0", UIMin = "0.0"))
 	float CameraLagMaxDistance;
 
+	/**
+	 * Get the target rotation we inherit, used as the base target for the boom rotation.
+	 * This is derived from attachment to our parent and considering the UsePawnControlRotation and absolute rotation flags.
+	 */
+	UFUNCTION(BlueprintCallable, Category=SpringArm)
+	FRotator GetTargetRotation() const;
+
 	/** Temporary variables when using camera lag, to record previous camera position */
 	FVector PreviousDesiredLoc;
 	FVector PreviousArmOrigin;
diff --git a/Engine/Source/Runtime/Engine/Classes/Kismet/GameplayStatics.h b/Engine/Source/Runtime/Engine/Classes/Kismet/GameplayStatics.h
--- a/Engine/Source/Runtime/Engine/Classes/Kismet/GameplayStatics.h
+++ b/Engine/Source/Runtime/Engine/Classes/Kismet/GameplayStatics.h
@@ -311,10 +311,14 @@ class ENGINE_API UGameplayStatics : public UBlueprintFunctionLibrary
 	 * @param EmitterTemplate - particle system to create
 	 * @param Location - location to place the effect in world space
 	 * @param Rotation - rotation to place the effect in world space	
+	 * @param Scale - scale to create the effect at
 	 * @param bAutoDestroy - Whether the component will automatically be destroyed when the particle system completes playing or whether it can be reactivated
 	 */
 	UFUNCTION(BlueprintCallable, Category="Effects|Components|ParticleSystem", meta=(Keywords = "particle system", WorldContext="WorldContextObject", UnsafeDuringActorConstruction = "true"))
-	static UParticleSystemComponent* SpawnEmitterAtLocation(const UObject* WorldContextObject, class UParticleSystem* EmitterTemplate, FVector Location, FRotator Rotation = FRotator::ZeroRotator, bool bAutoDestroy = true);
+	static UParticleSystemComponent* SpawnEmitterAtLocation(const UObject* WorldContextObject, UParticleSystem* EmitterTemplate, FVector Location, FRotator Rotation = FRotator::ZeroRotator, FVector Scale = FVector(1.f), bool bAutoDestroy = true);
+
+	// Backwards compatible version of SpawnEmitterAttached for C++ without Scale
+	static UParticleSystemComponent* SpawnEmitterAtLocation(const UObject* WorldContextObject, UParticleSystem* EmitterTemplate, FVector Location, FRotator Rotation, bool bAutoDestroy);
 
 	/** Plays the specified effect at the given location and rotation, fire and forget. The system will go away when the effect is complete. Does not replicate.
 	 * @param World - The World to spawn in
@@ -322,7 +326,12 @@ class ENGINE_API UGameplayStatics : public UBlueprintFunctionLibrary
 	 * @param SpawnTransform - transform with which to place the effect in world space
 	 * @param bAutoDestroy - Whether the component will automatically be destroyed when the particle system completes playing or whether it can be reactivated
 	 */
-	static UParticleSystemComponent* SpawnEmitterAtLocation(UWorld* World, class UParticleSystem* EmitterTemplate, const FTransform& SpawnTransform, bool bAutoDestroy = true);
+	static UParticleSystemComponent* SpawnEmitterAtLocation(UWorld* World, UParticleSystem* EmitterTemplate, const FTransform& SpawnTransform, bool bAutoDestroy = true);
+
+private:
+	static UParticleSystemComponent* InternalSpawnEmitterAtLocation(UWorld* World, UParticleSystem* EmitterTemplate, FVector Location, FRotator Rotation, FVector Scale, bool bAutoDestroy);
+
+public:
 
 	/** Plays the specified effect attached to and following the specified component. The system will go away when the effect is complete. Does not replicate.
 	* @param EmitterTemplate - particle system to create
@@ -330,11 +339,15 @@ class ENGINE_API UGameplayStatics : public UBlueprintFunctionLibrary
 	 * @param AttachPointName - Optional named point within the AttachComponent to spawn the emitter at
 	 * @param Location - Depending on the value of LocationType this is either a relative offset from the attach component/point or an absolute world location that will be translated to a relative offset (if LocationType is KeepWorldPosition).
 	 * @param Rotation - Depending on the value of LocationType this is either a relative offset from the attach component/point or an absolute world rotation that will be translated to a relative offset (if LocationType is KeepWorldPosition).
+	 * @param Scale - Depending on the value of LocationType this is either a relative scale from the attach component or an absolute world scale that will be translated to a relative scale (if LocationType is KeepWorldPosition).
 	 * @param LocationType - Specifies whether Location is a relative offset or an absolute world position
 	 * @param bAutoDestroy - Whether the component will automatically be destroyed when the particle system completes playing or whether it can be reactivated
 	 */
 	UFUNCTION(BlueprintCallable, Category="Effects|Components|ParticleSystem", meta=(Keywords = "particle system", UnsafeDuringActorConstruction = "true"))
-	static UParticleSystemComponent* SpawnEmitterAttached(class UParticleSystem* EmitterTemplate, class USceneComponent* AttachToComponent, FName AttachPointName = NAME_None, FVector Location = FVector(ForceInit), FRotator Rotation = FRotator::ZeroRotator, EAttachLocation::Type LocationType = EAttachLocation::KeepRelativeOffset, bool bAutoDestroy = true);
+	static UParticleSystemComponent* SpawnEmitterAttached(class UParticleSystem* EmitterTemplate, class USceneComponent* AttachToComponent, FName AttachPointName = NAME_None, FVector Location = FVector(ForceInit), FRotator Rotation = FRotator::ZeroRotator, FVector Scale = FVector(1.f), EAttachLocation::Type LocationType = EAttachLocation::KeepRelativeOffset, bool bAutoDestroy = true);
+
+	// Backwards compatible version of SpawnEmitterAttached for C++ without Scale
+	static UParticleSystemComponent* SpawnEmitterAttached(class UParticleSystem* EmitterTemplate, class USceneComponent* AttachToComponent, FName AttachPointName, FVector Location, FRotator Rotation, EAttachLocation::Type LocationType, bool bAutoDestroy = true);
 
 	// --- Sound functions ------------------------------
 	
diff --git a/Engine/Source/Runtime/Engine/Classes/Kismet/KismetArrayLibrary.h b/Engine/Source/Runtime/Engine/Classes/Kismet/KismetArrayLibrary.h
--- a/Engine/Source/Runtime/Engine/Classes/Kismet/KismetArrayLibrary.h
+++ b/Engine/Source/Runtime/Engine/Classes/Kismet/KismetArrayLibrary.h
@@ -1,4 +1,4 @@
-// Copyright 1998-2017 Epic Games, Inc. All Rights Reserved.
+﻿// Copyright 1998-2017 Epic Games, Inc. All Rights Reserved.
 
 #pragma once
 
@@ -142,6 +142,17 @@ class ENGINE_API UKismetArrayLibrary : public UBlueprintFunctionLibrary
 	UFUNCTION(BlueprintCallable, CustomThunk, meta=(DisplayName = "Set Array Elem", ArrayParm = "TargetArray", ArrayTypeDependentParams = "Item", AutoCreateRefTerm = "Item"), Category="Utilities|Array")
 	static void Array_Set(const TArray<int32>& TargetArray, int32 Index, const int32& Item, bool bSizeToFit);
 
+	/*
+	 *Swaps the elements at the specified positions in the specified array
+	 *If the specified positions are equal, invoking this method leaves the array unchanged
+	 *
+	 *@param	TargetArray		The array to perform the operation on
+	 *@param    FirstIndex      The index of one element to be swapped
+	 *@param    SecondIndex     The index of the other element to be swapped
+	*/
+	UFUNCTION(BlueprintCallable, CustomThunk, meta=(DisplayName = "Swap Array Elements", CompactNodeTitle = "SWAP", ArrayParm = "TargetArray"), Category="Utilities|Array")
+	static void Array_Swap(const TArray<int32>& TargetArray, int32 FirstIndex, int32 SecondIndex);
+
 	/*  
 	 *Finds the index of the first instance of the item within the array
 	 *	
@@ -202,6 +213,7 @@ class ENGINE_API UKismetArrayLibrary : public UBlueprintFunctionLibrary
 	static int32 GenericArray_LastIndex(const void* TargetArray, const UArrayProperty* ArrayProp);
 	static void GenericArray_Get(void* TargetArray, const UArrayProperty* ArrayProp, int32 Index, void* Item);
 	static void GenericArray_Set(void* TargetArray, const UArrayProperty* ArrayProp, int32 Index, const void* NewItem, bool bSizeToFit);
+	static void GenericArray_Swap(const void* TargetArray, const UArrayProperty* ArrayProp, int32 First, int32 Second);
 	static int32 GenericArray_Find(const void* TargetArray, const UArrayProperty* ArrayProperty, const void* ItemToFind);
 	static void GenericArray_SetArrayPropertyByName(UObject* OwnerObject, FName ArrayPropertyName, const void* SrcArrayAddr);
 	static bool GenericArray_IsValidIndex(const void* TargetArray, const UArrayProperty* ArrayProp, int32 IndexToTest);
@@ -535,6 +547,28 @@ public:
 		InnerProp->DestroyValue(StorageSpace);
 	}
 
+
+	DECLARE_FUNCTION(execArray_Swap)
+	{
+		Stack.MostRecentProperty = nullptr;
+		Stack.StepCompiledIn<UArrayProperty>(nullptr);
+		void* ArrayAddr = Stack.MostRecentPropertyAddress;
+		UArrayProperty* ArrayProperty = Cast<UArrayProperty>(Stack.MostRecentProperty);
+		if (!ArrayProperty)
+		{
+			Stack.bArrayContextFailed = true;
+			return;
+		}
+
+		P_GET_PROPERTY(UIntProperty, First);
+		P_GET_PROPERTY(UIntProperty, Second);
+
+		P_FINISH;
+		P_NATIVE_BEGIN;
+		GenericArray_Swap(ArrayAddr, ArrayProperty, First, Second);
+		P_NATIVE_END;
+	}
+
 	DECLARE_FUNCTION(execArray_Find)
 	{
 		Stack.MostRecentProperty = nullptr;
