Copying //UE4/Dev-Framework to //UE4/Dev-Main (Source: //UE4/Dev-Framework @ 3628051)

diff --git a/Engine/Source/Runtime/Engine/Classes/Kismet/KismetMathLibrary.h b/Engine/Source/Runtime/Engine/Classes/Kismet/KismetMathLibrary.h
--- a/Engine/Source/Runtime/Engine/Classes/Kismet/KismetMathLibrary.h
+++ b/Engine/Source/Runtime/Engine/Classes/Kismet/KismetMathLibrary.h
@@ -135,7 +135,7 @@ class ENGINE_API UKismetMathLibrary : public UBlueprintFunctionLibrary
 	// Boolean functions.
 	//
 	
-	/* Returns a uniformly distributed random bool*/
+	/** Returns a uniformly distributed random bool*/
 	UFUNCTION(BlueprintPure, Category="Math|Random", meta=(NotBlueprintThreadSafe))
 	static bool RandomBool();
 
@@ -153,35 +153,35 @@ class ENGINE_API UKismetMathLibrary : public UBlueprintFunctionLibrary
 	UFUNCTION(BlueprintPure, Category = "Math|Random", meta=(Weight = "0.5"))
 	static bool RandomBoolWithWeightFromStream(float Weight, const FRandomStream& RandomStream);
 
-	/* Returns the logical complement of the Boolean value (NOT A) */
+	/** Returns the logical complement of the Boolean value (NOT A) */
 	UFUNCTION(BlueprintPure, meta=(DisplayName = "NOT Boolean", CompactNodeTitle = "NOT", Keywords = "! not negate"), Category="Math|Boolean")
 	static bool Not_PreBool(bool A);
 
-	/* Returns true if the values are equal (A == B) */
+	/** Returns true if the values are equal (A == B) */
 	UFUNCTION(BlueprintPure, meta=(DisplayName = "Equal Boolean", CompactNodeTitle = "==", Keywords = "== equal"), Category="Math|Boolean")
 	static bool EqualEqual_BoolBool(bool A, bool B);
 
-	/* Returns true if the values are not equal (A != B) */
+	/** Returns true if the values are not equal (A != B) */
 	UFUNCTION(BlueprintPure, meta=(DisplayName = "NotEqual Boolean", CompactNodeTitle = "!=", Keywords = "!= not equal"), Category="Math|Boolean")
 	static bool NotEqual_BoolBool(bool A, bool B);
 
-	/* Returns the logical AND of two values (A AND B) */
+	/** Returns the logical AND of two values (A AND B) */
 	UFUNCTION(BlueprintPure, meta=(DisplayName = "AND Boolean", CompactNodeTitle = "AND", Keywords = "& and", CommutativeAssociativeBinaryOperator = "true"), Category="Math|Boolean")
 	static bool BooleanAND(bool A, bool B);
 
-	/* Returns the logical NAND of two values (A AND B) */
+	/** Returns the logical NAND of two values (A AND B) */
 	UFUNCTION(BlueprintPure, meta=(DisplayName = "NAND Boolean", CompactNodeTitle = "NAND", Keywords = "!& nand", CommutativeAssociativeBinaryOperator = "true"), Category="Math|Boolean")
 	static bool BooleanNAND(bool A, bool B);
 
-	/* Returns the logical OR of two values (A OR B) */
+	/** Returns the logical OR of two values (A OR B) */
 	UFUNCTION(BlueprintPure, meta=(DisplayName = "OR Boolean", CompactNodeTitle = "OR", Keywords = "| or", CommutativeAssociativeBinaryOperator = "true"), Category="Math|Boolean")
 	static bool BooleanOR(bool A, bool B);
 		
-	/* Returns the logical eXclusive OR of two values (A XOR B) */
+	/** Returns the logical eXclusive OR of two values (A XOR B) */
 	UFUNCTION(BlueprintPure, meta=(DisplayName = "XOR Boolean", CompactNodeTitle = "XOR", Keywords = "^ xor"), Category="Math|Boolean")
 	static bool BooleanXOR(bool A, bool B);
 
-	/* Returns the logical Not OR of two values (A NOR B) */
+	/** Returns the logical Not OR of two values (A NOR B) */
 	UFUNCTION(BlueprintPure, meta=(DisplayName = "NOR Boolean", CompactNodeTitle = "NOR", Keywords = "!^ nor"), Category="Math|Boolean")
 	static bool BooleanNOR(bool A, bool B);
 
@@ -189,55 +189,55 @@ class ENGINE_API UKismetMathLibrary : public UBlueprintFunctionLibrary
 	// Byte functions.
 	//
 
-	/* Multiplication (A * B) */
+	/** Multiplication (A * B) */
 	UFUNCTION(BlueprintPure, meta=(DisplayName = "Byte * Byte", CompactNodeTitle = "*", Keywords = "* multiply", CommutativeAssociativeBinaryOperator = "true"), Category="Math|Byte")
 	static uint8 Multiply_ByteByte(uint8 A, uint8 B);
 
-	/* Division (A / B) */
+	/** Division (A / B) */
 	UFUNCTION(BlueprintPure, meta=(DisplayName = "Byte / Byte", CompactNodeTitle = "/", Keywords = "/ divide division"), Category="Math|Byte")
 	static uint8 Divide_ByteByte(uint8 A, uint8 B = 1);
 
-	/* Modulo (A % B) */
+	/** Modulo (A % B) */
 	UFUNCTION(BlueprintPure, meta=(DisplayName = "% (Byte)", CompactNodeTitle = "%", Keywords = "% modulus"), Category="Math|Byte")
 	static uint8 Percent_ByteByte(uint8 A, uint8 B = 1);
 
-	/* Addition (A + B) */
+	/** Addition (A + B) */
 	UFUNCTION(BlueprintPure, meta=(DisplayName = "Byte + Byte", CompactNodeTitle = "+", Keywords = "+ add plus", CommutativeAssociativeBinaryOperator = "true"), Category="Math|Byte")
 	static uint8 Add_ByteByte(uint8 A, uint8 B = 1);
 
-	/* Subtraction (A - B) */
+	/** Subtraction (A - B) */
 	UFUNCTION(BlueprintPure, meta=(DisplayName = "Byte - Byte", CompactNodeTitle = "-", Keywords = "- subtract minus"), Category="Math|Byte")
 	static uint8 Subtract_ByteByte(uint8 A, uint8 B = 1);
 
-	/* Returns the minimum value of A and B */
+	/** Returns the minimum value of A and B */
 	UFUNCTION(BlueprintPure, meta = (DisplayName = "Min (Byte)", CompactNodeTitle = "MIN", CommutativeAssociativeBinaryOperator = "true"), Category = "Math|Byte")
 	static uint8 BMin(uint8 A, uint8 B);
 
-	/* Returns the maximum value of A and B */
+	/** Returns the maximum value of A and B */
 	UFUNCTION(BlueprintPure, meta = (DisplayName = "Max (Byte)", CompactNodeTitle = "MAX", CommutativeAssociativeBinaryOperator = "true"), Category = "Math|Byte")
 	static uint8 BMax(uint8 A, uint8 B);
 	
-	/* Returns true if A is less than B (A < B) */
+	/** Returns true if A is less than B (A < B) */
 	UFUNCTION(BlueprintPure, meta=(DisplayName = "Byte < Byte", CompactNodeTitle = "<", Keywords = "< less"), Category="Math|Byte")
 	static bool Less_ByteByte(uint8 A, uint8 B);
 
-	/* Returns true if A is greater than B (A > B) */
+	/** Returns true if A is greater than B (A > B) */
 	UFUNCTION(BlueprintPure, meta=(DisplayName = "Byte > Byte", CompactNodeTitle = ">", Keywords = "> greater"), Category="Math|Byte")
 	static bool Greater_ByteByte(uint8 A, uint8 B);
 
-	/* Returns true if A is less than or equal to B (A <= B) */
+	/** Returns true if A is less than or equal to B (A <= B) */
 	UFUNCTION(BlueprintPure, meta=(DisplayName = "Byte <= Byte", CompactNodeTitle = "<=", Keywords = "<= less"), Category="Math|Byte")
 	static bool LessEqual_ByteByte(uint8 A, uint8 B);
 
-	/* Returns true if A is greater than or equal to B (A >= B) */
+	/** Returns true if A is greater than or equal to B (A >= B) */
 	UFUNCTION(BlueprintPure, meta=(DisplayName = "Byte >= Byte", CompactNodeTitle = ">=", Keywords = ">= greater"), Category="Math|Byte")
 	static bool GreaterEqual_ByteByte(uint8 A, uint8 B);
 
-	/* Returns true if A is equal to B (A == B) */
+	/** Returns true if A is equal to B (A == B) */
 	UFUNCTION(BlueprintPure, meta=(DisplayName = "Equal (Byte)", CompactNodeTitle = "==", Keywords = "== equal"), Category="Math|Byte")
 	static bool EqualEqual_ByteByte(uint8 A, uint8 B);
 
-	/* Returns true if A is not equal to B (A != B) */
+	/** Returns true if A is not equal to B (A != B) */
 	UFUNCTION(BlueprintPure, meta=(DisplayName = "NotEqual (Byte)", CompactNodeTitle = "!=", Keywords = "!= not equal"), Category="Math|Byte")
 	static bool NotEqual_ByteByte(uint8 A, uint8 B);
 
@@ -245,71 +245,78 @@ class ENGINE_API UKismetMathLibrary : public UBlueprintFunctionLibrary
 	// Integer functions.
 	//
 
-	/* Multiplication (A * B) */
+	/** Multiplication (A * B) */
 	UFUNCTION(BlueprintPure, meta=(DisplayName = "integer * integer", CompactNodeTitle = "*", Keywords = "* multiply", CommutativeAssociativeBinaryOperator = "true"), Category="Math|Integer")
 	static int32 Multiply_IntInt(int32 A, int32 B);
 
-	/* Division (A / B) */
+	/** Division (A / B) */
 	UFUNCTION(BlueprintPure, meta=(DisplayName = "integer / integer", CompactNodeTitle = "/", Keywords = "/ divide division"), Category="Math|Integer")
 	static int32 Divide_IntInt(int32 A, int32 B = 1);
 
-	/* Modulo (A % B) */
+	/** Modulo (A % B) */
 	UFUNCTION(BlueprintPure, meta=(DisplayName = "% (integer)", CompactNodeTitle = "%", Keywords = "% modulus"), Category="Math|Integer")
 	static int32 Percent_IntInt(int32 A, int32 B = 1);
 
-	/* Addition (A + B) */
+	/** Addition (A + B) */
 	UFUNCTION(BlueprintPure, meta=(DisplayName = "integer + integer", CompactNodeTitle = "+", Keywords = "+ add plus", CommutativeAssociativeBinaryOperator = "true"), Category="Math|Integer")
 	static int32 Add_IntInt(int32 A, int32 B = 1);
 
-	/* Subtraction (A - B) */
+	/** Subtraction (A - B) */
 	UFUNCTION(BlueprintPure, meta=(DisplayName = "integer - integer", CompactNodeTitle = "-", Keywords = "- subtract minus"), Category="Math|Integer")
 	static int32 Subtract_IntInt(int32 A, int32 B = 1);
 
-	/* Returns true if A is less than B (A < B) */
+	/** Returns true if A is less than B (A < B) */
 	UFUNCTION(BlueprintPure, meta=(DisplayName = "integer < integer", CompactNodeTitle = "<", Keywords = "< less"), Category="Math|Integer")
 	static bool Less_IntInt(int32 A, int32 B);
 
-	/* Returns true if A is greater than B (A > B) */
+	/** Returns true if A is greater than B (A > B) */
 	UFUNCTION(BlueprintPure, meta=(DisplayName = "integer > integer", CompactNodeTitle = ">", Keywords = "> greater"), Category="Math|Integer")
 	static bool Greater_IntInt(int32 A, int32 B);
 
-	/* Returns true if A is less than or equal to B (A <= B) */
+	/** Returns true if A is less than or equal to B (A <= B) */
 	UFUNCTION(BlueprintPure, meta=(DisplayName = "integer <= integer", CompactNodeTitle = "<=", Keywords = "<= less"), Category="Math|Integer")
 	static bool LessEqual_IntInt(int32 A, int32 B);
 
-	/* Returns true if A is greater than or equal to B (A >= B) */
+	/** Returns true if A is greater than or equal to B (A >= B) */
 	UFUNCTION(BlueprintPure, meta=(DisplayName = "integer >= integer", CompactNodeTitle = ">=", Keywords = ">= greater"), Category="Math|Integer")
 	static bool GreaterEqual_IntInt(int32 A, int32 B);
 
-	/* Returns true if A is equal to B (A == B) */
+	/** Returns true if A is equal to B (A == B) */
 	UFUNCTION(BlueprintPure, meta=(DisplayName = "Equal (integer)", CompactNodeTitle = "==", Keywords = "== equal"), Category="Math|Integer")
 	static bool EqualEqual_IntInt(int32 A, int32 B);
 
-	/* Returns true if A is not equal to B (A != B) */
+	/** Returns true if A is not equal to B (A != B) */
 	UFUNCTION(BlueprintPure, meta=(DisplayName = "NotEqual (integer)", CompactNodeTitle = "!=", Keywords = "!= not equal"), Category="Math|Integer")
 	static bool NotEqual_IntInt(int32 A, int32 B);
 
-	/* Bitwise AND (A & B) */
+	/** Returns true if value is between Min and Max (V >= Min && V <= Max)
+	 * If InclusiveMin is true, value needs to be equal or larger than Min, else it needs to be larger
+	 * If InclusiveMax is true, value needs to be smaller or equal than Max, else it needs to be smaller
+	 */
+	UFUNCTION(BlueprintPure, meta = (DisplayName = "InRange (integer)", Min = "0", Max = "10"), Category = "Math|Integer")
+	static bool InRange_IntInt(int32 Value, int32 Min, int32 Max, bool InclusiveMin = true, bool InclusiveMax = true);
+
+	/** Bitwise AND (A & B) */
 	UFUNCTION(BlueprintPure, meta=(DisplayName = "Bitwise AND", CompactNodeTitle = "&", Keywords = "& and", CommutativeAssociativeBinaryOperator = "true"), Category="Math|Integer")
 	static int32 And_IntInt(int32 A, int32 B);
 
-	/* Bitwise XOR (A ^ B) */
+	/** Bitwise XOR (A ^ B) */
 	UFUNCTION(BlueprintPure, meta=(DisplayName = "Bitwise XOR", CompactNodeTitle = "^", Keywords = "^ xor"), Category="Math|Integer")
 	static int32 Xor_IntInt(int32 A, int32 B);
 
-	/* Bitwise OR (A | B) */
+	/** Bitwise OR (A | B) */
 	UFUNCTION(BlueprintPure, meta=(DisplayName = "Bitwise OR", CompactNodeTitle = "|", Keywords = "| or", CommutativeAssociativeBinaryOperator = "true"), Category="Math|Integer")
 	static int32 Or_IntInt(int32 A, int32 B);
 
-	/* Bitwise NOT (~A) */
+	/** Bitwise NOT (~A) */
 	UFUNCTION(BlueprintPure, meta = (DisplayName = "Bitwise NOT", CompactNodeTitle = "~", Keywords = "~ not"), Category = "Math|Integer")
 	static int32 Not_Int(int32 A);
 
-	/* Sign (integer, returns -1 if A < 0, 0 if A is zero, and +1 if A > 0) */
+	/** Sign (integer, returns -1 if A < 0, 0 if A is zero, and +1 if A > 0) */
 	UFUNCTION(BlueprintPure, meta=(DisplayName = "Sign (int)"), Category="Math|Integer")
 	static int32 SignOfInteger(int32 A);
 
-	/* Returns a uniformly distributed random number between 0 and Max - 1 */
+	/** Returns a uniformly distributed random number between 0 and Max - 1 */
 	UFUNCTION(BlueprintPure, Category="Math|Random", meta=(NotBlueprintThreadSafe))
 	static int32 RandomInteger(int32 Max);
 
@@ -317,45 +324,45 @@ class ENGINE_API UKismetMathLibrary : public UBlueprintFunctionLibrary
 	UFUNCTION(BlueprintPure, Category="Math|Random", meta = (NotBlueprintThreadSafe))
 	static int32 RandomIntegerInRange(int32 Min, int32 Max);
 
-	/* Returns the minimum value of A and B */
+	/** Returns the minimum value of A and B */
 	UFUNCTION(BlueprintPure, meta=(DisplayName = "Min (int)", CompactNodeTitle = "MIN", CommutativeAssociativeBinaryOperator = "true"), Category="Math|Integer")
 	static int32 Min(int32 A, int32 B);
 
-	/* Returns the maximum value of A and B */
+	/** Returns the maximum value of A and B */
 	UFUNCTION(BlueprintPure, meta=(DisplayName = "Max (int)", CompactNodeTitle = "MAX", CommutativeAssociativeBinaryOperator = "true"), Category="Math|Integer")
 	static int32 Max(int32 A, int32 B);
 
-	/* Returns Value clamped to be between A and B (inclusive) */
+	/** Returns Value clamped to be between A and B (inclusive) */
 	UFUNCTION(BlueprintPure, meta=(DisplayName = "Clamp (int)"), Category="Math|Integer")
 	static int32 Clamp(int32 Value, int32 Min, int32 Max);
 
-	/* Returns the absolute (positive) value of A */
+	/** Returns the absolute (positive) value of A */
 	UFUNCTION(BlueprintPure, meta=(DisplayName = "Absolute (int)", CompactNodeTitle = "ABS"), Category="Math|Integer")
 	static int32 Abs_Int(int32 A);
 
-
 	//
 	// Float functions.
 	//
-	/* Power (Base to the Exp-th power) */
+
+	/** Power (Base to the Exp-th power) */
 	UFUNCTION(BlueprintPure, meta=(DisplayName = "Power" ), Category="Math|Float")
 	static float MultiplyMultiply_FloatFloat(float Base, float Exp);
 
-	/* Multiplication (A * B) */
+	/** Multiplication (A * B) */
 	UFUNCTION(BlueprintPure, meta=(DisplayName = "float * float", CompactNodeTitle = "*", Keywords = "* multiply", CommutativeAssociativeBinaryOperator = "true"), Category="Math|Float")
 	static float Multiply_FloatFloat(float A, float B);
 
-	/* Multiplication (A * B) */
+	/** Multiplication (A * B) */
 	UFUNCTION(BlueprintPure, meta=(DisplayName = "int * float", CompactNodeTitle = "*", Keywords = "* multiply"), Category="Math|Float")
 	static float Multiply_IntFloat(int32 A, float B);
 
-	/* Division (A / B) */
+	/** Division (A / B) */
 	UFUNCTION(BlueprintPure, CustomThunk, meta=(DisplayName = "float / float", CompactNodeTitle = "/", Keywords = "/ divide division"), Category="Math|Float")
 	static float Divide_FloatFloat(float A, float B = 1.f);
 	
 	static float GenericDivide_FloatFloat(float A, float B);
 
-	/* Custom thunk to allow script stack trace in case of divide by zero */
+	/** Custom thunk to allow script stack trace in case of divide by zero */
 	DECLARE_FUNCTION(execDivide_FloatFloat)
 	{
 		P_GET_PROPERTY(UFloatProperty, A);
@@ -373,13 +380,13 @@ class ENGINE_API UKismetMathLibrary : public UBlueprintFunctionLibrary
 		*(float*)RESULT_PARAM = GenericDivide_FloatFloat(A, B);
 	}
 
-	/* Modulo (A % B) */
+	/** Modulo (A % B) */
 	UFUNCTION(BlueprintPure, CustomThunk, meta = (DisplayName = "% (float)", CompactNodeTitle = "%", Keywords = "% modulus"), Category = "Math|Float")
 	static float Percent_FloatFloat(float A, float B = 1.f);
 
 	static float GenericPercent_FloatFloat(float A, float B);
 
-	/* Custom thunk to allow script stack trace in case of modulo by zero */
+	/** Custom thunk to allow script stack trace in case of modulo by zero */
 	DECLARE_FUNCTION(execPercent_FloatFloat)
 	{
 		P_GET_PROPERTY(UFloatProperty, A);
@@ -401,50 +408,50 @@ class ENGINE_API UKismetMathLibrary : public UBlueprintFunctionLibrary
 	UFUNCTION(BlueprintPure, Category="Math|Float")
 	static float Fraction(float A);
 
-	/* Addition (A + B) */
+	/** Addition (A + B) */
 	UFUNCTION(BlueprintPure, meta=(DisplayName = "float + float", CompactNodeTitle = "+", Keywords = "+ add plus", CommutativeAssociativeBinaryOperator = "true"), Category="Math|Float")
 	static float Add_FloatFloat(float A, float B = 1.f);
 
-	/* Subtraction (A - B) */
+	/** Subtraction (A - B) */
 	UFUNCTION(BlueprintPure, meta=(DisplayName = "float - float", CompactNodeTitle = "-", Keywords = "- subtract minus"), Category="Math|Float")
 	static float Subtract_FloatFloat(float A, float B = 1.f);
 
-	/*Returns true if A is Less than B (A < B) */
+	/** Returns true if A is Less than B (A < B) */
 	UFUNCTION(BlueprintPure, meta=(DisplayName = "float < float", CompactNodeTitle = "<", Keywords = "< less"), Category="Math|Float")
 	static bool Less_FloatFloat(float A, float B);
 
-	/*Returns true if A is greater than B (A > B) */
+	/** Returns true if A is greater than B (A > B) */
 	UFUNCTION(BlueprintPure, meta=(DisplayName = "float > float", CompactNodeTitle = ">", Keywords = "> greater"), Category="Math|Float")
 	static bool Greater_FloatFloat(float A, float B);
 
-	/* Returns true if A is Less than or equal to B (A <= B) */
+	/** Returns true if A is Less than or equal to B (A <= B) */
 	UFUNCTION(BlueprintPure, meta=(DisplayName = "float <= float", CompactNodeTitle = "<=", Keywords = "<= less"), Category="Math|Float")
 	static bool LessEqual_FloatFloat(float A, float B);
 
-	/* Returns true if A is greater than or equal to B (A >= B) */
+	/** Returns true if A is greater than or equal to B (A >= B) */
 	UFUNCTION(BlueprintPure, meta=(DisplayName = "float >= float", CompactNodeTitle = ">=", Keywords = ">= greater"), Category="Math|Float")
 	static bool GreaterEqual_FloatFloat(float A, float B);
 
-	/* Returns true if A is exactly equal to B (A == B)*/
+	/** Returns true if A is exactly equal to B (A == B)*/
 	UFUNCTION(BlueprintPure, meta=(DisplayName = "Equal (float)", CompactNodeTitle = "==", Keywords = "== equal"), Category="Math|Float")
 	static bool EqualEqual_FloatFloat(float A, float B);
 
-	/* Returns true if A is nearly equal to B (|A - B| < ErrorTolerance) */
+	/** Returns true if A is nearly equal to B (|A - B| < ErrorTolerance) */
 	UFUNCTION(BlueprintPure, meta=(DisplayName = "Nearly Equal (float)", Keywords = "== equal"), Category="Math|Float")
 	static bool NearlyEqual_FloatFloat(float A, float B, float ErrorTolerance = 1.e-6f);
 
-	/* Returns true if A does not equal B (A != B)*/
+	/** Returns true if A does not equal B (A != B)*/
 	UFUNCTION(BlueprintPure, meta=(DisplayName = "NotEqual (float)", CompactNodeTitle = "!=", Keywords = "!= not equal"), Category="Math|Float")
 	static bool NotEqual_FloatFloat(float A, float B);
 
-	/* Returns true if value is between Min and Max (V >= Min && V <= Max)
+	/** Returns true if value is between Min and Max (V >= Min && V <= Max)
 	 * If InclusiveMin is true, value needs to be equal or larger than Min, else it needs to be larger
 	 * If InclusiveMax is true, value needs to be smaller or equal than Max, else it needs to be smaller
 	 */
 	UFUNCTION(BlueprintPure, meta=(DisplayName = "InRange (float)", Min="0.0", Max="1.0"), Category="Math|Float")
 	static bool InRange_FloatFloat(float Value, float Min, float Max, bool InclusiveMin = true, bool InclusiveMax = true);
 
-	/* Returns the hypotenuse of a right-angled triangle given the width and height. */
+	/** Returns the hypotenuse of a right-angled triangle given the width and height. */
 	UFUNCTION(BlueprintPure, meta=(Keywords = "pythagorean theorem"), Category = "Math|Float")
 	static float Hypotenuse(float Width, float Height);
 	
@@ -456,55 +463,55 @@ class ENGINE_API UKismetMathLibrary : public UBlueprintFunctionLibrary
 	UFUNCTION(BlueprintPure, meta = (DisplayName = "Snap to grid (float)"), Category = "Math|Float")
 	static float GridSnap_Float(float Location, float GridSize);
 
-	/* Returns the absolute (positive) value of A */
+	/** Returns the absolute (positive) value of A */
 	UFUNCTION(BlueprintPure, meta=(DisplayName = "Absolute (float)", CompactNodeTitle = "ABS"), Category="Math|Float")
 	static float Abs(float A);
 
-	/* Returns the sine of A (expects Radians)*/
+	/** Returns the sine of A (expects Radians)*/
 	UFUNCTION(BlueprintPure, meta=(DisplayName = "Sin (Radians)", CompactNodeTitle = "SIN", Keywords = "sine"), Category="Math|Trig")
 	static float Sin(float A);
 
-	/* Returns the inverse sine (arcsin) of A (result is in Radians) */
+	/** Returns the inverse sine (arcsin) of A (result is in Radians) */
 	UFUNCTION(BlueprintPure, meta=(DisplayName = "Asin (Radians)", CompactNodeTitle = "ASIN", Keywords = "sine"), Category="Math|Trig")
 	static float Asin(float A);
 
-	/* Returns the cosine of A (expects Radians)*/
+	/** Returns the cosine of A (expects Radians)*/
 	UFUNCTION(BlueprintPure, meta=(DisplayName = "Cos (Radians)", CompactNodeTitle = "COS"), Category="Math|Trig")
 	static float Cos(float A);
 
-	/* Returns the inverse cosine (arccos) of A (result is in Radians) */
+	/** Returns the inverse cosine (arccos) of A (result is in Radians) */
 	UFUNCTION(BlueprintPure, meta=(DisplayName = "Acos (Radians)", CompactNodeTitle = "ACOS"), Category="Math|Trig")
 	static float Acos(float A);
 
-	/* Returns the tan of A (expects Radians)*/
+	/** Returns the tan of A (expects Radians)*/
 	UFUNCTION(BlueprintPure, meta=(DisplayName = "Tan (Radians)", CompactNodeTitle = "TAN"), Category="Math|Trig")
 	static float Tan(float A);
 
-	/* Returns the inverse tan (atan) (result is in Radians)*/
+	/** Returns the inverse tan (atan) (result is in Radians)*/
 	UFUNCTION(BlueprintPure, meta=(DisplayName = "Atan (Radians)"), Category="Math|Trig")
 	static float Atan(float A);
 
-	/* Returns the inverse tan (atan2) of A/B (result is in Radians)*/
+	/** Returns the inverse tan (atan2) of A/B (result is in Radians)*/
 	UFUNCTION(BlueprintPure, meta=(DisplayName = "Atan2 (Radians)"), Category="Math|Trig")
 	static float Atan2(float A, float B);
 
-	/* Returns exponential(e) to the power A (e^A)*/
+	/** Returns exponential(e) to the power A (e^A)*/
 	UFUNCTION(BlueprintPure, Category="Math|Float", meta=(CompactNodeTitle = "e"))
 	static float Exp(float A);
 
-	/* Returns log of A base B (if B^R == A, returns R)*/
+	/** Returns log of A base B (if B^R == A, returns R)*/
 	UFUNCTION(BlueprintPure, Category = "Math|Float")
 	static float Log(float A, float Base = 1.f);
 
-	/* Returns natural log of A (if e^R == A, returns R)*/
+	/** Returns natural log of A (if e^R == A, returns R)*/
 	UFUNCTION(BlueprintPure, Category="Math|Float")
 	static float Loge(float A);
 
-	/* Returns square root of A*/
+	/** Returns square root of A*/
 	UFUNCTION(BlueprintPure, Category="Math|Float", meta=(Keywords = "square root", CompactNodeTitle = "SQRT"))
 	static float Sqrt(float A);
 
-	/* Returns square of A (A*A)*/
+	/** Returns square of A (A*A)*/
 	UFUNCTION(BlueprintPure, Category="Math|Float", meta=(CompactNodeTitle = "^2"))
 	static float Square(float A);
 
@@ -516,47 +523,47 @@ class ENGINE_API UKismetMathLibrary : public UBlueprintFunctionLibrary
 	UFUNCTION(BlueprintPure, Category="Math|Random", meta=(NotBlueprintThreadSafe))
 	static float RandomFloatInRange(float Min, float Max);
 
-	/* Returns the value of PI */
+	/** Returns the value of PI */
 	UFUNCTION(BlueprintPure, meta=(DisplayName = "Get PI", CompactNodeTitle = "PI"), Category="Math|Trig")
 	static float GetPI();
 
-	/* Returns the value of TAU (= 2 * PI) */
+	/** Returns the value of TAU (= 2 * PI) */
 	UFUNCTION(BlueprintPure, meta=(DisplayName = "Get TAU", CompactNodeTitle = "TAU"), Category="Math|Trig")
 	static float GetTAU();
 
-	/* Returns radians value based on the input degrees */
+	/** Returns radians value based on the input degrees */
 	UFUNCTION(BlueprintPure, meta=(DisplayName = "Degrees To Radians", CompactNodeTitle = "D2R"), Category="Math|Trig")
 	static float DegreesToRadians(float A);
 
-	/* Returns degrees value based on the input radians */
+	/** Returns degrees value based on the input radians */
 	UFUNCTION(BlueprintPure, meta=(DisplayName = "Radians To Degrees", CompactNodeTitle = "R2D"), Category="Math|Trig")
 	static float RadiansToDegrees(float A);
 
-	/* Returns the sin of A (expects Degrees)*/
+	/** Returns the sin of A (expects Degrees)*/
 	UFUNCTION(BlueprintPure, meta=(DisplayName = "Sin (Degrees)", CompactNodeTitle = "SINd", Keywords = "sine"), Category="Math|Trig")
 	static float DegSin(float A);
 
-	/* Returns the inverse sin (arcsin) of A (result is in Degrees) */
+	/** Returns the inverse sin (arcsin) of A (result is in Degrees) */
 	UFUNCTION(BlueprintPure, meta=(DisplayName = "Asin (Degrees)", CompactNodeTitle = "ASINd", Keywords = "sine"), Category="Math|Trig")
 	static float DegAsin(float A);
 
-	/* Returns the cos of A (expects Degrees)*/
+	/** Returns the cos of A (expects Degrees)*/
 	UFUNCTION(BlueprintPure, meta=(DisplayName = "Cos (Degrees)", CompactNodeTitle = "COSd"), Category="Math|Trig")
 	static float DegCos(float A);
 
-	/* Returns the inverse cos (arccos) of A (result is in Degrees) */
+	/** Returns the inverse cos (arccos) of A (result is in Degrees) */
 	UFUNCTION(BlueprintPure, meta=(DisplayName = "Acos (Degrees)", CompactNodeTitle = "ACOSd"), Category="Math|Trig")
 	static float DegAcos(float A);
 
-	/* Returns the tan of A (expects Degrees)*/
+	/** Returns the tan of A (expects Degrees)*/
 	UFUNCTION(BlueprintPure, meta=(DisplayName = "Tan (Degrees)", CompactNodeTitle = "TANd"), Category="Math|Trig")
 	static float DegTan(float A);
 
-	/* Returns the inverse tan (atan) (result is in Degrees)*/
+	/** Returns the inverse tan (atan) (result is in Degrees)*/
 	UFUNCTION(BlueprintPure, meta=(DisplayName = "Atan (Degrees)"), Category="Math|Trig")
 	static float DegAtan(float A);
 
-	/* Returns the inverse tan (atan2) of A/B (result is in Degrees)*/
+	/** Returns the inverse tan (atan2) of A/B (result is in Degrees)*/
 	UFUNCTION(BlueprintPure, meta=(DisplayName = "Atan2 (Degrees)"), Category="Math|Trig")
 	static float DegAtan2(float A, float B);
 
@@ -570,15 +577,15 @@ class ENGINE_API UKismetMathLibrary : public UBlueprintFunctionLibrary
 	UFUNCTION(BlueprintPure, meta=(DisplayName = "Clamp Angle"), Category="Math|Float")
 	static float ClampAngle(float AngleDegrees, float MinAngleDegrees, float MaxAngleDegrees);
 
-	/* Returns the minimum value of A and B */
+	/** Returns the minimum value of A and B */
 	UFUNCTION(BlueprintPure, meta=(DisplayName = "Min (float)", CompactNodeTitle = "MIN", CommutativeAssociativeBinaryOperator = "true"), Category="Math|Float")
 	static float FMin(float A, float B);
 
-	/* Returns the maximum value of A and B */
+	/** Returns the maximum value of A and B */
 	UFUNCTION(BlueprintPure, meta=(DisplayName = "Max (float)", CompactNodeTitle = "MAX", CommutativeAssociativeBinaryOperator = "true"), Category="Math|Float")
 	static float FMax(float A, float B);
 
-	/* Returns Value clamped between A and B (inclusive) */
+	/** Returns Value clamped between A and B (inclusive) */
 	UFUNCTION(BlueprintPure, meta=(DisplayName = "Clamp (float)", Min="0.0", Max="1.0"), Category="Math|Float")
 	static float FClamp(float Value, float Min, float Max);
 
@@ -606,7 +613,7 @@ class ENGINE_API UKismetMathLibrary : public UBlueprintFunctionLibrary
 	UFUNCTION(BlueprintPure, Category="Math|Byte")
 	static void MinOfByteArray(const TArray<uint8>& ByteArray, int32& IndexOfMinValue, uint8& MinValue);
 
-	/* Linearly interpolates between A and B based on Alpha (100% of A when Alpha=0 and 100% of B when Alpha=1) */
+	/** Linearly interpolates between A and B based on Alpha (100% of A when Alpha=0 and 100% of B when Alpha=1) */
 	UFUNCTION(BlueprintPure, Category="Math|Float")
 	static float Lerp(float A, float B, float Alpha);
 	
@@ -623,35 +630,35 @@ class ENGINE_API UKismetMathLibrary : public UBlueprintFunctionLibrary
 	UFUNCTION(BlueprintPure, meta = (Keywords = "percentage normalize range"), Category = "Math|Float")
 	static float InverseLerp(float A, float B, float Value);
 
-	/* Easeing  between A and B using a specified easing function */
+	/** Easeing  between A and B using a specified easing function */
 	UFUNCTION(BlueprintPure, meta = (DisplayName = "Ease", BlueprintInternalUseOnly = "true"), Category = "Math|Interpolation")
 	static float Ease(float A, float B, float Alpha, TEnumAsByte<EEasingFunc::Type> EasingFunc, float BlendExp = 2, int32 Steps = 2);
 
-	/* Rounds A to the nearest integer */
+	/** Rounds A to the nearest integer */
 	UFUNCTION(BlueprintPure, Category="Math|Float")
 	static int32 Round(float A);
 
-	/* Rounds A to the largest previous integer */
+	/** Rounds A to the largest previous integer */
 	UFUNCTION(BlueprintPure, meta=(DisplayName = "Floor"), Category="Math|Float")
 	static int32 FFloor(float A);
 	
-	/* Rounds A to an integer with truncation towards zero.  (e.g. -1.7 truncated to -1, 2.8 truncated to 2) */
+	/** Rounds A to an integer with truncation towards zero.  (e.g. -1.7 truncated to -1, 2.8 truncated to 2) */
 	UFUNCTION(BlueprintPure, meta=(DisplayName = "Truncate", BlueprintAutocast), Category="Math|Float")
 	static int32 FTrunc(float A);
 	
-	/* Rounds A to an integer with truncation towards zero for each element in a vector.  (e.g. -1.7 truncated to -1, 2.8 truncated to 2) */
+	/** Rounds A to an integer with truncation towards zero for each element in a vector.  (e.g. -1.7 truncated to -1, 2.8 truncated to 2) */
 	UFUNCTION(BlueprintPure, meta = (DisplayName = "Truncate Vector", BlueprintAutocast), Category = "Math|Float")
 	static FIntVector FTruncVector(const FVector& InVector);
 
-	/* Rounds A to the smallest following integer */
+	/** Rounds A to the smallest following integer */
 	UFUNCTION(BlueprintPure, Category="Math|Float")
 	static int32 FCeil(float A);
 
-	/* Returns the number of times Divisor will go into Dividend (i.e., Dividend divided by Divisor), as well as the remainder */
+	/** Returns the number of times Divisor will go into Dividend (i.e., Dividend divided by Divisor), as well as the remainder */
 	UFUNCTION(BlueprintPure, meta=(DisplayName = "Division (whole and remainder)"), Category="Math|Float")
 	static int32 FMod(float Dividend, float Divisor, float& Remainder);
 
-	/* Sign (float, returns -1 if A < 0, 0 if A is zero, and +1 if A > 0) */
+	/** Sign (float, returns -1 if A < 0, 0 if A is zero, and +1 if A > 0) */
 	UFUNCTION(BlueprintPure, meta=(DisplayName = "Sign (float)"), Category="Math|Float")
 	static float SignOfFloat(float A);
 
@@ -703,31 +710,31 @@ class ENGINE_API UKismetMathLibrary : public UBlueprintFunctionLibrary
 	// Vector functions.
 	//
 
-	/* Scales Vector A by B */
+	/** Scales Vector A by B */
 	UFUNCTION(BlueprintPure, meta=(DisplayName = "vector * float", CompactNodeTitle = "*", Keywords = "* multiply"), Category="Math|Vector")
 	static FVector Multiply_VectorFloat(FVector A, float B);
 	
-	/* Scales Vector A by B */
+	/** Scales Vector A by B */
 	UFUNCTION(BlueprintPure, meta=(DisplayName = "vector * int", CompactNodeTitle = "*", Keywords = "* multiply"), Category="Math|Vector")
 	static FVector Multiply_VectorInt(FVector A, int32 B);
 
-	/* Element-wise Vector multiplication (Result = {A.x*B.x, A.y*B.y, A.z*B.z}) */
+	/** Element-wise Vector multiplication (Result = {A.x*B.x, A.y*B.y, A.z*B.z}) */
 	UFUNCTION(BlueprintPure, meta=(DisplayName = "vector * vector", CompactNodeTitle = "*", Keywords = "* multiply", CommutativeAssociativeBinaryOperator = "true"), Category="Math|Vector")
 	static FVector Multiply_VectorVector(FVector A, FVector B);
 
-	/* Vector divide by a float */
+	/** Vector divide by a float */
 	UFUNCTION(BlueprintPure, meta=(DisplayName = "vector / float", CompactNodeTitle = "/", Keywords = "/ divide division"), Category="Math|Vector")
 	static FVector Divide_VectorFloat(FVector A, float B = 1.f);
 
-	/* Vector divide by an integer */
+	/** Vector divide by an integer */
 	UFUNCTION(BlueprintPure, meta=(DisplayName = "vector / int", CompactNodeTitle = "/", Keywords = "/ divide division"), Category="Math|Vector")
 	static FVector Divide_VectorInt(FVector A, int32 B = 1);
 	
-	/* Element-wise Vector division (Result = {A.x/B.x, A.y/B.y, A.z/B.z}) */
+	/** Element-wise Vector division (Result = {A.x/B.x, A.y/B.y, A.z/B.z}) */
 	UFUNCTION(BlueprintPure, meta=(DisplayName = "vector / vector", CompactNodeTitle = "/", Keywords = "/ divide division"), Category="Math|Vector")
 	static FVector Divide_VectorVector(FVector A, FVector B = FVector(1.f,1.f,1.f));
 
-	/* Vector addition */
+	/** Vector addition */
 	UFUNCTION(BlueprintPure, meta=(DisplayName = "vector + vector", CompactNodeTitle = "+", Keywords = "+ add plus", CommutativeAssociativeBinaryOperator = "true"), Category="Math|Vector")
 	static FVector Add_VectorVector(FVector A, FVector B);
 
@@ -739,7 +746,7 @@ class ENGINE_API UKismetMathLibrary : public UBlueprintFunctionLibrary
 	UFUNCTION(BlueprintPure, meta=(DisplayName = "vector + int", CompactNodeTitle = "+", Keywords = "+ add plus"), Category="Math|Vector")
 	static FVector Add_VectorInt(FVector A, int32 B);
 
-	/* Vector subtraction */
+	/** Vector subtraction */
 	UFUNCTION(BlueprintPure, meta=(DisplayName = "vector - vector", CompactNodeTitle = "-", Keywords = "- subtract minus"), Category="Math|Vector")
 	static FVector Subtract_VectorVector(FVector A, FVector B);
 
@@ -751,43 +758,43 @@ class ENGINE_API UKismetMathLibrary : public UBlueprintFunctionLibrary
 	UFUNCTION(BlueprintPure, meta=(DisplayName = "vector - int", CompactNodeTitle = "-", Keywords = "- subtract minus"), Category="Math|Vector")
 	static FVector Subtract_VectorInt(FVector A, int32 B);
 
-	/* Returns result of vector A rotated by the inverse of Rotator B */
+	/** Returns result of vector A rotated by the inverse of Rotator B */
 	UFUNCTION(BlueprintPure, meta=(DisplayName = "UnrotateVector"), Category="Math|Vector")
 	static FVector LessLess_VectorRotator(FVector A, FRotator B);
 
-	/* Returns result of vector A rotated by Rotator B */
+	/** Returns result of vector A rotated by Rotator B */
 	UFUNCTION(BlueprintPure, meta=(DisplayName = "RotateVector"), Category="Math|Vector")
 	static FVector GreaterGreater_VectorRotator(FVector A, FRotator B);
 
-	/* Returns result of vector A rotated by AngleDeg around Axis */
+	/** Returns result of vector A rotated by AngleDeg around Axis */
 	UFUNCTION(BlueprintPure, meta=(DisplayName = "RotateVectorAroundAxis"), Category="Math|Vector")
 	static FVector RotateAngleAxis(FVector InVect, float AngleDeg, FVector Axis);
 
-	/* Returns true if vector A is equal to vector B (A == B) within a specified error tolerance */
+	/** Returns true if vector A is equal to vector B (A == B) within a specified error tolerance */
 	UFUNCTION(BlueprintPure, meta=(DisplayName = "Equal (vector)", CompactNodeTitle = "==", Keywords = "== equal"), Category="Math|Vector")
 	static bool EqualEqual_VectorVector(FVector A, FVector B, float ErrorTolerance = 1.e-4f);
 
-	/* Returns true if vector A is not equal to vector B (A != B) within a specified error tolerance */
+	/** Returns true if vector A is not equal to vector B (A != B) within a specified error tolerance */
 	UFUNCTION(BlueprintPure, meta=(DisplayName = "Not Equal (vector)", CompactNodeTitle = "!=", Keywords = "!= not equal"), Category="Math|Vector")
 	static bool NotEqual_VectorVector(FVector A, FVector B, float ErrorTolerance = 1.e-4f);
 
-	/* Returns the dot product of two 3d vectors */
+	/** Returns the dot product of two 3d vectors */
 	UFUNCTION(BlueprintPure, meta=(DisplayName = "Dot Product", CompactNodeTitle = "dot"), Category="Math|Vector" )
 	static float Dot_VectorVector(FVector A, FVector B);
 
-	/* Returns the cross product of two 3d vectors */
+	/** Returns the cross product of two 3d vectors */
 	UFUNCTION(BlueprintPure, meta=(DisplayName = "Cross Product", CompactNodeTitle = "cross"), Category="Math|Vector" )
 	static FVector Cross_VectorVector(FVector A, FVector B);
 
-	/* Returns the dot product of two 2d vectors */
+	/** Returns the dot product of two 2d vectors */
 	UFUNCTION(BlueprintPure, meta=(DisplayName="Dot Product (2D)", CompactNodeTitle="dot"), Category="Math|Vector")
 	static float DotProduct2D(FVector2D A, FVector2D B);
 
-	/* Returns the cross product of two 2d vectors */
+	/** Returns the cross product of two 2d vectors */
 	UFUNCTION(BlueprintPure, meta=(DisplayName="Cross Product (2D)", CompactNodeTitle="cross"), Category="Math|Vector")
 	static float CrossProduct2D(FVector2D A, FVector2D B);
 
-	/* Returns the length of the FVector */
+	/** Returns the length of the FVector */
 	UFUNCTION(BlueprintPure, meta=(DisplayName = "VectorLength", Keywords="magnitude"), Category="Math|Vector")
 	static float VSize(FVector A);
 
@@ -795,7 +802,7 @@ class ENGINE_API UKismetMathLibrary : public UBlueprintFunctionLibrary
 	UFUNCTION(BlueprintPure, meta=(DisplayName = "Vector2dLength", Keywords="magnitude"), Category="Math|Vector2D")
 	static float VSize2D(FVector2D A);
 
-	/* Returns the squared length of the FVector */
+	/** Returns the squared length of the FVector */
 	UFUNCTION(BlueprintPure, meta=(DisplayName = "VectorLengthSquared", Keywords="magnitude"), Category="Math|Vector")
 	static float VSizeSquared(FVector A);
 
@@ -803,23 +810,23 @@ class ENGINE_API UKismetMathLibrary : public UBlueprintFunctionLibrary
 	UFUNCTION(BlueprintPure, meta=(DisplayName = "Vector2dLengthSquared", Keywords="magnitude"), Category="Math|Vector2D")
 	static float VSize2DSquared(FVector2D A);
 
-	/* Returns a unit normal version of the FVector A */
+	/** Returns a unit normal version of the FVector A */
 	UFUNCTION(BlueprintPure, meta=(DisplayName = "Normalize", Keywords="Unit Vector"), Category="Math|Vector")
 	static FVector Normal(FVector A);
 
-	/* Returns a unit normal version of the vector2d A */
+	/** Returns a unit normal version of the vector2d A */
 	UFUNCTION(BlueprintPure, meta=(DisplayName = "Normalize2D", Keywords="Unit Vector"), Category="Math|Vector2D")
 	static FVector2D Normal2D(FVector2D A);
 
-	/* Linearly interpolates between A and B based on Alpha (100% of A when Alpha=0 and 100% of B when Alpha=1) */
+	/** Linearly interpolates between A and B based on Alpha (100% of A when Alpha=0 and 100% of B when Alpha=1) */
 	UFUNCTION(BlueprintPure, meta=(DisplayName = "Lerp (vector)"), Category="Math|Vector")
 	static FVector VLerp(FVector A, FVector B, float Alpha);
 
-	/* Easeing  between A and B using a specified easing function */
+	/** Easeing  between A and B using a specified easing function */
 	UFUNCTION(BlueprintPure, meta = (DisplayName = "Ease (vector)", BlueprintInternalUseOnly = "true"), Category = "Math|Interpolation")
 	static FVector VEase(FVector A, FVector B, float Alpha, TEnumAsByte<EEasingFunc::Type> EasingFunc, float BlendExp = 2, int32 Steps = 2);
 
-	/* Returns a random vector with length of 1 */
+	/** Returns a random vector with length of 1 */
 	UFUNCTION(BlueprintPure, Category="Math|Random", meta=(NotBlueprintThreadSafe))
 	static FVector RandomUnitVector();
 
@@ -869,7 +876,7 @@ class ENGINE_API UKismetMathLibrary : public UBlueprintFunctionLibrary
 		return RandomUnitVectorInEllipticalConeInRadians(ConeDir, FMath::DegreesToRadians(MaxYawInDegrees), FMath::DegreesToRadians(MaxPitchInDegrees));
 	}
 
-	// Mirrors a vector by a normal
+	/** Mirrors a vector by a normal */
 	UFUNCTION(BlueprintPure, Category="Math|Vector")
 	static FVector MirrorVectorByNormal(FVector InVect, FVector InNormal);
 
@@ -1002,23 +1009,23 @@ class ENGINE_API UKismetMathLibrary : public UBlueprintFunctionLibrary
 	// Rotator functions.
 	//
 
-	// Returns true if rotator A is equal to rotator B (A == B) within a specified error tolerance
+	/** Returns true if rotator A is equal to rotator B (A == B) within a specified error tolerance */
 	UFUNCTION(BlueprintPure, meta=(DisplayName = "Equal (Rotator)", CompactNodeTitle = "==", Keywords = "== equal"), Category="Math|Rotator")
 	static bool EqualEqual_RotatorRotator(FRotator A, FRotator B, float ErrorTolerance = 1.e-4f);
 
-	// Returns true if rotator A is not equal to rotator B (A != B) within a specified error tolerance
+	/** Returns true if rotator A is not equal to rotator B (A != B) within a specified error tolerance */
 	UFUNCTION(BlueprintPure, meta=(DisplayName = "Not Equal (Rotator)", CompactNodeTitle = "!=", Keywords = "!= not equal"), Category="Math|Rotator")
 	static bool NotEqual_RotatorRotator(FRotator A, FRotator B, float ErrorTolerance = 1.e-4f);
 
-	// Returns rotator representing rotator A scaled by B 
+	/** Returns rotator representing rotator A scaled by B */
 	UFUNCTION(BlueprintPure, meta=(DisplayName = "ScaleRotator", CompactNodeTitle = "*", Keywords = "* multiply rotate rotation"), Category="Math|Rotator")
 	static FRotator Multiply_RotatorFloat(FRotator A, float B);
 	
-	// Returns rotator representing rotator A scaled by B 
+	/** Returns rotator representing rotator A scaled by B */
 	UFUNCTION(BlueprintPure, meta=(DisplayName = "ScaleRotator (int)", CompactNodeTitle = "*", Keywords = "* multiply rotate rotation"), Category="Math|Rotator")
 	static FRotator Multiply_RotatorInt(FRotator A, int32 B);
 
-	/** Combine 2 rotations to give you the resulting rotation */
+	/** Combine 2 rotations to give you the resulting rotation of first applying A, then B. */
 	UFUNCTION(BlueprintPure, meta=(DisplayName = "CombineRotators", Keywords="rotate rotation add"), Category="Math|Rotator")
 	static FRotator ComposeRotators(FRotator A, FRotator B);
 
@@ -1034,11 +1041,11 @@ class ENGINE_API UKismetMathLibrary : public UBlueprintFunctionLibrary
 	UFUNCTION(BlueprintPure, Category="Math|Random", meta=(Keywords="rotate rotation", NotBlueprintThreadSafe))
 	static FRotator RandomRotator(bool bRoll = false);
 
-	/* Linearly interpolates between A and B based on Alpha (100% of A when Alpha=0 and 100% of B when Alpha=1) */
+	/** Linearly interpolates between A and B based on Alpha (100% of A when Alpha=0 and 100% of B when Alpha=1) */
 	UFUNCTION(BlueprintPure, meta=(DisplayName = "Lerp (Rotator)"), Category="Math|Rotator")
 	static FRotator RLerp(FRotator A, FRotator B, float Alpha, bool bShortestPath);
 
-	/* Easeing  between A and B using a specified easing function */
+	/** Easeing  between A and B using a specified easing function */
 	UFUNCTION(BlueprintPure, meta = (DisplayName = "Ease (Rotator)", BlueprintInternalUseOnly = "true"), Category = "Math|Interpolation")
 	static FRotator REase(FRotator A, FRotator B, float Alpha, bool bShortestPath, TEnumAsByte<EEasingFunc::Type> EasingFunc, float BlendExp = 2, int32 Steps = 2);
 
@@ -1068,13 +1075,11 @@ class ENGINE_API UKismetMathLibrary : public UBlueprintFunctionLibrary
 	UFUNCTION(BlueprintPure, Category="Math|Rotator")
 	static float NormalizeAxis(float Angle);
 
-
-
 	//
 	//	LinearColor functions
 	//
 
-	/* Linearly interpolates between A and B based on Alpha (100% of A when Alpha=0 and 100% of B when Alpha=1) */
+	/** Linearly interpolates between A and B based on Alpha (100% of A when Alpha=0 and 100% of B when Alpha=1) */
 	UFUNCTION(BlueprintPure, meta=(DisplayName = "Lerp (LinearColor)"), Category="Math|Color")
 	static FLinearColor LinearColorLerp(FLinearColor A, FLinearColor B, float Alpha);
 
@@ -1090,11 +1095,11 @@ class ENGINE_API UKismetMathLibrary : public UBlueprintFunctionLibrary
 	UFUNCTION(BlueprintPure, meta=(DisplayName = "Lerp Using HSV (LinearColor)"), Category="Math|Color")
 	static FLinearColor LinearColorLerpUsingHSV(FLinearColor A, FLinearColor B, float Alpha);
 
-	/* Element-wise multiplication of two linear colors (R*R, G*G, B*B, A*A) */
+	/** Element-wise multiplication of two linear colors (R*R, G*G, B*B, A*A) */
 	UFUNCTION(BlueprintPure, meta=(DisplayName = "LinearColor * (LinearColor)", CompactNodeTitle = "*"), Category="Math|Color")
 	static FLinearColor Multiply_LinearColorLinearColor(FLinearColor A, FLinearColor B);
 
-	/* Element-wise multiplication of a linear color by a float (F*R, F*G, F*B, F*A) */
+	/** Element-wise multiplication of a linear color by a float (F*R, F*G, F*B, F*A) */
 	UFUNCTION(BlueprintPure, meta=(DisplayName = "LinearColor * Float", CompactNodeTitle = "*", Keywords = "* multiply"), Category="Math|Color")
 	static FLinearColor Multiply_LinearColorFloat(FLinearColor A, float B);
 
@@ -1116,139 +1121,139 @@ class ENGINE_API UKismetMathLibrary : public UBlueprintFunctionLibrary
 	// DateTime functions.
 	//
 
-	/* Makes a DateTime struct */
+	/** Makes a DateTime struct */
 	UFUNCTION(BlueprintPure, Category="Math|DateTime", meta=(NativeMakeFunc, AdvancedDisplay = "3"))
 	static FDateTime MakeDateTime(int32 Year, int32 Month, int32 Day, int32 Hour = 0, int32 Minute = 0, int32 Second = 0, int32 Millisecond = 0);
 
-	/* Breaks a DateTime into its components */
+	/** Breaks a DateTime into its components */
 	UFUNCTION(BlueprintPure, Category="Math|DateTime", meta=(NativeBreakFunc))
 	static void BreakDateTime(FDateTime InDateTime, int32& Year, int32& Month, int32& Day, int32& Hour, int32& Minute, int32& Second, int32& Millisecond);
 
-	/* Addition (A + B) */
+	/** Addition (A + B) */
 	UFUNCTION(BlueprintPure, meta=(DisplayName="DateTime + Timespan", CompactNodeTitle="+", Keywords="+ add plus"), Category="Math|DateTime")
 	static FDateTime Add_DateTimeTimespan( FDateTime A, FTimespan B );
 
-	/* Subtraction (A - B) */
+	/** Subtraction (A - B) */
 	UFUNCTION(BlueprintPure, meta=(DisplayName="DateTime - Timespan", CompactNodeTitle="-", Keywords="- subtract minus"), Category="Math|DateTime")
 	static FDateTime Subtract_DateTimeTimespan(FDateTime A, FTimespan B);
 
-	/* Subtraction (A - B) */
+	/** Subtraction (A - B) */
 	UFUNCTION(BlueprintPure, meta = (DisplayName = "DateTime - DateTime", CompactNodeTitle = "-", Keywords = "- subtract minus"), Category = "Math|DateTime")
 	static FTimespan Subtract_DateTimeDateTime(FDateTime A, FDateTime B);
 
-	/* Returns true if the values are equal (A == B) */
+	/** Returns true if the values are equal (A == B) */
 	UFUNCTION(BlueprintPure, meta=(DisplayName="Equal (DateTime)", CompactNodeTitle="==", Keywords="== equal"), Category="Math|DateTime")
 	static bool EqualEqual_DateTimeDateTime( FDateTime A, FDateTime B );
 
-	/* Returns true if the values are not equal (A != B) */
+	/** Returns true if the values are not equal (A != B) */
 	UFUNCTION(BlueprintPure, meta=(DisplayName="NotEqual (DateTime)", CompactNodeTitle="!=", Keywords="!= not equal"), Category="Math|DateTime")
 	static bool NotEqual_DateTimeDateTime( FDateTime A, FDateTime B );
 
-	/* Returns true if A is greater than B (A > B) */
+	/** Returns true if A is greater than B (A > B) */
 	UFUNCTION(BlueprintPure, meta=(DisplayName="DateTime > DateTime", CompactNodeTitle=">", Keywords="> greater"), Category="Math|DateTime")
 	static bool Greater_DateTimeDateTime( FDateTime A, FDateTime B );
 
-	/* Returns true if A is greater than or equal to B (A >= B) */
+	/** Returns true if A is greater than or equal to B (A >= B) */
 	UFUNCTION(BlueprintPure, meta=(DisplayName="DateTime >= DateTime", CompactNodeTitle=">=", Keywords=">= greater"), Category="Math|DateTime")
 	static bool GreaterEqual_DateTimeDateTime( FDateTime A, FDateTime B );
 
-	/* Returns true if A is less than B (A < B) */
+	/** Returns true if A is less than B (A < B) */
 	UFUNCTION(BlueprintPure, meta=(DisplayName="DateTime < DateTime", CompactNodeTitle="<", Keywords="< less"), Category="Math|DateTime")
 	static bool Less_DateTimeDateTime( FDateTime A, FDateTime B );
 
-	/* Returns true if A is less than or equal to B (A <= B) */
+	/** Returns true if A is less than or equal to B (A <= B) */
 	UFUNCTION(BlueprintPure, meta=(DisplayName="DateTime <= DateTime", CompactNodeTitle="<=", Keywords="<= less"), Category="Math|DateTime")
 	static bool LessEqual_DateTimeDateTime( FDateTime A, FDateTime B );
 
-	/* Returns the date component of A */
+	/** Returns the date component of A */
 	UFUNCTION(BlueprintPure, meta=(DisplayName="GetDate"), Category="Math|DateTime")
 	static FDateTime GetDate( FDateTime A );
 
-	/* Returns the day component of A (1 to 31) */
+	/** Returns the day component of A (1 to 31) */
 	UFUNCTION(BlueprintPure, meta=(DisplayName="GetDay"), Category="Math|DateTime")
 	static int32 GetDay( FDateTime A );
 
-	/* Returns the day of year of A */
+	/** Returns the day of year of A */
 	UFUNCTION(BlueprintPure, meta=(DisplayName="GetDayOfYear"), Category="Math|DateTime")
 	static int32 GetDayOfYear( FDateTime A );
 
-	/* Returns the hour component of A (24h format) */
+	/** Returns the hour component of A (24h format) */
 	UFUNCTION(BlueprintPure, meta=(DisplayName="GetHour"), Category="Math|DateTime")
 	static int32 GetHour( FDateTime A );
 
-	/* Returns the hour component of A (12h format) */
+	/** Returns the hour component of A (12h format) */
 	UFUNCTION(BlueprintPure, meta=(DisplayName="GetHour12"), Category="Math|DateTime")
 	static int32 GetHour12( FDateTime A );
 
-	/* Returns the millisecond component of A */
+	/** Returns the millisecond component of A */
 	UFUNCTION(BlueprintPure, meta=(DisplayName="GetMillisecond"), Category="Math|DateTime")
 	static int32 GetMillisecond( FDateTime A );
 
-	/* Returns the minute component of A */
+	/** Returns the minute component of A */
 	UFUNCTION(BlueprintPure, meta=(DisplayName="GetMinute"), Category="Math|DateTime")
 	static int32 GetMinute( FDateTime A );
 
-	/* Returns the month component of A */
+	/** Returns the month component of A */
 	UFUNCTION(BlueprintPure, meta=(DisplayName="GetMonth"), Category="Math|DateTime")
 	static int32 GetMonth( FDateTime A );
 
-	/* Returns the second component of A */
+	/** Returns the second component of A */
 	UFUNCTION(BlueprintPure, meta=(DisplayName="GetSecond"), Category="Math|DateTime")
 	static int32 GetSecond( FDateTime A );
 
-	/* Returns the time elapsed since midnight of A */
+	/** Returns the time elapsed since midnight of A */
 	UFUNCTION(BlueprintPure, meta=(DisplayName="GetTimeOfDay"), Category="Math|DateTime")
 	static FTimespan GetTimeOfDay( FDateTime A );
 
-	/* Returns the year component of A */
+	/** Returns the year component of A */
 	UFUNCTION(BlueprintPure, meta=(DisplayName="GetYear"), Category="Math|DateTime")
 	static int32 GetYear( FDateTime A );
 
-	/* Returns whether A's time is in the afternoon */
+	/** Returns whether A's time is in the afternoon */
 	UFUNCTION(BlueprintPure, meta=(DisplayName="IsAfternoon"), Category="Math|DateTime")
 	static bool IsAfternoon( FDateTime A );
 
-	/* Returns whether A's time is in the morning */
+	/** Returns whether A's time is in the morning */
 	UFUNCTION(BlueprintPure, meta=(DisplayName="IsMorning"), Category="Math|DateTime")
 	static bool IsMorning( FDateTime A );
 
-	/* Returns the number of days in the given year and month */
+	/** Returns the number of days in the given year and month */
 	UFUNCTION(BlueprintPure, meta=(DisplayName="DaysInMonth"), Category="Math|DateTime")
 	static int32 DaysInMonth( int32 Year, int32 Month );
 
-	/* Returns the number of days in the given year */
+	/** Returns the number of days in the given year */
 	UFUNCTION(BlueprintPure, meta=(DisplayName="DaysInYear"), Category="Math|DateTime")
 	static int32 DaysInYear( int32 Year );
 
-	/* Returns whether given year is a leap year */
+	/** Returns whether given year is a leap year */
 	UFUNCTION(BlueprintPure, meta=(DisplayName="IsLeapYear"), Category="Math|DateTime")
 	static bool IsLeapYear( int32 Year );
 
-	/* Returns the maximum date and time value */
+	/** Returns the maximum date and time value */
 	UFUNCTION(BlueprintPure, meta=(DisplayName="MaxValue"), Category="Math|DateTime")
 	static FDateTime DateTimeMaxValue( );
 
-	/* Returns the minimum date and time value */
+	/** Returns the minimum date and time value */
 	UFUNCTION(BlueprintPure, meta=(DisplayName="MinValue"), Category="Math|DateTime")
 	static FDateTime DateTimeMinValue( );
 
-	/* Returns the local date and time on this computer */
+	/** Returns the local date and time on this computer */
 	UFUNCTION(BlueprintPure, meta=(DisplayName="Now"), Category="Math|DateTime")
 	static FDateTime Now( );
 
-	/* Returns the local date on this computer */
+	/** Returns the local date on this computer */
 	UFUNCTION(BlueprintPure, meta=(DisplayName="Today"), Category="Math|DateTime")
 	static FDateTime Today( );
 
-	/* Returns the UTC date and time on this computer */
+	/** Returns the UTC date and time on this computer */
 	UFUNCTION(BlueprintPure, meta=(DisplayName="UtcNow"), Category="Math|DateTime")
 	static FDateTime UtcNow( );
 
-	/* Converts a date string in ISO-8601 format to a DateTime object */
+	/** Converts a date string in ISO-8601 format to a DateTime object */
 	UFUNCTION(BlueprintPure, Category="Math|DateTime")
 	static bool DateTimeFromIsoString(FString IsoString, FDateTime& Result);
 
-	/* Converts a date string to a DateTime object */
+	/** Converts a date string to a DateTime object */
 	UFUNCTION(BlueprintPure, Category="Math|DateTime")
 	static bool DateTimeFromString(FString DateTimeString, FDateTime& Result);
 
@@ -1256,143 +1261,143 @@ class ENGINE_API UKismetMathLibrary : public UBlueprintFunctionLibrary
 	// Timespan functions.
 	//
 
-	/* Makes a Timespan struct */
+	/** Makes a Timespan struct */
 	UFUNCTION(BlueprintPure, Category="Math|Timespan", meta=(NativeMakeFunc))
 	static FTimespan MakeTimespan(int32 Days, int32 Hours, int32 Minutes, int32 Seconds, int32 Milliseconds);
 
-	/* Makes a Timespan struct */
+	/** Makes a Timespan struct */
 	UFUNCTION(BlueprintPure, Category="Math|Timespan", meta=(NativeMakeFunc))
 	static FTimespan MakeTimespan2(int32 Days, int32 Hours, int32 Minutes, int32 Seconds, int32 FractionNano);
 
-	/* Breaks a Timespan into its components */
+	/** Breaks a Timespan into its components */
 	UFUNCTION(BlueprintPure, Category="Math|Timespan", meta=(NativeBreakFunc))
 	static void BreakTimespan(FTimespan InTimespan, int32& Days, int32& Hours, int32& Minutes, int32& Seconds, int32& Milliseconds);
 
-	/* Breaks a Timespan into its components */
+	/** Breaks a Timespan into its components */
 	UFUNCTION(BlueprintPure, Category="Math|Timespan", meta=(NativeBreakFunc))
 	static void BreakTimespan2(FTimespan InTimespan, int32& Days, int32& Hours, int32& Minutes, int32& Seconds, int32& FractionNano);
 
-	/* Addition (A + B) */
+	/** Addition (A + B) */
 	UFUNCTION(BlueprintPure, meta=(DisplayName="Timespan + Timespan", CompactNodeTitle="+", Keywords="+ add plus"), Category="Math|Timespan")
 	static FTimespan Add_TimespanTimespan( FTimespan A, FTimespan B );
 
-	/* Subtraction (A - B) */
+	/** Subtraction (A - B) */
 	UFUNCTION(BlueprintPure, meta=(DisplayName="Timespan - Timespan", CompactNodeTitle="-", Keywords="- subtract minus"), Category="Math|Timespan")
 	static FTimespan Subtract_TimespanTimespan( FTimespan A, FTimespan B );
 
-	/* Scalar multiplication (A * s) */
+	/** Scalar multiplication (A * s) */
 	UFUNCTION(BlueprintPure, meta=(DisplayName="Timespan * float", CompactNodeTitle="*", Keywords="* multiply"), Category="Math|Timespan")
 	static FTimespan Multiply_TimespanFloat( FTimespan A, float Scalar );
 
-	/* Scalar division (A * s) */
+	/** Scalar division (A * s) */
 	UFUNCTION(BlueprintPure, meta=(DisplayName="Timespan * float", CompactNodeTitle="/", Keywords="/ divide"), Category="Math|Timespan")
 	static FTimespan Divide_TimespanFloat( FTimespan A, float Scalar );
 
-	/* Returns true if the values are equal (A == B) */
+	/** Returns true if the values are equal (A == B) */
 	UFUNCTION(BlueprintPure, meta=(DisplayName="Equal (Timespan)", CompactNodeTitle="==", Keywords="== equal"), Category="Math|Timespan")
 	static bool EqualEqual_TimespanTimespan( FTimespan A, FTimespan B );
 
-	/* Returns true if the values are not equal (A != B) */
+	/** Returns true if the values are not equal (A != B) */
 	UFUNCTION(BlueprintPure, meta=(DisplayName="NotEqual (Timespan)", CompactNodeTitle="!=", Keywords="!= not equal"), Category="Math|Timespan")
 	static bool NotEqual_TimespanTimespan( FTimespan A, FTimespan B );
 
-	/* Returns true if A is greater than B (A > B) */
+	/** Returns true if A is greater than B (A > B) */
 	UFUNCTION(BlueprintPure, meta=(DisplayName="Timespan > Timespan", CompactNodeTitle=">", Keywords="> greater"), Category="Math|Timespan")
 	static bool Greater_TimespanTimespan( FTimespan A, FTimespan B );
 
-	/* Returns true if A is greater than or equal to B (A >= B) */
+	/** Returns true if A is greater than or equal to B (A >= B) */
 	UFUNCTION(BlueprintPure, meta=(DisplayName="Timespan >= Timespan", CompactNodeTitle=">=", Keywords=">= greater"), Category="Math|Timespan")
 	static bool GreaterEqual_TimespanTimespan( FTimespan A, FTimespan B );
 
-	/* Returns true if A is less than B (A < B) */
+	/** Returns true if A is less than B (A < B) */
 	UFUNCTION(BlueprintPure, meta=(DisplayName="Timespan < Timespan", CompactNodeTitle="<", Keywords="< less"), Category="Math|Timespan")
 	static bool Less_TimespanTimespan( FTimespan A, FTimespan B );
 
-	/* Returns true if A is less than or equal to B (A <= B) */
+	/** Returns true if A is less than or equal to B (A <= B) */
 	UFUNCTION(BlueprintPure, meta=(DisplayName="Timespan <= Timespan", CompactNodeTitle="<=", Keywords="<= less"), Category="Math|Timespan")
 	static bool LessEqual_TimespanTimespan( FTimespan A, FTimespan B );
 
-	/* Returns the days component of A */
+	/** Returns the days component of A */
 	UFUNCTION(BlueprintPure, meta=(DisplayName="GetDays"), Category="Math|Timespan")
 	static int32 GetDays( FTimespan A );
 
-	/* Returns the absolute value of A */
+	/** Returns the absolute value of A */
 	UFUNCTION(BlueprintPure, meta=(DisplayName="GetDuration"), Category="Math|Timespan")
 	static FTimespan GetDuration( FTimespan A );
 
-	/* Returns the hours component of A */
+	/** Returns the hours component of A */
 	UFUNCTION(BlueprintPure, meta=(DisplayName="GetHours"), Category="Math|Timespan")
 	static int32 GetHours( FTimespan A );
 
-	/* Returns the milliseconds component of A */
+	/** Returns the milliseconds component of A */
 	UFUNCTION(BlueprintPure, meta=(DisplayName="GetMilliseconds"), Category="Math|Timespan")
 	static int32 GetMilliseconds( FTimespan A );
 
-	/* Returns the minutes component of A */
+	/** Returns the minutes component of A */
 	UFUNCTION(BlueprintPure, meta=(DisplayName="GetMinutes"), Category="Math|Timespan")
 	static int32 GetMinutes( FTimespan A );
 
-	/* Returns the seconds component of A */
+	/** Returns the seconds component of A */
 	UFUNCTION(BlueprintPure, meta=(DisplayName="GetSeconds"), Category="Math|Timespan")
 	static int32 GetSeconds( FTimespan A );
 
-	/* Returns the total number of days in A */
+	/** Returns the total number of days in A */
 	UFUNCTION(BlueprintPure, meta=(DisplayName="GetTotalDays"), Category="Math|Timespan")
 	static float GetTotalDays( FTimespan A );
 
-	/* Returns the total number of hours in A */
+	/** Returns the total number of hours in A */
 	UFUNCTION(BlueprintPure, meta=(DisplayName="GetTotalHours"), Category="Math|Timespan")
 	static float GetTotalHours( FTimespan A );
 
-	/* Returns the total number of milliseconds in A */
+	/** Returns the total number of milliseconds in A */
 	UFUNCTION(BlueprintPure, meta=(DisplayName="GetTotalMilliseconds"), Category="Math|Timespan")
 	static float GetTotalMilliseconds( FTimespan A );
 
-	/* Returns the total number of minutes in A */
+	/** Returns the total number of minutes in A */
 	UFUNCTION(BlueprintPure, meta=(DisplayName="GetTotalMinutes"), Category="Math|Timespan")
 	static float GetTotalMinutes( FTimespan A );
 
-	/* Returns the total number of seconds in A */
+	/** Returns the total number of seconds in A */
 	UFUNCTION(BlueprintPure, meta=(DisplayName="GetTotalSeconds"), Category="Math|Timespan")
 	static float GetTotalSeconds( FTimespan A );
 
-	/* Returns a time span that represents the specified number of days */
+	/** Returns a time span that represents the specified number of days */
 	UFUNCTION(BlueprintPure, meta=(DisplayName="FromDays"), Category="Math|Timespan")
 	static FTimespan FromDays( float Days );
 
-	/* Returns a time span that represents the specified number of hours */
+	/** Returns a time span that represents the specified number of hours */
 	UFUNCTION(BlueprintPure, meta=(DisplayName="FromHours"), Category="Math|Timespan")
 	static FTimespan FromHours( float Hours );
 
-	/* Returns a time span that represents the specified number of milliseconds */
+	/** Returns a time span that represents the specified number of milliseconds */
 	UFUNCTION(BlueprintPure, meta=(DisplayName="FromMilliseconds"), Category="Math|Timespan")
 	static FTimespan FromMilliseconds( float Milliseconds );
 
-	/* Returns a time span that represents the specified number of minutes */
+	/** Returns a time span that represents the specified number of minutes */
 	UFUNCTION(BlueprintPure, meta=(DisplayName="FromMinutes"), Category="Math|Timespan")
 	static FTimespan FromMinutes( float Minutes );
 
-	/* Returns a time span that represents the specified number of seconds */
+	/** Returns a time span that represents the specified number of seconds */
 	UFUNCTION(BlueprintPure, meta=(DisplayName="FromSeconds"), Category="Math|Timespan")
 	static FTimespan FromSeconds( float Seconds );
 
-	/* Returns the maximum time span value */
+	/** Returns the maximum time span value */
 	UFUNCTION(BlueprintPure, meta=(DisplayName="MaxValue"), Category="Math|Timespan")
 	static FTimespan TimespanMaxValue( );
 
-	/* Returns the minimum time span value */
+	/** Returns the minimum time span value */
 	UFUNCTION(BlueprintPure, meta=(DisplayName="MinValue"), Category="Math|Timespan")
 	static FTimespan TimespanMinValue( );
 
-	/* Returns the ratio between two time spans (A / B), handles zero values */
+	/** Returns the ratio between two time spans (A / B), handles zero values */
 	UFUNCTION(BlueprintPure, meta=(DisplayName="TimespanRatio"), Category="Math|Timespan")
 	static float TimespanRatio( FTimespan A, FTimespan B );
 
-	/* Returns a zero time span value */
+	/** Returns a zero time span value */
 	UFUNCTION(BlueprintPure, meta=(DisplayName="ZeroValue"), Category="Math|Timespan")
 	static FTimespan TimespanZeroValue( );
 
-	/* Converts a time span string to a Timespan object */
+	/** Converts a time span string to a Timespan object */
 	UFUNCTION(BlueprintPure, Category="Math|Timespan")
 	static bool TimespanFromString(FString TimespanString, FTimespan& Result);
 
@@ -1615,11 +1620,11 @@ class ENGINE_API UKismetMathLibrary : public UBlueprintFunctionLibrary
 	UFUNCTION(BlueprintPure, Category="Math|Color", meta=(DisplayName = "RGB to HSV"))
 	static void RGBToHSV(const FLinearColor InColor, float& H, float& S, float& V, float& A);
 
-	// Converts a HSV linear color (where H is in R, S is in G, and V is in B) to RGB
+	/** Converts a HSV linear color (where H is in R, S is in G, and V is in B) to RGB */
 	UFUNCTION(BlueprintPure, Category="Math|Color", meta=(DisplayName = "HSV to RGB (vector)", Keywords="cast convert"))
 	static void HSVToRGB_Vector(const FLinearColor HSV, FLinearColor& RGB);
 
-	// Converts a RGB linear color to HSV (where H is in R, S is in G, and V is in B)
+	/** Converts a RGB linear color to HSV (where H is in R, S is in G, and V is in B) */
 	UFUNCTION(BlueprintPure, Category="Math|Color", meta=(DisplayName = "RGB to HSV (vector)", Keywords="cast convert"))
 	static void RGBToHSV_Vector(const FLinearColor RGB, FLinearColor& HSV);
 
@@ -1671,16 +1676,15 @@ class ENGINE_API UKismetMathLibrary : public UBlueprintFunctionLibrary
 	UFUNCTION(BlueprintPure, meta=(DisplayName = "GetRotationXVector", Keywords="rotation rotate cast convert", BlueprintAutocast), Category="Math|Rotator")
 	static FVector Conv_RotatorToVector(FRotator InRot);
 
-
 	//
 	// Object operators and functions.
 	//
 	
-	/* Returns true if A and B are equal (A == B) */
+	/** Returns true if A and B are equal (A == B) */
 	UFUNCTION(BlueprintPure, meta=(DisplayName = "Equal (Object)", CompactNodeTitle = "==", Keywords = "== equal"), Category="Utilities")
 	static bool EqualEqual_ObjectObject(class UObject* A, class UObject* B);
 
-	/* Returns true if A and B are not equal (A != B) */
+	/** Returns true if A and B are not equal (A != B) */
 	UFUNCTION(BlueprintPure, meta=(DisplayName = "NotEqual (Object)", CompactNodeTitle = "!=", Keywords = "!= not equal"), Category="Utilities")
 	static bool NotEqual_ObjectObject(class UObject* A, class UObject* B);
 
@@ -1688,11 +1692,11 @@ class ENGINE_API UKismetMathLibrary : public UBlueprintFunctionLibrary
 	// Class operators and functions.
 	//
 
-	/* Returns true if A and B are equal (A == B) */
+	/** Returns true if A and B are equal (A == B) */
 	UFUNCTION(BlueprintPure, meta=(DisplayName = "Equal (Class)", CompactNodeTitle = "==", Keywords = "== equal"), Category="Utilities")
 	static bool EqualEqual_ClassClass(class UClass* A, class UClass* B);
 
-	/* Returns true if A and B are not equal (A != B) */
+	/** Returns true if A and B are not equal (A != B) */
 	UFUNCTION(BlueprintPure, meta=(DisplayName = "NotEqual (Class)", CompactNodeTitle = "!=", Keywords = "!= not equal"), Category="Utilities")
 	static bool NotEqual_ClassClass(class UClass* A, class UClass* B);
 
@@ -1709,11 +1713,11 @@ class ENGINE_API UKismetMathLibrary : public UBlueprintFunctionLibrary
 	// Name operators.
 	//
 	
-	/* Returns true if A and B are equal (A == B) */
+	/** Returns true if A and B are equal (A == B) */
 	UFUNCTION(BlueprintPure, meta=(DisplayName = "Equal (Name)", CompactNodeTitle = "==", Keywords = "== equal"), Category="Utilities|Name")
 	static bool EqualEqual_NameName(FName A, FName B);
 
-	/* Returns true if A and B are not equal (A != B) */
+	/** Returns true if A and B are not equal (A != B) */
 	UFUNCTION(BlueprintPure, meta=(DisplayName = "NotEqual (Name)", CompactNodeTitle = "!=", Keywords = "!= not equal"), Category="Utilities|Name")
 	static bool NotEqual_NameName(FName A, FName B);
 
@@ -1723,32 +1727,46 @@ class ENGINE_API UKismetMathLibrary : public UBlueprintFunctionLibrary
 	
 	/** 
 	 *	Transform a position by the supplied transform.
-	 *	For example, if T was an object's transform, would transform a position from local space to world space.
+	 *	For example, if T was an object's transform, this would transform a position from local space to world space.
 	 */
 	UFUNCTION(BlueprintPure, Category="Math|Transform", meta=(Keywords="location"))
 	static FVector TransformLocation(const FTransform& T, FVector Location);
 
 	/** 
 	 *	Transform a direction vector by the supplied transform - will not change its length. 
-	 *	For example, if T was an object's transform, would transform a direction from local space to world space.
+	 *	For example, if T was an object's transform, this would transform a direction from local space to world space.
 	 */
 	UFUNCTION(BlueprintPure, Category="Math|Transform")
 	static FVector TransformDirection(const FTransform& T, FVector Direction);
 
+	/** 
+	 *	Transform a rotator by the supplied transform. 
+	 *	For example, if T was an object's transform, this would transform a rotation from local space to world space.
+	 */
+	UFUNCTION(BlueprintPure, Category="Math|Transform")
+	static FRotator TransformRotation(const FTransform& T, FRotator Rotation);
+
 	/** 
 	 *	Transform a position by the inverse of the supplied transform.
-	 *	For example, if T was an object's transform, would transform a position from world space to local space.
+	 *	For example, if T was an object's transform, this would transform a position from world space to local space.
 	 */
 	UFUNCTION(BlueprintPure, Category="Math|Transform", meta=(Keywords="location"))
 	static FVector InverseTransformLocation(const FTransform& T, FVector Location);
 
 	/** 
 	 *	Transform a direction vector by the inverse of the supplied transform - will not change its length.
-	 *	For example, if T was an object's transform, would transform a direction from world space to local space.
+	 *	For example, if T was an object's transform, this would transform a direction from world space to local space.
 	 */
 	UFUNCTION(BlueprintPure, Category="Math|Transform")
 	static FVector InverseTransformDirection(const FTransform& T, FVector Direction);
 
+	/** 
+	 *	Transform a rotator by the inverse of the supplied transform. 
+	 *	For example, if T was an object's transform, this would transform a rotation from world space to local space.
+	 */
+	UFUNCTION(BlueprintPure, Category="Math|Transform")
+	static FRotator InverseTransformRotation(const FTransform& T, FRotator Rotation);
+
 	/**
 	 * Compose two transforms in order: A * B.
 	 *
@@ -1827,7 +1845,7 @@ class ENGINE_API UKismetMathLibrary : public UBlueprintFunctionLibrary
 	UFUNCTION(BlueprintPure, meta=(DisplayName = "vector2d * float", CompactNodeTitle = "*", Keywords = "* multiply"), Category="Math|Vector2D")
 	static FVector2D Multiply_Vector2DFloat(FVector2D A, float B);
 
-	/* Element-wise Vector multiplication (Result = {A.x*B.x, A.y*B.y}) */
+	/** Element-wise Vector multiplication (Result = {A.x*B.x, A.y*B.y}) */
 	UFUNCTION(BlueprintPure, meta = (DisplayName = "vector2d * vector2d", CompactNodeTitle = "*", Keywords = "* multiply", CommutativeAssociativeBinaryOperator = "true"), Category = "Math|Vector2D")
 	static FVector2D Multiply_Vector2DVector2D(FVector2D A, FVector2D B);
 
@@ -1835,7 +1853,7 @@ class ENGINE_API UKismetMathLibrary : public UBlueprintFunctionLibrary
 	UFUNCTION(BlueprintPure, meta=(DisplayName = "vector2d / float", CompactNodeTitle = "/", Keywords = "/ divide division"), Category="Math|Vector2D")
 	static FVector2D Divide_Vector2DFloat(FVector2D A, float B = 1.f);
 
-	/* Element-wise Vector divide (Result = {A.x/B.x, A.y/B.y}) */
+	/** Element-wise Vector divide (Result = {A.x/B.x, A.y/B.y}) */
 	UFUNCTION(BlueprintPure, meta = (DisplayName = "vector2d / vector2d", CompactNodeTitle = "/", Keywords = "/ divide division"), Category = "Math|Vector2D")
 	static FVector2D Divide_Vector2DVector2D(FVector2D A, FVector2D B);
 
@@ -1847,11 +1865,11 @@ class ENGINE_API UKismetMathLibrary : public UBlueprintFunctionLibrary
 	UFUNCTION(BlueprintPure, meta=(DisplayName = "vector2d - float", CompactNodeTitle = "-", Keywords = "- subtract minus"), Category="Math|Vector2D")
 	static FVector2D Subtract_Vector2DFloat(FVector2D A, float B);
 
-	/* Returns true if vector2D A is equal to vector2D B (A == B) within a specified error tolerance */
+	/** Returns true if vector2D A is equal to vector2D B (A == B) within a specified error tolerance */
     UFUNCTION(BlueprintPure, meta=(DisplayName = "Equal (vector2D)", CompactNodeTitle = "==", Keywords = "== equal"), Category="Math|Vector2D")
     static bool EqualEqual_Vector2DVector2D(FVector2D A, FVector2D B, float ErrorTolerance = 1.e-4f);
 
-    /* Returns true if vector2D A is not equal to vector2D B (A != B) within a specified error tolerance */
+    /** Returns true if vector2D A is not equal to vector2D B (A != B) within a specified error tolerance */
     UFUNCTION(BlueprintPure, meta=(DisplayName = "Not Equal (vector2D)", CompactNodeTitle = "!=", Keywords = "!= not equal"), Category="Math|Vector2D")
     static bool NotEqual_Vector2DVector2D(FVector2D A, FVector2D B, float ErrorTolerance = 1.e-4f);
 	
@@ -1891,7 +1909,6 @@ class ENGINE_API UKismetMathLibrary : public UBlueprintFunctionLibrary
 	UFUNCTION(BlueprintPure, Category="Math|Interpolation", meta=(Keywords="position"))
 	static FVector VInterpTo(FVector Current, FVector Target, float DeltaTime, float InterpSpeed);
 
-
 	/**
 	 * Tries to reach Target at a constant rate.
 	 *
@@ -2002,7 +2019,7 @@ class ENGINE_API UKismetMathLibrary : public UBlueprintFunctionLibrary
 	// Random stream functions
 	//
 
-	/* Returns a uniformly distributed random number between 0 and Max - 1 */
+	/** Returns a uniformly distributed random number between 0 and Max - 1 */
 	UFUNCTION(BlueprintPure, Category="Math|Random")
 	static int32 RandomIntegerFromStream(int32 Max, const FRandomStream& Stream);
 
@@ -2010,7 +2027,7 @@ class ENGINE_API UKismetMathLibrary : public UBlueprintFunctionLibrary
 	UFUNCTION(BlueprintPure, Category="Math|Random")
 	static int32 RandomIntegerInRangeFromStream(int32 Min, int32 Max, const FRandomStream& Stream);
 
-	/* Returns a random bool */
+	/** Returns a random bool */
 	UFUNCTION(BlueprintPure, Category="Math|Random")
 	static bool RandomBoolFromStream(const FRandomStream& Stream);
 
@@ -2022,7 +2039,7 @@ class ENGINE_API UKismetMathLibrary : public UBlueprintFunctionLibrary
 	UFUNCTION(BlueprintPure, Category="Math|Random")
 	static float RandomFloatInRangeFromStream(float Min, float Max, const FRandomStream& Stream);
 
-	/* Returns a random vector with length of 1.0 */
+	/** Returns a random vector with length of 1.0 */
 	UFUNCTION(BlueprintPure, Category="Math|Random")
 	static FVector RandomUnitVectorFromStream(const FRandomStream& Stream);
 
@@ -2177,7 +2194,6 @@ private:
 };
 
 
-
 // Conditionally inlined
 #if KISMET_MATH_INLINE_ENABLED
 #include "KismetMathLibrary.inl"
diff --git a/Engine/Source/Runtime/Engine/Classes/Kismet/KismetSystemLibrary.h b/Engine/Source/Runtime/Engine/Classes/Kismet/KismetSystemLibrary.h
--- a/Engine/Source/Runtime/Engine/Classes/Kismet/KismetSystemLibrary.h
+++ b/Engine/Source/Runtime/Engine/Classes/Kismet/KismetSystemLibrary.h
@@ -184,6 +184,10 @@ class ENGINE_API UKismetSystemLibrary : public UBlueprintFunctionLibrary
 	UFUNCTION(BlueprintPure, Category = "Utilities")
 	static bool IsValidSoftObjectReference(const TSoftObjectPtr<UObject>& SoftObjectReference);
 
+	/** Converts a Soft Object Reference to a string. The other direction is not provided because it cannot be validated */
+	UFUNCTION(BlueprintPure, meta = (DisplayName = "ToString (SoftObjectReference)", CompactNodeTitle = "->"), Category = "Utilities")
+	static FString Conv_SoftObjectReferenceToString(const TSoftObjectPtr<UObject>& SoftObjectReference);
+
 	/** Returns true if the values are equal (A == B) */
 	UFUNCTION(BlueprintPure, meta = (DisplayName = "Equal (SoftObjectReference)", CompactNodeTitle = "=="), Category = "Utilities")
 	static bool EqualEqual_SoftObjectReference(const TSoftObjectPtr<UObject>& A, const TSoftObjectPtr<UObject>& B);
@@ -196,6 +200,10 @@ class ENGINE_API UKismetSystemLibrary : public UBlueprintFunctionLibrary
 	UFUNCTION(BlueprintPure, Category = "Utilities")
 	static bool IsValidSoftClassReference(const TSoftClassPtr<UObject>& SoftClassReference);
 
+	/** Converts a Soft Class Reference to a string. The other direction is not provided because it cannot be validated */
+	UFUNCTION(BlueprintPure, meta = (DisplayName = "ToString (SoftObjectReference)", CompactNodeTitle = "->"), Category = "Utilities")
+	static FString Conv_SoftClassReferenceToString(const TSoftClassPtr<UObject>& SoftClassReference);
+
 	/** Returns true if the values are equal (A == B) */
 	UFUNCTION(BlueprintPure, meta = (DisplayName = "Equal (SoftClassReference)", CompactNodeTitle = "=="), Category = "Utilities")
 	static bool EqualEqual_SoftClassReference(const TSoftClassPtr<UObject>& A, const TSoftClassPtr<UObject>& B);
@@ -1644,6 +1652,10 @@ class ENGINE_API UKismetSystemLibrary : public UBlueprintFunctionLibrary
 	UFUNCTION(BlueprintPure, Category = "AssetManager")
 	static bool IsValidPrimaryAssetId(FPrimaryAssetId PrimaryAssetId);
 
+	/** Converts a Primary Asset Id to a string. The other direction is not provided because it cannot be validated */
+	UFUNCTION(BlueprintPure, meta = (DisplayName = "ToString (PrimaryAssetId)", CompactNodeTitle = "->"), Category = "AssetManager")
+	static FString Conv_PrimaryAssetIdToString(FPrimaryAssetId PrimaryAssetId);
+
 	/** Returns true if the values are equal (A == B) */
 	UFUNCTION(BlueprintPure, meta = (DisplayName = "Equal (PrimaryAssetId)", CompactNodeTitle = "=="), Category = "AssetManager")
 	static bool EqualEqual_PrimaryAssetId(FPrimaryAssetId A, FPrimaryAssetId B);
@@ -1652,10 +1664,14 @@ class ENGINE_API UKismetSystemLibrary : public UBlueprintFunctionLibrary
 	UFUNCTION(BlueprintPure, meta = (DisplayName = "NotEqual (PrimaryAssetId)", CompactNodeTitle = "!="), Category = "AssetManager")
 	static bool NotEqual_PrimaryAssetId(FPrimaryAssetId A, FPrimaryAssetId B);
 
-	/** Returns list of PrimaryAssetIds for a PrimaryAssetType */
+	/** Returns list of Primary Asset Ids for a PrimaryAssetType */
 	UFUNCTION(BlueprintPure, Category = "AssetManager")
 	static bool IsValidPrimaryAssetType(FPrimaryAssetType PrimaryAssetType);
 
+	/** Converts a Primary Asset Type to a string. The other direction is not provided because it cannot be validated */
+	UFUNCTION(BlueprintPure, meta = (DisplayName = "ToString (PrimaryAssetType)", CompactNodeTitle = "->"), Category = "AssetManager")
+	static FString Conv_PrimaryAssetTypeToString(FPrimaryAssetType PrimaryAssetType);
+
 	/** Returns true if the values are equal (A == B) */
 	UFUNCTION(BlueprintPure, meta = (DisplayName = "Equal (PrimaryAssetType)", CompactNodeTitle = "=="), Category = "AssetManager")
 	static bool EqualEqual_PrimaryAssetType(FPrimaryAssetType A, FPrimaryAssetType B);
diff --git a/Engine/Source/Runtime/Engine/Classes/Particles/ParticleSystemComponent.h b/Engine/Source/Runtime/Engine/Classes/Particles/ParticleSystemComponent.h
--- a/Engine/Source/Runtime/Engine/Classes/Particles/ParticleSystemComponent.h
+++ b/Engine/Source/Runtime/Engine/Classes/Particles/ParticleSystemComponent.h
@@ -1,4 +1,4 @@
-// Copyright 1998-2017 Epic Games, Inc. All Rights Reserved.
+﻿// Copyright 1998-2017 Epic Games, Inc. All Rights Reserved.
 
 
 #pragma once
@@ -341,37 +341,37 @@
 	UPROPERTY(transient, duplicatetransient)
 	TArray<class USkeletalMeshComponent*> SkelMeshComponents;
 
-	uint32 bWasCompleted:1;
+	uint8 bWasCompleted:1;
 
-	uint32 bSuppressSpawning:1;
+	uint8 bSuppressSpawning:1;
 
 private:
 	/** if true, this psys can tick in any thread **/
-	uint32 bIsElligibleForAsyncTick:1;
+	uint8 bIsElligibleForAsyncTick:1;
 	/** if true, bIsElligibleForAsyncTick is set up **/
-	uint32 bIsElligibleForAsyncTickComputed:1;
+	uint8 bIsElligibleForAsyncTickComputed:1;
 public:
 
-	uint32 bWasDeactivated:1;
+	uint8 bWasDeactivated:1;
 
 	/** True if this was active before being unregistered or otherwise reset, if so reactivate it */
-	uint32 bWasActive:1;
+	uint8 bWasActive:1;
 
 	/** If true, someone has requested this component reset. */
-	uint32 bResetTriggered : 1;
+	uint8 bResetTriggered : 1;
 
 	UPROPERTY(EditAnywhere, BlueprintReadWrite, Category=Particles)
-	uint32 bResetOnDetach:1;
+	uint8 bResetOnDetach:1;
 
 	/** whether to update the particle system on dedicated servers */
 	UPROPERTY()
-	uint32 bUpdateOnDedicatedServer:1;
+	uint8 bUpdateOnDedicatedServer:1;
 
 	/** Indicates that the component has not been ticked since being registered. */
-	uint32 bJustRegistered:1;
+	uint8 bJustRegistered:1;
 
 	/** This flag will be set the first time the PSysComp is activated... used to prevent auto activated PSysComps from calling InitParticles twice on level load */
-	uint32 bHasBeenActivated:1;
+	uint8 bHasBeenActivated:1;
 
 	/**
 	 * If true, this Particle System will be available for recycling after it has completed. Auto-destroyed systems cannot be recycled.
@@ -379,7 +379,7 @@
 	 * This is only an optimization and does not change particle system behavior, aside from not triggering normal component initialization events more than once.
 	 */
 	UPROPERTY(EditAnywhere, BlueprintReadWrite, AdvancedDisplay, Category=Particles)
-	uint32 bAllowRecycling:1;
+	uint8 bAllowRecycling:1;
 
 	/**
 	 * True if we should automatically attach to AutoAttachParent when activated, and detach from our parent when completed.
@@ -389,29 +389,71 @@
 	 * @see AutoAttachParent, AutoAttachSocketName, AutoAttachLocationType
 	 */
 	UPROPERTY(EditAnywhere, BlueprintReadOnly, Category=Attachment)
-	uint32 bAutoManageAttachment:1;
+	uint8 bAutoManageAttachment:1;
 	
-	/** The significance this component requires of it's emitters for them to be enabled. */
-	UPROPERTY()
-	EParticleSignificanceLevel RequiredSignificance;
-
-	/** Time in seconds since we were last considered significant. */
-	float LastSignificantTime;
-
 	/** If this component is having it's significance managed by gameplay code. */
-	uint32 bIsManagingSignificance : 1;
+	uint8 bIsManagingSignificance : 1;
 	/** If this component was previously having it's significance managed by gameplay code. Allows us to refresh render data when this changes. */
-	uint32 bWasManagingSignificance : 1;
-	
+	uint8 bWasManagingSignificance : 1;
+
+	UPROPERTY()
+	uint8 bWarmingUp:1;
+
+	/** indicates that the component's LODMethod overrides the Template's */
+	UPROPERTY(EditAnywhere, BlueprintReadWrite, Category=LOD)
+	uint8 bOverrideLODMethod:1;
+
+	/**
+	 *	Flag indicating that dynamic updating of render data should NOT occur during Tick.
+	 *	This is used primarily to allow for warming up and simulated effects to a certain state.
+	 */
+	UPROPERTY()
+	uint8 bSkipUpdateDynamicDataDuringTick:1;
+
+	/** This is set when any of our "don't tick me" timeout values have fired */
+	uint8 bForcedInActive:1;
+
+	/** If true, force an LOD update from the renderer. */
+	uint8 bForceLODUpdateFromRenderer:1;
+
+	/** If true, the ViewRelevanceFlags are dirty and should be recached */
+	uint8 bIsViewRelevanceDirty : 1;
+
+	uint8 bAutoDestroy : 1;
+
 private:
+	
+	/** Indicates if the component's transform has been changed and requires updating */
+	uint8 bIsTransformDirty:1;
+
 	/** Did we auto attach during activation? Used to determine if we should restore the relative transform during detachment. */
-	uint32 bDidAutoAttach:1;
+	uint8 bDidAutoAttach : 1;
 
+	/** Is AsyncComponentToWorld etc valid? */
+	uint8 bAsyncDataCopyIsValid:1;
+	uint8 bParallelRenderThreadUpdate:1;
+
+	/** If true, it means the ASync work is done and the finalize is not */
+	uint8 bNeedsFinalize:1;
+	/** If true, it means the Async work is in process and not yet completed */
+	uint8 bAsyncWorkOutstanding:1;
+
 	/** Restore relative transform from auto attachment and optionally detach from parent (regardless of whether it was an auto attachment). */
 	void CancelAutoAttachment(bool bDetachFromParent);
 
 public:
 
+	/** The method of LOD level determination to utilize for this particle system */
+	UPROPERTY(EditAnywhere, BlueprintReadWrite, Category=LOD)
+	TEnumAsByte<enum ParticleSystemLODMethod> LODMethod;
+
+	/** The significance this component requires of it's emitters for them to be enabled. */
+	UPROPERTY()
+	EParticleSignificanceLevel RequiredSignificance;
+
+	/** Current particle 'replay state'.  This setting controls whether we're currently simulating/rendering particles normally, or whether we should capture or playback particle replay data instead. */
+	TEnumAsByte<enum ParticleReplayState> ReplayState;
+
 	void ResetNextTick()
 	{
 		bResetTriggered = true;
@@ -448,16 +490,11 @@
 	UPROPERTY()
 	float WarmupTickRate;
 
-	UPROPERTY()
-	uint32 bWarmingUp:1;
-
 private:
 	int32 LODLevel;
 
 public:
 
-	uint32 bAutoDestroy:1;
-
 	/**
 	 * Number of seconds of emitter not being rendered that need to pass before it
 	 * no longer gets ticked/ becomes inactive.
@@ -469,9 +506,6 @@
 	/** Tracks the time since the last forced UpdateTransform. */
 	float TimeSinceLastForceUpdateTransform;
 
-	/** Indicates if the component's transform has been changed and requires updating */
-	bool bIsTransformDirty;
-
 public:
 	/**
 	 * Time between forced UpdateTransforms for systems that use dynamically calculated bounds,
@@ -497,48 +531,36 @@
 	/** Used to accumulate total tick time to determine whether system can be skipped ticking if not visible. */
 	float AccumTickTime;
 
-	/** indicates that the component's LODMethod overrides the Template's */
-	UPROPERTY(EditAnywhere, BlueprintReadWrite, Category=LOD)
-	uint32 bOverrideLODMethod:1;
+	/** Time in seconds since we were last considered significant. */
+	float LastSignificantTime;
 
-	/** The method of LOD level determination to utilize for this particle system */
-	UPROPERTY(EditAnywhere, BlueprintReadWrite, Category=LOD)
-	TEnumAsByte<enum ParticleSystemLODMethod> LODMethod;
+	/** LOD updating... */
+	float AccumLODDistanceCheckTime;
 
-	/**
-	 *	Flag indicating that dynamic updating of render data should NOT occur during Tick.
-	 *	This is used primarily to allow for warming up and simulated effects to a certain state.
-	 */
-	UPROPERTY()
-	uint32 bSkipUpdateDynamicDataDuringTick:1;
+private:
+	/** Remember the global detail mode when we last checked the emitters. */
+	uint32 LastCheckedDetailMode;
 
-	/** This is set when any of our "don't tick me" timeout values have fired */
-	uint32 bForcedInActive:1;
-
-	/** If true, force an LOD update from the renderer. */
-	uint32 bForceLODUpdateFromRenderer:1;
-
+public:
 	/** The view relevance flags for each LODLevel. */
 	TArray<FMaterialRelevance> CachedViewRelevanceFlags;
 
-	/** If true, the ViewRelevanceFlags are dirty and should be recached */
-	uint32 bIsViewRelevanceDirty:1;
-
 	/** Array of replay clips for this particle system component.  These are serialized to disk.  You really should never add anything to this in the editor.  It's exposed so that you can delete clips if you need to, but be careful when doing so! */
 	UPROPERTY()
 	TArray<class UParticleSystemReplay*> ReplayClips;
 
-	/** Current particle 'replay state'.  This setting controls whether we're currently simulating/rendering particles normally, or whether we should capture or playback particle replay data instead. */
-	TEnumAsByte<enum ParticleReplayState> ReplayState;
-
 	/** Clip ID number we're either playing back or capturing to, depending on the value of ReplayState. */
 	int32 ReplayClipIDNumber;
 
 	/** The current replay frame for playback */
 	int32 ReplayFrameIndex;
 
-	/** LOD updating... */
-	float AccumLODDistanceCheckTime;
+	/** Scales DeltaTime in UParticleSystemComponent::Tick(...) */
+	UPROPERTY(EditAnywhere, BlueprintReadWrite, Category=Particles)
+	float CustomTimeDilation;
 
+	/** This is created at start up and then added to each emitter */
+	float EmitterDelay;
+
 	/** The Spawn events that occurred in this PSysComp. */
 	TArray<struct FParticleEventSpawnData> SpawnEvents;
@@ -555,17 +577,6 @@
 	/** The Kismet events that occurred for this PSysComp. */
 	TArray<struct FParticleEventKismetData> KismetEvents;
 
-	/** Scales DeltaTime in UParticleSystemComponent::Tick(...) */
-	UPROPERTY(EditAnywhere, BlueprintReadWrite, Category=Particles)
-	float CustomTimeDilation;
-
-	/** This is created at start up and then added to each emitter */
-	float EmitterDelay;
-
-	// Called when the particle system is done
-	UPROPERTY(BlueprintAssignable)
-	FOnSystemFinished OnSystemFinished;
-
 public:
 	/**
 	 * Component we automatically attach to when activated, if bAutoManageAttachment is true.
@@ -582,12 +593,14 @@
 	UPROPERTY(EditAnywhere, BlueprintReadWrite, Category=Attachment, meta=(EditCondition="bAutoManageAttachment"))
 	FName AutoAttachSocketName;
 
+#if WITH_EDITORONLY_DATA
 	/**
 	 * DEPRECATED: Options for how we handle our location when we attach to the AutoAttachParent, if bAutoManageAttachment is true.
 	 * @see bAutoManageAttachment, EAttachLocation::Type
 	 */
 	UPROPERTY()
 	TEnumAsByte<EAttachLocation::Type> AutoAttachLocationType_DEPRECATED;
+#endif
 
 	/**
 	 * Options for how we handle our location when we attach to the AutoAttachParent, if bAutoManageAttachment is true.
@@ -639,9 +652,14 @@
 	FRotator SavedAutoAttachRelativeRotation;
 	FVector SavedAutoAttachRelativeScale3D;
 
+public:
+	class FFXSystemInterface* FXSystem;
+
+	// Called when the particle system is done
+	UPROPERTY(BlueprintAssignable)
+	FOnSystemFinished OnSystemFinished;
+
 private:
-	/** Cached copy of the transform for async work */
-	FTransform AsyncComponentToWorld;
 	/** Cached copy of the instance params */
 	TArray<struct FParticleSysParam> AsyncInstanceParameters;
 	/** Player locations computed before kicking off async task. Safe to access from async task or game thread*/
@@ -653,13 +671,6 @@
 	/** Cached copy of PartSysVelocity */
 	FVector AsyncPartSysVelocity;
 
-	/** Is AsyncComponentToWorld etc valid? */
-	bool bAsyncDataCopyIsValid;
-	bool bParallelRenderThreadUpdate;
-
-	/** Remember the global detail mode when we last checked the emitters. */
-	uint32 LastCheckedDetailMode;
-
 public:
 
 	/** Called from game code when the significance required for a component changes. */
@@ -1061,27 +1072,25 @@
 	void SetTrailSourceData(FName InFirstSocketName, FName InSecondSocketName, ETrailWidthMode InWidthMode, float InWidth);
 
 public:
-	TArray<struct FParticleEmitterInstance*> EmitterInstances;
-
-	class FFXSystemInterface* FXSystem;
-
 	/** Command fence used to shut down properly */
 	class FRenderCommandFence* ReleaseResourcesFence;
 
+	TArray<struct FParticleEmitterInstance*> EmitterInstances;
+
 	/** Static delegate called for all systems on an activation change. */
 	static FOnSystemPreActivationChange OnSystemPreActivationChange;
 
 private:
+	/** Cached copy of the transform for async work */
+	FTransform AsyncComponentToWorld;
+
 	/** Task ref for parallel portion */
 	FGraphEventRef AsyncWork;
+
 	/** Copy of delta time coming into TickComponent */
 	float DeltaTimeTick;
 	/** Info from async tick */
 	int32 TotalActiveParticles;
-	/** If true, it means the ASync work is done and the finalize is not */
-	bool bNeedsFinalize;
-	/** If true, it means the Async work is in process and not yet completed */
-	bool bAsyncWorkOutstanding;
 	/** Number of significant emitters. When this is 0, the system can either be deactivated or stopped ticking. */
 	uint32 NumSignificantEmitters;
 	/** Time in ms since a tick was last performed; used with MinTimeBetweenTicks (on UParticleSystem) to control tick rate */
