Copying //UE4/Dev-Framework to //UE4/Dev-Main (Source: //UE4/Dev-Framework @ 3628051)

diff --git a/Engine/Source/Runtime/Engine/Classes/PhysicsEngine/BodyInstance.h b/Engine/Source/Runtime/Engine/Classes/PhysicsEngine/BodyInstance.h
--- a/Engine/Source/Runtime/Engine/Classes/PhysicsEngine/BodyInstance.h
+++ b/Engine/Source/Runtime/Engine/Classes/PhysicsEngine/BodyInstance.h
@@ -165,7 +165,7 @@ private:
 	friend struct FBodyInstance;
 };
 
-enum class BodyInstanceSceneState
+enum class BodyInstanceSceneState : uint8
 {
 	NotAdded,
 	AwaitingAdd,
@@ -215,12 +215,32 @@ struct ENGINE_API FBodyInstance
 	// COLLISION SETTINGS
 
 	/** Types of objects that this physics objects will collide with. */
-	// @todo : make this to be transient, so that it doesn't have to save anymore
-	// we have to still load them until resave
 	UPROPERTY() 
 	struct FCollisionResponseContainer ResponseToChannels_DEPRECATED;
 
+public:
+	// Current state of the physx body for tracking deferred addition and removal.
+	BodyInstanceSceneState CurrentSceneState;
+
+	/** The set of values used in considering when put this body to sleep. */
+	UPROPERTY(EditAnywhere, AdvancedDisplay, BlueprintReadWrite, Category = Physics)
+	ESleepFamily SleepFamily;
+
+	/** Locks physical movement along specified axis.*/
+	UPROPERTY(EditAnywhere, Category = Physics, meta = (DisplayName = "Mode"))
+	TEnumAsByte<EDOFMode::Type> DOFMode;
+
 private:
+	/**
+	 * Type of collision enabled.
+	 * 
+	 *	No Collision      : Will not create any representation in the physics engine. Cannot be used for spatial queries (raycasts, sweeps, overlaps) or simulation (rigid body, constraints). Best performance possible (especially for moving objects)
+	 *	Query Only        : Only used for spatial queries (raycasts, sweeps, and overlaps). Cannot be used for simulation (rigid body, constraints). Useful for character movement and things that do not need physical simulation. Performance gains by keeping data out of simulation tree.
+	 *	Physics Only      : Only used only for physics simulation (rigid body, constraints). Cannot be used for spatial queries (raycasts, sweeps, overlaps). Useful for jiggly bits on characters that do not need per bone detection. Performance gains by keeping data out of query tree
+	 *	Collision Enabled : Can be used for both spatial queries (raycasts, sweeps, overlaps) and simulation (rigid body, constraints).
+	 */
+	UPROPERTY(EditAnywhere, Category=Custom)
+	TEnumAsByte<ECollisionEnabled::Type> CollisionEnabled;
 
 	/** Collision Profile Name **/
 	UPROPERTY(EditAnywhere, Category=Custom)
@@ -237,11 +257,11 @@ public:
 
 	/** If true Continuous Collision Detection (CCD) will be used for this component */
 	UPROPERTY(EditAnywhere, AdvancedDisplay, BlueprintReadOnly, Category = Collision)
-	uint32 bUseCCD : 1;
+	uint8 bUseCCD : 1;
 
 	/**	Should 'Hit' events fire when this object collides during physics simulation. */
 	UPROPERTY(EditAnywhere, BlueprintReadOnly, Category = Collision, meta = (DisplayName = "Simulation Generates Hit Events"))
-	uint32 bNotifyRigidBodyCollision : 1;
+	uint8 bNotifyRigidBodyCollision : 1;
 
 	/////////
 	// SIM SETTINGS
@@ -252,88 +272,88 @@ public:
 	 * For a Static Mesh Component, simulating requires simple collision to be setup on the StaticMesh asset.
 	 */
 	UPROPERTY(EditAnywhere, BlueprintReadWrite, Category = Physics)
-	uint32 bSimulatePhysics : 1;
+	uint8 bSimulatePhysics : 1;
 
 	/** If true, mass will not be automatically computed and you must set it directly */
 	UPROPERTY(EditAnywhere, Category = Physics, meta = (InlineEditConditionToggle))
-	uint32 bOverrideMass : 1;
+	uint8 bOverrideMass : 1;
 
 	/** If object should have the force of gravity applied */
 	UPROPERTY(EditAnywhere, BlueprintReadOnly, Category = Physics)
-	uint32 bEnableGravity : 1;
+	uint8 bEnableGravity : 1;
 
 	/** If true and is attached to a parent, the two bodies will be joined into a single rigid body. Physical settings like collision profile and body settings are determined by the root */
 	UPROPERTY(EditAnywhere, AdvancedDisplay, BlueprintReadWrite, Category = Physics, meta = (editcondition = "!bSimulatePhysics"))
-	uint32 bAutoWeld : 1;
+	uint8 bAutoWeld : 1;
 
 	/** If object should start awake, or if it should initially be sleeping */
 	UPROPERTY(EditAnywhere, AdvancedDisplay, BlueprintReadOnly, Category = Physics, meta = (editcondition = "bSimulatePhysics"))
-	uint32 bStartAwake:1;
+	uint8 bStartAwake:1;
 
 	/**	Should 'wake/sleep' events fire when this object is woken up or put to sleep by the physics simulation. */
 	UPROPERTY(EditAnywhere,AdvancedDisplay, BlueprintReadOnly, Category = Physics)
-	uint32 bGenerateWakeEvents : 1;
+	uint8 bGenerateWakeEvents : 1;
 
 	/** If true, it will update mass when scale changes **/
 	UPROPERTY()
-	uint32 bUpdateMassWhenScaleChanges:1;
+	uint8 bUpdateMassWhenScaleChanges:1;
 
 	/** When a Locked Axis Mode is selected, will lock translation on the specified axis*/
 	UPROPERTY(EditAnywhere, Category = Physics, meta=(DisplayName = "Lock Axis Translation"))
-	uint32 bLockTranslation : 1;
+	uint8 bLockTranslation : 1;
 	
 	/** When a Locked Axis Mode is selected, will lock rotation to the specified axis*/
 	UPROPERTY(EditAnywhere, Category = Physics, meta=(DisplayName = "Lock Axis Rotation"))
-	uint32 bLockRotation : 1;
+	uint8 bLockRotation : 1;
 
 	/** Lock translation along the X-axis*/
 	UPROPERTY(EditAnywhere, Category = Physics, meta = (DisplayName = "X"))
-	uint32 bLockXTranslation : 1;
+	uint8 bLockXTranslation : 1;
 
 	/** Lock translation along the Y-axis*/
 	UPROPERTY(EditAnywhere, Category = Physics, meta = (DisplayName = "Y"))
-	uint32 bLockYTranslation : 1;
+	uint8 bLockYTranslation : 1;
 
 	/** Lock translation along the Z-axis*/
 	UPROPERTY(EditAnywhere, Category = Physics, meta = (DisplayName = "Z"))
-	uint32 bLockZTranslation : 1;
+	uint8 bLockZTranslation : 1;
 
 	/** Lock rotation about the X-axis*/
 	UPROPERTY(EditAnywhere, Category = Physics, meta = (DisplayName = "X"))
-	uint32 bLockXRotation : 1;
+	uint8 bLockXRotation : 1;
 
 	/** Lock rotation about the Y-axis*/
 	UPROPERTY(EditAnywhere, Category = Physics, meta = (DisplayName = "Y"))
-	uint32 bLockYRotation : 1;
+	uint8 bLockYRotation : 1;
 
 	/** Lock rotation about the Z-axis*/
 	UPROPERTY(EditAnywhere, Category = Physics, meta = (DisplayName = "Z"))
-	uint32 bLockZRotation : 1;
+	uint8 bLockZRotation : 1;
 
 	/** Override the default max angular velocity */
 	UPROPERTY(EditAnywhere, Category = Physics, meta = (editcondition = "bSimulatePhysics", InlineEditConditionToggle))
-	uint32 bOverrideMaxAngularVelocity : 1;
+	uint8 bOverrideMaxAngularVelocity : 1;
 
 	/** When initializing dynamic instances their component or velocity can override the bStartAwake flag */
-	uint32 bWokenExternally : 1;
+	uint8 bWokenExternally : 1;
 
 	/**
 	* If true, this body will be put into the asynchronous physics scene. If false, it will be put into the synchronous physics scene.
 	* If the body is static, it will be placed into both scenes regardless of the value of bUseAsyncScene.
 	*/
 	UPROPERTY(EditAnywhere, AdvancedDisplay, BlueprintReadWrite, Category = Physics)
-	uint32 bUseAsyncScene : 1;
+	uint8 bUseAsyncScene : 1;
 protected:
 
 	/** Whether this body instance has its own custom MaxDepenetrationVelocity*/
 	UPROPERTY(EditAnywhere, Category = Physics, meta=(InlineEditConditionToggle))
-	uint32 bOverrideMaxDepenetrationVelocity : 1;
+	uint8 bOverrideMaxDepenetrationVelocity : 1;
 
 	/** Whether this instance of the object has its own custom walkable slope override setting. */
 	UPROPERTY(EditAnywhere, Category = Physics, meta = (InlineEditConditionToggle))
-	uint32 bOverrideWalkableSlopeOnInstance : 1;
+	uint8 bOverrideWalkableSlopeOnInstance : 1;
 
-	uint32 bHasSharedShapes : 1;
+	uint8 bHasSharedShapes : 1;
 
 	/** The maximum velocity used to depenetrate this object*/
 	UPROPERTY(EditAnywhere, AdvancedDisplay, BlueprintReadOnly, Category = Physics, meta = (editcondition = "bOverrideMaxDepenetrationVelocity", ClampMin = "0.0", UIMin = "0.0"))
@@ -380,6 +400,12 @@ public:
 	UPROPERTY(EditAnywhere, AdvancedDisplay, BlueprintReadWrite, Category = Physics)
 	FVector InertiaTensorScale;
 
+private:
+	/** Enum indicating what type of object this should be considered as when it moves */
+	UPROPERTY(EditAnywhere, Category=Custom)
+	TEnumAsByte<enum ECollisionChannel> ObjectType;
+
+public:
 	/** Use the collision profile found in the given BodySetup's default BodyInstance */
 	void UseExternalCollisionProfile(UBodySetup* InExternalCollisionProfileBodySetup);
 
@@ -394,7 +420,7 @@ public:
 	static EDOFMode::Type ResolveDOFMode(EDOFMode::Type DOFMode);
 
 	/** Constraint used to allow for easy DOF setup per bodyinstance */
-	FConstraintInstance * DOFConstraint;
+	FConstraintInstance* DOFConstraint;
 
 	/** The parent body that we are welded to*/
 	FBodyInstance* WeldParent;
@@ -1155,34 +1181,6 @@ private:
 
 #endif
 
-public:
-	// Current state of the physx body for tracking deferred addition and removal.
-	BodyInstanceSceneState CurrentSceneState;
-
-	/** The set of values used in considering when put this body to sleep. */
-	UPROPERTY(EditAnywhere, AdvancedDisplay, BlueprintReadWrite, Category = Physics)
-	ESleepFamily SleepFamily;
-
-	/** Locks physical movement along specified axis.*/
-	UPROPERTY(EditAnywhere, Category = Physics, meta = (DisplayName = "Mode"))
-	TEnumAsByte<EDOFMode::Type> DOFMode;
-
-private:
-	/**
-	 * Type of collision enabled.
-	 * 
-	 *	No Collision      : Will not create any representation in the physics engine. Cannot be used for spatial queries (raycasts, sweeps, overlaps) or simulation (rigid body, constraints). Best performance possible (especially for moving objects)
-	 *	Query Only        : Only used for spatial queries (raycasts, sweeps, and overlaps). Cannot be used for simulation (rigid body, constraints). Useful for character movement and things that do not need physical simulation. Performance gains by keeping data out of simulation tree.
-	 *	Physics Only      : Only used only for physics simulation (rigid body, constraints). Cannot be used for spatial queries (raycasts, sweeps, overlaps). Useful for jiggly bits on characters that do not need per bone detection. Performance gains by keeping data out of query tree
-	 *	Collision Enabled : Can be used for both spatial queries (raycasts, sweeps, overlaps) and simulation (rigid body, constraints).
-	 */
-	UPROPERTY(EditAnywhere, Category=Custom)
-	TEnumAsByte<ECollisionEnabled::Type> CollisionEnabled;
-
-	/** Enum indicating what type of object this should be considered as when it moves */
-	UPROPERTY(EditAnywhere, Category=Custom)
-	TEnumAsByte<enum ECollisionChannel> ObjectType;
-
 	void SetShapeFlagsInternal_AssumesShapeLocked(struct FSetShapeParams& Params, bool& bUpdateMassProperties);
 };
 
@@ -1239,4 +1237,4 @@ FORCEINLINE_DEBUGGABLE bool FBodyInstance::IsInstanceSimulatingPhysics() const
 }
 
 extern template ENGINE_API bool FBodyInstance::OverlapTestForBodiesImpl(const FVector& Position, const FQuat& Rotation, const TArray<FBodyInstance*>& Bodies) const;
-extern template ENGINE_API bool FBodyInstance::OverlapTestForBodiesImpl(const FVector& Position, const FQuat& Rotation, const TArray<FBodyInstance*, TInlineAllocator<1>>& Bodies) const;
+extern template ENGINE_API bool FBodyInstance::OverlapTestForBodiesImpl(const FVector& Position, const FQuat& Rotation, const TArray<FBodyInstance*, TInlineAllocator<1>>& Bodies) const;
\ No newline at end of file
diff --git a/Engine/Source/Runtime/Engine/Private/Actor.cpp b/Engine/Source/Runtime/Engine/Private/Actor.cpp
--- a/Engine/Source/Runtime/Engine/Private/Actor.cpp
+++ b/Engine/Source/Runtime/Engine/Private/Actor.cpp
@@ -102,6 +102,7 @@ void AActor::InitializeDefaults()
 	bActorLabelEditable = true;
 	SpriteScale = 1.0f;
 	bEnableAutoLODGeneration = true;	
+	InputConsumeOption_DEPRECATED = ICO_ConsumeBoundKeys;
 #endif // WITH_EDITORONLY_DATA
 	NetCullDistanceSquared = 225000000.0f;
 	NetDriverName = NAME_GameNetDriver;
@@ -109,7 +110,6 @@ void AActor::InitializeDefaults()
 	// will be updated in PostInitProperties
 	bActorEnableCollision = true;
 	bActorSeamlessTraveled = false;
-	InputConsumeOption_DEPRECATED = ICO_ConsumeBoundKeys;
 	bBlockInput = false;
 	bCanBeDamaged = true;
 	bFindCameraComponentWhenViewTarget = true;
@@ -538,7 +538,7 @@ void AActor::Serialize(FArchive& Ar)
 	// we'll let that drop
 	if (Ar.GetPortFlags() & PPF_DuplicateForPIE)
 	{
-		TArray<UActorComponent*> DuplicatingComponents;
+		TInlineComponentArray<UActorComponent*> DuplicatingComponents;
 		if (Ar.IsSaving())
 		{
 			DuplicatingComponents.Reserve(OwnedComponents.Num());
@@ -568,11 +568,6 @@ void AActor::PostLoad()
 		Owner->Children.Add(this);
 	}
 
-	if (GetLinkerUE4Version() < VER_UE4_CONSUME_INPUT_PER_BIND)
-	{
-		bBlockInput = (InputConsumeOption_DEPRECATED == ICO_ConsumeAll);
-	}
-
 	if (GetLinkerUE4Version() < VER_UE4_PRIVATE_REMOTE_ROLE)
 	{
 		bReplicates = (RemoteRole != ROLE_None);
@@ -584,6 +579,12 @@ void AActor::PostLoad()
 		bExchangedRoles = false;
 	}
 
+#if WITH_EDITORONLY_DATA
+	if (GetLinkerUE4Version() < VER_UE4_CONSUME_INPUT_PER_BIND)
+	{
+		bBlockInput = (InputConsumeOption_DEPRECATED == ICO_ConsumeAll);
+	}
+
 	if (AActor* ParentActor = ParentComponentActor_DEPRECATED.Get())
 	{
 		TInlineComponentArray<UChildActorComponent*> ParentChildActorComponents(ParentActor);
@@ -599,7 +600,6 @@ void AActor::PostLoad()
 
 	if ( GIsEditor )
 	{
-#if WITH_EDITORONLY_DATA
 		// Propagate the hidden at editor startup flag to the transient hidden flag
 		bHiddenEdTemporary = bHiddenEd;
 
@@ -608,8 +608,8 @@ void AActor::PostLoad()
 		{
 			UE_LOG(LogActor, Log,  TEXT("Loaded Actor (%s) with IsPendingKill() == true"), *GetName() );
 		}
-#endif // WITH_EDITORONLY_DATA
 	}
+#endif // WITH_EDITORONLY_DATA
 }
 
 void AActor::PostLoadSubobjects(FObjectInstancingGraph* OuterInstanceGraph)
@@ -2524,8 +2524,8 @@ void AActor::AddOwnedComponent(UActorComponent* Component)
 
 void AActor::RemoveOwnedComponent(UActorComponent* Component)
 {
-	// Note: we do not mark dirty here because this can be called when in editor when modifying transient components
-	// if a component is added during this time it should not dirty.  Higher level code in the editor should always dirty the package anyway
+	// Note: we do not mark dirty here because this can be called as part of component duplication when reinstancing components during blueprint compilation
+	// if a component is removed during this time it should not dirty.  Higher level code in the editor should always dirty the package anyway.
 	const bool bMarkDirty = false;
 	Modify(bMarkDirty);
 
@@ -3637,8 +3637,10 @@ void AActor::K2_DestroyComponent(UActorComponent* Component)
 bool AActor::SetRootComponent(class USceneComponent* NewRootComponent)
 {
 	/** Only components owned by this actor can be used as a its root component. */
-	if (ensure(NewRootComponent == NULL || NewRootComponent->GetOwner() == this))
+	if (ensure(NewRootComponent == nullptr || NewRootComponent->GetOwner() == this))
 	{
+		Modify();
+
 		RootComponent = NewRootComponent;
 		return true;
 	}
diff --git a/Engine/Source/Runtime/Engine/Private/ActorEditor.cpp b/Engine/Source/Runtime/Engine/Private/ActorEditor.cpp
--- a/Engine/Source/Runtime/Engine/Private/ActorEditor.cpp
+++ b/Engine/Source/Runtime/Engine/Private/ActorEditor.cpp
@@ -377,7 +377,7 @@ void AActor::PreEditUndo()
 	Super::PreEditUndo();
 }
 
-void AActor::PostEditUndo()
+bool AActor::InternalPostEditUndo()
 {
 	// Check if this Actor needs to be re-instanced
 	UClass* OldClass = GetClass();
@@ -390,7 +390,7 @@ void AActor::PostEditUndo()
 		};
 
 		// Early exit, letting anything more occur would be invalid due to the REINST_ class
-		return;
+		return false;
 	}
 
 	// Notify LevelBounds actor that level bounding box might be changed
@@ -406,6 +406,10 @@ void AActor::PostEditUndo()
 	if (!IsPendingKill())
 	{
 		ResetOwnedComponents();
+
+		// BP created components are not serialized, so this should be cleared and will be filled in as the construction scripts are run
+		BlueprintCreatedComponents.Reset();
+
 		// notify navigation system
 		UNavigationSystem::UpdateActorAndComponentsInNavOctree(*this);
 	}
@@ -414,50 +418,26 @@ void AActor::PostEditUndo()
 		UNavigationSystem::ClearNavOctreeAll(this);
 	}
 
-	Super::PostEditUndo();
+	// This is a normal undo, so call super
+	return true;
 }
 
-void AActor::PostEditUndo(TSharedPtr<ITransactionObjectAnnotation> TransactionAnnotation)
+void AActor::PostEditUndo()
 {
-	CurrentTransactionAnnotation = StaticCastSharedPtr<FActorTransactionAnnotation>(TransactionAnnotation);
-
-	// Check if this Actor needs to be re-instanced
-	UClass* OldClass = GetClass();
-	if (OldClass->HasAnyClassFlags(CLASS_NewerVersionExists))
-	{
-		UClass* NewClass = OldClass->GetAuthoritativeClass();
-		if (!ensure(NewClass != OldClass))
-		{
-			UE_LOG(LogActor, Warning, TEXT("WARNING: %s is out of date and is the same as its AuthoritativeClass during PostEditUndo!"), *OldClass->GetName());
-		};
-
-		// Early exit, letting anything more occur would be invalid due to the REINST_ class
-		return;
-	}
-
-	// Notify LevelBounds actor that level bounding box might be changed
-	if (!IsTemplate())
+	if (InternalPostEditUndo())
 	{
-		GetLevel()->MarkLevelBoundsDirty();
+		Super::PostEditUndo();
 	}
+}
 
-	// Restore OwnedComponents array
-	if (!IsPendingKill())
-	{
-		ResetOwnedComponents();
-
-		// BP created components are not serialized, so this should be cleared and will be filled in as the construction scripts are run
-		BlueprintCreatedComponents.Reset();
+void AActor::PostEditUndo(TSharedPtr<ITransactionObjectAnnotation> TransactionAnnotation)
+{
+	CurrentTransactionAnnotation = StaticCastSharedPtr<FActorTransactionAnnotation>(TransactionAnnotation);
 
-		// notify navigation system
-		UNavigationSystem::UpdateActorAndComponentsInNavOctree(*this);
-	}
-	else
+	if (InternalPostEditUndo())
 	{
-		UNavigationSystem::ClearNavOctreeAll(this);
+		Super::PostEditUndo(TransactionAnnotation);
 	}
-
-	Super::PostEditUndo(TransactionAnnotation);
 }
 
 // @todo: Remove this hack once we have decided on the scaling method to use.
diff --git a/Engine/Source/Runtime/Engine/Private/Animation/AnimNotify_PlayParticleEffect.cpp b/Engine/Source/Runtime/Engine/Private/Animation/AnimNotify_PlayParticleEffect.cpp
--- a/Engine/Source/Runtime/Engine/Private/Animation/AnimNotify_PlayParticleEffect.cpp
+++ b/Engine/Source/Runtime/Engine/Private/Animation/AnimNotify_PlayParticleEffect.cpp
@@ -14,6 +14,7 @@ UAnimNotify_PlayParticleEffect::UAnimNotify_PlayParticleEffect()
 	: Super()
 {
 	Attached = true;
+	Scale = FVector(1.f);
 
 #if WITH_EDITORONLY_DATA
 	NotifyColor = FColor(192, 255, 99, 255);
@@ -52,7 +53,7 @@ void UAnimNotify_PlayParticleEffect::Notify(class USkeletalMeshComponent* MeshCo
 
 		if (Attached)
 		{
-			UGameplayStatics::SpawnEmitterAttached(PSTemplate, MeshComp, SocketName, LocationOffset, RotationOffset);
+			UGameplayStatics::SpawnEmitterAttached(PSTemplate, MeshComp, SocketName, LocationOffset, RotationOffset, Scale);
 		}
 		else
 		{
@@ -60,6 +61,7 @@ void UAnimNotify_PlayParticleEffect::Notify(class USkeletalMeshComponent* MeshCo
 			FTransform SpawnTransform;
 			SpawnTransform.SetLocation(MeshTransform.TransformPosition(LocationOffset));
 			SpawnTransform.SetRotation(MeshTransform.GetRotation() * RotationOffsetQuat);
+			SpawnTransform.SetScale3D(Scale);
 			UGameplayStatics::SpawnEmitterAtLocation(MeshComp->GetWorld(), PSTemplate, SpawnTransform);
 		}
 	}
diff --git a/Engine/Source/Runtime/Engine/Private/AssetManager.cpp b/Engine/Source/Runtime/Engine/Private/AssetManager.cpp
--- a/Engine/Source/Runtime/Engine/Private/AssetManager.cpp
+++ b/Engine/Source/Runtime/Engine/Private/AssetManager.cpp
@@ -3106,7 +3106,7 @@ void UAssetManager::InitializeAssetBundlesFromMetadata(const UStruct* Struct, co
 				{
 					TSet<FName> LocalBundleSet;
 					TArray<FString> BundleList;
-					FString BundleString = PropertyToSearch->GetMetaData(AssetBundlesName);
+					const FString& BundleString = PropertyToSearch->GetMetaData(AssetBundlesName);
 					BundleString.ParseIntoArrayWS(BundleList, TEXT(","));
 
 					for (const FString& BundleNameString : BundleList)
diff --git a/Engine/Source/Runtime/Engine/Private/Character.cpp b/Engine/Source/Runtime/Engine/Private/Character.cpp
--- a/Engine/Source/Runtime/Engine/Private/Character.cpp
+++ b/Engine/Source/Runtime/Engine/Private/Character.cpp
@@ -421,7 +421,7 @@ void ACharacter::OnStartCrouch( float HeightAdjust, float ScaledHeightAdjust )
 	}
 	else
 	{
-		BaseTranslationOffset.Z = GetDefault<ACharacter>(GetClass())->BaseTranslationOffset.Z + HeightAdjust;
+		BaseTranslationOffset.Z = DefaultChar->BaseTranslationOffset.Z + HeightAdjust;
 	}
 
 	K2_OnStartCrouch(HeightAdjust, ScaledHeightAdjust);
diff --git a/Engine/Source/Runtime/Engine/Private/Collision/PhysXCollision.h b/Engine/Source/Runtime/Engine/Private/Collision/PhysXCollision.h
--- a/Engine/Source/Runtime/Engine/Private/Collision/PhysXCollision.h
+++ b/Engine/Source/Runtime/Engine/Private/Collision/PhysXCollision.h
@@ -191,7 +191,7 @@ public:
 		// Perform other kinds of zero-extent queries as zero-extent sphere queries
 		if ((CollisionShape.ShapeType != ECollisionShape::Sphere) && CollisionShape.IsNearlyZero())
 		{
-			PtrToUnionData = UnionData.SetSubtype<PxSphereGeometry>(PxSphereGeometry(KINDA_SMALL_NUMBER));
+			PtrToUnionData = UnionData.SetSubtype<PxSphereGeometry>(PxSphereGeometry(FCollisionShape::MinSphereRadius()));
 		}
 		else
 		{
@@ -200,23 +200,35 @@ public:
 			case ECollisionShape::Box:
 				{
 					PxVec3 BoxExtents = U2PVector(CollisionShape.GetBox());
-					BoxExtents.x = FMath::Max(BoxExtents.x, KINDA_SMALL_NUMBER);
-					BoxExtents.y = FMath::Max(BoxExtents.y, KINDA_SMALL_NUMBER);
-					BoxExtents.z = FMath::Max(BoxExtents.z, KINDA_SMALL_NUMBER);
+					BoxExtents.x = FMath::Max(BoxExtents.x, FCollisionShape::MinBoxExtent());
+					BoxExtents.y = FMath::Max(BoxExtents.y, FCollisionShape::MinBoxExtent());
+					BoxExtents.z = FMath::Max(BoxExtents.z, FCollisionShape::MinBoxExtent());
 
 					PtrToUnionData = UnionData.SetSubtype<PxBoxGeometry>(PxBoxGeometry(BoxExtents));
 					Rotation = U2PQuat(Rot);
 					break;
 				}
 			case ECollisionShape::Sphere:
-				PtrToUnionData = UnionData.SetSubtype<PxSphereGeometry>(PxSphereGeometry(FMath::Max(CollisionShape.GetSphereRadius(), KINDA_SMALL_NUMBER)));
+				PtrToUnionData = UnionData.SetSubtype<PxSphereGeometry>(PxSphereGeometry(FMath::Max(CollisionShape.GetSphereRadius(), FCollisionShape::MinSphereRadius())));
 				break;
 			case ECollisionShape::Capsule:
-				PtrToUnionData = UnionData.SetSubtype<PxCapsuleGeometry>(PxCapsuleGeometry(FMath::Max(CollisionShape.GetCapsuleRadius(), KINDA_SMALL_NUMBER), FMath::Max(CollisionShape.GetCapsuleAxisHalfLength(), KINDA_SMALL_NUMBER)));
-				Rotation = ConvertToPhysXCapsuleRot(Rot);
-				break;
+				{
+					const float CapsuleRadius = CollisionShape.GetCapsuleRadius();
+					const float CapsuleHalfHeight = CollisionShape.GetCapsuleHalfHeight();
+					if (CapsuleRadius < CapsuleHalfHeight)
+					{
+						PtrToUnionData = UnionData.SetSubtype<PxCapsuleGeometry>(PxCapsuleGeometry(FMath::Max(CapsuleRadius, FCollisionShape::MinCapsuleRadius()), FMath::Max(CollisionShape.GetCapsuleAxisHalfLength(), FCollisionShape::MinCapsuleAxisHalfHeight())));
+						Rotation = ConvertToPhysXCapsuleRot(Rot);
+					}
+					else
+					{
+						// Use a sphere instead.
+						PtrToUnionData = UnionData.SetSubtype<PxSphereGeometry>(PxSphereGeometry(FMath::Max(CapsuleRadius, FCollisionShape::MinSphereRadius())));
+					}					
+					break;
+				}
 			default:
-				// invalid point
+				// invalid type
 				ensure(false);
 			}
 		}
diff --git a/Engine/Source/Runtime/Engine/Private/Components/ActorComponent.cpp b/Engine/Source/Runtime/Engine/Private/Components/ActorComponent.cpp
--- a/Engine/Source/Runtime/Engine/Private/Components/ActorComponent.cpp
+++ b/Engine/Source/Runtime/Engine/Private/Components/ActorComponent.cpp
@@ -1,4 +1,4 @@
-// Copyright 1998-2017 Epic Games, Inc. All Rights Reserved.
+﻿// Copyright 1998-2017 Epic Games, Inc. All Rights Reserved.
 // ActorComponent.cpp: Actor component implementation.
 
 #include "Components/ActorComponent.h"
@@ -203,6 +203,7 @@ void UActorComponent::PostLoad()
 {
 	Super::PostLoad();
 
+#if WITH_EDITORONLY_DATA
 	if (GetLinkerUE4Version() < VER_UE4_ACTOR_COMPONENT_CREATION_METHOD)
 	{
 		if (IsTemplate())
@@ -239,6 +240,7 @@ void UActorComponent::PostLoad()
 			}
 		}
 	}
+#endif
 
 	if (CreationMethod == EComponentCreationMethod::SimpleConstructionScript)
 	{
@@ -475,13 +477,17 @@ void UActorComponent::BeginDestroy()
 bool UActorComponent::NeedsLoadForClient() const
 {
 	check(GetOuter());
-	return (!IsEditorOnly() && GetOuter()->NeedsLoadForClient() && Super::NeedsLoadForClient());
+	// For Component Blueprints, avoid calling into the class to avoid recursion
+	bool bNeedsLoadOuter = HasAnyFlags(RF_ClassDefaultObject) || GetOuter()->NeedsLoadForClient();
+	return (!IsEditorOnly() && bNeedsLoadOuter && Super::NeedsLoadForClient());
 }
 
 bool UActorComponent::NeedsLoadForServer() const
 {
 	check(GetOuter());
-	return (!IsEditorOnly() && GetOuter()->NeedsLoadForServer() && Super::NeedsLoadForServer());
+	// For Component Blueprints, avoid calling into the class to avoid recursion
+	bool bNeedsLoadOuter = HasAnyFlags(RF_ClassDefaultObject) || GetOuter()->NeedsLoadForServer();
+	return (!IsEditorOnly() && bNeedsLoadOuter && Super::NeedsLoadForServer());
 }
 
 int32 UActorComponent::GetFunctionCallspace( UFunction* Function, void* Parameters, FFrame* Stack )
@@ -1201,7 +1207,6 @@ void UActorComponent::CreatePhysicsState()
 
 		// Broadcast delegate
 		GlobalCreatePhysicsDelegate.Broadcast(this);
-		InstanceCreatePhysicsDelegate.Broadcast();
 	}
 }
 
@@ -1213,7 +1218,6 @@ void UActorComponent::DestroyPhysicsState()
 	{
 		// Broadcast delegate
 		GlobalDestroyPhysicsDelegate.Broadcast(this);
-		InstanceDestroyPhysicsDelegate.Broadcast();
 
 		ensureMsgf(bRegistered, TEXT("Component has physics state when not registered (%s)"), *GetFullName()); // should not have physics state unless we are registered
 
diff --git a/Engine/Source/Runtime/Engine/Private/Components/ChildActorComponent.cpp b/Engine/Source/Runtime/Engine/Private/Components/ChildActorComponent.cpp
--- a/Engine/Source/Runtime/Engine/Private/Components/ChildActorComponent.cpp
+++ b/Engine/Source/Runtime/Engine/Private/Components/ChildActorComponent.cpp
@@ -532,6 +532,10 @@ void UChildActorComponent::CreateChildActor()
 				{
 					Params.ObjectFlags &= ~RF_Transactional;
 				}
+				if (HasAllFlags(RF_Transient))
+				{
+					Params.ObjectFlags |= RF_Transient;
+				}
 
 				// Spawn actor of desired class
 				ConditionalUpdateComponentToWorld();
diff --git a/Engine/Source/Runtime/Engine/Private/Components/ForceFeedbackComponent.cpp b/Engine/Source/Runtime/Engine/Private/Components/ForceFeedbackComponent.cpp
--- a/Engine/Source/Runtime/Engine/Private/Components/ForceFeedbackComponent.cpp
+++ b/Engine/Source/Runtime/Engine/Private/Components/ForceFeedbackComponent.cpp
@@ -1,4 +1,4 @@
-// Copyright 1998-2017 Epic Games, Inc. All Rights Reserved.
+﻿// Copyright 1998-2017 Epic Games, Inc. All Rights Reserved.
 
 #include "Components/ForceFeedbackComponent.h"
 #include "Components/BillboardComponent.h"
@@ -355,7 +355,7 @@ bool UForceFeedbackComponent::Advance(const float DeltaTime)
 
 	const float Duration = ForceFeedbackEffect->GetDuration();
 
-	PlayTime += DeltaTime;
+	PlayTime += (bIgnoreTimeDilation ? FApp::GetDeltaTime() : DeltaTime);
 
 	if (PlayTime > Duration && (!bLooping || Duration == 0.f) )
 	{
diff --git a/Engine/Source/Runtime/Engine/Private/Components/MeshComponent.cpp b/Engine/Source/Runtime/Engine/Private/Components/MeshComponent.cpp
--- a/Engine/Source/Runtime/Engine/Private/Components/MeshComponent.cpp
+++ b/Engine/Source/Runtime/Engine/Private/Components/MeshComponent.cpp
@@ -233,10 +233,9 @@ void UMeshComponent::SetScalarParameterValueOnMaterials(const FName ParameterNam
 	}
 
 	// Look up material index array according to ParameterName
-	TArray<int32>* MaterialIndicesPtr = ScalarParameterMaterialIndices.Find(ParameterName);
-	if (MaterialIndicesPtr != nullptr)
+	if (FMaterialParameterCache* ParameterCache = MaterialParameterCache.Find(ParameterName))
 	{
-		const TArray<int32>& MaterialIndices = *MaterialIndicesPtr;
+		const TArray<int32>& MaterialIndices = ParameterCache->ScalarParameterMaterialIndices;
 		// Loop over all the material indices and update set the parameter value on the corresponding materials		
 		for ( int32 MaterialIndex : MaterialIndices)
 		{
@@ -266,10 +265,9 @@ void UMeshComponent::SetVectorParameterValueOnMaterials(const FName ParameterNam
 	}
 
 	// Look up material index array according to ParameterName
-	TArray<int32>* MaterialIndicesPtr = VectorParameterMaterialIndices.Find(ParameterName);
-	if (MaterialIndicesPtr != nullptr)
+	if (FMaterialParameterCache* ParameterCache = MaterialParameterCache.Find(ParameterName))
 	{
-		const TArray<int32>& MaterialIndices = *MaterialIndicesPtr;
+		const TArray<int32>& MaterialIndices = ParameterCache->VectorParameterMaterialIndices;
 		// Loop over all the material indices and update set the parameter value on the corresponding materials		
 		for ( int32 MaterialIndex : MaterialIndices )
 		{
@@ -296,8 +294,7 @@ void UMeshComponent::MarkCachedMaterialParameterNameIndicesDirty()
 void UMeshComponent::CacheMaterialParameterNameIndices()
 {
 	// Clean up possible previous data
-	ScalarParameterMaterialIndices.Reset();
-	VectorParameterMaterialIndices.Reset();
+	MaterialParameterCache.Reset();
 
 	// not sure if this is the best way to do this
 	const UWorld* World = GetWorld();
@@ -321,17 +318,16 @@ void UMeshComponent::CacheMaterialParameterNameIndices()
 			Material->GetAllScalarParameterNames(OutParameterNames, OutParameterIds);
 			for (FName& ParameterName : OutParameterNames)
 			{
-				// Add or retrieve TMap entry for this scalar parameter name
-				TArray<int32>& MaterialIndices = ScalarParameterMaterialIndices.FindOrAdd(ParameterName);
+				// Add or retrieve entry for this parameter name
+				FMaterialParameterCache& ParameterCache = MaterialParameterCache.FindOrAdd(ParameterName);
 				// Add the corresponding material index
-				MaterialIndices.Add(MaterialIndex);
+				ParameterCache.ScalarParameterMaterialIndices.Add(MaterialIndex);
 				
 				// GetScalarParameterDefault() expects to use a FMaterialResource, which means the world has to be rendering
 				if (bHasMaterialResource)
 				{
-					// set default value to it
-					float& DefaultValue = ScalarParameterDefaultValues.FindOrAdd(ParameterName);
-					DefaultValue = Material->GetScalarParameterDefault(ParameterName, FeatureLevel);
+					// store the default value
+					ParameterCache.ScalarParameterDefaultValue = Material->GetScalarParameterDefault(ParameterName, FeatureLevel);
 				}
 			}
 
@@ -343,10 +339,10 @@ void UMeshComponent::CacheMaterialParameterNameIndices()
 			Material->GetAllVectorParameterNames(OutParameterNames, OutParameterIds);
 			for (FName& ParameterName : OutParameterNames)
 			{
-				// Add or retrieve TMap entry for this vector parameter name
-				TArray<int32>& MaterialIndices = VectorParameterMaterialIndices.FindOrAdd(ParameterName);
+				// Add or retrieve entry for this parameter name
+				FMaterialParameterCache& ParameterCache = MaterialParameterCache.FindOrAdd(ParameterName);
 				// Add the corresponding material index
-				MaterialIndices.Add(MaterialIndex);
+				ParameterCache.VectorParameterMaterialIndices.Add(MaterialIndex);
 			}
 		}
 		++MaterialIndex;
diff --git a/Engine/Source/Runtime/Engine/Private/Components/SceneComponent.cpp b/Engine/Source/Runtime/Engine/Private/Components/SceneComponent.cpp
--- a/Engine/Source/Runtime/Engine/Private/Components/SceneComponent.cpp
+++ b/Engine/Source/Runtime/Engine/Private/Components/SceneComponent.cpp
@@ -505,19 +505,19 @@ void USceneComponent::UpdateComponentToWorldWithParent(USceneComponent* Parent,F
 #endif
 
 	// If our parent hasn't been updated before, we'll need walk up our parent attach hierarchy
-	if (Parent && !Parent->bWorldToComponentUpdated)
+	if (Parent && !Parent->bComponentToWorldUpdated)
 	{
 		//QUICK_SCOPE_CYCLE_COUNTER(STAT_USceneComponent_UpdateComponentToWorldWithParent_Parent);
 		Parent->UpdateComponentToWorld();
 
 		// Updating the parent may (depending on if we were already attached to parent) result in our being updated, so just return
-		if (bWorldToComponentUpdated)
+		if (bComponentToWorldUpdated)
 		{
 			return;
 		}
 	}
 
-	bWorldToComponentUpdated = true;
+	bComponentToWorldUpdated = true;
 
 	FTransform NewTransform(NoInit);
 
@@ -548,9 +548,7 @@ void USceneComponent::UpdateComponentToWorldWithParent(USceneComponent* Parent,F
 	{
 		//QUICK_SCOPE_CYCLE_COUNTER(STAT_USceneComponent_UpdateComponentToWorldWithParent_Changed);
 		// Update transform
-		PRAGMA_DISABLE_DEPRECATION_WARNINGS
 		ComponentToWorld = NewTransform;
-		PRAGMA_ENABLE_DEPRECATION_WARNINGS
 		PropagateTransformUpdate(true, UpdateTransformFlags, Teleport);
 	}
 	else
@@ -1845,7 +1843,7 @@ bool USceneComponent::AttachToComponent(USceneComponent* Parent, const FAttachme
 
 		if (Parent->IsNetSimulating() && !IsNetSimulating())
 		{
-			Parent->ClientAttachedChildren.Add(this);
+			Parent->ClientAttachedChildren.AddUnique(this);
 		}
 
 		// Now apply attachment rules
@@ -2195,7 +2193,7 @@ void USceneComponent::UpdateChildTransforms(EUpdateTransformFlags UpdateTransfor
 		if (ChildComp != nullptr)
 		{
 			// Don't update the child if it uses a completely absolute (world-relative) scheme, unless it has never been updated.
-			if (!(ChildComp->bAbsoluteLocation && ChildComp->bAbsoluteRotation && ChildComp->bAbsoluteScale) || !ChildComp->bWorldToComponentUpdated)
+			if (!(ChildComp->bAbsoluteLocation && ChildComp->bAbsoluteRotation && ChildComp->bAbsoluteScale) || !ChildComp->bComponentToWorldUpdated)
 			{
 				ChildComp->UpdateComponentToWorld(UpdateTransformFlagsFromParent, Teleport);
 			}
@@ -2203,6 +2201,7 @@ void USceneComponent::UpdateChildTransforms(EUpdateTransformFlags UpdateTransfor
 	}
 }
 
+#if WITH_EDITORONLY_DATA
 void USceneComponent::Serialize(FArchive& Ar)
 {
 	Super::Serialize(Ar);
@@ -2214,6 +2213,7 @@ void USceneComponent::Serialize(FArchive& Ar)
 		bAbsoluteLocation = bAbsoluteTranslation_DEPRECATED;
 	}
 }
+#endif
 
 
 void USceneComponent::PostInterpChange(UProperty* PropertyThatChanged)
@@ -2534,7 +2534,7 @@ void USceneComponent::BeginDestroy()
 
 bool USceneComponent::InternalSetWorldLocationAndRotation(FVector NewLocation, const FQuat& RotationQuat, bool bNoPhysics, ETeleportType Teleport)
 {
-	checkSlow(bWorldToComponentUpdated);
+	checkSlow(bComponentToWorldUpdated);
 	FQuat NewRotationQuat(RotationQuat);
 
 #if ENABLE_NAN_DIAGNOSTIC
@@ -2925,9 +2925,7 @@ void USceneComponent::ApplyWorldOffset(const FVector& InOffset, bool bWorldShift
 	// We do this because at level load/duplication ComponentToWorld is uninitialized
 	{
 		const FTransform RelativeTransform(RelativeRotationCache.RotatorToQuat(RelativeRotation), RelativeLocation, RelativeScale3D);
-		PRAGMA_DISABLE_DEPRECATION_WARNINGS
 		ComponentToWorld = CalcNewComponentToWorld(RelativeTransform);
-		PRAGMA_ENABLE_DEPRECATION_WARNINGS
 	}
 
 	// Update bounds
@@ -2940,9 +2938,7 @@ void USceneComponent::ApplyWorldOffset(const FVector& InOffset, bool bWorldShift
 		
 		// Calculate the new ComponentToWorld transform
 		const FTransform RelativeTransform(RelativeRotationCache.RotatorToQuat(RelativeRotation), RelativeLocation, RelativeScale3D);
-		PRAGMA_DISABLE_DEPRECATION_WARNINGS
 		ComponentToWorld = CalcNewComponentToWorld(RelativeTransform);
-		PRAGMA_ENABLE_DEPRECATION_WARNINGS
 	}
 
 	// Physics move is skipped if physics state is not created or physics scene supports origin shifting
@@ -3259,9 +3255,7 @@ void FScopedMovementUpdate::RevertMove()
 		if (IsTransformDirty())
 		{
 			// Teleport to start
-			PRAGMA_DISABLE_DEPRECATION_WARNINGS
 			Component->ComponentToWorld = InitialTransform;
-			PRAGMA_ENABLE_DEPRECATION_WARNINGS
 			Component->RelativeLocation = InitialRelativeLocation;
 			Component->RelativeRotation = InitialRelativeRotation;
 			Component->RelativeScale3D = InitialRelativeScale;
diff --git a/Engine/Source/Runtime/Engine/Private/DataChannel.cpp b/Engine/Source/Runtime/Engine/Private/DataChannel.cpp
--- a/Engine/Source/Runtime/Engine/Private/DataChannel.cpp
+++ b/Engine/Source/Runtime/Engine/Private/DataChannel.cpp
@@ -755,6 +755,12 @@ void UActorChannel::AppendMustBeMappedGuids( FOutBunch* Bunch )
 
 FPacketIdRange UChannel::SendBunch( FOutBunch* Bunch, bool Merge )
 {
+	if (!ensure(ChIndex != -1))
+	{
+		// Client "closing" but still processing bunches. Client->Server RPCs should avoid calling this, but perhaps more code needs to check this condition.
+		return FPacketIdRange(INDEX_NONE);
+	}
+
 	check(!Closing);
 	check(Connection->Channels[ChIndex]==this);
 	check(!Bunch->IsError());
diff --git a/Engine/Source/Runtime/Engine/Private/DeveloperSettings.cpp b/Engine/Source/Runtime/Engine/Private/DeveloperSettings.cpp
--- a/Engine/Source/Runtime/Engine/Private/DeveloperSettings.cpp
+++ b/Engine/Source/Runtime/Engine/Private/DeveloperSettings.cpp
@@ -103,7 +103,7 @@ void UDeveloperSettings::ImportConsoleVariableValues()
 			continue;
 		}
 
-		FString CVarName = Property->GetMetaData(DeveloperSettingsConsoleVariableMetaFName);
+		const FString& CVarName = Property->GetMetaData(DeveloperSettingsConsoleVariableMetaFName);
 		if (!CVarName.IsEmpty())
 		{
 			IConsoleVariable* CVar = IConsoleManager::Get().FindConsoleVariable(*CVarName);
@@ -126,7 +126,7 @@ void UDeveloperSettings::ExportValuesToConsoleVariables(UProperty* PropertyThatC
 {
 	if(PropertyThatChanged)
 	{
-		FString CVarName = PropertyThatChanged->GetMetaData(DeveloperSettingsConsoleVariableMetaFName);
+		const FString& CVarName = PropertyThatChanged->GetMetaData(DeveloperSettingsConsoleVariableMetaFName);
 		if (!CVarName.IsEmpty())
 		{
 			IConsoleVariable* CVar = IConsoleManager::Get().FindConsoleVariable(*CVarName);
diff --git a/Engine/Source/Runtime/Engine/Private/EdGraph/EdGraphPin.cpp b/Engine/Source/Runtime/Engine/Private/EdGraph/EdGraphPin.cpp
--- a/Engine/Source/Runtime/Engine/Private/EdGraph/EdGraphPin.cpp
+++ b/Engine/Source/Runtime/Engine/Private/EdGraph/EdGraphPin.cpp
@@ -80,6 +80,17 @@ TArray<TPair<UEdGraphPin*, FString>> PinAllocationTracking;
 /////////////////////////////////////////////////////
 // FEdGraphPinType
 
+void FEdGraphPinType::PostSerialize(const FArchive& Ar)
+{
+	if (Ar.UE4Ver() < VER_UE4_EDGRAPHPINTYPE_SERIALIZATION)
+	{
+		if (bIsArray_DEPRECATED)
+		{
+			ContainerType = EPinContainerType::Array;
+		}
+	}
+}
+
 bool FEdGraphPinType::Serialize(FArchive& Ar)
 {
 	if (Ar.UE4Ver() < VER_UE4_EDGRAPHPINTYPE_SERIALIZATION)
diff --git a/Engine/Source/Runtime/Engine/Private/GameFramework/ForceFeedbackEffect.cpp b/Engine/Source/Runtime/Engine/Private/GameFramework/ForceFeedbackEffect.cpp
--- a/Engine/Source/Runtime/Engine/Private/GameFramework/ForceFeedbackEffect.cpp
+++ b/Engine/Source/Runtime/Engine/Private/GameFramework/ForceFeedbackEffect.cpp
@@ -1,7 +1,8 @@
-// Copyright 1998-2017 Epic Games, Inc. All Rights Reserved.
+﻿// Copyright 1998-2017 Epic Games, Inc. All Rights Reserved.
 
 #include "GameFramework/ForceFeedbackEffect.h"
 #include "GenericPlatform/IInputInterface.h"
+#include "Misc/App.h"
 
 UForceFeedbackEffect::UForceFeedbackEffect(const FObjectInitializer& ObjectInitializer)
 	: Super(ObjectInitializer)
@@ -91,7 +92,7 @@ bool FActiveForceFeedbackEffect::Update(const float DeltaTime, FForceFeedbackVal
 
 	const float Duration = ForceFeedbackEffect->GetDuration();
 
-	PlayTime += DeltaTime;
+	PlayTime += (bIgnoreTimeDilation ? FApp::GetDeltaTime() : DeltaTime);
 
 	if (PlayTime > Duration && (!bLooping || Duration == 0.f) )
 	{
diff --git a/Engine/Source/Runtime/Engine/Private/GameFramework/SpringArmComponent.cpp b/Engine/Source/Runtime/Engine/Private/GameFramework/SpringArmComponent.cpp
--- a/Engine/Source/Runtime/Engine/Private/GameFramework/SpringArmComponent.cpp
+++ b/Engine/Source/Runtime/Engine/Private/GameFramework/SpringArmComponent.cpp
@@ -40,7 +40,7 @@ USpringArmComponent::USpringArmComponent(const FObjectInitializer& ObjectInitial
 	CameraLagMaxDistance = 0.f;
 }
 
-void USpringArmComponent::UpdateDesiredArmLocation(bool bDoTrace, bool bDoLocationLag, bool bDoRotationLag, float DeltaTime)
+FRotator USpringArmComponent::GetTargetRotation() const
 {
 	FRotator DesiredRot = GetComponentRotation();
 
@@ -57,9 +57,9 @@ void USpringArmComponent::UpdateDesiredArmLocation(bool bDoTrace, bool bDoLocati
 	}
 
 	// If inheriting rotation, check options for which components to inherit
-	if(!bAbsoluteRotation)
+	if (!bAbsoluteRotation)
 	{
-		if(!bInheritPitch)
+		if (!bInheritPitch)
 		{
 			DesiredRot.Pitch = RelativeRotation.Pitch;
 		}
@@ -75,6 +75,13 @@ void USpringArmComponent::UpdateDesiredArmLocation(bool bDoTrace, bool bDoLocati
 		}
 	}
 
+	return DesiredRot;
+}
+
+void USpringArmComponent::UpdateDesiredArmLocation(bool bDoTrace, bool bDoLocationLag, bool bDoRotationLag, float DeltaTime)
+{
+	FRotator DesiredRot = GetTargetRotation();
+
 	const float InverseCameraLagMaxTimeStep = (1.f / CameraLagMaxTimeStep);
 
 	// Apply 'lag' to rotation if desired
diff --git a/Engine/Source/Runtime/Engine/Private/GameplayStatics.cpp b/Engine/Source/Runtime/Engine/Private/GameplayStatics.cpp
--- a/Engine/Source/Runtime/Engine/Private/GameplayStatics.cpp
+++ b/Engine/Source/Runtime/Engine/Private/GameplayStatics.cpp
@@ -1,4 +1,4 @@
-// Copyright 1998-2017 Epic Games, Inc. All Rights Reserved.
+﻿// Copyright 1998-2017 Epic Games, Inc. All Rights Reserved.
 
 #include "Kismet/GameplayStatics.h"
 #include "Serialization/MemoryWriter.h"
@@ -364,6 +364,12 @@ float UGameplayStatics::ApplyDamage(AActor* DamagedActor, float BaseDamage, ACon
 
 UObject* UGameplayStatics::SpawnObject(TSubclassOf<UObject> ObjectClass, UObject* Outer)
 {
+	if (*ObjectClass == nullptr)
+	{
+		UE_LOG(LogScript, Warning, TEXT("UGameplayStatics::SpawnObject no class specified"));
+		return nullptr;
+	}
+
 	if (!Outer)
 	{
 		UE_LOG(LogScript, Warning, TEXT("UGameplayStatics::SpawnObject null outer"));
@@ -690,36 +696,49 @@ UParticleSystemComponent* CreateParticleSystem(UParticleSystem* EmitterTemplate,
 
 UParticleSystemComponent* UGameplayStatics::SpawnEmitterAtLocation(const UObject* WorldContextObject, UParticleSystem* EmitterTemplate, FVector SpawnLocation, FRotator SpawnRotation, bool bAutoDestroy)
 {
-	if (EmitterTemplate)
-	{
-		if (UWorld* World = GEngine->GetWorldFromContextObject(WorldContextObject, EGetWorldErrorMode::LogAndReturnNull))
-		{
-			UParticleSystemComponent* PSC = CreateParticleSystem(EmitterTemplate, World, World->GetWorldSettings(), bAutoDestroy);
+	return SpawnEmitterAtLocation(WorldContextObject, EmitterTemplate, SpawnLocation, SpawnRotation, FVector(1.f), bAutoDestroy);
+}
+
+UParticleSystemComponent* UGameplayStatics::InternalSpawnEmitterAtLocation(UWorld* World, UParticleSystem* EmitterTemplate, FVector SpawnLocation, FRotator SpawnRotation, FVector SpawnScale, bool bAutoDestroy)
+{
+	check(World && EmitterTemplate);
 
-			PSC->bAbsoluteLocation = true;
-			PSC->bAbsoluteRotation = true;
-			PSC->bAbsoluteScale = true;
-			PSC->RelativeLocation = SpawnLocation;
-			PSC->RelativeRotation = SpawnRotation;
-			PSC->RelativeScale3D = FVector(1.f);
+	UParticleSystemComponent* PSC = CreateParticleSystem(EmitterTemplate, World, World->GetWorldSettings(), bAutoDestroy);
 
-			PSC->RegisterComponentWithWorld(World);
+	PSC->bAbsoluteLocation = true;
+	PSC->bAbsoluteRotation = true;
+	PSC->bAbsoluteScale = true;
+	PSC->RelativeLocation = SpawnLocation;
+	PSC->RelativeRotation = SpawnRotation;
+	PSC->RelativeScale3D = SpawnScale;
 
-			PSC->ActivateSystem(true);
+	PSC->RegisterComponentWithWorld(World);
+
+	PSC->ActivateSystem(true);
 
 #if !(UE_BUILD_SHIPPING || UE_BUILD_TEST)
-			if (PSC->Template && PSC->Template->IsImmortal())
-			{
-				UE_LOG(LogParticles, Warning, TEXT("GameplayStatics::SpawnEmitterAtLocation spawned potentially immortal particle system! %s (%s) may stay in world despite never spawning particles after burst spawning is over."),
-					*(PSC->GetPathName()), *(PSC->Template->GetPathName())
-					);
-			}
+	if (PSC->Template && PSC->Template->IsImmortal())
+	{
+		UE_LOG(LogParticles, Warning, TEXT("GameplayStatics::SpawnEmitterAtLocation spawned potentially immortal particle system! %s (%s) may stay in world despite never spawning particles after burst spawning is over."),
+			*(PSC->GetPathName()), *(PSC->Template->GetPathName())
+		);
+	}
 #endif
 
-			return PSC;
+	return PSC;
+}
+
+UParticleSystemComponent* UGameplayStatics::SpawnEmitterAtLocation(const UObject* WorldContextObject, UParticleSystem* EmitterTemplate, FVector SpawnLocation, FRotator SpawnRotation, FVector SpawnScale, bool bAutoDestroy)
+{
+	UParticleSystemComponent* PSC = nullptr;
+	if (EmitterTemplate)
+	{
+		if (UWorld* World = GEngine->GetWorldFromContextObject(WorldContextObject, EGetWorldErrorMode::LogAndReturnNull))
+		{
+			PSC = InternalSpawnEmitterAtLocation(World, EmitterTemplate, SpawnLocation, SpawnRotation, SpawnScale, bAutoDestroy);
 		}
 	}
-	return nullptr;
+	return PSC;
 }
 
 UParticleSystemComponent* UGameplayStatics::SpawnEmitterAtLocation(UWorld* World, UParticleSystem* EmitterTemplate, const FTransform& SpawnTransform, bool bAutoDestroy)
@@ -727,34 +746,17 @@ UParticleSystemComponent* UGameplayStatics::SpawnEmitterAtLocation(UWorld* World
 	UParticleSystemComponent* PSC = nullptr;
 	if (World && EmitterTemplate)
 	{
-		PSC = CreateParticleSystem(EmitterTemplate, World, World->GetWorldSettings(), bAutoDestroy);
-
-		PSC->bAbsoluteLocation = true;
-		PSC->bAbsoluteRotation = true;
-		PSC->bAbsoluteScale = true;
-		PSC->RelativeLocation = SpawnTransform.GetLocation();
-		PSC->RelativeRotation = SpawnTransform.GetRotation().Rotator();
-		PSC->RelativeScale3D = SpawnTransform.GetScale3D();
-
-		PSC->RegisterComponentWithWorld(World);
-
-		PSC->ActivateSystem(true);
-
-#if !(UE_BUILD_SHIPPING || UE_BUILD_TEST)
-		if (PSC->Template && PSC->Template->IsImmortal())
-		{
-			UE_LOG(LogParticles, Warning, TEXT("GameplayStatics::SpawnEmitterAtLocation spawned potentially immortal particle system! %s (%s) may stay in world despite never spawning particles after burst spawning is over."),
-				*(PSC->GetPathName()), *(PSC->Template->GetPathName())
-				);
-		}
-#endif
-
+		PSC = InternalSpawnEmitterAtLocation(World, EmitterTemplate, SpawnTransform.GetLocation(), SpawnTransform.GetRotation().Rotator(), SpawnTransform.GetScale3D(), bAutoDestroy);
 	}
 	return PSC;
 }
 
-
 UParticleSystemComponent* UGameplayStatics::SpawnEmitterAttached(UParticleSystem* EmitterTemplate, USceneComponent* AttachToComponent, FName AttachPointName, FVector Location, FRotator Rotation, EAttachLocation::Type LocationType, bool bAutoDestroy)
+{
+	return SpawnEmitterAttached(EmitterTemplate, AttachToComponent, AttachPointName, Location, Rotation, FVector(1.f), LocationType, bAutoDestroy);
+}
+
+UParticleSystemComponent* UGameplayStatics::SpawnEmitterAttached(UParticleSystem* EmitterTemplate, USceneComponent* AttachToComponent, FName AttachPointName, FVector Location, FRotator Rotation, FVector Scale, EAttachLocation::Type LocationType, bool bAutoDestroy)
 {
 	UParticleSystemComponent* PSC = nullptr;
 	if (EmitterTemplate)
@@ -775,11 +777,11 @@ UParticleSystemComponent* UGameplayStatics::SpawnEmitterAttached(UParticleSystem
 				if (LocationType == EAttachLocation::KeepWorldPosition)
 				{
 					const FTransform ParentToWorld = AttachToComponent->GetSocketTransform(AttachPointName);
-					const FTransform ComponentToWorld(Rotation, Location);
+					const FTransform ComponentToWorld(Rotation, Location, Scale);
 					const FTransform RelativeTM = ComponentToWorld.GetRelativeTransform(ParentToWorld);
 					PSC->RelativeLocation = RelativeTM.GetLocation();
 					PSC->RelativeRotation = RelativeTM.GetRotation().Rotator();
-					PSC->RelativeScale3D = FVector(1.f);
+					PSC->RelativeScale3D = RelativeTM.GetScale3D();
 				}
 				else
 				{
@@ -788,14 +790,14 @@ UParticleSystemComponent* UGameplayStatics::SpawnEmitterAttached(UParticleSystem
 					
 					if (LocationType == EAttachLocation::SnapToTarget)
 					{
-						// SnapToTarget indicates we "keep world scale", and since we pass in no scale, we assume we want to keep a scale of 1,
-						// which indicates we want the inverse of the parent-to-world scale to maintain world scale of 1.
+						// SnapToTarget indicates we "keep world scale", this indicates we we want the inverse of the parent-to-world scale 
+						// to calculate world scale at Scale 1, and then apply the passed in Scale
 						const FTransform ParentToWorld = AttachToComponent->GetSocketTransform(AttachPointName);
-						PSC->RelativeScale3D = ParentToWorld.GetSafeScaleReciprocal(ParentToWorld.GetScale3D());
+						PSC->RelativeScale3D = Scale * ParentToWorld.GetSafeScaleReciprocal(ParentToWorld.GetScale3D());
 					}
 					else
 					{
-						PSC->RelativeScale3D = FVector(1.f);
+						PSC->RelativeScale3D = Scale;
 					}
 				}
 
diff --git a/Engine/Source/Runtime/Engine/Private/InstancedStaticMesh.cpp b/Engine/Source/Runtime/Engine/Private/InstancedStaticMesh.cpp
--- a/Engine/Source/Runtime/Engine/Private/InstancedStaticMesh.cpp
+++ b/Engine/Source/Runtime/Engine/Private/InstancedStaticMesh.cpp
@@ -77,7 +77,7 @@ void FStaticMeshInstanceBuffer::SetupCPUAccess(UInstancedStaticMeshComponent* In
 
 	const bool bNeedsCPUAccess = (InComponent->CastShadow && InComponent->bAffectDistanceFieldLighting 
 		// Distance field algorithms need access to instance data on the CPU
-		&& (CVar->GetValueOnGameThread() != 0 || (InComponent->GetStaticMesh() && InComponent->GetStaticMesh()->bGenerateMeshDistanceField))) 
+		&& (CVar->GetValueOnAnyThread(true) != 0 || (InComponent->GetStaticMesh() && InComponent->GetStaticMesh()->bGenerateMeshDistanceField)))
 		// If the buffer is created dynamically it mean we want to be able to do partial data update, so data should be accessible CPU side (i.e not trashed after creation of RHI resource)
 		|| IsDynamic
 		|| RequireCPUAccess;
diff --git a/Engine/Source/Runtime/Engine/Private/KismetArrayLibrary.cpp b/Engine/Source/Runtime/Engine/Private/KismetArrayLibrary.cpp
--- a/Engine/Source/Runtime/Engine/Private/KismetArrayLibrary.cpp
+++ b/Engine/Source/Runtime/Engine/Private/KismetArrayLibrary.cpp
@@ -1,4 +1,4 @@
-// Copyright 1998-2017 Epic Games, Inc. All Rights Reserved.
+﻿// Copyright 1998-2017 Epic Games, Inc. All Rights Reserved.
 
 #include "Kismet/KismetArrayLibrary.h"
 #include "GameFramework/Actor.h"
@@ -16,6 +16,7 @@ const FName SetOutOfBoundsWarning = FName("SetOutOfBoundsWarning");
 const FName InsertOutOfBoundsWarning = FName("InsertOutOfBoundsWarning");
 const FName RemoveOutOfBoundsWarning = FName("RemoveOutOfBoundsWarning");
 const FName ResizeArrayNegativeWarning = FName("ResizeArrayNegativeWarning");
+const FName SwapElementsInArrayWarning = FName("SwapElementsInArrayWarning");
 
 UKismetArrayLibrary::UKismetArrayLibrary(const FObjectInitializer& ObjectInitializer)
 	: Super(ObjectInitializer)
@@ -50,6 +51,12 @@ UKismetArrayLibrary::UKismetArrayLibrary(const FObjectInitializer& ObjectInitial
 			LOCTEXT("ResizeArrayNegativeWarning", "Array resized to negative size")
 		)
 	);
+	FBlueprintSupport::RegisterBlueprintWarning(
+		FBlueprintWarningDeclaration(
+			SwapElementsInArrayWarning,
+			LOCTEXT("SwapElementsInArrayWarning", "Array swap access out of bounds")
+		)
+	);
 }
 
 void UKismetArrayLibrary::FilterArray(const TArray<AActor*>& TargetArray, TSubclassOf<class AActor> FilterClass, TArray<AActor*>& FilteredArray)
@@ -181,7 +188,7 @@ void UKismetArrayLibrary::GenericArray_Shuffle(void* TargetArray, const UArrayPr
 		int32 LastIndex = ArrayHelper.Num() - 1;
 		for (int32 i = 0; i <= LastIndex; ++i)
 		{
-			int32 Index = FMath::RandRange(0, LastIndex);
+			int32 Index = FMath::RandRange(i, LastIndex);
 			if (i != Index)
 			{
 				ArrayHelper.SwapValues(i, Index);
@@ -290,6 +297,40 @@ void UKismetArrayLibrary::GenericArray_Set(void* TargetArray, const UArrayProper
 	}
 }
 
+void UKismetArrayLibrary::GenericArray_Swap(const void* TargetArray, const UArrayProperty* ArrayProp, int32 First, int32 Second)
+{
+	if (TargetArray)
+	{
+		FScriptArrayHelper ArrayHelper(ArrayProp, TargetArray);
+
+		if (ArrayHelper.IsValidIndex(First) && ArrayHelper.IsValidIndex(Second))
+		{
+			// If First and Second indices are the same - nothing to do
+			if (First != Second)
+			{
+				ArrayHelper.SwapValues(First, Second);
+			}
+		}
+		else if (ArrayHelper.Num() == 0)
+		{
+			FFrame::KismetExecutionMessage(*FString::Printf(TEXT("Attempted to swap elements in empty array %s!"),
+				*ArrayProp->GetName()),
+				ELogVerbosity::Warning,
+				SwapElementsInArrayWarning);
+		}
+		else
+		{
+			FFrame::KismetExecutionMessage(*FString::Printf(TEXT("Attempted to swap elements [%d] and [%d] in array %s. Available index range: [0 %d)!"),
+				First,
+				Second,
+				*ArrayProp->GetName(),
+				ArrayHelper.Num() - 1),
+				ELogVerbosity::Warning,
+				SwapElementsInArrayWarning);
+		}
+	}
+}
+
 int32 UKismetArrayLibrary::GenericArray_Find(const void* TargetArray, const UArrayProperty* ArrayProperty, const void* ItemToFind)
 {
 	int32 ResultIndex = INDEX_NONE;
@@ -423,6 +464,12 @@ void UKismetArrayLibrary::Array_Set(const TArray<int32>& TargetArray, int32 Inde
 	check(0);
 }
 
+void UKismetArrayLibrary::Array_Swap(const TArray<int32>& TargetArray, int32 First, int32 Second)
+{
+	// We should never hit these!  They're stubs to avoid NoExport on the class.  Call the Generic* equivalent instead
+	check(0);
+}
+
 int32 UKismetArrayLibrary::Array_Find(const TArray<int32>& TargetArray, const int32& ItemToFind)
 {
 	// We should never hit these!  They're stubs to avoid NoExport on the class.  Call the Generic* equivalent instead
diff --git a/Engine/Source/Runtime/Engine/Private/KismetMathLibrary.cpp b/Engine/Source/Runtime/Engine/Private/KismetMathLibrary.cpp
--- a/Engine/Source/Runtime/Engine/Private/KismetMathLibrary.cpp
+++ b/Engine/Source/Runtime/Engine/Private/KismetMathLibrary.cpp
@@ -129,13 +129,11 @@ void UKismetMathLibrary::ReportError_DaysInMonth()
 }
 
 
-
 // Include code in this source file if it's not being inlined in the header.
 #if !KISMET_MATH_INLINE_ENABLED
 #include "Kismet/KismetMathLibrary.inl"
 #endif
 
-
 bool UKismetMathLibrary::RandomBoolWithWeight(float Weight)
 {
 	//If the Weight equals to 0.0f then always return false
@@ -414,7 +412,6 @@ FRotator UKismetMathLibrary::RandomRotator(bool bRoll)
 	return RRot;
 }
 
-
 FVector UKismetMathLibrary::GetReflectionVector(FVector Direction, FVector SurfaceNormal)
 {
 	return FMath::GetReflectionVector(Direction, SurfaceNormal);
@@ -450,6 +447,16 @@ FVector UKismetMathLibrary::GetVectorArrayAverage(const TArray<FVector>& Vectors
 	return Average;
 }
 
+FRotator UKismetMathLibrary::TransformRotation(const FTransform& T, FRotator Rotation)
+{
+	return T.TransformRotation(Rotation.Quaternion()).Rotator();
+}
+
+FRotator UKismetMathLibrary::InverseTransformRotation(const FTransform& T, FRotator Rotation)
+{
+	return T.InverseTransformRotation(Rotation.Quaternion()).Rotator();
+}
+
 FRotator UKismetMathLibrary::ComposeRotators(FRotator A, FRotator B)
 {
 	FQuat AQuat = FQuat(A);
@@ -817,7 +824,6 @@ FRotator UKismetMathLibrary::MakeRotationFromAxes(FVector Forward, FVector Right
 	return RotMatrix.Rotator();
 }
 
-
 int32 UKismetMathLibrary::RandomIntegerFromStream(int32 Max, const FRandomStream& Stream)
 {
 	return Stream.RandHelper(Max);
@@ -828,6 +834,11 @@ int32 UKismetMathLibrary::RandomIntegerInRangeFromStream(int32 Min, int32 Max, c
 	return Stream.RandRange(Min, Max);
 }
 
+bool UKismetMathLibrary::InRange_IntInt(int32 Value, int32 Min, int32 Max, bool InclusiveMin, bool InclusiveMax)
+{
+	return ((InclusiveMin ? (Value >= Min) : (Value > Min)) && (InclusiveMax ? (Value <= Max) : (Value < Max)));
+}
+
 bool UKismetMathLibrary::RandomBoolFromStream(const FRandomStream& Stream)
 {
 	return (Stream.RandRange(0, 1) == 1) ? true : false;
@@ -880,7 +891,6 @@ void UKismetMathLibrary::SetRandomStreamSeed(FRandomStream& Stream, int32 NewSee
 	Stream.Initialize(NewSeed);
 }
 
-
 void UKismetMathLibrary::MinimumAreaRectangle(class UObject* WorldContextObject, const TArray<FVector>& InVerts, const FVector& SampleSurfaceNormal, FVector& OutRectCenter, FRotator& OutRectRotation, float& OutSideLengthX, float& OutSideLengthY, bool bDebugDraw)
 {
 	float MinArea = -1.f;
diff --git a/Engine/Source/Runtime/Engine/Private/KismetSystemLibrary.cpp b/Engine/Source/Runtime/Engine/Private/KismetSystemLibrary.cpp
--- a/Engine/Source/Runtime/Engine/Private/KismetSystemLibrary.cpp
+++ b/Engine/Source/Runtime/Engine/Private/KismetSystemLibrary.cpp
@@ -855,6 +855,11 @@ bool UKismetSystemLibrary::IsValidSoftObjectReference(const TSoftObjectPtr<UObje
 	return !SoftObjectReference.IsNull();
 }
 
+FString UKismetSystemLibrary::Conv_SoftObjectReferenceToString(const TSoftObjectPtr<UObject>& SoftObjectReference)
+{
+	return SoftObjectReference.ToString();
+}
+
 bool UKismetSystemLibrary::EqualEqual_SoftObjectReference(const TSoftObjectPtr<UObject>& A, const TSoftObjectPtr<UObject>& B)
 {
 	return A == B;
@@ -870,6 +875,11 @@ bool UKismetSystemLibrary::IsValidSoftClassReference(const TSoftClassPtr<UObject
 	return !SoftClassReference.IsNull();
 }
 
+FString UKismetSystemLibrary::Conv_SoftClassReferenceToString(const TSoftClassPtr<UObject>& SoftClassReference)
+{
+	return SoftClassReference.ToString();
+}
+
 bool UKismetSystemLibrary::EqualEqual_SoftClassReference(const TSoftClassPtr<UObject>& A, const TSoftClassPtr<UObject>& B)
 {
 	return A == B;
@@ -2565,6 +2575,11 @@ bool UKismetSystemLibrary::IsValidPrimaryAssetId(FPrimaryAssetId PrimaryAssetId)
 	return PrimaryAssetId.IsValid();
 }
 
+FString UKismetSystemLibrary::Conv_PrimaryAssetIdToString(FPrimaryAssetId PrimaryAssetId)
+{
+	return PrimaryAssetId.ToString();
+}
+
 bool UKismetSystemLibrary::EqualEqual_PrimaryAssetId(FPrimaryAssetId A, FPrimaryAssetId B)
 {
 	return A == B;
@@ -2580,6 +2595,11 @@ bool UKismetSystemLibrary::IsValidPrimaryAssetType(FPrimaryAssetType PrimaryAsse
 	return PrimaryAssetType.IsValid();
 }
 
+FString UKismetSystemLibrary::Conv_PrimaryAssetTypeToString(FPrimaryAssetType PrimaryAssetType)
+{
+	return PrimaryAssetType.ToString();
+}
+
 bool UKismetSystemLibrary::EqualEqual_PrimaryAssetType(FPrimaryAssetType A, FPrimaryAssetType B)
 {
 	return A == B;
diff --git a/Engine/Source/Runtime/Engine/Private/Materials/Material.cpp b/Engine/Source/Runtime/Engine/Private/Materials/Material.cpp
--- a/Engine/Source/Runtime/Engine/Private/Materials/Material.cpp
+++ b/Engine/Source/Runtime/Engine/Private/Materials/Material.cpp
@@ -1457,41 +1457,41 @@ void UMaterial::FixupMaterialUsageAfterLoad()
 
 void UMaterial::GetAllVectorParameterNames(TArray<FName> &OutParameterNames, TArray<FGuid> &OutParameterIds) const
 {
-	OutParameterNames.Empty();
-	OutParameterIds.Empty();
+	OutParameterNames.Reset();
+	OutParameterIds.Reset();
 	GetAllParameterNames<UMaterialExpressionVectorParameter>(OutParameterNames, OutParameterIds);
 }
 void UMaterial::GetAllScalarParameterNames(TArray<FName> &OutParameterNames, TArray<FGuid> &OutParameterIds) const
 {
-	OutParameterNames.Empty();
-	OutParameterIds.Empty();
+	OutParameterNames.Reset();
+	OutParameterIds.Reset();
 	GetAllParameterNames<UMaterialExpressionScalarParameter>(OutParameterNames, OutParameterIds);
 }
 void UMaterial::GetAllTextureParameterNames(TArray<FName> &OutParameterNames, TArray<FGuid> &OutParameterIds) const
 {
-	OutParameterNames.Empty();
-	OutParameterIds.Empty();
+	OutParameterNames.Reset();
+	OutParameterIds.Reset();
 	GetAllParameterNames<UMaterialExpressionTextureSampleParameter>(OutParameterNames, OutParameterIds);
 }
 
 void UMaterial::GetAllFontParameterNames(TArray<FName> &OutParameterNames, TArray<FGuid> &OutParameterIds) const
 {
-	OutParameterNames.Empty();
-	OutParameterIds.Empty();
+	OutParameterNames.Reset();
+	OutParameterIds.Reset();
 	GetAllParameterNames<UMaterialExpressionFontSampleParameter>(OutParameterNames, OutParameterIds);
 }
 
 void UMaterial::GetAllStaticSwitchParameterNames(TArray<FName> &OutParameterNames, TArray<FGuid> &OutParameterIds) const
 {
-	OutParameterNames.Empty();
-	OutParameterIds.Empty();
+	OutParameterNames.Reset();
+	OutParameterIds.Reset();
 	GetAllParameterNames<UMaterialExpressionStaticBoolParameter>(OutParameterNames, OutParameterIds);
 }
 
 void UMaterial::GetAllStaticComponentMaskParameterNames(TArray<FName> &OutParameterNames, TArray<FGuid> &OutParameterIds) const
 {
-	OutParameterNames.Empty();
-	OutParameterIds.Empty();
+	OutParameterNames.Reset();
+	OutParameterIds.Reset();
 	GetAllParameterNames<UMaterialExpressionStaticComponentMaskParameter>(OutParameterNames, OutParameterIds);
 }
 
diff --git a/Engine/Source/Runtime/Engine/Private/Materials/MaterialExpressions.cpp b/Engine/Source/Runtime/Engine/Private/Materials/MaterialExpressions.cpp
--- a/Engine/Source/Runtime/Engine/Private/Materials/MaterialExpressions.cpp
+++ b/Engine/Source/Runtime/Engine/Private/Materials/MaterialExpressions.cpp
@@ -688,7 +688,7 @@ bool UMaterialExpression::CanEditChange(const UProperty* InProperty) const
 
 		if (InProperty->HasMetaData(OverridingInputPropertyMetaData))
 		{
-			FString OverridingPropertyName = InProperty->GetMetaData(OverridingInputPropertyMetaData);
+			const FString& OverridingPropertyName = InProperty->GetMetaData(OverridingInputPropertyMetaData);
 
 			UStructProperty* StructProp = FindField<UStructProperty>(GetClass(), *OverridingPropertyName);
 			if (ensure(StructProp != nullptr))
diff --git a/Engine/Source/Runtime/Engine/Private/NetworkDriver.cpp b/Engine/Source/Runtime/Engine/Private/NetworkDriver.cpp
--- a/Engine/Source/Runtime/Engine/Private/NetworkDriver.cpp
+++ b/Engine/Source/Runtime/Engine/Private/NetworkDriver.cpp
@@ -1367,6 +1367,13 @@ void UNetDriver::InternalProcessRemoteFunction
 		Ch->ReplicateActor();
 	}
 
+	// Clients may be "closing" this connection but still processing bunches, we can't send anything if we have an invalid ChIndex.
+	if (Ch->ChIndex == -1)
+	{
+		ensure(!IsServer);
+		return;
+	}
+
 	// Form the RPC preamble.
 	FOutBunch Bunch( Ch, 0 );
 
diff --git a/Engine/Source/Runtime/Engine/Private/Particles/ParticleComponents.cpp b/Engine/Source/Runtime/Engine/Private/Particles/ParticleComponents.cpp
--- a/Engine/Source/Runtime/Engine/Private/Particles/ParticleComponents.cpp
+++ b/Engine/Source/Runtime/Engine/Private/Particles/ParticleComponents.cpp
@@ -1,4 +1,4 @@
-// Copyright 1998-2017 Epic Games, Inc. All Rights Reserved.
+﻿// Copyright 1998-2017 Epic Games, Inc. All Rights Reserved.
 
 /*=============================================================================
 	UnParticleComponent.cpp: Particle component implementation.
@@ -3419,11 +3419,13 @@ void UParticleSystemComponent::Serialize( FArchive& Ar )
 	}
 
 	Ar.UsingCustomVersion(FFrameworkObjectVersion::GUID);
+#if WITH_EDITORONLY_DATA
 
 	if (Ar.CustomVer(FFrameworkObjectVersion::GUID) < FFrameworkObjectVersion::ExplicitAttachmentRules)
 	{
 		USceneComponent::ConvertAttachLocation(AutoAttachLocationType_DEPRECATED, AutoAttachLocationRule, AutoAttachRotationRule, AutoAttachScaleRule);
 	}
+#endif
 }
 
 void UParticleSystemComponent::BeginDestroy()
diff --git a/Engine/Source/Runtime/Engine/Private/PhysicsEngine/BodyInstance.cpp b/Engine/Source/Runtime/Engine/Private/PhysicsEngine/BodyInstance.cpp
--- a/Engine/Source/Runtime/Engine/Private/PhysicsEngine/BodyInstance.cpp
+++ b/Engine/Source/Runtime/Engine/Private/PhysicsEngine/BodyInstance.cpp
@@ -437,6 +437,12 @@ FBodyInstance::FBodyInstance()
 	, Scale3D(1.0f)
 	, SceneIndexSync(0)
 	, SceneIndexAsync(0)
+#if WITH_PHYSX
+	, CurrentSceneState(BodyInstanceSceneState::NotAdded)
+#endif // WITH_PHYSX
+	, SleepFamily(ESleepFamily::Normal)
+	, DOFMode(0)
+	, CollisionEnabled(ECollisionEnabled::QueryAndPhysics)
 	, CollisionProfileName(UCollisionProfile::CustomCollisionProfileName)
 	, MaskFilter(0)
 	, bUseCCD(false)
@@ -471,17 +477,18 @@ FBodyInstance::FBodyInstance()
 	, COMNudge(ForceInit)
 	, MassScale(1.f)
 	, InertiaTensorScale(1.f)
-	, DOFConstraint(NULL)
-	, WeldParent(NULL)
-	, PhysMaterialOverride(NULL)
+	, ObjectType(ECC_WorldStatic)
+	, DOFConstraint(nullptr)
+	, WeldParent(nullptr)
+	, PhysMaterialOverride(nullptr)
 	, CustomSleepThresholdMultiplier(1.f)
 	, StabilizationThresholdMultiplier(1.f)
 	, PhysicsBlendWeight(0.f)
 	, PositionSolverIterationCount(8)
 #if WITH_PHYSX
-	, RigidActorSync(NULL)
-	, RigidActorAsync(NULL)
-	, BodyAggregate(NULL)
+	, RigidActorSync(nullptr)
+	, RigidActorAsync(nullptr)
+	, BodyAggregate(nullptr)
 	, RigidActorSyncId(PX_SERIAL_OBJECT_ID_INVALID)
 	, RigidActorAsyncId(PX_SERIAL_OBJECT_ID_INVALID)
 #endif // WITH_PHYSX
@@ -489,12 +496,7 @@ FBodyInstance::FBodyInstance()
 #if WITH_PHYSX
 	, InitialLinearVelocity(0.0f)
 	, PhysxUserData(this)
-	, CurrentSceneState(BodyInstanceSceneState::NotAdded)
 #endif // WITH_PHYSX
-	, SleepFamily(ESleepFamily::Normal)
-	, DOFMode(0)
-	, CollisionEnabled(ECollisionEnabled::QueryAndPhysics)
-	, ObjectType(ECC_WorldStatic)
 {
 	MaxAngularVelocity = UPhysicsSettings::Get()->MaxAngularVelocity;
 }
@@ -2327,7 +2329,7 @@ bool FBodyInstance::UpdateBodyScale(const FVector& InScale3D, bool bForceUpdate)
 
 					PShape->getSphereGeometry(PSphereGeom);
 					 
-					PSphereGeom.radius = FMath::Max(SphereElem->Radius * AdjustedScale3DAbs.X, KINDA_SMALL_NUMBER);
+					PSphereGeom.radius = FMath::Max(SphereElem->Radius * AdjustedScale3DAbs.X, FCollisionShape::MinSphereRadius());
 					PLocalPose.p = U2PVector(RelativeTM.TransformPosition(SphereElem->Center));
 					PLocalPose.p *= AdjustedScale3D.X;
 
@@ -2347,9 +2349,9 @@ bool FBodyInstance::UpdateBodyScale(const FVector& InScale3D, bool bForceUpdate)
 					FKBoxElem* BoxElem = ShapeElem->GetShapeCheck<FKBoxElem>();
 					PShape->getBoxGeometry(PBoxGeom);
 
-					PBoxGeom.halfExtents.x = FMath::Max((0.5f * BoxElem->X * AdjustedScale3DAbs.X), KINDA_SMALL_NUMBER);
-					PBoxGeom.halfExtents.y = FMath::Max((0.5f * BoxElem->Y * AdjustedScale3DAbs.Y), KINDA_SMALL_NUMBER);
-					PBoxGeom.halfExtents.z = FMath::Max((0.5f * BoxElem->Z * AdjustedScale3DAbs.Z), KINDA_SMALL_NUMBER);
+					PBoxGeom.halfExtents.x = FMath::Max((0.5f * BoxElem->X * AdjustedScale3DAbs.X), FCollisionShape::MinBoxExtent());
+					PBoxGeom.halfExtents.y = FMath::Max((0.5f * BoxElem->Y * AdjustedScale3DAbs.Y), FCollisionShape::MinBoxExtent());
+					PBoxGeom.halfExtents.z = FMath::Max((0.5f * BoxElem->Z * AdjustedScale3DAbs.Z), FCollisionShape::MinBoxExtent());
 
 					FTransform BoxTransform = BoxElem->GetTransform() * RelativeTM;
 					PLocalPose = PxTransform(U2PTransform(BoxTransform));
@@ -2383,12 +2385,13 @@ bool FBodyInstance::UpdateBodyScale(const FVector& InScale3D, bool bForceUpdate)
 					float Radius = FMath::Max(SphylElem->Radius * ScaleRadius, 0.1f);
 					float Length = SphylElem->Length + SphylElem->Radius * 2.f;
 					float HalfLength = Length * ScaleLength * 0.5f;
-					Radius = FMath::Clamp(Radius, 0.1f, HalfLength);	//radius is capped by half length
+					Radius = FMath::Min(Radius, HalfLength);	//radius is capped by half length
+					Radius = FMath::Max(Radius, FCollisionShape::MinCapsuleRadius()); // bounded by minimum limit.
 					float HalfHeight = HalfLength - Radius;
-					HalfHeight = FMath::Max(0.1f, HalfHeight);
+					HalfHeight = FMath::Max(FCollisionShape::MinCapsuleAxisHalfHeight(), HalfHeight);
 
-					PCapsuleGeom.halfHeight = FMath::Max(HalfHeight, KINDA_SMALL_NUMBER);
-					PCapsuleGeom.radius = FMath::Max(Radius, KINDA_SMALL_NUMBER);
+					PCapsuleGeom.halfHeight = HalfHeight;
+					PCapsuleGeom.radius = Radius;
 
 					PLocalPose = PxTransform(U2PVector(RelativeTM.TransformPosition(SphylElem->Center)), U2PQuat(SphylElem->Rotation.Quaternion()) * U2PSphylBasis);
 					PLocalPose.p.x *= AdjustedScale3D.X;
diff --git a/Engine/Source/Runtime/Engine/Private/PlayerController.cpp b/Engine/Source/Runtime/Engine/Private/PlayerController.cpp
--- a/Engine/Source/Runtime/Engine/Private/PlayerController.cpp
+++ b/Engine/Source/Runtime/Engine/Private/PlayerController.cpp
@@ -1,4 +1,4 @@
-// Copyright 1998-2017 Epic Games, Inc. All Rights Reserved.
+﻿// Copyright 1998-2017 Epic Games, Inc. All Rights Reserved.
 
 #include "GameFramework/PlayerController.h"
 #include "Misc/PackageName.h"
@@ -1927,7 +1927,7 @@ bool APlayerController::GetHitResultUnderCursorForObjects(const TArray<TEnumAsBy
 		}
 	}
 
-	if(!bHit)	//If there was no hit we reset the results. This is redundent but helps Blueprint users
+	if(!bHit)	//If there was no hit we reset the results. This is redundant but helps Blueprint users
 	{
 		HitResult = FHitResult();
 	}
@@ -1949,7 +1949,7 @@ bool APlayerController::GetHitResultUnderFinger(ETouchIndex::Type FingerIndex, E
 		}
 	}
 
-	if(!bHit)	//If there was no hit we reset the results. This is redundent but helps Blueprint users
+	if(!bHit)	//If there was no hit we reset the results. This is redundant but helps Blueprint users
 	{
 		HitResult = FHitResult();
 	}
@@ -1971,7 +1971,7 @@ bool APlayerController::GetHitResultUnderFingerByChannel(ETouchIndex::Type Finge
 		}
 	}
 
-	if(!bHit)	//If there was no hit we reset the results. This is redundent but helps Blueprint users
+	if(!bHit)	//If there was no hit we reset the results. This is redundant but helps Blueprint users
 	{
 		HitResult = FHitResult();
 	}
@@ -1993,7 +1993,7 @@ bool APlayerController::GetHitResultUnderFingerForObjects(ETouchIndex::Type Fing
 		}
 	}
 
-	if(!bHit)	//If there was no hit we reset the results. This is redundent but helps Blueprint users
+	if(!bHit)	//If there was no hit we reset the results. This is redundant but helps Blueprint users
 	{
 		HitResult = FHitResult();
 	}
@@ -3618,7 +3618,7 @@ void APlayerController::ClientPrestreamTextures_Implementation( AActor* ForcedAc
 	}
 }
 
-void APlayerController::ClientPlayForceFeedback_Implementation( UForceFeedbackEffect* ForceFeedbackEffect, bool bLooping, FName Tag)
+void APlayerController::ClientPlayForceFeedback_Implementation( UForceFeedbackEffect* ForceFeedbackEffect, bool bLooping, bool bIgnoreTimeDilation, FName Tag)
 {
 	if (ForceFeedbackEffect)
 	{
@@ -3633,7 +3633,7 @@ void APlayerController::ClientPlayForceFeedback_Implementation( UForceFeedbackEf
 			}
 		}
 
-		ActiveForceFeedbackEffects.Emplace(ForceFeedbackEffect, bLooping, Tag);
+		ActiveForceFeedbackEffects.Emplace(ForceFeedbackEffect, bLooping, bIgnoreTimeDilation, Tag);
 
 #if !(UE_BUILD_SHIPPING || UE_BUILD_TEST)
 		ForceFeedbackEffectHistoryEntries.Emplace(ActiveForceFeedbackEffects.Last(), GetWorld()->GetTimeSeconds());
@@ -3896,6 +3896,8 @@ void APlayerController::ProcessForceFeedbackAndHaptics(const float DeltaTime, co
 
 	if (bProcessFeedback)
 	{
+		UWorld* World = GetWorld();
+
 		// --- Force Feedback --------------------------
 		for (int32 Index = ActiveForceFeedbackEffects.Num() - 1; Index >= 0; --Index)
 		{
@@ -3910,7 +3912,7 @@ void APlayerController::ProcessForceFeedbackAndHaptics(const float DeltaTime, co
 			DynamicEntry.Value.Update(ForceFeedbackValues);
 		}
 
-		if (FForceFeedbackManager* ForceFeedbackManager = FForceFeedbackManager::Get(GetWorld()))
+		if (FForceFeedbackManager* ForceFeedbackManager = FForceFeedbackManager::Get(World))
 		{
 			ForceFeedbackManager->Update(GetFocalLocation(), ForceFeedbackValues);
 		}
@@ -3966,24 +3968,22 @@ void APlayerController::ProcessForceFeedbackAndHaptics(const float DeltaTime, co
 		{
 			InputInterface->SetForceFeedbackChannelValues(ControllerId, (bForceFeedbackEnabled ? ForceFeedbackValues : FForceFeedbackValues()));
 
-			bool bDisableHaptics = (CVarDisableHaptics.GetValueOnGameThread() > 0);
+			const bool bDisableHaptics = (CVarDisableHaptics.GetValueOnGameThread() > 0);
 			if (!bDisableHaptics)
 			{
-
 				// Haptic Updates
 				if (bLeftHapticsNeedUpdate)
 				{
 					InputInterface->SetHapticFeedbackValues(ControllerId, (int32)EControllerHand::Left, LeftHaptics);
 				}
-
 				if (bRightHapticsNeedUpdate)
 				{
 					InputInterface->SetHapticFeedbackValues(ControllerId, (int32)EControllerHand::Right, RightHaptics);
 				}
-			}
-			if (bGunHapticsNeedUpdate)
-			{
-				InputInterface->SetHapticFeedbackValues(ControllerId, (int32)EControllerHand::Gun, GunHaptics);
+				if (bGunHapticsNeedUpdate)
+				{
+					InputInterface->SetHapticFeedbackValues(ControllerId, (int32)EControllerHand::Gun, GunHaptics);
+				}
 			}
 		}
 	}
diff --git a/Engine/Source/Runtime/Engine/Private/StreamableManager.cpp b/Engine/Source/Runtime/Engine/Private/StreamableManager.cpp
--- a/Engine/Source/Runtime/Engine/Private/StreamableManager.cpp
+++ b/Engine/Source/Runtime/Engine/Private/StreamableManager.cpp
@@ -315,6 +315,8 @@ FStreamableManager* FStreamableHandle::GetOwningManager() const
 
 void FStreamableHandle::CancelHandle()
 {
+	check(IsInGameThread());
+
 	if (bReleased || bCanceled || !OwningManager)
 	{
 		// Too late to cancel it
@@ -332,6 +334,7 @@ void FStreamableHandle::CancelHandle()
 	TSharedRef<FStreamableHandle> SharedThis = AsShared();
 
 	ExecuteDelegate(CancelDelegate, SharedThis);
+	UnbindDelegates();
 
 	// Remove from referenced list
 	for (const FSoftObjectPath& AssetRef : RequestedAssets)
@@ -370,6 +373,8 @@ void FStreamableHandle::CancelHandle()
 
 void FStreamableHandle::ReleaseHandle()
 {
+	check(IsInGameThread());
+
 	if (bReleased || bCanceled)
 	{
 		// Too late to release it
@@ -426,6 +431,8 @@ void FStreamableHandle::StartStalledHandle()
 
 FStreamableHandle::~FStreamableHandle()
 {
+	check(IsInGameThread());
+
 	if (IsActive())
 	{
 		bReleased = true;
@@ -443,6 +450,7 @@ void FStreamableHandle::CompleteLoad()
 		bLoadCompleted = true;
 
 		ExecuteDelegate(CompleteDelegate, AsShared());
+		UnbindDelegates();
 
 		// Update any meta handles that are still active
 		for (TWeakPtr<FStreamableHandle> WeakHandle : ParentHandles)
@@ -526,8 +534,17 @@ void FStreamableHandle::CallUpdateDelegate()
 	}
 }
 
+void FStreamableHandle::UnbindDelegates()
+{
+	CancelDelegate.Unbind();
+	UpdateDelegate.Unbind();
+	CompleteDelegate.Unbind();
+}
+
 void FStreamableHandle::AsyncLoadCallbackWrapper(const FName& PackageName, UPackage* Package, EAsyncLoadingResult::Type Result, FSoftObjectPath TargetName)
 {
+	check(IsInGameThread());
+
 	// Needed so we can bind with a shared pointer for safety
 	if (OwningManager)
 	{
@@ -601,9 +618,10 @@ struct FStreamable
 				ActiveHandle->bCanceled = true;
 				ActiveHandle->OwningManager = nullptr;
 
-				if (!ActiveHandle->bReleased && ActiveHandle->CancelDelegate.IsBound())
+				if (!ActiveHandle->bReleased)
 				{
 					FStreamableHandle::ExecuteDelegate(ActiveHandle->CancelDelegate, ActiveHandle);
+					ActiveHandle->UnbindDelegates();
 				}
 			}
 		}
@@ -1055,6 +1073,8 @@ void FStreamableManager::FindInMemory( FSoftObjectPath& InOutTargetName, struct
 
 void FStreamableManager::AsyncLoadCallback(FSoftObjectPath TargetName)
 {
+	check(IsInGameThread());
+
 	FStreamable* Existing = FindStreamable(TargetName);
 
 	UE_LOG(LogStreamableManager, Verbose, TEXT("Stream Complete callback %s"), *TargetName.ToString());
diff --git a/Engine/Source/Runtime/Engine/Private/TextureDerivedData.cpp b/Engine/Source/Runtime/Engine/Private/TextureDerivedData.cpp
--- a/Engine/Source/Runtime/Engine/Private/TextureDerivedData.cpp
+++ b/Engine/Source/Runtime/Engine/Private/TextureDerivedData.cpp
@@ -1374,31 +1374,6 @@ bool UTexture::IsCachedCookedPlatformDataLoaded( const ITargetPlatform* TargetPl
 	return true;
 }
 
-void UTexture2D::WillNeverCacheCookedPlatformDataAgain()
-{
-	Super::WillNeverCacheCookedPlatformDataAgain();
-#if WITH_EDITORONLY_DATA 
-	// "Source" is only in WITH_EDITORONLY_DATA
-	// UE_LOG(LogTemp, Display, TEXT("Cleared source texture data for texture %s"), *GetName());
-	// clear source mips if we are in the editor then we don't have this luxury 
-	check( IsAsyncCacheComplete());
-
-	const TMap<FString, FTexturePlatformData*> *CookedPlatformDataPtr = GetCookedPlatformData();
-	if ( CookedPlatformDataPtr != NULL )
-	{
-		for ( const auto& TexturePlatformData : *CookedPlatformDataPtr )
-		{
-			check( TexturePlatformData.Value->AsyncTask == NULL);
-		}
-	}
-
-	// Don't release the source memory because at build time, textures may need other texture source data.
-	// Releasing it would make the source data unavailable for texture referenced through CompositeTexture.
-	// Otherwise, texture data is usually already released when loaded with AllowAsyncLoading.
-	//	Source.ReleaseSourceMemory();
-#endif
-}
-
 bool UTexture::IsAsyncCacheComplete()
 {
 	bool bComplete = true;
diff --git a/Engine/Source/Runtime/Engine/Private/UnrealEngine.cpp b/Engine/Source/Runtime/Engine/Private/UnrealEngine.cpp
--- a/Engine/Source/Runtime/Engine/Private/UnrealEngine.cpp
+++ b/Engine/Source/Runtime/Engine/Private/UnrealEngine.cpp
@@ -681,6 +681,7 @@ namespace
 		check(NewWorld);
 
 		OutPackage->PIEInstanceID = WorldContext.PIEInstance;
+		OutPackage->SetPackageFlags(PKG_PlayInEditor);
 
 		// Rename streaming levels to PIE
 		for (ULevelStreaming* StreamingLevel : NewWorld->StreamingLevels)
@@ -1575,63 +1576,68 @@ void UEngine::UpdateTimeAndHandleMaxTickRate()
 	FTimedMemReport::Get().PumpTimedMemoryReports();
 #endif
 
-	// start at now minus a bit so we don't get a zero delta.
-	static double LastTime = FPlatformTime::Seconds() - 0.0001;
+	// This is always in realtime and is not adjusted by fixed framerate. Start slightly below current real time
+	static double LastRealTime = FPlatformTime::Seconds() - 0.0001;
 	static bool bTimeWasManipulated = false;
 	bool bTimeWasManipulatedDebug = bTimeWasManipulated;	//Just used for logging of previous frame
 
 	// Figure out whether we want to use real or fixed time step.
 	const bool bUseFixedTimeStep = FApp::IsBenchmarking() || FApp::UseFixedTimeStep();
 
+	// Updates logical last time to match logical current time from last tick
 	FApp::UpdateLastTime();
 
 	// Calculate delta time and update time.
 	if( bUseFixedTimeStep )
 	{
+		// In fixed time step we set the real times to the logical time, this causes it to lie about how fast it is going
 		bTimeWasManipulated = true;
 		const float FrameRate = FApp::GetFixedDeltaTime();
 		FApp::SetDeltaTime(FrameRate);
-		LastTime = FApp::GetCurrentTime();
+		LastRealTime = FApp::GetCurrentTime();
 		FApp::SetCurrentTime(FApp::GetCurrentTime() + FApp::GetDeltaTime());
 	}
 	else
 	{
-		FApp::SetCurrentTime(FPlatformTime::Seconds());
+		// Updates logical time to real time, this may be changed by fixed frame rate below
+		double CurrentRealTime = FPlatformTime::Seconds();
+		FApp::SetCurrentTime(CurrentRealTime);
+
 		// Did we just switch from a fixed time step to real-time?  If so, then we'll update our
 		// cached 'last time' so our current interval isn't huge (or negative!)
 		if( bTimeWasManipulated && !bUseFixedFrameRate )
 		{
-			LastTime = FApp::GetCurrentTime() - FApp::GetDeltaTime();
+			LastRealTime = CurrentRealTime - FApp::GetDeltaTime();
 			bTimeWasManipulated = false;
 		}
 
-		// Calculate delta time.
-		float DeltaTime = FApp::GetCurrentTime() - LastTime;
+		// Calculate delta time, this is in real time seconds
+		float DeltaRealTime = CurrentRealTime - LastRealTime;
 
 		// Negative delta time means something is wrong with the system. Error out so user can address issue.
-		if( DeltaTime < 0 )
+		if(DeltaRealTime < 0 )
 		{
 #if PLATFORM_ANDROID
 			UE_LOG(LogEngine, Warning, TEXT("Detected negative delta time - ignoring"));
 #else
 			// AMD dual-core systems are a known issue that require AMD CPU drivers to be installed. Installer will take care of this for shipping.
 			UE_LOG(LogEngine, Fatal,TEXT("Detected negative delta time - on AMD systems please install http://files.aoaforums.com/I3199-setup.zip.html - DeltaTime:%f, bUseFixedFrameRate:%d, bTimeWasManipulatedDebug:%d, FixedFrameRate:%f"), 
-				DeltaTime, bUseFixedFrameRate, bTimeWasManipulatedDebug, FixedFrameRate);
+				DeltaRealTime, bUseFixedFrameRate, bTimeWasManipulatedDebug, FixedFrameRate);
 #endif
-			DeltaTime = 0.01;
+			DeltaRealTime = 0.01;
 		}
 
 		// Give engine chance to update frame rate smoothing
-		UpdateRunningAverageDeltaTime(DeltaTime);
+		UpdateRunningAverageDeltaTime(DeltaRealTime);
 
 		// Get max tick rate based on network settings and current delta time.
-		const float GivenMaxTickRate = GetMaxTickRate(DeltaTime);
+		const float GivenMaxTickRate = GetMaxTickRate(DeltaRealTime);
 		const float MaxTickRate = FABTest::StaticIsActive() ? 0.0f : (bUseFixedFrameRate ? FixedFrameRate : GivenMaxTickRate);
 		float WaitTime		= 0;
 		// Convert from max FPS to wait time.
 		if( MaxTickRate > 0 )
 		{
-			WaitTime = FMath::Max( 1.f / MaxTickRate - DeltaTime, 0.f );
+			WaitTime = FMath::Max( 1.f / MaxTickRate - DeltaRealTime, 0.f );
 		}
 
 		// Enforce maximum framerate and smooth framerate by waiting.
@@ -1642,7 +1648,7 @@ void UEngine::UpdateTimeAndHandleMaxTickRate()
 			FThreadIdleStats::FScopeIdle Scope;
 			
 			FSimpleScopeSecondsCounter ActualWaitTimeCounter(ActualWaitTime);
-			double WaitEndTime = FApp::GetCurrentTime() + WaitTime;
+			double WaitEndTime = CurrentRealTime + WaitTime;
 
 			SCOPE_CYCLE_COUNTER(STAT_GameTickWaitTime);
 			SCOPE_CYCLE_COUNTER(STAT_GameIdleTime);
@@ -1667,34 +1673,34 @@ void UEngine::UpdateTimeAndHandleMaxTickRate()
 					FPlatformProcess::SleepNoStats( 0 );
 				}
 			}
+			CurrentRealTime = FPlatformTime::Seconds();
 
 			if(bUseFixedFrameRate)
 			{
-				const float FrameRate = 1.f / FixedFrameRate;
-				FApp::SetDeltaTime(FrameRate);
-				FApp::SetCurrentTime(LastTime + FApp::GetDeltaTime());
+				// We are on fixed framerate but had to delay, we set the current time with a fixed time step, which will set Delta below
+				const double FrameTime = 1.0 / FixedFrameRate;
+				FApp::SetCurrentTime(LastRealTime + FrameTime);
 				bTimeWasManipulated = true;
 			}
 			else
 			{
-				FApp::SetCurrentTime(FPlatformTime::Seconds());
+				FApp::SetCurrentTime(CurrentRealTime);
 			}
 		}
 		else if(bUseFixedFrameRate && MaxTickRate == FixedFrameRate)
 		{
-			//We are doing fixed framerate and the real delta time is bigger than our desired delta time. In this case we start falling behind real time (and that's ok)
-			const float FrameRate = 1.f / FixedFrameRate;
-			FApp::SetDeltaTime(FrameRate);
-			FApp::SetCurrentTime(LastTime + FApp::GetDeltaTime());
+			// We are doing fixed framerate and the real delta time is bigger than our desired delta time. In this case we start falling behind real time (and that's ok)
+			const double FrameTime = 1.0 / FixedFrameRate;
+			FApp::SetCurrentTime(LastRealTime + FrameTime);
 			bTimeWasManipulated = true;
 		}
 
-
 		SET_FLOAT_STAT(STAT_GameTickWantedWaitTime,WaitTime * 1000.f);
 		double AdditionalWaitTimeInMs = (ActualWaitTime - static_cast<double>(WaitTime)) * 1000.0;
 		SET_FLOAT_STAT(STAT_GameTickAdditionalWaitTime,FMath::Max<float>(static_cast<float>(AdditionalWaitTimeInMs),0.f));
 
-		FApp::SetDeltaTime(FApp::GetCurrentTime() - LastTime);
+		// Update logical delta time based on logical current time
+		FApp::SetDeltaTime(FApp::GetCurrentTime() - LastRealTime);
 		FApp::SetIdleTime(ActualWaitTime);
 
 		// Negative delta time means something is wrong with the system. Error out so user can address issue.
@@ -1708,7 +1714,8 @@ void UEngine::UpdateTimeAndHandleMaxTickRate()
 #endif
 			FApp::SetDeltaTime(0.01);
 		}
-		LastTime			= FApp::GetCurrentTime();
+
+		LastRealTime = CurrentRealTime;
 
 		// Enforce a maximum delta time if wanted.
 		UGameEngine* GameEngine = Cast<UGameEngine>(this);
@@ -1748,7 +1755,7 @@ void UEngine::UpdateTimeAndHandleMaxTickRate()
 		{
 			// in seconds
 			FApp::SetDeltaTime(1.0f / OverrideFPS);
-			LastTime = FApp::GetCurrentTime();
+			LastRealTime = FApp::GetCurrentTime();
 			FApp::SetCurrentTime(FApp::GetCurrentTime() + FApp::GetDeltaTime());
 			bTimeWasManipulated = true;
 		}
@@ -10448,7 +10455,7 @@ bool UEngine::LoadMap( FWorldContext& WorldContext, FURL URL, class UPendingNetG
 		// See if the level is already in memory
 		WorldPackage = FindPackage(nullptr, *URL.Map);
 
-		const bool bPackageAlreadyLoaded = (WorldPackage != nullptr);
+		bool bPackageAlreadyLoaded = (WorldPackage != nullptr);
 
 		// If the level isn't already in memory, load level from disk
 		if (WorldPackage == nullptr)
@@ -10475,6 +10482,8 @@ bool UEngine::LoadMap( FWorldContext& WorldContext, FURL URL, class UPendingNetG
 			NewWorld = UWorld::FollowWorldRedirectorInPackage(WorldPackage);
 			if ( NewWorld )
 			{
+				// Treat this as an already loaded package because we were loaded by the redirector
+				bPackageAlreadyLoaded = true;
 				WorldPackage = NewWorld->GetOutermost();
 			}
 		}
diff --git a/Engine/Source/Runtime/Engine/Private/UserDefinedStruct.cpp b/Engine/Source/Runtime/Engine/Private/UserDefinedStruct.cpp
--- a/Engine/Source/Runtime/Engine/Private/UserDefinedStruct.cpp
+++ b/Engine/Source/Runtime/Engine/Private/UserDefinedStruct.cpp
@@ -6,6 +6,7 @@
 #include "UObject/StructOnScope.h"
 #include "UObject/UnrealType.h"
 #include "UObject/LinkerLoad.h"
+#include "UObject/FrameworkObjectVersion.h"
 #include "Misc/SecureHash.h"
 #include "UObject/PropertyPortFlags.h"
 #include "Misc/PackageName.h" // for FPackageName::GetLongPackageAssetName()
@@ -75,24 +76,37 @@ void UUserDefinedStruct::Serialize(FArchive& Ar)
 {
 	Super::Serialize( Ar );
 
-	if (Ar.IsLoading() && (EUserDefinedStructureStatus::UDSS_UpToDate == Status))
+	Ar.UsingCustomVersion(FFrameworkObjectVersion::GUID);
+
+	if (Ar.IsLoading())
 	{
-		// We need to force the editor data to be preload in case anyone needs to extract variable
-		// information at editor time about the user structure.
-		if ( EditorData != nullptr )
+		if (Ar.CustomVer(FFrameworkObjectVersion::GUID) < FFrameworkObjectVersion::UserDefinedStructsBlueprintVisible)
 		{
-			Ar.Preload(EditorData);
-			if (!(Ar.GetPortFlags() & PPF_Duplicate))
+			for (TFieldIterator<UProperty> PropIt(this); PropIt; ++PropIt)
 			{
-				FStructureEditorUtils::RecreateDefaultInstanceInEditorData(this);
+				PropIt->PropertyFlags |= CPF_BlueprintVisible;
 			}
 		}
 
-		const FStructureEditorUtils::EStructureError Result = FStructureEditorUtils::IsStructureValid(this, NULL, &ErrorMessage);
-		if (FStructureEditorUtils::EStructureError::Ok != Result)
+		if (EUserDefinedStructureStatus::UDSS_UpToDate == Status)
 		{
-			Status = EUserDefinedStructureStatus::UDSS_Error;
-			UE_LOG(LogClass, Log, TEXT("UUserDefinedStruct.Serialize '%s' validation: %s"), *GetName(), *ErrorMessage);
+			// We need to force the editor data to be preload in case anyone needs to extract variable
+			// information at editor time about the user structure.
+			if (EditorData != nullptr)
+			{
+				Ar.Preload(EditorData);
+				if (!(Ar.GetPortFlags() & PPF_Duplicate))
+				{
+					FStructureEditorUtils::RecreateDefaultInstanceInEditorData(this);
+				}
+			}
+
+			const FStructureEditorUtils::EStructureError Result = FStructureEditorUtils::IsStructureValid(this, NULL, &ErrorMessage);
+			if (FStructureEditorUtils::EStructureError::Ok != Result)
+			{
+				Status = EUserDefinedStructureStatus::UDSS_Error;
+				UE_LOG(LogClass, Log, TEXT("UUserDefinedStruct.Serialize '%s' validation: %s"), *GetName(), *ErrorMessage);
+			}
 		}
 	}
 }
diff --git a/Engine/Source/Runtime/Engine/Private/World.cpp b/Engine/Source/Runtime/Engine/Private/World.cpp
--- a/Engine/Source/Runtime/Engine/Private/World.cpp
+++ b/Engine/Source/Runtime/Engine/Private/World.cpp
@@ -906,7 +906,7 @@ bool UWorld::PreSaveRoot(const TCHAR* Filename)
 	{
 		for (UBlueprint* Blueprint : PersistentLevel->GetLevelBlueprints())
 		{
-			FKismetEditorUtilities::CompileBlueprint(Blueprint, EBlueprintCompileOptions::SkipGarbageCollection);
+			FKismetEditorUtilities::CompileBlueprint(Blueprint, EBlueprintCompileOptions::SkipGarbageCollection | EBlueprintCompileOptions::SkipSave);
 		}
 	}
 #endif
@@ -2527,6 +2527,7 @@ void UWorld::RenameToPIEWorld(int32 PIEInstanceID)
 	UPackage* WorldPackage = GetOutermost();
 
 	WorldPackage->PIEInstanceID = PIEInstanceID;
+	WorldPackage->SetPackageFlags(PKG_PlayInEditor);
 
 	const FString PIEPackageName = *UWorld::ConvertToPIEPackageName(WorldPackage->GetName(), PIEInstanceID);
 	WorldPackage->Rename(*PIEPackageName);
diff --git a/Engine/Source/Runtime/Engine/Public/WorldCollision.h b/Engine/Source/Runtime/Engine/Public/WorldCollision.h
--- a/Engine/Source/Runtime/Engine/Public/WorldCollision.h
+++ b/Engine/Source/Runtime/Engine/Public/WorldCollision.h
@@ -64,7 +64,12 @@ struct FCollisionShape
 {
 	ECollisionShape::Type ShapeType;
 
-	/** Union that supports upto 3 floats **/
+	static FORCEINLINE CONSTEXPR float MinBoxExtent()				{ return KINDA_SMALL_NUMBER; }
+	static FORCEINLINE CONSTEXPR float MinSphereRadius()			{ return KINDA_SMALL_NUMBER; }
+	static FORCEINLINE CONSTEXPR float MinCapsuleRadius()			{ return KINDA_SMALL_NUMBER; }
+	static FORCEINLINE CONSTEXPR float MinCapsuleAxisHalfHeight()	{ return KINDA_SMALL_NUMBER; }
+
+	/** Union that supports up to 3 floats **/
 	union
 	{
 		struct 
@@ -154,16 +159,16 @@ struct FCollisionShape
 		{
 		case ECollisionShape::Box:
 			{	
-				return (Box.HalfExtentX <= KINDA_SMALL_NUMBER && Box.HalfExtentY <= KINDA_SMALL_NUMBER && Box.HalfExtentZ <= KINDA_SMALL_NUMBER);
+				return (Box.HalfExtentX <= FCollisionShape::MinBoxExtent() && Box.HalfExtentY <= FCollisionShape::MinBoxExtent() && Box.HalfExtentZ <= FCollisionShape::MinBoxExtent());
 			}
 		case  ECollisionShape::Sphere:
 			{
-				return (Sphere.Radius <= KINDA_SMALL_NUMBER);
+				return (Sphere.Radius <= FCollisionShape::MinSphereRadius());
 			}
 		case ECollisionShape::Capsule:
 			{
 				// @Todo check height? It didn't check before, so I'm keeping this way for time being
-				return (Capsule.Radius <= KINDA_SMALL_NUMBER);
+				return (Capsule.Radius <= FCollisionShape::MinCapsuleRadius());
 			}
 		}
 
@@ -197,7 +202,7 @@ struct FCollisionShape
 	float GetCapsuleAxisHalfLength() const
 	{
 		ensure (ShapeType == ECollisionShape::Capsule);
-		return FMath::Max<float>(Capsule.HalfHeight - Capsule.Radius, 1.f);
+		return FMath::Max<float>(Capsule.HalfHeight - Capsule.Radius, FCollisionShape::FCollisionShape::MinCapsuleAxisHalfHeight());
 	}
 
 	/** Utility function to get Box Extention */
diff --git a/Engine/Source/Runtime/Slate/Public/Widgets/Input/SButton.h b/Engine/Source/Runtime/Slate/Public/Widgets/Input/SButton.h
--- a/Engine/Source/Runtime/Slate/Public/Widgets/Input/SButton.h
+++ b/Engine/Source/Runtime/Slate/Public/Widgets/Input/SButton.h
@@ -196,25 +196,9 @@ protected:
 	/** Padding that accounts for the button border when pressed */
 	FMargin PressedBorderPadding;
 
-	/** True if this button is currently in a pressed state */
-	bool bIsPressed;
-
 	/** The location in screenspace the button was pressed */
 	FVector2D PressedScreenSpacePosition;
 
-	/** The delegate to execute when the button is clicked */
-	FOnClicked OnClicked;
-
-	/** The delegate to execute when the button is pressed */
-	FSimpleDelegate OnPressed;
-
-	/** The delegate to execute when the button is released */
-	FSimpleDelegate OnReleased;
-
-	FSimpleDelegate OnHovered;
-
-	FSimpleDelegate OnUnhovered;
-
 	/** Style resource for the button */
 	const FButtonStyle* Style;
 
@@ -227,17 +211,18 @@ protected:
 	/** Brush resource that represents a button when it is disabled */
 	const FSlateBrush* DisabledImage;
 
-	/** Sets whether a click should be triggered on mouse down, mouse up, or that both a mouse down and up are required. */
-	EButtonClickMethod::Type ClickMethod;
+	/** The delegate to execute when the button is clicked */
+	FOnClicked OnClicked;
 
-	/** How should the button be clicked with touch events? */
-	EButtonTouchMethod::Type TouchMethod;
+	/** The delegate to execute when the button is pressed */
+	FSimpleDelegate OnPressed;
 
-	/** How should the button be clicked with keyboard/controller button events? */
-	EButtonPressMethod::Type PressMethod;
+	/** The delegate to execute when the button is released */
+	FSimpleDelegate OnReleased;
 
-	/** Can this button be focused? */
-	bool bIsFocusable;
+	FSimpleDelegate OnHovered;
+
+	FSimpleDelegate OnUnhovered;
 
 	/** Press the button */
 	virtual void Press();
@@ -260,6 +245,21 @@ protected:
 	/** The Sound to play when the button is pressed */
 	FSlateSound PressedSound;
 
+	/** Sets whether a click should be triggered on mouse down, mouse up, or that both a mouse down and up are required. */
+	TEnumAsByte<EButtonClickMethod::Type> ClickMethod;
+
+	/** How should the button be clicked with touch events? */
+	TEnumAsByte<EButtonTouchMethod::Type> TouchMethod;
+
+	/** How should the button be clicked with keyboard/controller button events? */
+	TEnumAsByte<EButtonPressMethod::Type> PressMethod;
+
+	/** Can this button be focused? */
+	uint8 bIsFocusable:1;
+
+	/** True if this button is currently in a pressed state */
+	uint8 bIsPressed:1;
+
 private:
 
 	virtual FVector2D ComputeDesiredSize(float) const override;
diff --git a/Engine/Source/Runtime/Slate/Public/Widgets/Input/SSpinBox.h b/Engine/Source/Runtime/Slate/Public/Widgets/Input/SSpinBox.h
--- a/Engine/Source/Runtime/Slate/Public/Widgets/Input/SSpinBox.h
+++ b/Engine/Source/Runtime/Slate/Public/Widgets/Input/SSpinBox.h
@@ -346,6 +346,7 @@ public:
 		if ( MouseEvent.GetEffectingButton() == EKeys::LeftMouseButton && PointerDraggingSliderIndex == INDEX_NONE )
 		{
 			DistanceDragged = 0;
+			InternalValue = ValueAttribute.Get();
 			PreDragValue = InternalValue;
 			PointerDraggingSliderIndex = MouseEvent.GetPointerIndex();
 			CachedMousePosition = MouseEvent.GetScreenSpacePosition().IntPoint();
@@ -614,6 +615,7 @@ public:
 		if ( Key == EKeys::Escape && HasMouseCapture())
 		{
 			bDragging = false;
+			PointerDraggingSliderIndex = INDEX_NONE;
 
 			InternalValue = PreDragValue;
 			NotifyValueCommitted();
@@ -868,6 +870,7 @@ protected:
 		if ( !this->HasMouseCapture() )
 		{
 			bDragging = false;
+			PointerDraggingSliderIndex = INDEX_NONE;
 		}
 	}
 	
diff --git a/Engine/Source/Runtime/SlateCore/Private/Fonts/SlateFontInfo.cpp b/Engine/Source/Runtime/SlateCore/Private/Fonts/SlateFontInfo.cpp
--- a/Engine/Source/Runtime/SlateCore/Private/Fonts/SlateFontInfo.cpp
+++ b/Engine/Source/Runtime/SlateCore/Private/Fonts/SlateFontInfo.cpp
@@ -18,9 +18,10 @@ FSlateFontInfo::FSlateFontInfo( )
 	, CompositeFont()
 	, TypefaceFontName()
 	, Size(0)
-	, FontName_DEPRECATED()
-	, Hinting_DEPRECATED(EFontHinting::Default)
 	, FontFallback(EFontFallback::FF_Max)
+#if WITH_EDITORONLY_DATA
+	, Hinting_DEPRECATED(EFontHinting::Default)
+#endif
 {
 }
 
@@ -32,9 +33,10 @@ FSlateFontInfo::FSlateFontInfo( TSharedPtr<const FCompositeFont> InCompositeFont
 	, CompositeFont(InCompositeFont)
 	, TypefaceFontName(InTypefaceFontName)
 	, Size(InSize)
-	, FontName_DEPRECATED()
-	, Hinting_DEPRECATED(EFontHinting::Default)
 	, FontFallback(EFontFallback::FF_Max)
+#if WITH_EDITORONLY_DATA
+	, Hinting_DEPRECATED(EFontHinting::Default)
+#endif
 {
 	if (!InCompositeFont.IsValid())
 	{
@@ -50,9 +52,10 @@ FSlateFontInfo::FSlateFontInfo( const UObject* InFontObject, const int32 InSize,
 	, CompositeFont()
 	, TypefaceFontName(InTypefaceFontName)
 	, Size(InSize)
-	, FontName_DEPRECATED()
-	, Hinting_DEPRECATED(EFontHinting::Default)
 	, FontFallback(EFontFallback::FF_Max)
+#if WITH_EDITORONLY_DATA
+	, Hinting_DEPRECATED(EFontHinting::Default)
+#endif
 {
 	if (InFontObject)
 	{
@@ -76,14 +79,16 @@ FSlateFontInfo::FSlateFontInfo( const FString& InFontName, uint16 InSize, EFontH
 	, CompositeFont()
 	, TypefaceFontName()
 	, Size(InSize)
+	, FontFallback(EFontFallback::FF_Max)
+#if WITH_EDITORONLY_DATA
 	, FontName_DEPRECATED(*InFontName)
 	, Hinting_DEPRECATED(InHinting)
-	, FontFallback(EFontFallback::FF_Max)
+#endif
 {
 	//Useful for debugging style breakages
 	//check( FPaths::FileExists( FontName.ToString() ) );
 
-	UpgradeLegacyFontInfo();
+	UpgradeLegacyFontInfo(FName(*InFontName), InHinting);
 }
 
 
@@ -93,14 +98,16 @@ FSlateFontInfo::FSlateFontInfo( const FName& InFontName, uint16 InSize, EFontHin
 	, CompositeFont()
 	, TypefaceFontName()
 	, Size(InSize)
+	, FontFallback(EFontFallback::FF_Max)
+#if WITH_EDITORONLY_DATA
 	, FontName_DEPRECATED(InFontName)
 	, Hinting_DEPRECATED(InHinting)
-	, FontFallback(EFontFallback::FF_Max)
+#endif
 {
 	//Useful for debugging style breakages
 	//check( FPaths::FileExists( FontName.ToString() ) );
 
-	UpgradeLegacyFontInfo();
+	UpgradeLegacyFontInfo(InFontName, InHinting);
 }
 
 
@@ -110,14 +117,16 @@ FSlateFontInfo::FSlateFontInfo( const ANSICHAR* InFontName, uint16 InSize, EFont
 	, CompositeFont()
 	, TypefaceFontName()
 	, Size(InSize)
+	, FontFallback(EFontFallback::FF_Max)
+#if WITH_EDITORONLY_DATA
 	, FontName_DEPRECATED(InFontName)
 	, Hinting_DEPRECATED(InHinting)
-	, FontFallback(EFontFallback::FF_Max)
+#endif
 {
 	//Useful for debugging style breakages
 	//check( FPaths::FileExists( FontName.ToString() ) );
 
-	UpgradeLegacyFontInfo();
+	UpgradeLegacyFontInfo(FName(InFontName), InHinting);
 }
 
 
@@ -127,14 +136,16 @@ FSlateFontInfo::FSlateFontInfo( const WIDECHAR* InFontName, uint16 InSize, EFont
 	, CompositeFont()
 	, TypefaceFontName()
 	, Size(InSize)
+	, FontFallback(EFontFallback::FF_Max)
+#if WITH_EDITORONLY_DATA
 	, FontName_DEPRECATED(InFontName)
 	, Hinting_DEPRECATED(InHinting)
-	, FontFallback(EFontFallback::FF_Max)
+#endif
 {
 	//Useful for debugging style breakages
 	//check( FPaths::FileExists( FontName.ToString() ) );
 
-	UpgradeLegacyFontInfo();
+	UpgradeLegacyFontInfo(FName(InFontName), InHinting);
 }
 
 
@@ -162,21 +173,23 @@ const FCompositeFont* FSlateFontInfo::GetCompositeFont() const
 }
 
 
+#if WITH_EDITORONLY_DATA
 void FSlateFontInfo::PostSerialize(const FArchive& Ar)
 {
 	if (Ar.UE4Ver() < VER_UE4_SLATE_COMPOSITE_FONTS && !FontObject)
 	{
-		UpgradeLegacyFontInfo();
+		UpgradeLegacyFontInfo(FontName_DEPRECATED, Hinting_DEPRECATED);
 	}
 }
+#endif
 
 
-void FSlateFontInfo::UpgradeLegacyFontInfo()
+void FSlateFontInfo::UpgradeLegacyFontInfo(FName LegacyFontName, EFontHinting LegacyHinting)
 {
 	static const FName SpecialName_DefaultSystemFont("DefaultSystemFont");
 
 	// Special case for using the default system font
-	CompositeFont = (FontName_DEPRECATED == SpecialName_DefaultSystemFont)
+	CompositeFont = (LegacyFontName == SpecialName_DefaultSystemFont)
 		? FLegacySlateFontInfoCache::Get().GetSystemFont()
-		: FLegacySlateFontInfoCache::Get().GetCompositeFont(FontName_DEPRECATED, Hinting_DEPRECATED);
+		: FLegacySlateFontInfoCache::Get().GetCompositeFont(LegacyFontName, LegacyHinting);
 }
