DUPE CHANGE : Merging using //Orion/Dev-General->//UE4/Main (CL 3616757, 3611167)

- Fixed various crashes w.r.t. anim update montage only mode and LOD change

diff --git a/Engine/Source/Runtime/Engine/Classes/Animation/AnimInstance.h b/Engine/Source/Runtime/Engine/Classes/Animation/AnimInstance.h
--- a/Engine/Source/Runtime/Engine/Classes/Animation/AnimInstance.h
+++ b/Engine/Source/Runtime/Engine/Classes/Animation/AnimInstance.h
@@ -322,7 +322,7 @@ struct FMontageEvaluationState
 	{}
 
 	// The montage to evaluate
-	UAnimMontage* Montage;
+	TWeakObjectPtr<UAnimMontage> Montage;
 
 	// The weight to use for this montage
 	float MontageWeight;
diff --git a/Engine/Source/Runtime/Engine/Private/Animation/AnimInstanceProxy.cpp b/Engine/Source/Runtime/Engine/Private/Animation/AnimInstanceProxy.cpp
--- a/Engine/Source/Runtime/Engine/Private/Animation/AnimInstanceProxy.cpp
+++ b/Engine/Source/Runtime/Engine/Private/Animation/AnimInstanceProxy.cpp
@@ -128,6 +128,7 @@ void FAnimInstanceProxy::Initialize(UAnimInstance* InAnimInstance)
 	AnimInstanceName = *InAnimInstance->GetFullName();
 #endif
 
+	UpdateCounter.Reset();
 	ReinitializeSlotNodes();
 
 	if (const USkeletalMeshComponent* SkelMeshComp = InAnimInstance->GetOwningComponent())
@@ -975,11 +976,23 @@ void FAnimInstanceProxy::SlotEvaluatePose(const FName& SlotNodeName, const FComp
 #if DEBUG_MONTAGEINSTANCE_WEIGHT
 	float TotalWeight = 0.f;
 #endif // DEBUG_MONTAGEINSTANCE_WEIGHT
+
 	for (const FMontageEvaluationState& EvalState : MontageEvaluationData)
 	{
-		if (EvalState.Montage->IsValidSlot(SlotNodeName))
+		// If MontageEvaluationData is not valid anymore, pass-through AnimSlot.
+		// This can happen if InitAnim pushes a RefreshBoneTransforms when not rendered,
+		// with EMeshComponentUpdateFlag::OnlyTickMontagesWhenNotRendered set.
+		if (!EvalState.Montage.IsValid())
+		{
+			BlendedPose = SourcePose;
+			BlendedCurve = SourceCurve;
+			return;
+		}
+
+		const UAnimMontage* const Montage = EvalState.Montage.Get();
+		if (Montage->IsValidSlot(SlotNodeName))
 		{
-			FAnimTrack const* const AnimTrack = EvalState.Montage->GetAnimationData(SlotNodeName);
+			FAnimTrack const* const AnimTrack = Montage->GetAnimationData(SlotNodeName);
 
 			// Find out additive type for pose.
 			EAdditiveAnimationType const AdditiveAnimType = AnimTrack->IsAdditive() 
@@ -993,13 +1006,13 @@ void FAnimInstanceProxy::SlotEvaluatePose(const FName& SlotNodeName, const FComp
 			NewPose.Curve.InitFrom(RequiredBones);
 
 			// Extract pose from Track
-			FAnimExtractContext ExtractionContext(EvalState.MontagePosition, EvalState.Montage->HasRootMotion() && RootMotionMode != ERootMotionMode::NoRootMotionExtraction);
+			FAnimExtractContext ExtractionContext(EvalState.MontagePosition, Montage->HasRootMotion() && RootMotionMode != ERootMotionMode::NoRootMotionExtraction);
 			AnimTrack->GetAnimationPose(NewPose.Pose, NewPose.Curve, ExtractionContext);
 
 			// add montage curves 
 			FBlendedCurve MontageCurve;
 			MontageCurve.InitFrom(RequiredBones);
-			EvalState.Montage->EvaluateCurveData(MontageCurve, EvalState.MontagePosition);
+			Montage->EvaluateCurveData(MontageCurve, EvalState.MontagePosition);
 			NewPose.Curve.Combine(MontageCurve);
 
 #if DEBUG_MONTAGEINSTANCE_WEIGHT
@@ -1115,22 +1128,25 @@ void FAnimInstanceProxy::GetSlotWeight(const FName& SlotNodeName, float& out_Slo
 	// first get all the montage instance weight this slot node has
 	for (const FMontageEvaluationState& EvalState : MontageEvaluationData)
 	{
-		if (EvalState.Montage->IsValidSlot(SlotNodeName))
+		if (EvalState.Montage.IsValid())
 		{
-			NewSlotNodeWeight += EvalState.MontageWeight;
-
-			if( !EvalState.Montage->IsValidAdditiveSlot(SlotNodeName) )
+			const UAnimMontage* const Montage = EvalState.Montage.Get();
+			if (Montage->IsValidSlot(SlotNodeName))
 			{
-				NonAdditiveTotalWeight += EvalState.MontageWeight;
-			}
+				NewSlotNodeWeight += EvalState.MontageWeight;
+				if (!Montage->IsValidAdditiveSlot(SlotNodeName))
+				{
+					NonAdditiveTotalWeight += EvalState.MontageWeight;
+				}
 
 #if DEBUGMONTAGEWEIGHT			
-			TotalDesiredWeight += EvalState->DesiredWeight;
+				TotalDesiredWeight += EvalState->DesiredWeight;
 #endif
 #if !NO_LOGGING
 			UE_LOG(LogAnimation, Verbose, TEXT("GetSlotWeight : Owner: %s, AnimMontage: %s,  (DesiredWeight:%0.2f, Weight:%0.2f)"),
 						*GetActorName(), *EvalState.Montage->GetName(), EvalState.DesiredWeight, EvalState.MontageWeight);
 #endif
+			}
 		}
 	}
 
diff --git a/Engine/Source/Runtime/Engine/Private/Animation/AnimSingleNodeInstanceProxy.cpp b/Engine/Source/Runtime/Engine/Private/Animation/AnimSingleNodeInstanceProxy.cpp
--- a/Engine/Source/Runtime/Engine/Private/Animation/AnimSingleNodeInstanceProxy.cpp
+++ b/Engine/Source/Runtime/Engine/Private/Animation/AnimSingleNodeInstanceProxy.cpp
@@ -16,6 +16,8 @@ void FAnimSingleNodeInstanceProxy::Initialize(UAnimInstance* InAnimInstance)
 	PreviewPoseCurrentTime = 0.0f;
 #endif
 
+	UpdateCounter.Reset();
+
 	// it's already doing it when evaluate
 	BlendSpaceInput = FVector::ZeroVector;
 	CurrentTime = 0.f;
