Make editor scene viewport backbuffer format consistent with game back buffer format to avoid quantization.  Noticable in some automation tests

diff --git a/Engine/Source/Runtime/Engine/Private/Slate/SceneViewport.cpp b/Engine/Source/Runtime/Engine/Private/Slate/SceneViewport.cpp
--- a/Engine/Source/Runtime/Engine/Private/Slate/SceneViewport.cpp
+++ b/Engine/Source/Runtime/Engine/Private/Slate/SceneViewport.cpp
@@ -20,6 +20,8 @@
 
 extern EWindowMode::Type GetWindowModeType(EWindowMode::Type WindowMode);
 
+static EPixelFormat SceneTargetFormat = PF_A2B10G10R10;
+
 FSceneViewport::FSceneViewport( FViewportClient* InViewportClient, TSharedPtr<SViewport> InViewportWidget )
 	: FViewport( InViewportClient )
 	, CurrentReplyState( FReply::Unhandled() )
@@ -1746,7 +1748,7 @@ void FSceneViewport::InitDynamicRHI()
 			//add sufficient entires for buffering.
 			for (int32 i = BufferedSlateHandles.Num(); i < NumBufferedFrames; i++)
 			{
-				BufferedSlateHandles.Add(new FSlateRenderTargetRHI(nullptr, 0, 0));
+				BufferedSlateHandles.Add(new FSlateRenderTargetRHI(nullptr, 0, 0)); 
 				BufferedRenderTargetsRHI.Add(nullptr);
 				BufferedShaderResourceTexturesRHI.Add(nullptr);
 			}
@@ -1766,9 +1768,9 @@ void FSceneViewport::InitDynamicRHI()
 		for (int32 i = 0; i < NumBufferedFrames; ++i)
 		{
 			// try to allocate texture via StereoRenderingDevice; if not successful, use the default way
-			if (!bStereo || !GEngine->StereoRenderingDevice->AllocateRenderTargetTexture(i, TexSizeX, TexSizeY, PF_B8G8R8A8, 1, TexCreate_None, TexCreate_RenderTargetable, BufferedRTRHI, BufferedSRVRHI))
+			if (!bStereo || !GEngine->StereoRenderingDevice->AllocateRenderTargetTexture(i, TexSizeX, TexSizeY, SceneTargetFormat, 1, TexCreate_None, TexCreate_RenderTargetable, BufferedRTRHI, BufferedSRVRHI))
 			{
-				RHICreateTargetableShaderResource2D(TexSizeX, TexSizeY, PF_B8G8R8A8, 1, TexCreate_None, TexCreate_RenderTargetable, false, CreateInfo, BufferedRTRHI, BufferedSRVRHI);
+				RHICreateTargetableShaderResource2D(TexSizeX, TexSizeY, SceneTargetFormat, 1, TexCreate_None, TexCreate_RenderTargetable, false, CreateInfo, BufferedRTRHI, BufferedSRVRHI);
 			}
 			BufferedRenderTargetsRHI[i] = BufferedRTRHI;
 			BufferedShaderResourceTexturesRHI[i] = BufferedSRVRHI;
diff --git a/Engine/Source/Runtime/Engine/Public/UnrealClient.h b/Engine/Source/Runtime/Engine/Public/UnrealClient.h
--- a/Engine/Source/Runtime/Engine/Public/UnrealClient.h
+++ b/Engine/Source/Runtime/Engine/Public/UnrealClient.h
@@ -1068,7 +1068,7 @@ public:
 		FTexture2DRHIRef ShaderResourceTextureRHI;
 
 		FRHIResourceCreateInfo CreateInfo;
-		RHICreateTargetableShaderResource2D( SizeX, SizeY, PF_B8G8R8A8, 1, TexCreate_None, TexCreate_RenderTargetable, false, CreateInfo, RenderTargetTextureRHI, ShaderResourceTextureRHI );
+		RHICreateTargetableShaderResource2D( SizeX, SizeY, PF_A2B10G10R10, 1, TexCreate_None, TexCreate_RenderTargetable, false, CreateInfo, RenderTargetTextureRHI, ShaderResourceTextureRHI );
 	}
 
 	// @todo UE4 DLL: Without these functions we get unresolved linker errors with FRenderResource
