Force an Update to occur after EMeshComponentUpdateFlag::OnlyTickMontagesWhenNotRendered happens, so Eval data is never out of sync with update data.

diff --git a/Engine/Source/Runtime/Engine/Private/Animation/AnimInstance.cpp b/Engine/Source/Runtime/Engine/Private/Animation/AnimInstance.cpp
--- a/Engine/Source/Runtime/Engine/Private/Animation/AnimInstance.cpp
+++ b/Engine/Source/Runtime/Engine/Private/Animation/AnimInstance.cpp
@@ -339,6 +339,9 @@ void UAnimInstance::UpdateAnimation(float DeltaSeconds, bool bNeedsValidRootMoti
 	SCOPE_CYCLE_COUNTER(STAT_UpdateAnimation);
 	FScopeCycleCounterUObject AnimScope(this);
 
+	// acquire the proxy as we need to update
+	FAnimInstanceProxy& Proxy = GetProxyOnGameThread<FAnimInstanceProxy>();
+
 	if (const USkeletalMeshComponent* SkelMeshComp = GetSkelMeshComponent())
 	{
 		/**
@@ -356,6 +359,12 @@ void UAnimInstance::UpdateAnimation(float DeltaSeconds, bool bNeedsValidRootMoti
 			*/
 			NotifyQueue.Reset(GetSkelMeshComponent());
 
+			/** 
+				Reset UpdateCounter(), this will force Update to occur if Eval is triggered without an Update.
+				This is to ensure that SlotNode EvaluationData is resynced to evaluate properly.
+			*/
+			Proxy.ResetUpdateCounter();
+
 			UpdateMontage(DeltaSeconds);
 
 			/**
@@ -373,9 +382,6 @@ void UAnimInstance::UpdateAnimation(float DeltaSeconds, bool bNeedsValidRootMoti
 		}
 	}
 
-	// acquire the proxy as we need to update
-	FAnimInstanceProxy& Proxy = GetProxyOnGameThread<FAnimInstanceProxy>();
-
 #if WITH_EDITOR
 	if (GIsEditor)
 	{
diff --git a/Engine/Source/Runtime/Engine/Public/Animation/AnimInstanceProxy.h b/Engine/Source/Runtime/Engine/Public/Animation/AnimInstanceProxy.h
--- a/Engine/Source/Runtime/Engine/Public/Animation/AnimInstanceProxy.h
+++ b/Engine/Source/Runtime/Engine/Public/Animation/AnimInstanceProxy.h
@@ -218,6 +218,8 @@ public:
 	const FGraphTraversalCounter& GetUpdateCounter() const  { return UpdateCounter; }
 	const FGraphTraversalCounter& GetEvaluationCounter() const { return EvaluationCounter; }
 	const FGraphTraversalCounter& GetSlotNodeInitializationCounter() const { return SlotNodeInitializationCounter; }
+	
+	void ResetUpdateCounter() { UpdateCounter.Reset(); }
 
 	/** Access root motion params */
 	FRootMotionMovementParams& GetExtractedRootMotion() { return ExtractedRootMotion; }
