Expose FSplinePoint members to blueprints

diff --git a/Engine/Source/Runtime/Engine/Classes/Components/SplineComponent.h b/Engine/Source/Runtime/Engine/Classes/Components/SplineComponent.h
--- a/Engine/Source/Runtime/Engine/Classes/Components/SplineComponent.h
+++ b/Engine/Source/Runtime/Engine/Classes/Components/SplineComponent.h
@@ -13,7 +13,7 @@ class FPrimitiveDrawInterface;
 class FSceneView;
 
 /** Permitted spline point types for SplineComponent. */
-UENUM()
+UENUM(BlueprintType)
 namespace ESplinePointType
 {
 	enum Type
@@ -114,25 +114,25 @@ struct FSplinePoint
 {
 	GENERATED_BODY()
 
-	UPROPERTY(EditAnywhere, Category = SplinePoint)
+	UPROPERTY(EditAnywhere, BlueprintReadWrite, Category = SplinePoint)
 	float InputKey;
 
-	UPROPERTY(EditAnywhere, Category = SplinePoint)
+	UPROPERTY(EditAnywhere, BlueprintReadWrite, Category = SplinePoint)
 	FVector Position;
 
-	UPROPERTY(EditAnywhere, Category = SplinePoint)
+	UPROPERTY(EditAnywhere, BlueprintReadWrite, Category = SplinePoint)
 	FVector ArriveTangent;
 
-	UPROPERTY(EditAnywhere, Category = SplinePoint)
+	UPROPERTY(EditAnywhere, BlueprintReadWrite, Category = SplinePoint)
 	FVector LeaveTangent;
 
-	UPROPERTY(EditAnywhere, Category = SplinePoint)
+	UPROPERTY(EditAnywhere, BlueprintReadWrite, Category = SplinePoint)
 	FRotator Rotation;
 
-	UPROPERTY(EditAnywhere, Category = SplinePoint)
+	UPROPERTY(EditAnywhere, BlueprintReadWrite, Category = SplinePoint)
 	FVector Scale;
 
-	UPROPERTY(EditAnywhere, Category = SplinePoint)
+	UPROPERTY(EditAnywhere, BlueprintReadWrite, Category = SplinePoint)
 	TEnumAsByte<ESplinePointType::Type> Type;
 
 	/** Default constructor */
