Texture streaming logic no longer access the creation flags through the texture RHI on the gamethread.
This fixes an ensure on XBoxOne

diff --git a/Engine/Source/Runtime/Engine/Private/Texture2D.cpp b/Engine/Source/Runtime/Engine/Private/Texture2D.cpp
--- a/Engine/Source/Runtime/Engine/Private/Texture2D.cpp
+++ b/Engine/Source/Runtime/Engine/Private/Texture2D.cpp
@@ -1116,6 +1116,7 @@ FTexture2DResource::FTexture2DResource( UTexture2D* InOwner, int32 InitialMipCou
 :	Owner( InOwner )
 ,	ResourceMem( InOwner->ResourceMem )
 ,	bReadyForStreaming(false)
+,	bUseVirtualUpdatePath(false)
 #if STATS
 ,	TextureSize( 0 )
 #endif
@@ -1197,6 +1198,10 @@ void FTexture2DResource::InitRHI()
 		TexCreateFlags |= TexCreate_NoTiling;
 	}
 
+	// Determine if this texture should use the virtual update path when streaming in and out mips. 
+	// Note that because of "r.VirtualTextureReducedMemory" it might use a virtual allocation initially.
+	bUseVirtualUpdatePath = CanCreateAsVirtualTexture(TexCreateFlags);
+
 	EPixelFormat EffectiveFormat = Owner->GetPixelFormat();
 
 	CreateSamplerStates(UTexture2D::GetGlobalMipMapLODBias() + GetDefaultMipMapBias());
@@ -1215,7 +1220,7 @@ void FTexture2DResource::InitRHI()
 			static auto CVarVirtualTextureReducedMemoryEnabled = IConsoleManager::Get().FindTConsoleVariableDataInt(TEXT("r.VirtualTextureReducedMemory"));
 			check(CVarVirtualTextureReducedMemoryEnabled);
 
-			if ( Owner->bIsStreamable && CanCreateAsVirtualTexture(TexCreateFlags) && (CVarVirtualTextureReducedMemoryEnabled->GetValueOnRenderThread() == 0 || RequestedMips > UTexture2D::GetMinTextureResidentMipCount()) )
+			if ( Owner->bIsStreamable && bUseVirtualUpdatePath && (CVarVirtualTextureReducedMemoryEnabled->GetValueOnRenderThread() == 0 || RequestedMips > UTexture2D::GetMinTextureResidentMipCount()) )
 			{
 				TexCreateFlags |= TexCreate_Virtual;
 
@@ -1468,7 +1473,7 @@ bool UTexture2D::StreamIn(int32 NewMipCount, bool bHighPrio)
 {
 	check(IsInGameThread());
 	FTexture2DResource* Texture2DResource = (FTexture2DResource*)Resource;
-	if (bIsStreamable && !PendingUpdate && Texture2DResource && Texture2DResource->bReadyForStreaming && Texture2DResource->TextureRHI && NewMipCount > GetNumResidentMips())
+	if (bIsStreamable && !PendingUpdate && Texture2DResource && Texture2DResource->bReadyForStreaming && NewMipCount > GetNumResidentMips())
 	{
 #if WITH_EDITORONLY_DATA
 		if (FPlatformProperties::HasEditorOnlyData())
@@ -1485,11 +1490,8 @@ bool UTexture2D::StreamIn(int32 NewMipCount, bool bHighPrio)
 		else
 #endif
 		{
-			const uint32 TexCreateFlags = Texture2DResource->TextureRHI->GetFlags();
-			const bool bIsVirtualTexture = (TexCreateFlags & TexCreate_Virtual) == TexCreate_Virtual;
-
 			// If the future texture is to be a virtual texture, use the virtual stream in path.
-			if (CanCreateAsVirtualTexture(TexCreateFlags) || bIsVirtualTexture)
+			if (Texture2DResource->bUseVirtualUpdatePath)
 			{
 				PendingUpdate = new FTexture2DStreamIn_IO_Virtual(this, NewMipCount, bHighPrio);
 			}
@@ -1516,12 +1518,9 @@ bool UTexture2D::StreamOut(int32 NewMipCount)
 {
 	check(IsInGameThread());
 	FTexture2DResource* Texture2DResource = (FTexture2DResource*)Resource;
-	if (bIsStreamable && !PendingUpdate && Texture2DResource && Texture2DResource->bReadyForStreaming && Texture2DResource->TextureRHI && NewMipCount < GetNumResidentMips())
+	if (bIsStreamable && !PendingUpdate && Texture2DResource && Texture2DResource->bReadyForStreaming && NewMipCount < GetNumResidentMips())
 	{
-		const uint32 TexCreateFlags = Texture2DResource->TextureRHI->GetFlags();
-		const bool bIsVirtualTexture = (TexCreateFlags & TexCreate_Virtual) == TexCreate_Virtual;
-
-		if (bIsVirtualTexture)
+		if (Texture2DResource->bUseVirtualUpdatePath)
 		{
 			PendingUpdate = new FTexture2DStreamOut_Virtual(this, NewMipCount);
 		}
diff --git a/Engine/Source/Runtime/Engine/Public/TextureResource.h b/Engine/Source/Runtime/Engine/Public/TextureResource.h
--- a/Engine/Source/Runtime/Engine/Public/TextureResource.h
+++ b/Engine/Source/Runtime/Engine/Public/TextureResource.h
@@ -148,6 +148,9 @@ private:
 	/** Whether the texture RHI has been initialized.														*/
 	bool bReadyForStreaming;
 
+	/** Whether this texture should be updated using the virtual allocations.								*/
+	bool bUseVirtualUpdatePath;
+
 	EMipFadeSettings MipFadeSetting;
 
 	/** First mip level used in Texture2DRHI. This is always correct as long as Texture2DRHI is allocated, regardless of streaming status. */
