Merging cl 3667233
Merging //UE4/Release-4.18 to Release-Staging-4.18 (//UE4/Release-Staging-4.18)
Static analysis fix

diff --git a/Engine/Source/Runtime/Engine/Private/Slate/DebugCanvas.cpp b/Engine/Source/Runtime/Engine/Private/Slate/DebugCanvas.cpp
--- a/Engine/Source/Runtime/Engine/Private/Slate/DebugCanvas.cpp
+++ b/Engine/Source/Runtime/Engine/Private/Slate/DebugCanvas.cpp
@@ -179,7 +179,10 @@ void FDebugCanvasDrawer::InitDebugCanvas(UWorld* InWorld)
 
 		if (LayerID != INVALID_LAYER_ID && (!RenderThreadCanvas->IsUsingInternalTexture() || !bCanvasRenderedLastFrame))
 		{
-			GEngine->StereoRenderingDevice->GetStereoLayers()->DestroyLayer(LayerID);
+			if (GEngine && GEngine->StereoRenderingDevice.IsValid() && GEngine->StereoRenderingDevice->GetStereoLayers())
+			{
+				GEngine->StereoRenderingDevice->GetStereoLayers()->DestroyLayer(LayerID);
+			}
 			LayerID = INVALID_LAYER_ID;
 		}
 	}
