Removing over-aggressive reentrance checks for material instance parameter default value retrieval. These shouldn't ever trigger as a cyclic instance won't have any parameters to get the values for, but leaving in the editor-path to catch the theoretical code-driven content creation case.

diff --git a/Engine/Source/Runtime/Engine/Private/Materials/MaterialInstance.cpp b/Engine/Source/Runtime/Engine/Private/Materials/MaterialInstance.cpp
--- a/Engine/Source/Runtime/Engine/Private/Materials/MaterialInstance.cpp
+++ b/Engine/Source/Runtime/Engine/Private/Materials/MaterialInstance.cpp
@@ -1795,10 +1795,12 @@ void UMaterialInstance::GetDependentFunctions(TArray<UMaterialFunctionInterface*
 
 bool UMaterialInstance::GetScalarParameterDefaultValue(const FMaterialParameterInfo& ParameterInfo, float& OutValue, bool bOveriddenOnly) const
 {
+#if WITH_EDITOR
 	if (GetReentrantFlag())
 	{
 		return false;
 	}
+#endif
 
 	// In the case of duplicate parameters with different values, this will return the
 	// first matching expression found, not necessarily the one that's used for rendering
@@ -1832,7 +1834,9 @@ bool UMaterialInstance::GetScalarParameterDefaultValue(const FMaterialParameterI
 	
 	if (Parent)
 	{
+#if WITH_EDITOR
 		FMICReentranceGuard	Guard(this);
+#endif
 		return Parent->GetScalarParameterDefaultValue(ParameterInfo, OutValue, bOveriddenOnly);
 	}
 
@@ -1841,10 +1845,12 @@ bool UMaterialInstance::GetScalarParameterDefaultValue(const FMaterialParameterI
 
 bool UMaterialInstance::GetVectorParameterDefaultValue(const FMaterialParameterInfo& ParameterInfo, FLinearColor& OutValue, bool bOveriddenOnly) const
 {
+#if WITH_EDITOR
 	if (GetReentrantFlag())
 	{
 		return false;
 	}
+#endif
 
 	// In the case of duplicate parameters with different values, this will return the
 	// first matching expression found, not necessarily the one that's used for rendering
@@ -1878,7 +1884,9 @@ bool UMaterialInstance::GetVectorParameterDefaultValue(const FMaterialParameterI
 	
 	if (Parent)
 	{
+#if WITH_EDITOR
 		FMICReentranceGuard	Guard(this);
+#endif
 		return Parent->GetVectorParameterDefaultValue(ParameterInfo, OutValue, bOveriddenOnly);
 	}
 
@@ -1887,10 +1895,12 @@ bool UMaterialInstance::GetVectorParameterDefaultValue(const FMaterialParameterI
 
 bool UMaterialInstance::GetTextureParameterDefaultValue(const FMaterialParameterInfo& ParameterInfo, UTexture*& OutValue) const
 {
+#if WITH_EDITOR
 	if (GetReentrantFlag())
 	{
 		return false;
 	}
+#endif
 
 	// In the case of duplicate parameters with different values, this will return the
 	// first matching expression found, not necessarily the one that's used for rendering
@@ -1924,7 +1934,9 @@ bool UMaterialInstance::GetTextureParameterDefaultValue(const FMaterialParameter
 	
 	if (Parent)
 	{
+#if WITH_EDITOR
 		FMICReentranceGuard	Guard(this);
+#endif
 		return Parent->GetTextureParameterDefaultValue(ParameterInfo, OutValue);
 	}
 
@@ -1933,10 +1945,12 @@ bool UMaterialInstance::GetTextureParameterDefaultValue(const FMaterialParameter
 
 bool UMaterialInstance::GetFontParameterDefaultValue(const FMaterialParameterInfo& ParameterInfo, UFont*& OutFontValue, int32& OutFontPage) const
 {
+#if WITH_EDITOR
 	if (GetReentrantFlag())
 	{
 		return false;
 	}
+#endif
 
 	// In the case of duplicate parameters with different values, this will return the
 	// first matching expression found, not necessarily the one that's used for rendering
@@ -1970,7 +1984,9 @@ bool UMaterialInstance::GetFontParameterDefaultValue(const FMaterialParameterInf
 	
 	if (Parent)
 	{
+#if WITH_EDITOR
 		FMICReentranceGuard	Guard(this);
+#endif
 		return Parent->GetFontParameterDefaultValue(ParameterInfo, OutFontValue, OutFontPage);
 	}
 
