Copying 3694952
Fix for adjancey DDC serialization warning.

diff --git a/Engine/Plugins/Experimental/AlembicImporter/Source/AlembicLibrary/Private/AbcImporter.cpp b/Engine/Plugins/Experimental/AlembicImporter/Source/AlembicLibrary/Private/AbcImporter.cpp
--- a/Engine/Plugins/Experimental/AlembicImporter/Source/AlembicLibrary/Private/AbcImporter.cpp
+++ b/Engine/Plugins/Experimental/AlembicImporter/Source/AlembicLibrary/Private/AbcImporter.cpp
@@ -1886,12 +1886,6 @@ bool FAbcImporter::BuildSkeletalMesh(
 		}
 	}
 
-	// Build the adjacency index buffer used for tessellation.
-	TArray<FSoftSkinVertex> SoftSkinVertices;
-	LODModel.GetVertices(SoftSkinVertices);
-
-	MeshUtilities.BuildSkeletalAdjacencyIndexBuffer(SoftSkinVertices, LODModel.NumTexCoords, LODModel.IndexBuffer, LODModel.AdjacencyBuffer);
-
 	// Compute the required bones for this model.
 	USkeletalMesh::CalculateRequiredBones(LODModel, RefSkeleton, NULL);
 
diff --git a/Engine/Source/Developer/MeshUtilities/Private/MeshUtilities.cpp b/Engine/Source/Developer/MeshUtilities/Private/MeshUtilities.cpp
--- a/Engine/Source/Developer/MeshUtilities/Private/MeshUtilities.cpp
+++ b/Engine/Source/Developer/MeshUtilities/Private/MeshUtilities.cpp
@@ -1169,14 +1169,6 @@ void FMeshUtilities::BuildSkeletalModelFromChunks(FSkeletalMeshLODModel& LODMode
 	}
 	Chunks.Empty();
 
-	// Build the adjacency index buffer used for tessellation.
-	{
-		TArray<FSoftSkinVertex> Vertices;
-		LODModel.GetVertices(Vertices);
-
-		BuildSkeletalAdjacencyIndexBuffer(Vertices, LODModel.NumTexCoords, LODModel.IndexBuffer, LODModel.AdjacencyBuffer);
-	}
-
 	// Compute the required bones for this model.
 	USkeletalMesh::CalculateRequiredBones(LODModel, RefSkeleton, NULL);
 #endif // #if WITH_EDITORONLY_DATA
diff --git a/Engine/Source/Runtime/Engine/Private/SkeletalMesh.cpp b/Engine/Source/Runtime/Engine/Private/SkeletalMesh.cpp
--- a/Engine/Source/Runtime/Engine/Private/SkeletalMesh.cpp
+++ b/Engine/Source/Runtime/Engine/Private/SkeletalMesh.cpp
@@ -1029,30 +1029,6 @@ void USkeletalMesh::Serialize( FArchive& Ar )
 		}
 	}
 
-	// Build adjacency information for meshes that have not yet had it built.
-#if WITH_EDITOR
-	for ( int32 LODIndex = 0; LODIndex < ImportedModel->LODModels.Num(); ++LODIndex )
-	{
-		FSkeletalMeshLODModel& LODModel = ImportedModel->LODModels[LODIndex];
-
-		// If we have no adjacency info, but we do have valid mesh, build it now
-		if ((LODModel.AdjacencyBuffer.Num() == 0 && LODModel.IndexBuffer.Num() > 0)
-#if WITH_APEX_CLOTHING
-		|| (Ar.IsLoading() && Ar.UE4Ver() < VER_UE4_APEX_CLOTH_TESSELLATION && LODModel.HasClothData())
-#endif // WITH_APEX_CLOTHING
-			)
-		{
-			TArray<FSoftSkinVertex> Vertices;
-			IMeshUtilities& MeshUtilities = FModuleManager::Get().LoadModuleChecked<IMeshUtilities>("MeshUtilities");
-			UE_LOG(LogSkeletalMesh, Warning, TEXT("Building adjacency information for skeletal mesh '%s'. Please resave the asset."), *GetPathName() );
-			LODModel.GetVertices( Vertices );
-			MeshUtilities.BuildSkeletalAdjacencyIndexBuffer( Vertices, LODModel.NumTexCoords, LODModel.IndexBuffer, LODModel.AdjacencyBuffer );
-
-			// @skeltodo - Bump SkelMesh DDC GUID here
-		}
-	}
-#endif // #if WITH_EDITOR
-
 	// make sure we're counting properly
 	if (!Ar.IsLoading() && !Ar.IsSaving())
 	{
@@ -1486,12 +1462,6 @@ void USkeletalMesh::RemoveLegacyClothingSections()
 						}
 					}
 
-					// Fix adjacency buffer
-					TArray<FSoftSkinVertex> Vertices;
-					IMeshUtilities& MeshUtilities = FModuleManager::Get().LoadModuleChecked<IMeshUtilities>("MeshUtilities");
-					LodModel.GetVertices(Vertices);
-					MeshUtilities.BuildSkeletalAdjacencyIndexBuffer(Vertices, LodModel.NumTexCoords, LodModel.IndexBuffer, LodModel.AdjacencyBuffer);
-
 					PostEditChange();
 				}
 			}
diff --git a/Engine/Source/Runtime/Engine/Private/SkeletalMeshLODModel.cpp b/Engine/Source/Runtime/Engine/Private/SkeletalMeshLODModel.cpp
--- a/Engine/Source/Runtime/Engine/Private/SkeletalMeshLODModel.cpp
+++ b/Engine/Source/Runtime/Engine/Private/SkeletalMeshLODModel.cpp
@@ -634,15 +634,10 @@ void FSkeletalMeshLODModel::Serialize(FArchive& Ar, UObject* Owner, int32 Idx)
 			if (!StripFlags.IsClassDataStripped(LodAdjacencyStripFlag))
 			{
 				// For old content, load as a multi-size container, but convert into regular array
-				if (Ar.IsLoading() && Ar.CustomVer(FSkeletalMeshCustomVersion::GUID) < FSkeletalMeshCustomVersion::SplitModelAndRenderData)
 				{
-					FMultiSizeIndexContainer TempMultiSizeIndexContainer;
-					TempMultiSizeIndexContainer.Serialize(Ar, false);
-					TempMultiSizeIndexContainer.GetIndexBuffer(AdjacencyBuffer);
-				}
-				else
-				{
-					Ar << AdjacencyBuffer;
+					// Serialize and discard the adjacency data, it's now build for the DDC
+					FMultiSizeIndexContainer TempMultiSizeAdjacencyIndexContainer;
+					TempMultiSizeAdjacencyIndexContainer.Serialize(Ar, false);
 				}
 			}
 
@@ -729,7 +724,6 @@ void FSkeletalMeshLODModel::GetResourceSizeEx(FResourceSizeEx& CumulativeResourc
 	CumulativeResourceSize.AddUnknownMemoryBytes(ActiveBoneIndices.GetAllocatedSize());
 	CumulativeResourceSize.AddUnknownMemoryBytes(RequiredBones.GetAllocatedSize());
 	CumulativeResourceSize.AddUnknownMemoryBytes(IndexBuffer.GetAllocatedSize());
-	CumulativeResourceSize.AddUnknownMemoryBytes(AdjacencyBuffer.GetAllocatedSize());
 
 	CumulativeResourceSize.AddUnknownMemoryBytes(RawPointIndices.GetBulkDataSize());
 	CumulativeResourceSize.AddUnknownMemoryBytes(LegacyRawPointIndices.GetBulkDataSize());
diff --git a/Engine/Source/Runtime/Engine/Private/SkeletalMeshLODRenderData.cpp b/Engine/Source/Runtime/Engine/Private/SkeletalMeshLODRenderData.cpp
--- a/Engine/Source/Runtime/Engine/Private/SkeletalMeshLODRenderData.cpp
+++ b/Engine/Source/Runtime/Engine/Private/SkeletalMeshLODRenderData.cpp
@@ -11,6 +11,7 @@
 
 #if WITH_EDITOR
 #include "SkeletalMeshModel.h"
+#include "MeshUtilities.h"
 #endif // WITH_EDITOR
 
 
@@ -234,7 +235,11 @@ void FSkeletalMeshLODRenderData::BuildFromLODModel(const FSkeletalMeshLODModel*
 	const uint8 DataTypeSize = (ImportedModel->NumVertices < MAX_uint16) ? sizeof(uint16) : sizeof(uint32);
 
 	MultiSizeIndexContainer.RebuildIndexBuffer(DataTypeSize, ImportedModel->IndexBuffer);
-	AdjacencyMultiSizeIndexContainer.RebuildIndexBuffer(DataTypeSize, ImportedModel->AdjacencyBuffer);
+	
+	TArray<uint32> BuiltAdjacencyIndices;
+	IMeshUtilities& MeshUtilities = FModuleManager::Get().LoadModuleChecked<IMeshUtilities>("MeshUtilities");
+	MeshUtilities.BuildSkeletalAdjacencyIndexBuffer(Vertices, ImportedModel->NumTexCoords, ImportedModel->IndexBuffer, BuiltAdjacencyIndices);
+	AdjacencyMultiSizeIndexContainer.RebuildIndexBuffer(DataTypeSize, BuiltAdjacencyIndices);
 
 	// MorphTargetVertexInfoBuffers are created in InitResources
 
diff --git a/Engine/Source/Runtime/Engine/Private/SkeletalMeshRenderData.cpp b/Engine/Source/Runtime/Engine/Private/SkeletalMeshRenderData.cpp
--- a/Engine/Source/Runtime/Engine/Private/SkeletalMeshRenderData.cpp
+++ b/Engine/Source/Runtime/Engine/Private/SkeletalMeshRenderData.cpp
@@ -29,7 +29,7 @@ namespace SkeletalMeshCookStats
 // differences, etc.) replace the version GUID below with a new one.
 // In case of merge conflicts with DDC versions, you MUST generate a new GUID
 // and set this new GUID as the version.                                       
-#define SKELETALMESH_DERIVEDDATA_VER TEXT("62BD032753D246BCA16F277150000DAC")
+#define SKELETALMESH_DERIVEDDATA_VER TEXT("AC033446D70841C0A314C0C4C56E39F6")
 
 static const FString& GetSkeletalMeshDerivedDataVersion()
 {
diff --git a/Engine/Source/Runtime/Engine/Public/Rendering/SkeletalMeshLODModel.h b/Engine/Source/Runtime/Engine/Public/Rendering/SkeletalMeshLODModel.h
--- a/Engine/Source/Runtime/Engine/Public/Rendering/SkeletalMeshLODModel.h
+++ b/Engine/Source/Runtime/Engine/Public/Rendering/SkeletalMeshLODModel.h
@@ -191,9 +191,6 @@ public:
 	*/
 	TArray<FBoneIndexType> ActiveBoneIndices;
 
-	/** Adjacency buffer information */
-	TArray<uint32> AdjacencyBuffer;
-
 	/**
 	* Bones that should be updated when rendering this LOD. This may include bones that are not required for rendering.
 	* All parents for bones in this array should be present as well - that is, a complete path from the root to each bone.
