Copy 3695687
Fix uninitialized variable for Mac automation failures.

diff --git a/Engine/Source/Runtime/Engine/Private/SkeletalMeshModel.cpp b/Engine/Source/Runtime/Engine/Private/SkeletalMeshModel.cpp
--- a/Engine/Source/Runtime/Engine/Private/SkeletalMeshModel.cpp
+++ b/Engine/Source/Runtime/Engine/Private/SkeletalMeshModel.cpp
@@ -7,6 +7,7 @@
 #include "Engine/SkeletalMesh.h"
 
 FSkeletalMeshModel::FSkeletalMeshModel()
+	: bGuidIsHash(false)	
 {
 }
 
