Updated physicalized animation and physics stability tests to be more relaxed on console.
Added cvar for adjusting tolerance on less precise autopamted tests, AutomationTestUseRelaxedPhysicsStabilityThresholds.
Added GetConsoleVariableFloat/Int global blueprint functions.

diff --git a/Engine/Source/Developer/FunctionalTesting/Private/AutomationBlueprintFunctionLibrary.cpp b/Engine/Source/Developer/FunctionalTesting/Private/AutomationBlueprintFunctionLibrary.cpp
--- a/Engine/Source/Developer/FunctionalTesting/Private/AutomationBlueprintFunctionLibrary.cpp
+++ b/Engine/Source/Developer/FunctionalTesting/Private/AutomationBlueprintFunctionLibrary.cpp
@@ -47,6 +47,12 @@ static TAutoConsoleVariable<int32> CVarAutomationScreenshotResolutionHeight(
 	TEXT("The height of automation screenshots."),
 	ECVF_Default);
 
+static TAutoConsoleVariable<int32> CVarMaxContactOffset(
+	TEXT("AutomationTestUseRelaxedPhysicsStabilityThresholds"),
+	0,
+	TEXT("Relaxes the thresold of automated physics stability tests. Default: Off"),
+	ECVF_Default);
+
 #if (WITH_DEV_AUTOMATION_TESTS || WITH_PERF_AUTOMATION_TESTS)
 
 template<typename T>
diff --git a/Engine/Source/Runtime/Engine/Classes/Kismet/KismetSystemLibrary.h b/Engine/Source/Runtime/Engine/Classes/Kismet/KismetSystemLibrary.h
--- a/Engine/Source/Runtime/Engine/Classes/Kismet/KismetSystemLibrary.h
+++ b/Engine/Source/Runtime/Engine/Classes/Kismet/KismetSystemLibrary.h
@@ -342,6 +342,24 @@ class ENGINE_API UKismetSystemLibrary : public UBlueprintFunctionLibrary
 	UFUNCTION(BlueprintCallable, Category="Development",meta=(WorldContext="WorldContextObject"))
 	static void ExecuteConsoleCommand(UObject* WorldContextObject, const FString& Command, class APlayerController* SpecificPlayer = NULL );
 
+	/**
+	 * Attempts to retrieve the value of the specified float console variable, if it exists.
+	 * 
+	 * @param	VariableName	Name of the console variable to find.
+	 * @return	The value if found, 0 otherwise.
+	 */
+	UFUNCTION(BlueprintCallable, Category="Development",meta=(WorldContext="WorldContextObject"))
+	static float GetConsoleVariableFloatValue(UObject* WorldContextObject, const FString& VariableName);
+
+	/**
+	 * Attempts to retrieve the value of the specified integer console variable, if it exists.
+	 * 
+	 * @param	VariableName	Name of the console variable to find.
+	 * @return	The value if found, 0 otherwise.
+	 */
+	UFUNCTION(BlueprintCallable, Category="Development",meta=(WorldContext="WorldContextObject"))
+	static int32 GetConsoleVariableIntValue(UObject* WorldContextObject, const FString& VariableName);
+
 	/** 
 	 *	Exit the current game 
 	 * @param	SpecificPlayer	The specific player to quit the game. If not specified, player 0 will quit.
diff --git a/Engine/Source/Runtime/Engine/Private/KismetSystemLibrary.cpp b/Engine/Source/Runtime/Engine/Private/KismetSystemLibrary.cpp
--- a/Engine/Source/Runtime/Engine/Private/KismetSystemLibrary.cpp
+++ b/Engine/Source/Runtime/Engine/Private/KismetSystemLibrary.cpp
@@ -269,6 +269,41 @@ void UKismetSystemLibrary::ExecuteConsoleCommand(UObject* WorldContextObject, co
 	}
 }
 
+float UKismetSystemLibrary::GetConsoleVariableFloatValue(UObject* WorldContextObject, const FString& VariableName)
+{
+	float Value = 0.0f;
+
+	TConsoleVariableData<float>* Variable = IConsoleManager::Get().FindTConsoleVariableDataFloat(*VariableName);
+	if (Variable)
+	{
+		Value = Variable->GetValueOnGameThread();
+	}
+	else
+	{
+		UE_LOG(LogBlueprintUserMessages, Warning, TEXT("Failed to find console variable '%s'."), *VariableName);
+	}
+
+	return Value;
+}
+
+int32 UKismetSystemLibrary::GetConsoleVariableIntValue(UObject* WorldContextObject, const FString& VariableName)
+{
+	int32 Value = 0;
+
+	TConsoleVariableData<int32>* Variable = IConsoleManager::Get().FindTConsoleVariableDataInt(*VariableName);
+	if (Variable)
+	{
+		Value = Variable->GetValueOnGameThread();
+	}
+	else
+	{
+		UE_LOG(LogBlueprintUserMessages, Warning, TEXT("Failed to find console variable '%s'."), *VariableName);
+	}
+
+	return Value;
+}
+
+
 void UKismetSystemLibrary::QuitGame(UObject* WorldContextObject, class APlayerController* SpecificPlayer, TEnumAsByte<EQuitPreference::Type> QuitPreference)
 {
 	APlayerController* TargetPC = SpecificPlayer ? SpecificPlayer : UGameplayStatics::GetPlayerController(WorldContextObject, 0);
