Copying //UE4/Dev-Framework to //UE4/Dev-Main (Source: //UE4/Dev-Framework @ 3716594)

diff --git a/Engine/Source/Runtime/Engine/Classes/Components/SkeletalMeshComponent.h b/Engine/Source/Runtime/Engine/Classes/Components/SkeletalMeshComponent.h
--- a/Engine/Source/Runtime/Engine/Classes/Components/SkeletalMeshComponent.h
+++ b/Engine/Source/Runtime/Engine/Classes/Components/SkeletalMeshComponent.h
@@ -68,9 +68,6 @@
 	TArray<FTransform> BoneSpaceTransforms;
 	FVector RootBoneTranslation;
 
-	// Double buffer curve data
-	FBlendedHeapCurve	Curve;
-
 	// Are we performing interpolation this tick
 	bool bDoInterpolation;
 
@@ -83,6 +80,9 @@
 	// duplicate the cache curves
 	bool bDuplicateToCacheCurve;
 
+	// Double buffer curve data
+	FBlendedHeapCurve	Curve;
+
 	FAnimationEvaluationContext()
 	{
 		Clear();
@@ -265,11 +265,6 @@
 	
 	/** @Todo anim: Matinee related data end - this needs to be replaced to new system. */
 
-protected:
-	/** Whether to use Animation Blueprint or play Single Animation Asset. */
-	UPROPERTY(EditAnywhere, BlueprintReadOnly, Category=Animation)
-	TEnumAsByte<EAnimationMode::Type>	AnimationMode;
-
 public:
 #if WITH_EDITORONLY_DATA
 	/** The blueprint for creating an AnimationScript. */
@@ -304,13 +299,18 @@
 
 	/** Temporary array of local-space (relative to parent bone) rotation/translation for each bone. */
 	TArray<FTransform> BoneSpaceTransforms;
-	
+
+	/** Offset of the root bone from the reference pose. Used to offset bounding box. */
+	UPROPERTY(transient)
+	FVector RootBoneTranslation;
+
+	/** If bEnableLineCheckWithBounds is true, scale the bounds by this value before doing line check. */
+	UPROPERTY()
+	FVector LineCheckBoundsScale;
+
 	/** Temporary storage for curves */
 	FBlendedHeapCurve AnimCurves;
 
-	/** Temporary fix for local space kinematics. This only works for bodies that have no constraints and is needed by vehicles. Proper support will remove this flag */
-	bool bLocalSpaceKinematics;
-
 	// Update Rate
 
 	/** Cached BoneSpaceTransforms for Update Rate optimization. */
@@ -332,13 +332,6 @@
 	UPROPERTY(EditAnywhere, AdvancedDisplay, BlueprintReadOnly, Category = Physics)
 	EDynamicActorScene UseAsyncScene;
 
-	/** If true, there is at least one body in the current PhysicsAsset with a valid bone in the current SkeletalMesh */
-	UPROPERTY(transient)
-	uint32 bHasValidBodies:1;
-
-	/** Indicates that this SkeletalMeshComponent has deferred kinematic bone updates until next physics sim.  */
-	uint32 bDeferredKinematicUpdate:1;
-
 	/** If we are running physics, should we update non-simulated bones based on the animation bone positions. */
 	UPROPERTY(EditAnywhere, AdvancedDisplay, BlueprintReadWrite, Category=SkeletalMesh)
 	TEnumAsByte<EKinematicBonesUpdateToPhysics::Type> KinematicBonesUpdateType;
@@ -347,9 +340,32 @@
 	UPROPERTY(EditAnywhere, BlueprintReadWrite, Category = Physics)
 	TEnumAsByte<EPhysicsTransformUpdateMode::Type> PhysicsTransformUpdateMode;
 
+	/** whether we need to teleport cloth. */
+	EClothingTeleportMode ClothTeleportMode;
+
+protected:
+	/** Whether to use Animation Blueprint or play Single Animation Asset. */
+	UPROPERTY(EditAnywhere, BlueprintReadOnly, Category=Animation)
+	TEnumAsByte<EAnimationMode::Type>	AnimationMode;
+
+private:
+	/** Teleport type to use on the next update */
+	ETeleportType PendingTeleportType;
+
+public:
+	/** Temporary fix for local space kinematics. This only works for bodies that have no constraints and is needed by vehicles. Proper support will remove this flag */
+	uint8 bLocalSpaceKinematics:1;
+
+	/** If true, there is at least one body in the current PhysicsAsset with a valid bone in the current SkeletalMesh */
+	UPROPERTY(transient)
+	uint8 bHasValidBodies:1;
+
+	/** Indicates that this SkeletalMeshComponent has deferred kinematic bone updates until next physics sim.  */
+	uint8 bDeferredKinematicUpdate:1;
+
 	/** Enables blending in of physics bodies whether Simulate or not*/
 	UPROPERTY(transient)
-	uint32 bBlendPhysics:1;
+	uint8 bBlendPhysics:1;
 
 	/**
 	 *  If true, simulate physics for this component on a dedicated server.
@@ -357,41 +373,42 @@
 	 *	Note: This property cannot be changed at runtime.
 	 */
 	UPROPERTY(EditAnywhere, AdvancedDisplay, BlueprintReadWrite, Category = SkeletalMesh)
-	uint32 bEnablePhysicsOnDedicatedServer:1;
+	uint8 bEnablePhysicsOnDedicatedServer:1;
 
 	/**
 	 *	If we should pass joint position to joints each frame, so that they can be used by motorized joints to drive the
 	 *	ragdoll based on the animation.
 	 */
 	UPROPERTY(EditAnywhere, AdvancedDisplay, BlueprintReadWrite, Category=SkeletalMesh)
-	uint32 bUpdateJointsFromAnimation:1;
+	uint8 bUpdateJointsFromAnimation:1;
 
 	/** Disable cloth simulation and play original animation without simulation */
 	UPROPERTY(EditAnywhere, BlueprintReadWrite, Category=Clothing)
-	uint32 bDisableClothSimulation:1;
+	uint8 bDisableClothSimulation:1;
 
 private:
 	/** Disable animation curves for this component. If this is set true, no curves will be processed */
 	UPROPERTY(EditAnywhere, AdvancedDisplay, Category = SkeletalMesh)
-	uint32 bAllowAnimCurveEvaluation : 1;
+	uint8 bAllowAnimCurveEvaluation : 1;
 
 	/** DEPRECATED. Use bAllowAnimCurveEvaluation instead */
 	DEPRECATED(4.18, "This property is deprecated. Please use bAllowAnimCurveEvaluatiuon instead. Note that the meaning is reversed.")	
 	UPROPERTY()
-	uint32 bDisableAnimCurves_DEPRECATED : 1;
+	uint8 bDisableAnimCurves_DEPRECATED : 1;
 
-	/** You can choose to disable certain curves if you prefer. 
-	 * This is transient curves that will be ignored by animation system if you choose this */
-	UPROPERTY(transient)
-	TArray<FName> DisallowedAnimCurves;
+	/** Whether or not we're taking cloth sim information from our master component */
+	uint8 bBindClothToMasterComponent:1;
 
+	/** Flag denoting whether or not the clothing transform needs to update */
+	uint8 bPendingClothTransformUpdate:1;
+
 public:
 	/** can't collide with part of environment if total collision volumes exceed 16 capsules or 32 planes per convex */
 	UPROPERTY(EditAnywhere, BlueprintReadWrite, Category=Clothing)
-	uint32 bCollideWithEnvironment:1;
+	uint8 bCollideWithEnvironment:1;
 	/** can't collide with part of attached children if total collision volumes exceed 16 capsules or 32 planes per convex */
 	UPROPERTY(EditAnywhere, BlueprintReadWrite, Category=Clothing)
-	uint32 bCollideWithAttachedChildren:1;
+	uint8 bCollideWithAttachedChildren:1;
 	/**
 	 * It's worth trying this option when you feel that the current cloth simulation is unstable.
 	 * The scale of the actor is maintained during the simulation. 
@@ -400,7 +417,7 @@
 	 * Known issues: - Currently there's simulation issues when this feature is used in 3.x (DE4076) So if localSpaceSim is enabled there's no inertia effect when the global pose of the clothing actor changes.
 	 */
 	UPROPERTY(EditAnywhere, BlueprintReadWrite, Category = Clothing)
-	uint32 bLocalSpaceSimulation : 1;
+	uint8 bLocalSpaceSimulation : 1;
 
 	/**
 	 * cloth morph target option
@@ -408,25 +425,14 @@
 	 * so it's impossible to change this option in run-time
 	 */
 	UPROPERTY(EditAnywhere, Category = Clothing)
-	uint32 bClothMorphTarget : 1;
+	uint8 bClothMorphTarget : 1;
 
 	/** reset the clothing after moving the clothing position (called teleport) */
 	UPROPERTY(EditAnywhere, BlueprintReadWrite, Category=Clothing)
-	uint32 bResetAfterTeleport:1;
+	uint8 bResetAfterTeleport:1;
 
-	/**
-	 * weight to blend between simulated results and key-framed positions
-	 * if weight is 1.0, shows only cloth simulation results and 0.0 will show only skinned results
-	 */
-	UPROPERTY(EditAnywhere, BlueprintReadWrite, Category = Clothing)
-	float ClothBlendWeight;
-
 	/** To save previous state */
-	uint32 bPrevDisableClothSimulation:1;
+	uint8 bPrevDisableClothSimulation : 1;
 
-	/** Offset of the root bone from the reference pose. Used to offset bounding box. */
-	UPROPERTY(transient)
-	FVector RootBoneTranslation;
-
 	/** 
 	 * Optimization
@@ -434,37 +440,25 @@
 	
 	 /** Whether animation and world transform updates are deferred. If this is on, the kinematic bodies (scene query data) will not update until the next time the physics simulation is run */
 	UPROPERTY(EditAnywhere, AdvancedDisplay, BlueprintReadOnly, Category = SkeletalMesh)
-	uint32 bDeferMovementFromSceneQueries : 1;
+	uint8 bDeferMovementFromSceneQueries : 1;
 
 	/** Skips Ticking and Bone Refresh. */
 	UPROPERTY(EditAnywhere, AdvancedDisplay, BlueprintReadWrite, Category=SkeletalMesh)
-	uint32 bNoSkeletonUpdate:1;
+	uint8 bNoSkeletonUpdate:1;
 
 	/** pauses this component's animations (doesn't tick them, but still refreshes bones) */
 	UPROPERTY(EditAnywhere, AdvancedDisplay, BlueprintReadWrite, Category=Animation)
-	uint32 bPauseAnims:1;
+	uint8 bPauseAnims:1;
 
 	/** On InitAnim should we set to ref pose (if false use first tick of animation data)*/
 	UPROPERTY(EditAnywhere, AdvancedDisplay, Category = Animation)
-	bool bUseRefPoseOnInitAnim;
+	uint8 bUseRefPoseOnInitAnim:1;
 
 	/**
 	* Uses skinned data for collision data.
 	*/
 	UPROPERTY(EditAnywhere, AdvancedDisplay, BlueprintReadOnly, Category=SkeletalMesh)
-	uint32 bEnablePerPolyCollision:1;
+	uint8 bEnablePerPolyCollision:1;
 
 	/**
-	* Used for per poly collision. In 99% of cases you will be better off using a Physics Asset.
-	* This BodySetup is per instance because all modification of vertices is done in place */
-	UPROPERTY(transient)
-	class UBodySetup * BodySetup;
-
-	void CreateBodySetup();
-
-#if !(UE_BUILD_SHIPPING || UE_BUILD_TEST)
-	virtual void SendRenderDebugPhysics(FPrimitiveSceneProxy* OverrideSceneProxy = nullptr) override;
-#endif
-
-	/**
 	 * Misc 
@@ -472,66 +466,140 @@
 	
 	/** If true, force the mesh into the reference pose - is an optimization. */
 	UPROPERTY()
-	bool bForceRefpose;
+	uint8 bForceRefpose:1;
 	
 	/** If true TickPose() will not be called from the Component's TickComponent function.
 	* It will instead be called from Autonomous networking updates. See ACharacter. */
 	UPROPERTY(Transient)
-	uint32 bOnlyAllowAutonomousTickPose : 1;
+	uint8 bOnlyAllowAutonomousTickPose : 1;
 
 	/** True if calling TickPose() from Autonomous networking updates. See ACharacter. */
 	UPROPERTY(Transient)
-	uint32 bIsAutonomousTickPose : 1;
+	uint8 bIsAutonomousTickPose : 1;
 
 	/** If bForceRefPose was set last tick. */
 	UPROPERTY()
-	uint32 bOldForceRefPose:1;
+	uint8 bOldForceRefPose:1;
 
 	/** Bool that enables debug drawing of the skeleton before it is passed to the physics. Useful for debugging animation-driven physics. */
 	UPROPERTY()
-	uint32 bShowPrePhysBones:1;
+	uint8 bShowPrePhysBones:1;
 
 	/** If false, indicates that on the next call to UpdateSkelPose the RequiredBones array should be recalculated. */
 	UPROPERTY(transient)
-	uint32 bRequiredBonesUpToDate:1;
+	uint8 bRequiredBonesUpToDate:1;
 
 	/** If true, AnimTree has been initialised. */
 	UPROPERTY(transient)
-	uint32 bAnimTreeInitialised:1;
+	uint8 bAnimTreeInitialised:1;
 
 	/** If true, the Location of this Component will be included into its bounds calculation
 	* (this can be useful when using SMU_OnlyTickPoseWhenRendered on a character that moves away from the root and no bones are left near the origin of the component) */
 	UPROPERTY(EditAnywhere, AdvancedDisplay, BlueprintReadOnly, Category = SkeletalMesh)
-	uint32 bIncludeComponentLocationIntoBounds : 1;
+	uint8 bIncludeComponentLocationIntoBounds : 1;
 
 	/** If true, line checks will test against the bounding box of this skeletal mesh component and return a hit if there is a collision. */
 	UPROPERTY()
-	uint32 bEnableLineCheckWithBounds:1;
+	uint8 bEnableLineCheckWithBounds:1;
 
 protected:
+
+	/** Whether the clothing simulation is suspended (not the same as disabled, we no longer run the sim but keep the last valid sim data around) */
+	uint8 bClothingSimulationSuspended:1;
+
 #if WITH_EDITORONLY_DATA
 	/** If true, this will Tick until disabled */
 	UPROPERTY(AdvancedDisplay, EditInstanceOnly, transient, Category = SkeletalMesh)
-	uint32 bUpdateAnimationInEditor : 1;
+	uint8 bUpdateAnimationInEditor : 1;
 #endif
+
+private:
+
+	UPROPERTY(Transient)
+	uint8 bNeedsQueuedAnimEventsDispatched:1;
+
+	uint8 bPostEvaluatingAnimation:1;
+
+	/** You can choose to disable certain curves if you prefer. 
+	 * This is transient curves that will be ignored by animation system if you choose this */
+	UPROPERTY(transient)
+	TArray<FName> DisallowedAnimCurves;
 
 public:
 
 	/** Cache AnimCurveUidVersion from Skeleton and this will be used to identify if it needs to be updated */
 	UPROPERTY(transient)
 	uint16 CachedAnimCurveUidVersion;
-	
-	/** If bEnableLineCheckWithBounds is true, scale the bounds by this value before doing line check. */
-	UPROPERTY()
-	FVector LineCheckBoundsScale;
 
+	/**
+	 * weight to blend between simulated results and key-framed positions
+	 * if weight is 1.0, shows only cloth simulation results and 0.0 will show only skinned results
+	 */
+	UPROPERTY(EditAnywhere, BlueprintReadWrite, Category = Clothing)
+	float ClothBlendWeight;
+
+	/**
+	* Used for per poly collision. In 99% of cases you will be better off using a Physics Asset.
+	* This BodySetup is per instance because all modification of vertices is done in place */
+	UPROPERTY(transient)
+	class UBodySetup * BodySetup;
+
+	void CreateBodySetup();
+
+#if !(UE_BUILD_SHIPPING || UE_BUILD_TEST)
+	virtual void SendRenderDebugPhysics(FPrimitiveSceneProxy* OverrideSceneProxy = nullptr) override;
+#endif
+
+public:
+
 	/** Threshold for physics asset bodies above which we use an aggregate for broadphase collisions */
 	int32 RagdollAggregateThreshold;
 
+	float ClothMaxDistanceScale;
+
 	/** Notification when constraint is broken. */
 	UPROPERTY(BlueprintAssignable)
 	FConstraintBrokenSignature OnConstraintBroken;
 
+	/** Class of the object responsible for  */
+	UPROPERTY(EditAnywhere, Category = ClothingSimulation)
+	TSubclassOf<class UClothingSimulationFactory> ClothingSimulationFactory;
+
+	struct FPendingRadialForces
+	{
+		enum EType
+		{
+			AddImpulse,
+			AddForce
+		};
+
+		FVector Origin;
+		float Radius;
+		float Strength;
+		ERadialImpulseFalloff Falloff;
+		bool bIgnoreMass;
+		EType Type;
+		int32 FrameNum;
+
+		FPendingRadialForces(FVector InOrigin, float InRadius, float InStrength, ERadialImpulseFalloff InFalloff, bool InIgnoreMass, EType InType)
+			: Origin(InOrigin)
+			, Radius(InRadius)
+			, Strength(InStrength)
+			, Falloff(InFalloff)
+			, bIgnoreMass(InIgnoreMass)
+			, Type(InType)
+			, FrameNum(GFrameNumber)
+		{
+		}
+	};
+
+	const TArray<FPendingRadialForces>& GetPendingRadialForces() const
+	{
+		return PendingRadialForces;
+	}
+	/** Array of physical interactions for the frame. This is a temporary solution for a more permanent force system and should not be used directly*/
+	TArray<FPendingRadialForces> PendingRadialForces;
+
 	UFUNCTION(BlueprintCallable, Category="Components|SkeletalMesh", meta=(Keywords = "AnimBlueprint"))
 	void SetAnimInstanceClass(class UClass* NewClass);
 
@@ -803,18 +871,7 @@
 	virtual void ClearAnimNotifyErrors(UObject* InSourceNotify){}
 #endif
 
-protected:
-
-	/** Whether the clothing simulation is suspended (not the same as disabled, we no longer run the sim but keep the last valid sim data around) */
-	bool bClothingSimulationSuspended;
-
 public:
-	/** Temporary array of bone indices required this frame. Filled in by UpdateSkelPose. */
-	TArray<FBoneIndexType> RequiredBones;
-
-	/** Temporary array of bone indices required to populate component space transforms */
-	TArray<FBoneIndexType> FillComponentSpaceTransformsRequiredBones;
-
 	/** 
 	 *	Index of the 'Root Body', or top body in the asset hierarchy. 
 	 *	Filled in by InitInstance, so we don't need to save it.
@@ -831,38 +888,11 @@
 	/** Reset Root Body Index */
 	void ResetRootBodyIndex();
 
-	struct FPendingRadialForces
-	{
-		enum EType
-		{
-			AddImpulse,
-			AddForce
-		};
+	/** Temporary array of bone indices required this frame. Filled in by UpdateSkelPose. */
+	TArray<FBoneIndexType> RequiredBones;
 
-		FVector Origin;
-		float Radius;
-		float Strength;
-		ERadialImpulseFalloff Falloff;
-		bool bIgnoreMass;
-		EType Type;
-		int32 FrameNum;
+	/** Temporary array of bone indices required to populate component space transforms */
+	TArray<FBoneIndexType> FillComponentSpaceTransformsRequiredBones;
 
-		FPendingRadialForces(FVector InOrigin, float InRadius, float InStrength, ERadialImpulseFalloff InFalloff, bool InIgnoreMass, EType InType)
-			: Origin(InOrigin)
-			, Radius(InRadius)
-			, Strength(InStrength)
-			, Falloff(InFalloff)
-			, bIgnoreMass(InIgnoreMass)
-			, Type(InType)
-			, FrameNum(GFrameNumber)
-		{
-		}
-	};
-
-	const TArray<FPendingRadialForces>& GetPendingRadialForces() const
-	{
-		return PendingRadialForces;
-	}
-
 	/** Array of FBodyInstance objects, storing per-instance state about about each body. */
 	TArray<struct FBodyInstance*> Bodies;
@@ -878,10 +908,6 @@
 
 	FSkeletalMeshComponentClothTickFunction ClothTickFunction;
 
-	/** Class of the object responsible for  */
-	UPROPERTY(EditAnywhere, Category = ClothingSimulation)
-	TSubclassOf<class UClothingSimulationFactory> ClothingSimulationFactory;
-
 	/**
 	* Gets the teleportation rotation threshold.
 	* 
@@ -939,9 +965,20 @@
 	void ComputeTeleportRotationThresholdInRadians();
 	void ComputeTeleportDistanceThresholdInRadians();
 
+	// Can't rely on time value, because those may be affected by dilation and whether or not
+	// the game is paused.
+	// Also can't just rely on a flag as other components (like CharacterMovementComponent) may tick
+	// the pose and we can't guarantee tick order.
+	UPROPERTY(Transient)
+	uint32 LastPoseTickFrame;
+
+	/** Ref for the clothing parallel task, so we can detect whether or not a sim is running */
+	FGraphEventRef ParallelClothTask;
+
 public:
-	/** whether we need to teleport cloth. */
-	EClothingTeleportMode ClothTeleportMode;
 
+	/** Checked whether we have already ticked the pose this frame */
+	bool PoseTickedThisFrame() const;
+
 	bool IsClothBoundToMasterComponent() const { return bBindClothToMasterComponent; }
 
@@ -959,9 +996,6 @@
 
 private:
 
-	/** Array of physical interactions for the frame. This is a temporary solution for a more permanent force system and should not be used directly*/
-	TArray<FPendingRadialForces> PendingRadialForces;
-
 	/** Let the cloth tick and completion tasks have access to private clothing data */
 	friend FSkeletalMeshComponentClothTickFunction;
 	friend class FParallelClothCompletionTask;
@@ -969,9 +1003,6 @@
 	/** Debug mesh component should be able to access for visualisation */
 	friend class UDebugSkelMeshComponent;
 
-	/** Whether or not we're taking cloth sim information from our master component */
-	bool bBindClothToMasterComponent;
-
 	/** Copies the data from the external cloth simulation context. We copy instead of flipping because the API has to return the full struct to make backwards compat easy*/
 	void UpdateClothSimulationContext(float InDeltaTime);
 
@@ -986,18 +1017,9 @@
 	IClothingSimulation* ClothingSimulation;
 	IClothingSimulationContext* ClothingSimulationContext;
 
-	/** Ref for the clothing parallel task, so we can detect whether or not a sim is running */
-	FGraphEventRef ParallelClothTask;
-
 	/** Stalls on any currently running clothing simulations, needed when changing core sim state */
 	void HandleExistingParallelClothSimulation();
 
-	/** Flag denoting whether or not the clothing transform needs to update */
-	bool bPendingClothTransformUpdate;
-
-	/** Teleport type to use on the next update */
-	ETeleportType PendingTeleportType;
-
 	/** Called by the clothing completion event to perform a writeback of the simulation data 
 	 * to the game thread, the task is friended to gain access to this and not allow any
 	 * external callers to trigger writebacks
@@ -1017,8 +1039,6 @@
 
 public:
 
-	float ClothMaxDistanceScale;
-
 	static uint32 GetPhysicsSceneType(const UPhysicsAsset& PhysAsset, const FPhysScene& PhysScene, EDynamicActorScene SimulationScene);
 
 private:
@@ -1635,15 +1655,15 @@
 
 	void GetWindForCloth_GameThread(FVector& WindVector, float& WindAdaption) const;
 	
-	//Data for parallel evaluation of animation
-	FAnimationEvaluationContext AnimEvaluationContext;
-
 	// Reference to our current parallel animation evaluation task (if there is one)
 	FGraphEventRef				ParallelAnimationEvaluationTask;
 
 	// Reference to our current blend physics task (if there is one)
 	FGraphEventRef				ParallelBlendPhysicsCompletionTask;
 
+	//Data for parallel evaluation of animation
+	FAnimationEvaluationContext AnimEvaluationContext;
+
 public:
 	// Parallel evaluation wrappers
 	void ParallelAnimationEvaluation();
@@ -1713,6 +1733,7 @@
 
 	friend class FTickClothingTask;
 
+#if WITH_EDITORONLY_DATA
 	// these are deprecated variables from removing SingleAnimSkeletalComponent
 	// remove if this version goes away : VER_UE4_REMOVE_SINGLENODEINSTANCE
 	// deprecated variable to be re-save
@@ -1738,7 +1759,7 @@
 	// Default setting for playrate of SequenceToPlay to play. 
 	UPROPERTY()
 	float DefaultPlayRate_DEPRECATED;
-	
+#endif
 
 	/*
 	 * Update MorphTargetCurves from mesh - these are not animation curves, but SetMorphTarget and similar functions that can set to this mesh component
@@ -1749,22 +1770,7 @@
 	 */
 	void ResetMorphTargetCurves();
 
-	// Can't rely on time value, because those may be affected by dilation and whether or not
-	// the game is paused.
-	// Also can't just rely on a flag as other components (like CharacterMovementComponent) may tick
-	// the pose and we can't guarantee tick order.
-	UPROPERTY(Transient)
-	uint32 LastPoseTickFrame;
-
 public:
-	/** Keep track of when animation has been ticked to ensure it is ticked only once per frame. */
-	DEPRECATED(4.16, "This property is deprecated. Please use PoseTickedThisFrame instead.")
-	UPROPERTY(Transient)
-	float LastPoseTickTime;
-
-	/** Checked whether we have already ticked the pose this frame */
-	bool PoseTickedThisFrame() const;
-
 	/** Take extracted RootMotion and convert it from local space to world space. */
 	FTransform ConvertLocalRootMotionToWorld(const FTransform& InTransform);
 
@@ -1828,10 +1834,4 @@
 
 	// Are we currently within PostAnimEvaluation
 	bool IsPostEvaluatingAnimation() const { return bPostEvaluatingAnimation; }
-
-private: 
-	UPROPERTY(Transient)
-	bool bNeedsQueuedAnimEventsDispatched;
-
-	bool bPostEvaluatingAnimation;
 };
diff --git a/Engine/Source/Runtime/Engine/Classes/Components/SkinnedMeshComponent.h b/Engine/Source/Runtime/Engine/Classes/Components/SkinnedMeshComponent.h
--- a/Engine/Source/Runtime/Engine/Classes/Components/SkinnedMeshComponent.h
+++ b/Engine/Source/Runtime/Engine/Classes/Components/SkinnedMeshComponent.h
@@ -206,9 +206,6 @@
 
 	/** The index for the ComponentSpaceTransforms buffer we can currently read from */
 	int32 CurrentReadComponentTransforms;
-protected:
-	/** Are we using double buffered ComponentSpaceTransforms */
-	bool bDoubleBufferedComponentSpaceTransforms;
 
 	/** 
 	 * If set, this component has slave pose components that are associated with this 
@@ -225,6 +222,14 @@
 	/** Incremented every time the master bone map changes. Used to keep in sync with any duplicate data needed by other threads */
 	int32 MasterBoneMapCacheCount;
 
+public:
+	/**
+	 * Wireframe color
+	 */
+	UPROPERTY(EditAnywhere, AdvancedDisplay, BlueprintReadOnly, Category=SkeletalMesh)
+	FColor WireframeColor;
+
+protected:
 	/** Information for current ref pose override, if present */
 	FSkelMeshRefPoseOverride* RefPoseOverride;
 
@@ -240,14 +245,6 @@
 	 */
 	void GetCPUSkinnedVertices(TArray<struct FFinalSkinVertex>& OutVertices, int32 InLODIndex);
 
-	/** 
-	 * When true, we will just using the bounds from our MasterPoseComponent.  This is useful for when we have a Mesh Parented
-	 * to the main SkelMesh (e.g. outline mesh or a full body overdraw effect that is toggled) that is always going to be the same
-	 * bounds as parent.  We want to do no calculations in that case.
-	 */
-	UPROPERTY(EditAnywhere, AdvancedDisplay, BlueprintReadWrite, Category = SkeletalMesh)
-	uint32 bUseBoundsFromMasterPoseComponent:1;
-
 	/** Array indicating all active morph targets. This array is updated inside RefreshBoneTransforms based on the Anim Blueprint. */
 	TArray<FActiveMorphTarget> ActiveMorphTargets;
 
@@ -305,14 +302,6 @@
 	/**	High (best) DistanceFactor that was desired for rendering this USkeletalMesh last frame. Represents how big this mesh was in screen space   */
 	float MaxDistanceFactor;
 
-	/** LOD array info. Each index will correspond to the LOD index **/
-	UPROPERTY(transient)
-	TArray<struct FSkelMeshComponentLODInfo> LODInfo;
-
-	//
-	// Rendering options.
-	//
-	
 	/**
 	 * Allows adjusting the desired streaming distance of streaming textures that uses UV 0.
 	 * 1.0 is the default, whereas a higher value makes the textures stream in sooner from far away.
@@ -322,35 +311,51 @@
 	UPROPERTY(EditAnywhere, AdvancedDisplay, BlueprintReadWrite, Category=SkeletalMesh)
 	float StreamingDistanceMultiplier;
 
-	/**
-	 * Wireframe color
+	/** LOD array info. Each index will correspond to the LOD index **/
+	UPROPERTY(transient)
+	TArray<struct FSkelMeshComponentLODInfo> LODInfo;
+
+	/** Array of bone visibilities (containing one of the values in EBoneVisibilityStatus for each bone).  A bone is only visible if it is *exactly* 1 (BVS_Visible) */
+	TArray<uint8> BoneVisibilityStates;
+
+	/** This is update frequency flag even when our Owner has not been rendered recently
+	 * 
+	 * SMU_AlwaysTickPoseAndRefreshBones,			// Always Tick and Refresh BoneTransforms whether rendered or not
+	 * SMU_AlwaysTickPose,							// Always Tick, but Refresh BoneTransforms only when rendered
+	 * SMU_OnlyTickPoseWhenRendered,				// Tick only when rendered, and it will only RefreshBoneTransforms when rendered
+	 * 
 	 */
-	UPROPERTY(EditAnywhere, AdvancedDisplay, BlueprintReadOnly, Category=SkeletalMesh)
-	FColor WireframeColor;
+	UPROPERTY(EditAnywhere, AdvancedDisplay, BlueprintReadWrite, Category=SkeletalMesh)
+	TEnumAsByte<EMeshComponentUpdateFlag::Type> MeshComponentUpdateFlag;
 
+	/** 
+	 * When true, we will just using the bounds from our MasterPoseComponent.  This is useful for when we have a Mesh Parented
+	 * to the main SkelMesh (e.g. outline mesh or a full body overdraw effect that is toggled) that is always going to be the same
+	 * bounds as parent.  We want to do no calculations in that case.
+	 */
+	UPROPERTY(EditAnywhere, AdvancedDisplay, BlueprintReadWrite, Category = SkeletalMesh)
+	uint8 bUseBoundsFromMasterPoseComponent:1;
+
 	/** Forces the mesh to draw in wireframe mode. */
 	UPROPERTY()
-	uint32 bForceWireframe:1;
+	uint8 bForceWireframe:1;
 
 	/** Draw the skeleton hierarchy for this skel mesh. */
 	UPROPERTY()
-	uint32 bDisplayBones_DEPRECATED:1;
+	uint8 bDisplayBones_DEPRECATED:1;
 
 	/** Disable Morphtarget for this component. */
 	UPROPERTY(EditAnywhere, AdvancedDisplay, BlueprintReadWrite, Category = SkeletalMesh)
-	uint32 bDisableMorphTarget:1;
+	uint8 bDisableMorphTarget:1;
 
 	/** Don't bother rendering the skin. */
 	UPROPERTY()
-	uint32 bHideSkin:1;
-	/** Array of bone visibilities (containing one of the values in EBoneVisibilityStatus for each bone).  A bone is only visible if it is *exactly* 1 (BVS_Visible) */
-	TArray<uint8> BoneVisibilityStates;
-
+	uint8 bHideSkin:1;
 	/**
 	 *	If true, use per-bone motion blur on this skeletal mesh (requires additional rendering, can be disabled to save performance).
 	 */
 	UPROPERTY(EditAnywhere, AdvancedDisplay, BlueprintReadOnly, Category=SkeletalMesh)
-	uint32 bPerBoneMotionBlur:1;
+	uint8 bPerBoneMotionBlur:1;
 
 	//
 	// Misc.
@@ -358,36 +363,20 @@
 	
 	/** When true, skip using the physics asset etc. and always use the fixed bounds defined in the SkeletalMesh. */
 	UPROPERTY(EditAnywhere, AdvancedDisplay, BlueprintReadWrite, Category=SkeletalMesh)
-	uint32 bComponentUseFixedSkelBounds:1;
+	uint8 bComponentUseFixedSkelBounds:1;
 
 	/** If true, when updating bounds from a PhysicsAsset, consider _all_ BodySetups, not just those flagged with bConsiderForBounds. */
 	UPROPERTY(EditAnywhere, AdvancedDisplay, BlueprintReadWrite, Category=SkeletalMesh)
-	uint32 bConsiderAllBodiesForBounds:1;
+	uint8 bConsiderAllBodiesForBounds:1;
 
-	/** This is update frequency flag even when our Owner has not been rendered recently
-	 * 
-	 * SMU_AlwaysTickPoseAndRefreshBones,			// Always Tick and Refresh BoneTransforms whether rendered or not
-	 * SMU_AlwaysTickPose,							// Always Tick, but Refresh BoneTransforms only when rendered
-	 * SMU_OnlyTickPoseWhenRendered,				// Tick only when rendered, and it will only RefreshBoneTransforms when rendered
-	 * 
-	 */
-	UPROPERTY(EditAnywhere, AdvancedDisplay, BlueprintReadWrite, Category=SkeletalMesh)
-	TEnumAsByte<EMeshComponentUpdateFlag::Type> MeshComponentUpdateFlag;
 
-private:
-	/** If true, UpdateTransform will always result in a call to MeshObject->Update. */
-	UPROPERTY(transient)
-	uint32 bForceMeshObjectUpdate:1;
-
-public:
-
 	/** Whether or not we can highlight selected sections - this should really only be done in the editor */
 	UPROPERTY(transient)
-	uint32 bCanHighlightSelectedSections:1;
+	uint8 bCanHighlightSelectedSections:1;
 
 	/** true if mesh has been recently rendered, false otherwise */
 	UPROPERTY(transient)
-	uint32 bRecentlyRendered:1;
+	uint8 bRecentlyRendered:1;
 
 	/** 
 	 * Whether to use the capsule representation (when present) from a skeletal mesh's ShadowPhysicsAsset for direct shadowing from lights.
@@ -395,13 +384,52 @@
 	 * This flag will force bCastInsetShadow to be enabled.
 	 */
 	UPROPERTY(EditAnywhere, AdvancedDisplay, BlueprintReadOnly, Category=Lighting, meta=(EditCondition="CastShadow", DisplayName = "Capsule Direct Shadow"))
-	uint32 bCastCapsuleDirectShadow:1;
+	uint8 bCastCapsuleDirectShadow:1;
 
 	/** 
 	 * Whether to use the capsule representation (when present) from a skeletal mesh's ShadowPhysicsAsset for shadowing indirect lighting (from lightmaps or skylight).
 	 */
 	UPROPERTY(EditAnywhere, AdvancedDisplay, BlueprintReadOnly, Category=Lighting, meta=(EditCondition="CastShadow", DisplayName = "Capsule Indirect Shadow"))
-	uint32 bCastCapsuleIndirectShadow:1;
+	uint8 bCastCapsuleIndirectShadow:1;
+
+	/** Whether or not to CPU skin this component, requires render data refresh after changing */
+	UPROPERTY(transient)
+	uint8 bCPUSkinning : 1;
+
+	// Update Rate
+	/** if TRUE, Owner will determine how often animation will be updated and evaluated. See AnimUpdateRateTick() 
+	 * This allows to skip frames for performance. (For example based on visibility and size on screen). */
+	UPROPERTY(EditAnywhere, AdvancedDisplay, BlueprintReadWrite, Category=Optimization)
+	uint8 bEnableUpdateRateOptimizations:1;
+
+	/** Enable on screen debugging of update rate optimization. 
+	 * Red = Skipping 0 frames, Green = skipping 1 frame, Blue = skipping 2 frames, black = skipping more than 2 frames. 
+	 * @todo: turn this into a console command. */
+	UPROPERTY(EditAnywhere, AdvancedDisplay, BlueprintReadWrite, Category=Optimization)
+	uint8 bDisplayDebugUpdateRateOptimizations:1;
+
+protected:
+	/** Are we using double buffered ComponentSpaceTransforms */
+	uint8 bDoubleBufferedComponentSpaceTransforms : 1;
+
+	/** Track whether we still need to flip to recently modified buffer */
+	uint8 bNeedToFlipSpaceBaseBuffers : 1;
+
+	/** true when CachedLocalBounds is up to date. */
+	UPROPERTY(Transient)
+	mutable uint8 bCachedLocalBoundsUpToDate:1;
+
+private:
+	/** If true, UpdateTransform will always result in a call to MeshObject->Update. */
+	UPROPERTY(transient)
+	uint8 bForceMeshObjectUpdate:1;
+
+public:
+	/** Object responsible for sending bone transforms, morph target state etc. to render thread. */
+	class FSkeletalMeshObject*	MeshObject;
+
+	/** Gets the skeletal mesh resource used for rendering the component. */
+	FSkeletalMeshRenderData* GetSkeletalMeshRenderData() const;
 
 	/** 
 	 * Controls how dark the capsule indirect shadow can be.
@@ -409,10 +437,6 @@
 	UPROPERTY(EditAnywhere, AdvancedDisplay, BlueprintReadOnly, Category=Lighting, meta=(UIMin = "0", UIMax = "1", EditCondition="bCastCapsuleIndirectShadow", DisplayName = "Capsule Indirect Shadow Min Visibility"))
 	float CapsuleIndirectShadowMinVisibility;
 
-	/** Whether or not to CPU skin this component, requires render data refresh after changing */
-	UPROPERTY(transient)
-	uint32 bCPUSkinning : 1;
-
 	/** 
 	 * Override the Physics Asset of the mesh. It uses SkeletalMesh.PhysicsAsset, but if you'd like to override use this function
 	 * 
@@ -518,12 +542,6 @@
 	FName GetParentBone(FName BoneName) const;
 
 public:
-	/** Object responsible for sending bone transforms, morph target state etc. to render thread. */
-	class FSkeletalMeshObject*	MeshObject;
-
-	/** Gets the skeletal mesh resource used for rendering the component. */
-	FSkeletalMeshRenderData* GetSkeletalMeshRenderData() const;
-
 	//~ Begin UObject Interface
 	virtual void BeginDestroy() override;
 	virtual void Serialize(FArchive& Ar) override;
@@ -744,22 +762,6 @@
 
 	void SetComponentSpaceTransformsDoubleBuffering(bool bInDoubleBufferedComponentSpaceTransforms);
 
-
-	DEPRECATED(4.13, "GetComponentSpaceTransforms is now renamed GetComponentSpaceTransforms")
-	const TArray<FTransform>& GetSpaceBases() const { return GetComponentSpaceTransforms(); }
-
-	DEPRECATED(4.13, "GetEditableSpaceBases is now renamed GetEditableComponentSpaceTransforms")
-	TArray<FTransform>& GetEditableSpaceBases() { return GetEditableComponentSpaceTransforms(); }
-
-	DEPRECATED(4.13, "GetEditableSpaceBases is now renamed GetEditableComponentSpaceTransforms")
-	const TArray<FTransform>& GetEditableSpaceBases() const { return GetEditableComponentSpaceTransforms(); }
-
-	DEPRECATED(4.13, "GetNumSpaceBases is now renamed GetNumComponentSpaceTransforms")
-	int32 GetNumSpaceBases() const { return GetNumComponentSpaceTransforms(); }
-
-	DEPRECATED(4.13, "SetSpaceBaseDoubleBuffering is now renamed SetComponentSpaceTransformsDoubleBuffering")
-	void SetSpaceBaseDoubleBuffering(bool bInDoubleBufferedBlendSpaces) { SetComponentSpaceTransformsDoubleBuffering(bInDoubleBufferedBlendSpaces);  }
-
 	const FBoxSphereBounds& GetCachedLocalBounds() { return CachedLocalBounds; } 
 
 protected:
@@ -767,9 +769,6 @@
 	/** Flip the editable space base buffer */
 	void FlipEditableSpaceBases();
 
-	/** Track whether we still need to flip to recently modified buffer */
-	bool bNeedToFlipSpaceBaseBuffers;
-
 	/** 
 	 * Should update transform in Tick
 	 * 
@@ -797,10 +796,6 @@
 	UPROPERTY(Transient)
 	mutable FBoxSphereBounds CachedLocalBounds;
 
-	/** true when CachedLocalBounds is up to date. */
-	UPROPERTY(Transient)
-	mutable bool bCachedLocalBoundsUpToDate;
-
 public:
 
 	/** Invalidate Cached Bounds, when Mesh Component has been updated. */
@@ -821,19 +816,6 @@
 	 */
 	bool ShouldUpdateBoneVisibility() const;
 
-	// Update Rate
-public:
-	/** if TRUE, Owner will determine how often animation will be updated and evaluated. See AnimUpdateRateTick() 
-	 * This allows to skip frames for performance. (For example based on visibility and size on screen). */
-	UPROPERTY(EditAnywhere, AdvancedDisplay, BlueprintReadWrite, Category=Optimization)
-	bool bEnableUpdateRateOptimizations;
-
-	/** Enable on screen debugging of update rate optimization. 
-	 * Red = Skipping 0 frames, Green = skipping 1 frame, Blue = skipping 2 frames, black = skipping more than 2 frames. 
-	 * @todo: turn this into a console command. */
-	UPROPERTY(EditAnywhere, AdvancedDisplay, BlueprintReadWrite, Category=Optimization)
-	bool bDisplayDebugUpdateRateOptimizations;
-
 protected:
 
 	/** Removes update rate params and internal tracker data */
@@ -1097,12 +1079,12 @@
 
 	// Animation update rate control.
 public:
-	/** Animation Update Rate optimization parameters. */
-	struct FAnimUpdateRateParameters* AnimUpdateRateParams;
-
 	/** Delegate when AnimUpdateRateParams is created, to override its default settings. */
 	FOnAnimUpdateRateParamsCreated OnAnimUpdateRateParamsCreated;
 
+	/** Animation Update Rate optimization parameters. */
+	struct FAnimUpdateRateParameters* AnimUpdateRateParams;
+
 	/** Updates AnimUpdateRateParams, used by SkinnedMeshComponents.
 	* 
 	* @param bRecentlyRendered : true if at least one SkinnedMeshComponent on this Actor has been rendered in the last second.
diff --git a/Engine/Source/Runtime/Engine/Classes/EdGraph/EdGraphNode.h b/Engine/Source/Runtime/Engine/Classes/EdGraph/EdGraphNode.h
--- a/Engine/Source/Runtime/Engine/Classes/EdGraph/EdGraphNode.h
+++ b/Engine/Source/Runtime/Engine/Classes/EdGraph/EdGraphNode.h
@@ -45,11 +45,11 @@ struct ENGINE_API FEdGraphTerminalType
 
 	/** Category */
 	UPROPERTY()
-	FString TerminalCategory;
+	FName TerminalCategory;
 
 	/** Sub-category */
 	UPROPERTY()
-	FString TerminalSubCategory;
+	FName TerminalSubCategory;
 
 	/** Sub-category object */
 	UPROPERTY()
@@ -340,15 +340,28 @@ public:
 	/** Get all pins this node owns */
 	TArray<UEdGraphPin*> GetAllPins() { return Pins; }
 
+	struct FNameParameterHelper
+	{
+		FNameParameterHelper(const FName InNameParameter) : NameParameter(InNameParameter) { }
+		FNameParameterHelper(const FString& InNameParameter) : NameParameter(*InNameParameter) { }
+		FNameParameterHelper(const TCHAR* InNameParameter) : NameParameter(InNameParameter) { }
+
+		FName operator*() const { return NameParameter; }
+
+	private:
+		FName NameParameter;
+	};
+
 	/** Create a new pin on this node using the supplied info, and return the new pin */
+	DEPRECATED(4.19, "Use version that supplies Pin Category, SubCategory, and Name as an FName and uses PinContainerType instead of separate booleans for array, set, and map.")
 	UEdGraphPin* CreatePin(
 		EEdGraphPinDirection Dir, 
-		const FString& PinCategory, 
-		const FString& PinSubCategory, 
+		const FNameParameterHelper PinCategory, 
+		const FNameParameterHelper PinSubCategory, 
 		UObject* PinSubCategoryObject, 
 		bool bIsArray, 
 		bool bIsReference, 
-		const FString& PinName, 
+		const FNameParameterHelper PinName, 
 		bool bIsConst = false, 
 		int32 Index = INDEX_NONE, 
 		bool bIsSet = false, 
@@ -356,30 +369,111 @@ public:
 		const FEdGraphTerminalType& ValueTerminalType = FEdGraphTerminalType());
 
 	/** Create a new pin on this node using the supplied info, and return the new pin */
+	DEPRECATED(4.19, "Use version that supplies Pin Category, SubCategory, and Name as an FName and uses a parameter structure for optional paramaters.")
 	UEdGraphPin* CreatePin(
-		EEdGraphPinDirection Dir, 
-		const FString& PinCategory, 
-		const FString& PinSubCategory, 
-		UObject* PinSubCategoryObject, 
-		const FString& PinName, 
+		EEdGraphPinDirection Dir,
+		const FNameParameterHelper PinCategory,
+		const FNameParameterHelper PinSubCategory,
+		UObject* PinSubCategoryObject,
+		const FNameParameterHelper PinName,
 		EPinContainerType PinContainerType = EPinContainerType::None,
-		bool bIsReference = false, 
-		bool bIsConst = false, 
-		int32 Index = INDEX_NONE, 
+		bool bIsReference = false,
+		bool bIsConst = false,
+		int32 Index = INDEX_NONE,
 		const FEdGraphTerminalType& ValueTerminalType = FEdGraphTerminalType());
 
+	/** Parameter struct of less common options for CreatePin */
+	struct ENGINE_API FCreatePinParams
+	{
+		FCreatePinParams()
+			: ContainerType(EPinContainerType::None)
+			, bIsReference(false)
+			, bIsConst(false)
+			, Index(INDEX_NONE)
+		{
+		}
+
+		FCreatePinParams(const FEdGraphPinType& PinType);
+
+		EPinContainerType ContainerType;
+		bool bIsReference;
+		bool bIsConst;
+		int32 Index;
+		FEdGraphTerminalType ValueTerminalType;
+	};
+
+	/** Create a new pin on this node using the supplied info, and return the new pin */
+	UEdGraphPin* CreatePin(EEdGraphPinDirection Dir, const FName PinCategory, const FName PinName, const FCreatePinParams& PinParams = FCreatePinParams())
+	{
+		return CreatePin(Dir, PinCategory, NAME_None, nullptr, PinName, PinParams);
+	}
+
+	UEdGraphPin* CreatePin(EEdGraphPinDirection Dir, const FName PinCategory, const FName PinSubCategory, const FName PinName, const FCreatePinParams& PinParams = FCreatePinParams())
+	{
+		return CreatePin(Dir, PinCategory, PinSubCategory, nullptr, PinName, PinParams);
+	}
+
+	UEdGraphPin* CreatePin(EEdGraphPinDirection Dir, const FName PinCategory, UObject* PinSubCategoryObject, const FName PinName, const FCreatePinParams& PinParams = FCreatePinParams())
+	{
+		return CreatePin(Dir, PinCategory, NAME_None, PinSubCategoryObject, PinName, PinParams);
+	}
+
+	UEdGraphPin* CreatePin(EEdGraphPinDirection Dir, const FName PinCategory, const FName PinSubCategory, UObject* PinSubCategoryObject, const FName PinName, const FCreatePinParams& PinParams = FCreatePinParams());
+
 	/** Create a new pin on this node using the supplied pin type, and return the new pin */
-	UEdGraphPin* CreatePin(EEdGraphPinDirection Dir, const FEdGraphPinType& InPinType, const FString& PinName, int32 Index = INDEX_NONE);
+	UEdGraphPin* CreatePin(EEdGraphPinDirection Dir, const FEdGraphPinType& InPinType, const FName PinName, int32 Index = INDEX_NONE);
+
+	/** Create a new pin on this node using the supplied pin type, and return the new pin */
+	DEPRECATED(4.19, "Use version that passes PinName as FName instead.")
+	UEdGraphPin* CreatePin(EEdGraphPinDirection Dir, const FEdGraphPinType& InPinType, const FString& PinName, int32 Index = INDEX_NONE)
+	{
+		return CreatePin(Dir, InPinType, FName(*PinName), Index);
+	}
+
+	/** Create a new pin on this node using the supplied pin type, and return the new pin */
+	//DEPRECATED(4.19, "Remove when removing FString version. Exists just to resolve ambiguity")
+	UEdGraphPin* CreatePin(EEdGraphPinDirection Dir, const FEdGraphPinType& InPinType, const TCHAR* PinName, int32 Index = INDEX_NONE)
+	{
+		return CreatePin(Dir, InPinType, FName(PinName), Index);
+	}
 
 	/** Destroys the specified pin, does not modify its owning pin's Pins list */
 	static void DestroyPin(UEdGraphPin* Pin);
 
 	/** Find a pin on this node with the supplied name and optional direction */
-	UEdGraphPin* FindPin(const FString& PinName, const EEdGraphPinDirection Direction = EGPD_MAX) const;
+	UEdGraphPin* FindPin(const FName PinName, const EEdGraphPinDirection Direction = EGPD_MAX) const;
 
 	/** Find a pin on this node with the supplied name and optional direction and assert if it is not present */
-	UEdGraphPin* FindPinChecked(const FString& PinName, const EEdGraphPinDirection Direction = EGPD_MAX) const;
-	
+	UEdGraphPin* FindPinChecked(const FName PinName, const EEdGraphPinDirection Direction = EGPD_MAX) const
+	{
+		UEdGraphPin* Result = FindPin(PinName, Direction);
+		check(Result);
+		return Result;
+	}
+
+	/** Find a pin on this node with the supplied name and optional direction */
+	UEdGraphPin* FindPin(const FString& PinName, const EEdGraphPinDirection Direction = EGPD_MAX) const
+	{
+		return FindPin(*PinName, Direction);
+	}
+
+	/** Find a pin on this node with the supplied name and optional direction and assert if it is not present */
+	UEdGraphPin* FindPinChecked(const FString& PinName, const EEdGraphPinDirection Direction = EGPD_MAX) const
+	{
+		return FindPinChecked(*PinName, Direction);
+	}
+
+	/** Find a pin on this node with the supplied name and optional direction */
+	UEdGraphPin* FindPin(const TCHAR* PinName, const EEdGraphPinDirection Direction = EGPD_MAX) const;
+
+	/** Find a pin on this node with the supplied name and optional direction and assert if it is not present */
+	UEdGraphPin* FindPinChecked(const TCHAR* PinName, const EEdGraphPinDirection Direction = EGPD_MAX) const
+	{
+		UEdGraphPin* Result = FindPin(PinName, Direction);
+		check(Result);
+		return Result;
+	}
+
 	/** Find the pin on this node with the supplied guid */
 	UEdGraphPin* FindPinById(const FGuid PinId) const;
 
@@ -443,15 +537,15 @@ public:
 	}
 
 	/** Generate a unique pin name, trying to stick close to a passed in name */
-	virtual FString CreateUniquePinName(FString SourcePinName) const
+	virtual FName CreateUniquePinName(FName SourcePinName) const
 	{
-		FString PinName(SourcePinName);
+		FName PinName(SourcePinName);
 		
 		int32 Index = 1;
-		while (FindPin(PinName) != NULL)
+		while (FindPin(PinName) != nullptr)
 		{
 			++Index;
-			PinName = SourcePinName + FString::FromInt(Index);
+			PinName = *FString::Printf(TEXT("%s%d"),*SourcePinName.ToString(),Index);
 		}
 
 		return PinName;
@@ -665,7 +759,7 @@ public:
 	virtual void AddSearchMetaDataInfo(TArray<struct FSearchTagDataPair>& OutTaggedMetaData) const;
 
 	/** Return the requested metadata for the pin if there is any */
-	virtual FString GetPinMetaData(FString InPinName, FName InKey) { return FString(); }
+	virtual FString GetPinMetaData(FName InPinName, FName InKey) { return FString(); }
 
 	/** Return false if the node and any expansion will isolate itself during compile */
 	virtual bool IsCompilerRelevant() const { return true; }
diff --git a/Engine/Source/Runtime/Engine/Classes/EdGraph/EdGraphPin.h b/Engine/Source/Runtime/Engine/Classes/EdGraph/EdGraphPin.h
--- a/Engine/Source/Runtime/Engine/Classes/EdGraph/EdGraphPin.h
+++ b/Engine/Source/Runtime/Engine/Classes/EdGraph/EdGraphPin.h
@@ -1,4 +1,4 @@
-// Copyright 1998-2017 Epic Games, Inc. All Rights Reserved.
+﻿// Copyright 1998-2017 Epic Games, Inc. All Rights Reserved.
 
 #pragma once
 
@@ -86,11 +86,11 @@ struct FEdGraphPinType
 
 	/** Category of pin type */
 	UPROPERTY()
-	FString PinCategory;
+	FName PinCategory;
 
 	/** Sub-category of pin type */
 	UPROPERTY()
-	FString PinSubCategory;
+	FName PinSubCategory;
 
 	/** Sub-category object */
 	UPROPERTY()
@@ -141,10 +141,23 @@ public:
 	{
 	}
 
-	DEPRECATED(4.17, "Use version that takes PinContainerType instead of separate booleans for array, set, and map")
-	FEdGraphPinType(FString InPinCategory, FString InPinSubCategory, UObject* InPinSubCategoryObject, bool bInIsArray, bool bInIsReference, bool bInIsSet, bool bInIsMap, const FEdGraphTerminalType& InValueTerminalType )
-		: PinCategory(MoveTemp(InPinCategory))
-		, PinSubCategory(MoveTemp(InPinSubCategory))
+	//DEPRECATED(4.19, "Remove this helper when removing deprecation as it will no longer be needed")
+	struct FNameParameterHelper
+	{
+		FNameParameterHelper(const FName InNameParameter) : NameParameter(InNameParameter) { }
+		FNameParameterHelper(const FString& InNameParameter) : NameParameter(*InNameParameter) { }
+		FNameParameterHelper(const TCHAR* InNameParameter) : NameParameter(InNameParameter) { }
+
+		FName operator*() const { return NameParameter; }
+
+	private:
+		FName NameParameter;
+	};
+
+	DEPRECATED(4.17, "Use version that supplies PinCategory and SubCategory as FName, and takes PinContainerType instead of separate booleans for array, set, and map")
+	FEdGraphPinType(const FNameParameterHelper InPinCategory, const FNameParameterHelper InPinSubCategory, UObject* InPinSubCategoryObject, bool bInIsArray, bool bInIsReference, bool bInIsSet, bool bInIsMap, const FEdGraphTerminalType& InValueTerminalType )
+		: PinCategory(*InPinCategory)
+		, PinSubCategory(*InPinSubCategory)
 		, PinSubCategoryObject(InPinSubCategoryObject)
 		, PinValueType(InValueTerminalType)
 		, ContainerType(ToPinContainerType(bInIsArray, bInIsSet, bInIsMap))
@@ -155,9 +168,10 @@ public:
 	{
 	}
 
-	FEdGraphPinType(FString InPinCategory, FString InPinSubCategory, UObject* InPinSubCategoryObject, EPinContainerType InPinContainerType, bool bInIsReference, const FEdGraphTerminalType& InValueTerminalType )
-		: PinCategory(MoveTemp(InPinCategory))
-		, PinSubCategory(MoveTemp(InPinSubCategory))
+	DEPRECATED(4.19, "Use version that supplies SubCategory as FName")
+	FEdGraphPinType(FName InPinCategory, const FString& InPinSubCategory, UObject* InPinSubCategoryObject, EPinContainerType InPinContainerType, bool bInIsReference, const FEdGraphTerminalType& InValueTerminalType)
+		: PinCategory(InPinCategory)
+		, PinSubCategory(*InPinSubCategory)
 		, PinSubCategoryObject(InPinSubCategoryObject)
 		, PinValueType(InValueTerminalType)
 		, ContainerType(InPinContainerType)
@@ -167,6 +181,34 @@ public:
 		, bIsWeakPointer(false)
 	{
 	}
+
+	//DEPRECATED(4.19, "Remove this constructor when removing FString version, exists only to resolve ambiguity between FName/FString constructors when TCHAR* supplied")
+	FEdGraphPinType(FName InPinCategory, const TCHAR* InPinSubCategory, UObject* InPinSubCategoryObject, EPinContainerType InPinContainerType, bool bInIsReference, const FEdGraphTerminalType& InValueTerminalType)
+		: PinCategory(InPinCategory)
+		, PinSubCategory(InPinSubCategory)
+		, PinSubCategoryObject(InPinSubCategoryObject)
+		, PinValueType(InValueTerminalType)
+		, ContainerType(InPinContainerType)
+		, bIsArray_DEPRECATED(false)
+		, bIsReference(bInIsReference)
+		, bIsConst(false)
+		, bIsWeakPointer(false)
+	{
+	}
+
+	FEdGraphPinType(FName InPinCategory, FName InPinSubCategory, UObject* InPinSubCategoryObject, EPinContainerType InPinContainerType, bool bInIsReference, const FEdGraphTerminalType& InValueTerminalType )
+		: PinCategory(InPinCategory)
+		, PinSubCategory(InPinSubCategory)
+		, PinSubCategoryObject(InPinSubCategoryObject)
+		, PinValueType(InValueTerminalType)
+		, ContainerType(InPinContainerType)
+		, bIsArray_DEPRECATED(false)
+		, bIsReference(bInIsReference)
+		, bIsConst(false)
+		, bIsWeakPointer(false)
+	{
+	}
+
 	bool operator == ( const FEdGraphPinType& Other ) const
 	{
 		return (PinCategory == Other.PinCategory)
@@ -186,8 +228,8 @@ public:
 
 	void ResetToDefaults()
 	{
-		PinCategory.Reset();
-		PinSubCategory.Reset();
+		PinCategory = NAME_None;
+		PinSubCategory = NAME_None;
 		PinSubCategoryObject = nullptr;
 		PinValueType = FEdGraphTerminalType();
 		PinSubCategoryMemberReference.Reset();
@@ -298,7 +340,7 @@ public:
 	FGuid PinId;
 
 	/** Name of this pin. */
-	FString PinName;
+	FName PinName;
 
 #if WITH_EDITORONLY_DATA
 	/** Used as the display name if set. */
@@ -474,11 +516,11 @@ public:
 
 		PinType.ResetToDefaults();
 
-		PinName.Empty();
+		PinName = NAME_None;
 #if WITH_EDITORONLY_DATA
 		PinFriendlyName = FText::GetEmpty();
 #endif // WITH_EDITORONLY_DATA
-		AutogeneratedDefaultValue.Empty();
+		AutogeneratedDefaultValue.Reset();
 		ResetDefaultValue();
 
 #if WITH_EDITORONLY_DATA
@@ -517,7 +559,8 @@ public:
 	ENGINE_API bool ExportTextItem(FString& ValueStr, int32 PortFlags) const;
 	ENGINE_API bool ImportTextItem(const TCHAR*& Buffer, int32 PortFlags, class UObject* Parent, FOutputDevice* ErrorText);
 
-	ENGINE_API const FString& GetName() const { return PinName; }
+	ENGINE_API const FString GetName() const { return PinName.ToString(); }
+	ENGINE_API const FName GetFName() const { return PinName; }
 	ENGINE_API UEdGraphNode* GetOuter() const { return GetOwningNodeUnchecked(); }
 	ENGINE_API bool IsPendingKill() const {	return bWasTrashed; }
 	ENGINE_API FEdGraphTerminalType GetPrimaryTerminalType() const;
diff --git a/Engine/Source/Runtime/Engine/Classes/EdGraph/EdGraphSchema.h b/Engine/Source/Runtime/Engine/Classes/EdGraph/EdGraphSchema.h
--- a/Engine/Source/Runtime/Engine/Classes/EdGraph/EdGraphSchema.h
+++ b/Engine/Source/Runtime/Engine/Classes/EdGraph/EdGraphSchema.h
@@ -831,7 +831,7 @@ class ENGINE_API UEdGraphSchema : public UObject
 	 * @param	SourcePin	The pin where the link begins.
 	 * @param	TargetLink	The pin where the link ends.
 	 */
-	virtual void BreakSinglePinLink(UEdGraphPin* SourcePin, UEdGraphPin* TargetPin);
+	virtual void BreakSinglePinLink(UEdGraphPin* SourcePin, UEdGraphPin* TargetPin) const;
 
 	/** Split a pin in to subelements */
 	virtual void SplitPin(UEdGraphPin* Pin, bool bNotify = true) const { };
@@ -852,7 +852,7 @@ class ENGINE_API UEdGraphSchema : public UObject
 	virtual bool IsSelfPin(const UEdGraphPin& Pin) const   {return false;}
 
 	/** Is given string a delegate category name ? */
-	virtual bool IsDelegateCategory(const FString& Category) const   {return false;}
+	virtual bool IsDelegateCategory(const FName Category) const   {return false;}
 
 	/** 
 	 * Populate new graph with any default nodes
@@ -974,7 +974,7 @@ class ENGINE_API UEdGraphSchema : public UObject
 	 * @param InSourcePinDirection			Direction of the source pin
 	 * @return								Returns the new pin if created
 	 */
-	virtual UEdGraphPin* DropPinOnNode(UEdGraphNode* InTargetNode, const FString& InSourcePinName, const FEdGraphPinType& InSourcePinType, EEdGraphPinDirection InSourcePinDirection) const { return nullptr; }
+	virtual UEdGraphPin* DropPinOnNode(UEdGraphNode* InTargetNode, const FName& InSourcePinName, const FEdGraphPinType& InSourcePinType, EEdGraphPinDirection InSourcePinDirection) const { return nullptr; }
 
 	/**
 	 * Checks if the node supports dropping a pin on it
diff --git a/Engine/Source/Runtime/Engine/Classes/Engine/AssetManager.h b/Engine/Source/Runtime/Engine/Classes/Engine/AssetManager.h
--- a/Engine/Source/Runtime/Engine/Classes/Engine/AssetManager.h
+++ b/Engine/Source/Runtime/Engine/Classes/Engine/AssetManager.h
@@ -465,6 +465,9 @@ protected:
 	/** Handles updating manager if deleted object is relevant*/
 	virtual void OnInMemoryAssetDeleted(UObject *Object);
 
+	/** Called when object is saved */
+	virtual void OnObjectPreSave(UObject* Object);
+
 	/** When asset is renamed */
 	virtual void OnAssetRenamed(const FAssetData& NewData, const FString& OldPath);
 
diff --git a/Engine/Source/Runtime/Engine/Classes/Engine/DataAsset.h b/Engine/Source/Runtime/Engine/Classes/Engine/DataAsset.h
--- a/Engine/Source/Runtime/Engine/Classes/Engine/DataAsset.h
+++ b/Engine/Source/Runtime/Engine/Classes/Engine/DataAsset.h
@@ -29,8 +29,13 @@ private:
 
 /**
  * A DataAsset that implements GetPrimaryAssetId and has asset bundle support, which makes it something that can be manually loaded/unloaded from the AssetManager
+ * Making blueprint subclasses of this is useful because you can make blueprint-only primary asset types
+ * Blueprint subclasses will end up with a PrimaryAssetType equal to the name of the first native class found going up the hierarchy, or the top level blueprint class that directly subclasses this
+ * IE, if you have PrimaryDataAsset -> ParentNativeClass -> ChildNativeClass -> BlueprintAsset the type will be ChildNativeClass
+ * Whereas if you have PrimaryDataAsset -> ParentBlueprintClass -> ChildBlueprintClass the type will be ParentBlueprintClass
+ * To override this behavior, override GetPrimaryAssetId in your native class
  */
-UCLASS(abstract, MinimalAPI)
+UCLASS(abstract, MinimalAPI, Blueprintable)
 class UPrimaryDataAsset : public UDataAsset
 {
 	GENERATED_BODY()
diff --git a/Engine/Source/Runtime/Engine/Classes/Engine/EngineTypes.h b/Engine/Source/Runtime/Engine/Classes/Engine/EngineTypes.h
--- a/Engine/Source/Runtime/Engine/Classes/Engine/EngineTypes.h
+++ b/Engine/Source/Runtime/Engine/Classes/Engine/EngineTypes.h
@@ -2075,7 +2075,7 @@ struct TStructOpsTypeTraits<FHitResult> : public TStructOpsTypeTraitsBase2<FHitR
 
 
 /** Whether to teleport physics body or not */
-enum class ETeleportType
+enum class ETeleportType : uint8
 {
 	/** Do not teleport physics body. This means velocity will reflect the movement between initial and final position, and collisions along the way will occur */
 	None,
diff --git a/Engine/Source/Runtime/Engine/Classes/Engine/InheritableComponentHandler.h b/Engine/Source/Runtime/Engine/Classes/Engine/InheritableComponentHandler.h
--- a/Engine/Source/Runtime/Engine/Classes/Engine/InheritableComponentHandler.h
+++ b/Engine/Source/Runtime/Engine/Classes/Engine/InheritableComponentHandler.h
@@ -127,7 +127,6 @@ public:
 	//~ Begin UObject Interface
 	virtual void Serialize(FArchive& Ar) override;
 	virtual void PostLoad() override;
-	virtual void GetPreloadDependencies(TArray<UObject*>& OutDeps) override;
 	//~ End UObject Interface
 
 	void PreloadAllTempates();
diff --git a/Engine/Source/Runtime/Engine/Classes/Engine/LevelBounds.h b/Engine/Source/Runtime/Engine/Classes/Engine/LevelBounds.h
--- a/Engine/Source/Runtime/Engine/Classes/Engine/LevelBounds.h
+++ b/Engine/Source/Runtime/Engine/Classes/Engine/LevelBounds.h
@@ -13,7 +13,7 @@
  */
 class FEditorTickableLevelBounds 
 #if WITH_EDITOR
-	: FTickableGameObject
+	: public FTickableGameObject
 #endif
 {
 };
@@ -68,6 +68,7 @@ private:
 	virtual void Tick(float DeltaTime) override;
 	virtual UWorld* GetTickableGameObjectWorld() const override { return GetWorld(); }
 	virtual TStatId GetStatId() const override;
+	virtual ETickableTickType GetTickableTickType() const override;
 	virtual bool IsTickable() const override;
 	virtual bool IsTickableInEditor() const override;
 	
diff --git a/Engine/Source/Runtime/Engine/Classes/Engine/LevelStreamingKismet.h b/Engine/Source/Runtime/Engine/Classes/Engine/LevelStreamingKismet.h
--- a/Engine/Source/Runtime/Engine/Classes/Engine/LevelStreamingKismet.h
+++ b/Engine/Source/Runtime/Engine/Classes/Engine/LevelStreamingKismet.h
@@ -41,8 +41,11 @@ class ULevelStreamingKismet : public ULevelStreaming
 	* @param bOutSuccess - Whether operation was successful (map was found and added to the sub-levels list)
  	* @return Streaming level object for a level instance
  	*/ 
- 	UFUNCTION(BlueprintCallable, Category = LevelStreaming,meta=(WorldContext="WorldContextObject"))
- 	static ENGINE_API ULevelStreamingKismet* LoadLevelInstance(UObject* WorldContextObject, const FString& LevelName, const FVector& Location, const FRotator& Rotation, bool& bOutSuccess);
+ 	UFUNCTION(BlueprintCallable, Category = LevelStreaming, meta=(WorldContext="WorldContextObject"))
+ 	static ENGINE_API ULevelStreamingKismet* LoadLevelInstance(UObject* WorldContextObject, FString LevelName, FVector Location, FRotator Rotation, bool& bOutSuccess);
+
+ 	UFUNCTION(BlueprintCallable, Category = LevelStreaming, meta=(DisplayName = "Load Level Instance", WorldContext="WorldContextObject"))
+ 	static ENGINE_API ULevelStreamingKismet* LoadLevelInstanceBySoftObjectPtr(UObject* WorldContextObject, TSoftObjectPtr<UWorld> Level, FVector Location, FRotator Rotation, bool& bOutSuccess);
  	
 	//~ Begin UObject Interface
 	virtual void PostLoad() override;
@@ -55,5 +58,8 @@ class ULevelStreamingKismet : public ULevelStreaming
 private:
 	// Counter used by LoadLevelInstance to create unique level names
 	static int32 UniqueLevelInstanceId;
+
+ 	static ENGINE_API ULevelStreamingKismet* LoadLevelInstance_Internal(UWorld* World, const FString& LongPackageName, FVector Location, FRotator Rotation, bool& bOutSuccess);
+
 };
 
