Copying //UE4/Dev-Framework to //UE4/Dev-Main (Source: //UE4/Dev-Framework @ 3716594)

diff --git a/Engine/Source/Runtime/Engine/Classes/GameFramework/Actor.h b/Engine/Source/Runtime/Engine/Classes/GameFramework/Actor.h
--- a/Engine/Source/Runtime/Engine/Classes/GameFramework/Actor.h
+++ b/Engine/Source/Runtime/Engine/Classes/GameFramework/Actor.h
@@ -34,7 +34,6 @@ struct FNetViewer;
 struct FNetworkObjectInfo;
 
 ENGINE_API DECLARE_LOG_CATEGORY_EXTERN(LogActor, Log, Warning);
- 
 
 // Delegate signatures
 DECLARE_DYNAMIC_MULTICAST_DELEGATE_FiveParams( FTakeAnyDamageSignature, AActor*, DamagedActor, float, Damage, const class UDamageType*, DamageType, class AController*, InstigatedBy, AActor*, DamageCauser );
@@ -80,26 +79,22 @@ UCLASS(BlueprintType, Blueprintable, config=Engine, meta=(ShortTooltip="An Actor
 class ENGINE_API AActor : public UObject
 {
 	/**
-	* The functions of interest to initialization order for an Actor is roughly as follows:
-	* PostLoad/PostActorCreated - Do any setup of the actor required for construction. PostLoad for serialized actors, PostActorCreated for spawned.  
-	* AActor::OnConstruction - The construction of the actor, this is where Blueprint actors have their components created and blueprint variables are initialized
-	* AActor::PreInitializeComponents - Called before InitializeComponent is called on the actor's components
-	* UActorComponent::InitializeComponent - Each component in the actor's components array gets an initialize call (if bWantsInitializeComponent is true for that component)
-	* AActor::PostInitializeComponents - Called after the actor's components have been initialized
-	* AActor::BeginPlay - Called when the level is started
-	*/
+	 * The functions of interest to initialization order for an Actor is roughly as follows:
+	 * PostLoad/PostActorCreated - Do any setup of the actor required for construction. PostLoad for serialized actors, PostActorCreated for spawned.  
+	 * AActor::OnConstruction - The construction of the actor, this is where Blueprint actors have their components created and blueprint variables are initialized
+	 * AActor::PreInitializeComponents - Called before InitializeComponent is called on the actor's components
+	 * UActorComponent::InitializeComponent - Each component in the actor's components array gets an initialize call (if bWantsInitializeComponent is true for that component)
+	 * AActor::PostInitializeComponents - Called after the actor's components have been initialized
+	 * AActor::BeginPlay - Called when the level is started
+	 */
 
 	GENERATED_BODY()
 public:
 
-	/**
-	 * Default constructor for AActor
-	 */
+	/** Default constructor for AActor */
 	AActor();
 
-	/**
-	 * Constructor for AActor that takes an ObjectInitializer
-	 */
+	/** Constructor for AActor that takes an ObjectInitializer */
 	AActor(const FObjectInitializer& ObjectInitializer);
 
 private:
@@ -107,6 +102,7 @@ private:
 	void InitializeDefaults();
 
 public:
+	/** Returns the properties used for network replication */
 	void GetLifetimeReplicatedProps(TArray<FLifetimeProperty>& OutLifetimeProps) const override;
 
 	/**
@@ -115,19 +111,14 @@ public:
 	 * @see https://docs.unrealengine.com/latest/INT/API/Runtime/Engine/Engine/FTickFunction/
 	 * @see AddTickPrerequisiteActor(), AddTickPrerequisiteComponent()
 	 */
-	UPROPERTY(EditDefaultsOnly, Category="Tick")
+	UPROPERTY(EditDefaultsOnly, Category=Tick)
 	struct FActorTickFunction PrimaryActorTick;
 
-	/** Allow each actor to run at a different time speed. The DeltaTime for a frame is multiplied by the global TimeDilation (in WorldSettings) and this CustomTimeDilation for this actor's tick.  */
-	UPROPERTY(BlueprintReadWrite, AdvancedDisplay, Category="Misc")
-	float CustomTimeDilation;	
-
-public:
 	/**
 	 * Allows us to only see this Actor in the Editor, and not in the actual game.
 	 * @see SetActorHiddenInGame()
 	 */
-	UPROPERTY(Interp, EditAnywhere, Category=Rendering, BlueprintReadOnly, replicated, meta=(DisplayName = "Actor Hidden In Game", SequencerTrackClass = "MovieSceneVisibilityTrack"))
+	UPROPERTY(Interp, EditAnywhere, Category=Rendering, BlueprintReadOnly, Replicated, meta=(DisplayName="Actor Hidden In Game", SequencerTrackClass="MovieSceneVisibilityTrack"))
 	uint8 bHidden:1;
 
 	/** If true, when the actor is spawned it will be sent to the client but receive no further replication updates from the server afterwards. */
@@ -163,24 +154,20 @@ public:
 	 * If true, this actor is no longer replicated to new clients, and is "torn off" (becomes a ROLE_Authority) on clients to which it was being replicated.
 	 * @see TornOff()
 	 */
-	UPROPERTY(replicated)
+	UPROPERTY(Replicated)
 	uint8 bTearOff:1;    
 
 	/** Networking - Server - TearOff this actor to stop replication to clients. Will set bTearOff to true. */
-	UFUNCTION(BlueprintCallable, Category = Replication)
+	UFUNCTION(BlueprintCallable, Category=Replication)
 	virtual void TearOff();
 
 	/**
 	 * Whether we have already exchanged Role/RemoteRole on the client, as when removing then re-adding a streaming level.
 	 * Causes all initialization to be performed again even though the actor may not have actually been reloaded.
 	 */
-	UPROPERTY(transient)
+	UPROPERTY(Transient)
 	uint8 bExchangedRoles:1;
 
-	/** Is this actor still pending a full net update due to clients that weren't able to replicate the actor at the time of LastNetUpdateTime */
-	DEPRECATED(4.16, "bPendingNetUpdate has been deprecated. Please use bPendingNetUpdate that is now on FNetworkObjectInfo (obtained by calling GetNetworkObjectInfo)")
-	uint8 bPendingNetUpdate:1;
-
 	/** This actor will be loaded on network clients during map load */
 	UPROPERTY(Category=Replication, EditDefaultsOnly)
 	uint8 bNetLoadOnClient:1;
@@ -189,47 +176,64 @@ public:
 	UPROPERTY(Category=Replication, EditDefaultsOnly, BlueprintReadWrite)
 	uint8 bNetUseOwnerRelevancy:1;
 
-	/** If true, all input on the stack below this actor will not be considered */
-	UPROPERTY(EditDefaultsOnly, Category=Input)
-	uint8 bBlockInput:1;
-
-	/** True if this actor is currently running user construction script (used to defer component registration) */
-	uint8 bRunningUserConstructionScript:1;
+	/** If true, this actor will be replicated to network replays (default is true) */
+	UPROPERTY()
+	uint8 bRelevantForNetworkReplays:1;
 
 	/**
 	 * Whether we allow this Actor to tick before it receives the BeginPlay event.
 	 * Normally we don't tick actors until after BeginPlay; this setting allows this behavior to be overridden.
 	 * This Actor must be able to tick for this setting to be relevant.
 	 */
-	UPROPERTY(EditDefaultsOnly, Category="Tick")
+	UPROPERTY(EditDefaultsOnly, Category=Tick)
 	uint8 bAllowTickBeforeBeginPlay:1;
 
-private:
-	/** Whether FinishSpawning has been called for this Actor.  If it has not, the Actor is in a mal-formed state */
-	uint8 bHasFinishedSpawning:1;
-
-	/** Whether we've tried to register tick functions. Reset when they are unregistered. */
-	uint8 bTickFunctionsRegistered : 1;
+	/** If true then destroy self when "finished", meaning all relevant components report that they are done and no timelines or timers are in flight. */
+	UPROPERTY(BlueprintReadWrite, Category=Actor)
+	uint8 bAutoDestroyWhenFinished:1;
 
-	/** Whether we've deferred the RegisterAllComponents() call at spawn time. Reset when RegisterAllComponents() is called. */
-	uint8 bHasDeferredComponentRegistration : 1;
+	/** If true, all input on the stack below this actor will not be considered */
+	UPROPERTY(EditDefaultsOnly, Category=Input)
+	uint8 bBlockInput:1;
 
 	/**
-	 * Enables any collision on this actor.
-	 * @see SetActorEnableCollision(), GetActorEnableCollision()
+	 * Whether this actor can take damage. Must be true for damage events (e.g. ReceiveDamage()) to be called.
+	 * @see https://www.unrealengine.com/blog/damage-in-ue4
+	 * @see TakeDamage(), ReceiveDamage()
 	 */
+	UPROPERTY(EditAnywhere, BlueprintReadWrite, SaveGame, Replicated, Category=Actor)
+	uint8 bCanBeDamaged:1;
+
+	/** This actor collides with the world when placing in the editor, even if RootComponent collision is disabled. Does not affect spawning, @see SpawnCollisionHandlingMethod */
 	UPROPERTY()
-	uint8 bActorEnableCollision:1;
+	uint8 bCollideWhenPlacing:1;
 
-	/** Flag indicating we have checked initial simulating physics state to sync networked proxies to the server. */
-	uint8 bNetCheckedInitialPhysicsState:1;
+	/** If true, this actor should search for an owned camera component to view through when used as a view target. */
+	UPROPERTY(EditAnywhere, BlueprintReadWrite, Category=Actor, AdvancedDisplay)
+	uint8 bFindCameraComponentWhenViewTarget:1;
+	
+    /** If true, this actor will generate overlap events when spawned as part of level streaming. You might enable this is in the case where a streaming level loads around an actor and you want overlaps to trigger. */
+	UPROPERTY(EditAnywhere, BlueprintReadWrite, Category=Actor)
+	uint8 bGenerateOverlapEventsDuringLevelStreaming:1;
+
+	/** Whether this actor should not be affected by world origin shifting. */
+	UPROPERTY(EditAnywhere, AdvancedDisplay, Category=Actor)
+	uint8 bIgnoresOriginShifting:1;
+
+	/** If true, and if World setting has bEnableHierarchicalLOD equal to true, then it will generate LODActor from groups of clustered Actor */
+	UPROPERTY(EditAnywhere, AdvancedDisplay, Category=LOD, meta=(DisplayName="Include Actor for HLOD Mesh generation"))
+	uint8 bEnableAutoLODGeneration:1;
+
+	/** Indicates the actor was pulled through a seamless travel.  */
+	UPROPERTY()
+	uint8 bActorSeamlessTraveled:1;
 
 protected:
 	/**
 	 * If true, this actor will replicate to remote machines
 	 * @see SetReplicates()
 	 */
-	UPROPERTY(EditDefaultsOnly, BlueprintReadOnly, Category = "Replication")
+	UPROPERTY(EditDefaultsOnly, BlueprintReadOnly, Category=Replication)
 	uint8 bReplicates:1;
 
 	/** This function should only be used in the constructor of classes that need to set the RemoteRole for backwards compatibility purposes */
@@ -242,29 +246,73 @@ protected:
 	 */
 	virtual bool HasNetOwner() const;
 
+	/** Called when owner changes, does nothing by default but can be overridden */
 	UFUNCTION()
 	virtual void OnRep_Owner();
 
-private:
+	/** If true, this actor can be put inside of a GC Cluster to improve Garbage Collection performance */
+	UPROPERTY(Category=Actor, EditAnywhere, AdvancedDisplay)
+	uint8 bCanBeInCluster:1;
+
 	/**
-	 * Describes how much control the remote machine has over the actor.
+	 * If false, the Blueprint ReceiveTick() event will be disabled on dedicated servers.
+	 * @see AllowReceiveTickEventOnDedicatedServer()
 	 */
-	UPROPERTY(Replicated, transient)
-	TEnumAsByte<enum ENetRole> RemoteRole;
+	UPROPERTY()
+	uint8 bAllowReceiveTickEventOnDedicatedServer:1;
 
-	/**
-	 * Owner of this Actor, used primarily for replication (bNetUseOwnerRelevancy & bOnlyRelevantToOwner) and visibility (PrimitiveComponent bOwnerNoSee and bOnlyOwnerSee)
-	 * @see SetOwner(), GetOwner()
+private:
+	/** Whether FinishSpawning has been called for this Actor.  If it has not, the Actor is in a malformed state */
+	uint8 bHasFinishedSpawning:1;
+
+	/** 
+	 *	Indicates that PreInitializeComponents/PostInitializeComponents have been called on this Actor 
+	 *	Prevents re-initializing of actors spawned during level startup
 	 */
-	UPROPERTY(ReplicatedUsing=OnRep_Owner)
-	AActor* Owner;
+	uint8 bActorInitialized:1;
 
-public:
+	/** Whether we've tried to register tick functions. Reset when they are unregistered. */
+	uint8 bTickFunctionsRegistered:1;
 
-	/** Used to specify the net driver to replicate on (NAME_None || NAME_GameNetDriver is the default net driver) */
+	/** Whether we've deferred the RegisterAllComponents() call at spawn time. Reset when RegisterAllComponents() is called. */
+	uint8 bHasDeferredComponentRegistration:1;
+
+	/** True if this actor is currently running user construction script (used to defer component registration) */
+	uint8 bRunningUserConstructionScript:1;
+
+	/**
+	 * Enables any collision on this actor.
+	 * @see SetActorEnableCollision(), GetActorEnableCollision()
+	 */
 	UPROPERTY()
-	FName NetDriverName;
+	uint8 bActorEnableCollision:1;
+
+	/** Flag indicating we have checked initial simulating physics state to sync networked proxies to the server. */
+	uint8 bNetCheckedInitialPhysicsState:1;
 
+	/** Set when actor is about to be deleted. */
+	uint8 bActorIsBeingDestroyed:1;
+
+	enum class EActorBeginPlayState : uint8
+	{
+		HasNotBegunPlay,
+		BeginningPlay,
+		HasBegunPlay,
+	};
+
+	/** 
+	 *	Indicates that BeginPlay has been called for this Actor.
+	 *  Set back to HasNotBegunPlay once EndPlay has been called.
+	 */
+	EActorBeginPlayState ActorHasBegunPlay:2;
+
+	static uint32 BeginPlayCallDepth;
+
+	/** Describes how much control the remote machine has over the actor. */
+	UPROPERTY(Replicated, Transient)
+	TEnumAsByte<enum ENetRole> RemoteRole;
+
+public:
 	/**
 	 * Set whether this actor replicates to network clients. When this actor is spawned on the server it will be sent to clients as well.
 	 * Properties flagged for replication will update on clients if they change on the server.
@@ -272,14 +320,14 @@ public:
 	 * @param bInReplicates Whether this Actor replicates to network clients.
 	 * @see https://docs.unrealengine.com/latest/INT/Gameplay/Networking/Replication/
 	 */
-	UFUNCTION(BlueprintCallable, BlueprintAuthorityOnly, Category = "Replication")
+	UFUNCTION(BlueprintCallable, BlueprintAuthorityOnly, Category=Replication)
 	void SetReplicates(bool bInReplicates);
 
 	/**
-	* Set whether this actor's movement replicates to network clients.
-	* @param bInReplicateMovement Whether this Actor's movement replicates to clients.
-	*/
-	UFUNCTION(BlueprintCallable, Category = "Replication")
+	 * Set whether this actor's movement replicates to network clients.
+	 * @param bInReplicateMovement Whether this Actor's movement replicates to clients.
+	 */
+	UFUNCTION(BlueprintCallable, Category=Replication)
 	virtual void SetReplicateMovement(bool bInReplicateMovement);
 
 	/** Sets whether or not this Actor is an autonomous proxy, which is an actor on a network client that is controlled by a user on that client. */
@@ -289,7 +337,7 @@ public:
 	void CopyRemoteRoleFrom(const AActor* CopyFromActor);
 
 	/** Returns how much control the remote machine has over this actor. */
-	UFUNCTION(BlueprintCallable, Category = "Replication")
+	UFUNCTION(BlueprintCallable, Category=Replication)
 	ENetRole GetRemoteRole() const;
 
 	/** Used for replication of our RootComponent's position and velocity */
@@ -297,17 +345,39 @@ public:
 	struct FRepMovement ReplicatedMovement;
 
 	/** How long this Actor lives before dying, 0=forever. Note this is the INITIAL value and should not be modified once play has begun. */
-	UPROPERTY(EditAnywhere, BlueprintReadOnly, Category = Actor)
+	UPROPERTY(EditAnywhere, BlueprintReadOnly, Category=Actor)
 	float InitialLifeSpan;
 
+	/** Allow each actor to run at a different time speed. The DeltaTime for a frame is multiplied by the global TimeDilation (in WorldSettings) and this CustomTimeDilation for this actor's tick.  */
+	UPROPERTY(BlueprintReadWrite, AdvancedDisplay, Category=Actor)
+	float CustomTimeDilation;
+
+	/**
+	 * The time this actor was created, relative to World->GetTimeSeconds().
+	 * @see UWorld::GetTimeSeconds()
+	 */
+	float CreationTime;
+
 private:
 	/**
 	 * Used for replicating attachment of this actor's RootComponent to another actor.
 	 * This is filled in via GatherCurrentMovement() when the RootComponent has an AttachParent.
 	 */
-	UPROPERTY(Transient, replicatedUsing=OnRep_AttachmentReplication)
+	UPROPERTY(Transient, ReplicatedUsing=OnRep_AttachmentReplication)
 	struct FRepAttachment AttachmentReplication;
 
+	/**
+	 * Owner of this Actor, used primarily for replication (bNetUseOwnerRelevancy & bOnlyRelevantToOwner) and visibility (PrimitiveComponent bOwnerNoSee and bOnlyOwnerSee)
+	 * @see SetOwner(), GetOwner()
+	 */
+	UPROPERTY(ReplicatedUsing=OnRep_Owner)
+	AActor* Owner;
+
+protected:
+	/** Used to specify the net driver to replicate on (NAME_None || NAME_GameNetDriver is the default net driver) */
+	UPROPERTY()
+	FName NetDriverName;
+
 public:
 	/** Get read-only access to current AttachmentReplication. */
 	const struct FRepAttachment& GetAttachmentReplication() const { return AttachmentReplication; }
@@ -321,7 +391,7 @@ public:
 	TEnumAsByte<enum ENetRole> Role;
 
 	/** Dormancy setting for actor to take itself off of the replication list without being destroyed on clients. */
-	UPROPERTY(BlueprintReadOnly, EditDefaultsOnly, Category = Replication)
+	UPROPERTY(BlueprintReadOnly, EditDefaultsOnly, Category=Replication)
 	TEnumAsByte<enum ENetDormancy> NetDormancy;
 
 	/** Gives the actor a chance to pause replication to a player represented by the passed in actor - only called on server */
@@ -330,6 +400,10 @@ public:
 	/** Called on the client when the replication paused value is changed */
 	virtual void OnReplicationPausedChanged(bool bIsReplicationPaused);
 
+	/** Controls how to handle spawning this actor in a situation where it's colliding with something else. "Default" means AlwaysSpawn here. */
+	UPROPERTY(EditAnywhere, BlueprintReadWrite, Category=Actor)
+	ESpawnActorCollisionHandlingMethod SpawnCollisionHandlingMethod;
+
 	/** Automatically registers this actor to receive input from a player. */
 	UPROPERTY(EditAnywhere, Category=Input)
 	TEnumAsByte<EAutoReceiveInput::Type> AutoReceiveInput;
@@ -351,37 +425,27 @@ public:
 	UPROPERTY(BlueprintReadOnly, EditDefaultsOnly, Category=Replication)
 	float NetCullDistanceSquared;   
 
-	/** Internal - used by UWorld::ServerTickClients() */
-	UPROPERTY(transient)
+	/** Internal - used by UNetDriver */
+	UPROPERTY(Transient)
 	int32 NetTag;
 
-	/** Next time this actor will be considered for replication, set by SetNetUpdateTime() */
-	DEPRECATED(4.16, "NetUpdateTime has been deprecated. Please use NextUpdateTime that is now on FNetworkObjectInfo (obtained by calling GetNetworkObjectInfo)")
-	float NetUpdateTime;
-
 	/** How often (per second) this actor will be considered for replication, used to determine NetUpdateTime */
 	UPROPERTY(Category=Replication, EditDefaultsOnly, BlueprintReadWrite)
 	float NetUpdateFrequency;
 
 	/** Used to determine what rate to throttle down to when replicated properties are changing infrequently */
-	UPROPERTY( Category=Replication, EditDefaultsOnly, BlueprintReadWrite )
+	UPROPERTY(Category=Replication, EditDefaultsOnly, BlueprintReadWrite)
 	float MinNetUpdateFrequency;
 
 	/** Priority for this actor when checking for replication in a low bandwidth or saturated situation, higher priority means it is more likely to replicate */
 	UPROPERTY(Category=Replication, EditDefaultsOnly, BlueprintReadWrite)
 	float NetPriority;
 
-	/** Last time this actor was updated for replication via NetUpdateTime
-	* @warning: internal net driver time, not related to WorldSettings.TimeSeconds */
-	DEPRECATED(4.16, "LastNetUpdateTime has been deprecated. Please use LastNetUpdateTime that is now on FNetworkObjectInfo (obtained by calling GetNetworkObjectInfo)")
-	float LastNetUpdateTime;
-
 	/**
 	 * Set the name of the net driver associated with this actor.  Will move the actor out of the list of network actors from the old net driver and add it to the new list
-	 *
 	 * @param NewNetDriverName name of the new net driver association
 	 */
-	void SetNetDriverName(FName NewNetDriverName );
+	void SetNetDriverName(FName NewNetDriverName);
 
 	/** @return name of the net driver associated with this actor (all RPCs will go out via this connection) */
 	FName GetNetDriverName() const { return NetDriverName; }
@@ -395,116 +459,23 @@ public:
 	/** Called on the actor when a subobject is dynamically destroyed via replication */
 	virtual void OnSubobjectDestroyFromReplication(UObject *Subobject);
 
-	/** Called on the actor right before replication occurs. 
-	 * Only called on Server, and for autonomous proxies if recording a Client Replay. */
-	virtual void PreReplication( IRepChangedPropertyTracker & ChangedPropertyTracker );
-
-	/** Called on the actor right before replication occurs.
-	 * Called for everyone when recording a Client Replay, including Simulated Proxies. */
-	virtual void PreReplicationForReplay(IRepChangedPropertyTracker & ChangedPropertyTracker);
-
-	/** Called by the networking system to call PreReplication on this actor and its components using the given NetDriver to find or create RepChangedPropertyTrackers. */
-	void CallPreReplication(UNetDriver* NetDriver);
-
-	/** If true then destroy self when "finished", meaning all relevant components report that they are done and no timelines or timers are in flight. */
-	UPROPERTY(BlueprintReadWrite, Category=Actor)
-	uint8 bAutoDestroyWhenFinished:1;
-
-	/**
-	 * Whether this actor can take damage. Must be true for damage events (e.g. ReceiveDamage()) to be called.
-	 * @see https://www.unrealengine.com/blog/damage-in-ue4
-	 * @see TakeDamage(), ReceiveDamage()
-	 */
-	UPROPERTY(EditAnywhere, BlueprintReadWrite, SaveGame, Replicated, Category=Actor)
-	uint8 bCanBeDamaged:1;
-
-private:
 	/**
-	 * Set when actor is about to be deleted.
+	 * Called on the actor right before replication occurs. 
+	 * Only called on Server, and for autonomous proxies if recording a Client Replay.
 	 */
-	UPROPERTY(Transient, DuplicateTransient)
-	uint8 bActorIsBeingDestroyed:1;    
-
-public:
-
-	/** This actor collides with the world when placing in the editor, even if RootComponent collision is disabled. Does not affect spawning, @see SpawnCollisionHandlingMethod */
-	UPROPERTY()
-	uint8 bCollideWhenPlacing:1;
-
-	/** If true, this actor should search for an owned camera component to view through when used as a view target. */
-	UPROPERTY(EditAnywhere, BlueprintReadWrite, Category=Actor, AdvancedDisplay)
-	uint8 bFindCameraComponentWhenViewTarget:1;
-	
-	/** If true, this actor will be replicated to network replays (default is true) */
-	UPROPERTY()
-	uint8 bRelevantForNetworkReplays:1;
-	
-    /** If true, this actor will generate overlap events when spawned as part of level streaming. You might enable this is in the case where a streaming level loads around an actor and you want overlaps to trigger. */
-	UPROPERTY(EditAnywhere, BlueprintReadWrite, Category = Actor)
-	uint8 bGenerateOverlapEventsDuringLevelStreaming:1;
-
-protected:
-
-	/** If true, this actor can be put inside of a GC Cluster to improve Garbage Collection performance */
-	UPROPERTY(Category = Actor, EditAnywhere, AdvancedDisplay)
-	uint8 bCanBeInCluster:1;
-
-protected:
+	virtual void PreReplication(IRepChangedPropertyTracker & ChangedPropertyTracker);
 
 	/**
-	 * If false, the Blueprint ReceiveTick() event will be disabled on dedicated servers.
-	 * @see AllowReceiveTickEventOnDedicatedServer()
-	 */
-	UPROPERTY()
-	uint8 bAllowReceiveTickEventOnDedicatedServer:1;
-
-private:
-
-	/** 
-	 *	Indicates that PreInitializeComponents/PostInitializeComponents have been called on this Actor 
-	 *	Prevents re-initializing of actors spawned during level startup
+	 * Called on the actor right before replication occurs.
+	 * Called for everyone when recording a Client Replay, including Simulated Proxies.
 	 */
-	uint8 bActorInitialized:1;
-	
-	enum class EActorBeginPlayState : uint8
-	{
-		HasNotBegunPlay,
-		BeginningPlay,
-		HasBegunPlay,
-	};
-
-	/** 
-	 *	Indicates that BeginPlay has been called for this Actor.
-	 *  Set back to false once EndPlay has been called.
-	 */
-	EActorBeginPlayState ActorHasBegunPlay:2;
-
-public:
-	/** Indicates the actor was pulled through a seamless travel.  */
-	UPROPERTY()
-	uint8 bActorSeamlessTraveled:1;
-
-	/** Whether this actor should not be affected by world origin shifting. */
-	UPROPERTY(EditAnywhere, AdvancedDisplay, Category=Actor)
-	uint8 bIgnoresOriginShifting:1;
-	
-	/** If true, and if World setting has bEnableHierarchicalLOD equal to true, then it will generate LODActor from groups of clustered Actor */
-	UPROPERTY(EditAnywhere, AdvancedDisplay, Category = LOD, meta = (DisplayName = "Include Actor for HLOD Mesh generation"))
-	uint8 bEnableAutoLODGeneration:1;
-
-public:
+	virtual void PreReplicationForReplay(IRepChangedPropertyTracker & ChangedPropertyTracker);
 
-	/** Controls how to handle spawning this actor in a situation where it's colliding with something else. "Default" means AlwaysSpawn here. */
-	UPROPERTY(EditAnywhere, BlueprintReadWrite, Category=Actor)
-	ESpawnActorCollisionHandlingMethod SpawnCollisionHandlingMethod;
+	/** Called by the networking system to call PreReplication on this actor and its components using the given NetDriver to find or create RepChangedPropertyTrackers. */
+	void CallPreReplication(UNetDriver* NetDriver);	
 	
-	/** The time this actor was created, relative to World->GetTimeSeconds().
-	* @see UWorld::GetTimeSeconds()
-	*/
-	float CreationTime;
-
 	/** Pawn responsible for damage caused by this actor. */
-	UPROPERTY(BlueprintReadWrite, replicatedUsing=OnRep_Instigator, meta=(ExposeOnSpawn=true), Category=Actor)
+	UPROPERTY(BlueprintReadWrite, ReplicatedUsing=OnRep_Instigator, meta=(ExposeOnSpawn=true), Category=Actor)
 	class APawn* Instigator;
 
 	/** Called on clients when Instigator is replicated. */
@@ -512,35 +483,35 @@ public:
 	virtual void OnRep_Instigator();
 
 	/** Array of Actors whose Owner is this actor */
-	UPROPERTY(transient)
+	UPROPERTY(Transient)
 	TArray<AActor*> Children;
 
 protected:
-	/**
-	 * Collision primitive that defines the transform (location, rotation, scale) of this Actor.
-	 */
+	/** Collision primitive that defines the transform (location, rotation, scale) of this Actor. */
 	UPROPERTY(BlueprintGetter=K2_GetRootComponent, Category="Utilities|Transformation")
 	USceneComponent* RootComponent;
 
 #if WITH_EDITORONLY_DATA
 	/** Local space pivot offset for the actor */
-	UPROPERTY(EditAnywhere, BlueprintReadOnly, AdvancedDisplay, Category = Actor)
+	UPROPERTY(EditAnywhere, BlueprintReadOnly, AdvancedDisplay, Category=Actor)
 	FVector PivotOffset;
 #endif
 
 	/** The matinee actors that control this actor. */
-	UPROPERTY(transient)
+	UPROPERTY(Transient)
 	TArray<class AMatineeActor*> ControllingMatineeActors;
 
 	/** Handle for efficient management of LifeSpanExpired timer */
 	FTimerHandle TimerHandle_LifeSpanExpired;
 
 public:
-
 	/** Return the value of bAllowReceiveTickEventOnDedicatedServer, indicating whether the Blueprint ReceiveTick() event will occur on dedicated servers. */
 	FORCEINLINE bool AllowReceiveTickEventOnDedicatedServer() const { return bAllowReceiveTickEventOnDedicatedServer; }
 
-	/** Layer's the actor belongs to.  This is outside of the editoronly data to allow hiding of LD-specified layers at runtime for profiling. */
+	/** Returns if this actor is currently running the User Construction Script */
+	FORCEINLINE bool IsRunningUserConstructionScript() const { return bRunningUserConstructionScript; }
+
+	/** Layers the actor belongs to.  This is outside of the editoronly data to allow hiding of LD-specified layers at runtime for profiling. */
 	UPROPERTY()
 	TArray< FName > Layers;
 
@@ -555,18 +526,21 @@ private:
 	UPROPERTY()
 	TWeakObjectPtr<UChildActorComponent> ParentComponent;	
 
-public:
-
 #if WITH_EDITORONLY_DATA
 
+public:
 	/** The group this actor is a part of. */
-	UPROPERTY(transient)
+	UPROPERTY(Transient)
 	AActor* GroupActor;
 
 	/** The scale to apply to any billboard components in editor builds (happens in any WITH_EDITOR build, including non-cooked games). */
 	UPROPERTY(EditAnywhere, BlueprintReadWrite, Category=Rendering, meta=(DisplayName="Editor Billboard Scale"))
 	float SpriteScale;
 
+	/** Bitflag to represent which views this actor is hidden in, via per-view layer visibility. */
+	UPROPERTY(Transient)
+	uint64 HiddenEditorViews;
+
 private:
 	/**
 	 * The friendly name for this actor, displayed in the editor.  You should always use AActor::GetActorLabel() to access the actual label to display,
@@ -579,26 +553,11 @@ private:
 	UPROPERTY()
 	FName FolderPath;
 
-protected:
-
-	UPROPERTY()
-	uint8 bActorLabelEditable:1;    // Is the actor label editable by the user?
-
 public:
 	/** Whether this actor is hidden within the editor viewport. */
 	UPROPERTY()
 	uint8 bHiddenEd:1;
 
-protected:
-	/** Whether the actor can be manipulated by editor operations. */
-	UPROPERTY()
-	uint8 bEditable:1;
-
-	/** Whether this actor should be listed in the scene outliner. */
-	UPROPERTY()
-	uint8 bListedInSceneOutliner:1;
-
-public:
 	/** True if this actor is the preview actor dragged out of the content browser */
 	UPROPERTY()
 	uint8 bIsEditorPreviewActor:1;
@@ -607,15 +566,8 @@ public:
 	UPROPERTY()
 	uint8 bHiddenEdLayer:1;
 
-private:
-	/** Whether this actor is temporarily hidden within the editor; used for show/hide/etc functionality w/o dirtying the actor. */
-	UPROPERTY(transient)
-	uint8 bHiddenEdTemporary:1;
-
-public:
-
 	/** Whether this actor is hidden by the level browser. */
-	UPROPERTY(transient)
+	UPROPERTY(Transient)
 	uint8 bHiddenEdLevel:1;
 
 	/** If true, prevents the actor from being moved in the editor viewport. */
@@ -625,23 +577,31 @@ public:
 	/** If true during PostEditMove the construction script will be run every time. If false it will only run when the drag finishes. */
 	uint8 bRunConstructionScriptOnDrag:1;
 
-	/** Returns how many lights are uncached for this actor. */
-	int32 GetNumUncachedLights();
-#endif // WITH_EDITORONLY_DATA
+protected:
+	/** Is the actor label editable by the user? */
+	UPROPERTY()
+	uint8 bActorLabelEditable:1; 
+
+	/** Whether the actor can be manipulated by editor operations. */
+	UPROPERTY()
+	uint8 bEditable:1;
+
+	/** Whether this actor should be listed in the scene outliner. */
+	UPROPERTY()
+	uint8 bListedInSceneOutliner:1;
 
 private:
+	/** Whether this actor is temporarily hidden within the editor; used for show/hide/etc functionality w/o dirtying the actor. */
+	UPROPERTY(Transient)
+	uint8 bHiddenEdTemporary:1;
 
-	static uint32 BeginPlayCallDepth;
+#endif // WITH_EDITORONLY_DATA
 
 public:
-
 	/** Array of tags that can be used for grouping and categorizing. */
 	UPROPERTY(EditAnywhere, BlueprintReadWrite, AdvancedDisplay, Category=Actor)
 	TArray<FName> Tags;
 
-	/** Bitflag to represent which views this actor is hidden in, via per-view layer visibility. */
-	UPROPERTY(transient)
-	uint64 HiddenEditorViews;
 
 	//~==============================================================================================
 	// Delegates
@@ -655,16 +615,16 @@ public:
 	FTakePointDamageSignature OnTakePointDamage;
 	
 	/** 
-	 *	Called when another actor begins to overlap this actor, for example a player walking into a trigger.
-	 *	For events when objects have a blocking collision, for example a player hitting a wall, see 'Hit' events.
-	 *	@note Components on both this and the other Actor must have bGenerateOverlapEvents set to true to generate overlap events.
+	 * Called when another actor begins to overlap this actor, for example a player walking into a trigger.
+	 * For events when objects have a blocking collision, for example a player hitting a wall, see 'Hit' events.
+	 * @note Components on both this and the other Actor must have bGenerateOverlapEvents set to true to generate overlap events.
 	 */
 	UPROPERTY(BlueprintAssignable, Category="Collision")
 	FActorBeginOverlapSignature OnActorBeginOverlap;
 
 	/** 
-	 *	Called when another actor stops overlapping this actor. 
-	 *	@note Components on both this and the other Actor must have bGenerateOverlapEvents set to true to generate overlap events.
+	 * Called when another actor stops overlapping this actor. 
+	 * @note Components on both this and the other Actor must have bGenerateOverlapEvents set to true to generate overlap events.
 	 */
 	UPROPERTY(BlueprintAssignable, Category="Collision")
 	FActorEndOverlapSignature OnActorEndOverlap;
@@ -735,18 +695,18 @@ public:
 	UFUNCTION(BlueprintCallable, meta=(BlueprintInternalUseOnly="true", HideSelfPin="true", HidePin="InputAxisKey"))
 	FVector GetInputVectorAxisValue(const FKey InputAxisKey) const;
 
-	/** Returns the instigator for this actor, or NULL if there is none. */
+	/** Returns the instigator for this actor, or nullptr if there is none. */
 	UFUNCTION(BlueprintCallable, meta=(BlueprintProtected = "true"), Category="Game")
 	APawn* GetInstigator() const;
 
 	/**
 	 * Get the instigator, cast as a specific class.
-	 * @return The instigator for this weapon if it is the specified type, NULL otherwise.
+	 * @return The instigator for this weapon if it is the specified type, nullptr otherwise.
 	 */
 	template <class T>
 	T* GetInstigator() const { return Cast<T>(Instigator); };
 
-	/** Returns the instigator's controller for this actor, or NULL if there is none. */
+	/** Returns the instigator's controller for this actor, or nullptr if there is none. */
 	UFUNCTION(BlueprintCallable, meta=(BlueprintProtected = "true"), Category="Game")
 	AController* GetInstigatorController() const;
 
@@ -770,7 +730,6 @@ public:
 		return (RootComponent ? RootComponent->GetComponentTransform() : FTransform::Identity);
 	}
 
-
 	/** Returns the location of the RootComponent of this Actor */
 	UFUNCTION(BlueprintCallable, meta=(DisplayName = "GetActorLocation", Keywords="position"), Category="Utilities|Transformation")
 	FVector K2_GetActorLocation() const;
@@ -849,14 +808,14 @@ public:
 	bool K2_SetActorRotation(FRotator NewRotation, bool bTeleportPhysics);
 	
 	/**
-	* Set the Actor's rotation instantly to the specified rotation.
-	*
-	* @param	NewRotation	The new rotation for the Actor.
-	* @param	Teleport	How we teleport the physics state (if physics collision is enabled for this object).
-	*						If equal to ETeleportType::TeleportPhysics, physics velocity for this object is unchanged (so ragdoll parts are not affected by change in location).
-	*						If equal to ETeleportType::None, physics velocity is updated based on the change in position (affecting ragdoll parts).
-	* @return	Whether the rotation was successfully set.
-	*/
+	 * Set the Actor's rotation instantly to the specified rotation.
+	 *
+	 * @param	NewRotation	The new rotation for the Actor.
+	 * @param	Teleport	How we teleport the physics state (if physics collision is enabled for this object).
+	 *						If equal to ETeleportType::TeleportPhysics, physics velocity for this object is unchanged (so ragdoll parts are not affected by change in location).
+	 *						If equal to ETeleportType::None, physics velocity is updated based on the change in position (affecting ragdoll parts).
+	 * @return	Whether the rotation was successfully set.
+	 */
 	bool SetActorRotation(FRotator NewRotation, ETeleportType Teleport = ETeleportType::None);
 	bool SetActorRotation(const FQuat& NewRotation, ETeleportType Teleport = ETeleportType::None);
 
@@ -955,7 +914,6 @@ public:
 	 */
 	void AddActorWorldOffset(FVector DeltaLocation, bool bSweep=false, FHitResult* OutSweepHitResult=nullptr, ETeleportType Teleport = ETeleportType::None);
 
-
 	/**
 	 * Adds a delta to the rotation of this actor in world space.
 	 * 
@@ -1090,9 +1048,7 @@ public:
 	UFUNCTION(BlueprintCallable, Category="Utilities|Transformation")
 	void SetActorRelativeScale3D(FVector NewRelativeScale);
 
-	/**
-	 * Return the actor's relative scale 3d
-	 */
+	/** Return the actor's relative scale 3d */
 	UFUNCTION(BlueprintCallable, Category="Utilities|Orientation")
 	FVector GetActorRelativeScale3D() const;
 
@@ -1142,22 +1098,14 @@ public:
 	UFUNCTION(BlueprintCallable, meta=(DeprecatedFunction, DeprecationMessage = "Use Component.DestroyComponent instead", BlueprintProtected = "true", DisplayName = "DestroyComponent"))
 	void K2_DestroyComponent(UActorComponent* Component);
 
-	/** 
-	 *  Attaches the RootComponent of this Actor to the supplied component, optionally at a named socket. It is not valid to call this on components that are not Registered. 
-	 *   @param AttachLocationType	Type of attachment, AbsoluteWorld to keep its world position, RelativeOffset to keep the object's relative offset and SnapTo to snap to the new parent.
-	 */
-	DEPRECATED(4.12, "Please use AttachToComponent.")
-	void AttachRootComponentTo(USceneComponent* InParent, FName InSocketName = NAME_None, EAttachLocation::Type AttachLocationType = EAttachLocation::KeepRelativeOffset, bool bWeldSimulatedBodies = false);
-
 	/**
-	 *  Attaches the RootComponent of this Actor to the supplied component, optionally at a named socket. It is not valid to call this on components that are not Registered.
-	 *   @param AttachLocationType	Type of attachment, AbsoluteWorld to keep its world position, RelativeOffset to keep the object's relative offset and SnapTo to snap to the new parent.
+	 * Attaches the RootComponent of this Actor to the supplied component, optionally at a named socket. It is not valid to call this on components that are not Registered.
+	 * @param AttachLocationType	Type of attachment, AbsoluteWorld to keep its world position, RelativeOffset to keep the object's relative offset and SnapTo to snap to the new parent.
 	 */
 	DEPRECATED(4.17, "Use AttachToComponent instead.")
 	UFUNCTION(BlueprintCallable, meta = (DisplayName = "AttachActorToComponent (Deprecated)", AttachLocationType = "KeepRelativeOffset"), Category = "Utilities|Transformation")
 	void K2_AttachRootComponentTo(USceneComponent* InParent, FName InSocketName = NAME_None, EAttachLocation::Type AttachLocationType = EAttachLocation::KeepRelativeOffset, bool bWeldSimulatedBodies = true);
 
-
 	/**
 	 * Attaches the RootComponent of this Actor to the supplied component, optionally at a named socket. It is not valid to call this on components that are not Registered.
 	 * @param  Parent					Parent to attach to.
@@ -1177,18 +1125,9 @@ public:
 	void AttachToComponent(USceneComponent* Parent, const FAttachmentTransformRules& AttachmentRules, FName SocketName = NAME_None);
 
 	/**
-	 * Attaches the RootComponent of this Actor to the RootComponent of the supplied actor, optionally at a named socket.
-	 * @param InParentActor				Actor to attach this actor's RootComponent to
-	 * @param InSocketName				Socket name to attach to, if any
-	 * @param AttachLocationType	Type of attachment, AbsoluteWorld to keep its world position, RelativeOffset to keep the object's relative offset and SnapTo to snap to the new parent.
+	 *  Attaches the RootComponent of this Actor to the supplied component, optionally at a named socket. It is not valid to call this on components that are not Registered.
+	 *   @param AttachLocationType	Type of attachment, AbsoluteWorld to keep its world position, RelativeOffset to keep the object's relative offset and SnapTo to snap to the new parent.
 	 */
-	DEPRECATED(4.12, "Please use AttachToActor.")
-	void AttachRootComponentToActor(AActor* InParentActor, FName InSocketName = NAME_None, EAttachLocation::Type AttachLocationType = EAttachLocation::KeepRelativeOffset, bool bWeldSimulatedBodies = false);
-
-	/**
-	*  Attaches the RootComponent of this Actor to the supplied component, optionally at a named socket. It is not valid to call this on components that are not Registered.
-	*   @param AttachLocationType	Type of attachment, AbsoluteWorld to keep its world position, RelativeOffset to keep the object's relative offset and SnapTo to snap to the new parent.
-	*/
 	DEPRECATED(4.17, "Use AttachToActor instead.")
 	UFUNCTION(BlueprintCallable, meta = (DisplayName = "AttachActorToActor (Deprecated)", AttachLocationType = "KeepRelativeOffset"), Category = "Utilities|Transformation")
 	void K2_AttachRootComponentToActor(AActor* InParentActor, FName InSocketName = NAME_None, EAttachLocation::Type AttachLocationType = EAttachLocation::KeepRelativeOffset, bool bWeldSimulatedBodies = true);
@@ -1214,16 +1153,16 @@ public:
 	void K2_AttachToActor(AActor* ParentActor, FName SocketName, EAttachmentRule LocationRule, EAttachmentRule RotationRule, EAttachmentRule ScaleRule, bool bWeldSimulatedBodies);
 
 	/** 
-	 *  Snap the RootComponent of this Actor to the supplied Actor's root component, optionally at a named socket. It is not valid to call this on components that are not Registered. 
-	 *  If InSocketName == NAME_None, it will attach to origin of the InParentActor. 
+	 * Snap the RootComponent of this Actor to the supplied Actor's root component, optionally at a named socket. It is not valid to call this on components that are not Registered. 
+	 * If InSocketName == NAME_None, it will attach to origin of the InParentActor. 
 	 */
 	DEPRECATED(4.17, "Use AttachRootComponentTo with EAttachLocation::SnapToTarget option instead")
 	UFUNCTION(BlueprintCallable, meta=(DeprecatedFunction, DeprecationMessage = "Use AttachRootComponentTo with EAttachLocation::SnapToTarget option instead", DisplayName = "SnapActorTo"), Category="Utilities|Transformation")
 	void SnapRootComponentTo(AActor* InParentActor, FName InSocketName);
 
 	/** 
-	 *  Detaches the RootComponent of this Actor from any SceneComponent it is currently attached to. 
-	 *   @param bMaintainWorldTransform	If true, update the relative location/rotation of this component to keep its world position the same
+	 * Detaches the RootComponent of this Actor from any SceneComponent it is currently attached to. 
+	 * @param bMaintainWorldTransform	If true, update the relative location/rotation of this component to keep its world position the same
 	 */
 	DEPRECATED(4.17, "Use DetachFromActor instead")
 	UFUNCTION(BlueprintCallable, meta=(DisplayName = "DetachActorFromActor (Deprecated)"), Category="Utilities|Transformation")
@@ -1244,24 +1183,13 @@ public:
 	 */
 	void DetachFromActor(const FDetachmentTransformRules& DetachmentRules);
 
-	/** 
-	 *	DEPRECATED: Detaches all SceneComponents in this Actor from the supplied parent SceneComponent. 
+	/**
+	 *	Detaches all SceneComponents in this Actor from the supplied parent SceneComponent.
 	 *	@param InParentComponent		SceneComponent to detach this actor's components from
-	 *	@param bMaintainWorldTransform	If true, update the relative location/rotation of this component to keep its world position the same
+	 *	@param DetachmentRules			Rules to apply when detaching components
 	 */
-	DEPRECATED(4.12, "Please use DetachAllSceneComponents")
-	void DetachSceneComponentsFromParent(USceneComponent* InParentComponent, bool bMaintainWorldPosition = true);
-
-	/**
-	*	Detaches all SceneComponents in this Actor from the supplied parent SceneComponent.
-	*	@param InParentComponent		SceneComponent to detach this actor's components from
-	*	@param DetachmentRules			Rules to apply when detaching components
-	*/
 	void DetachAllSceneComponents(class USceneComponent* InParentComponent, const FDetachmentTransformRules& DetachmentRules);
 
-	//~==============================================================================
-	// Tags
-
 	/** See if this actor contains the supplied tag */
 	UFUNCTION(BlueprintCallable, Category="Utilities")
 	bool ActorHasTag(FName Tag) const;
@@ -1307,26 +1235,9 @@ public:
 	class UMaterialInstanceDynamic* MakeMIDForMaterial(class UMaterialInterface* Parent);
 
 	/** The number of seconds (in game time) since this Actor was created, relative to Get Game Time In Seconds. */
-	UFUNCTION(BlueprintPure, Category = Actor)
+	UFUNCTION(BlueprintPure, Category=Actor)
 	float GetGameTimeSinceCreation();
 
-	//~=============================================================================
-	// AI functions.
-	
-	/**
-	 * Trigger a noise caused by a given Pawn, at a given location.
-	 * Note that the NoiseInstigator Pawn MUST have a PawnNoiseEmitterComponent for the noise to be detected by a PawnSensingComponent.
-	 * Senders of MakeNoise should have an Instigator if they are not pawns, or pass a NoiseInstigator.
-	 *
-	 * @param Loudness The relative loudness of this noise. Usual range is 0 (no noise) to 1 (full volume). If MaxRange is used, this scales the max range, otherwise it affects the hearing range specified by the sensor.
-	 * @param NoiseInstigator Pawn responsible for this noise.  Uses the actor's Instigator if NoiseInstigator=NULL
-	 * @param NoiseLocation Position of noise source.  If zero vector, use the actor's location.
-	 * @param MaxRange Max range at which the sound may be heard. A value of 0 indicates no max range (though perception may have its own range). Loudness scales the range. (Note: not supported for legacy PawnSensingComponent, only for AIPerception)
-	 * @param Tag Identifier for the noise.
-	 */
-	UFUNCTION(BlueprintCallable, BlueprintAuthorityOnly, Category="AI", meta=(BlueprintProtected = "true"))
-	void MakeNoise(float Loudness=1.f, APawn* NoiseInstigator=NULL, FVector NoiseLocation=FVector::ZeroVector, float MaxRange = 0.f, FName Tag = NAME_None);
-
 protected:
 	/** Event when play begins for this actor. */
 	UFUNCTION(BlueprintImplementableEvent, meta=(DisplayName = "BeginPlay"))
@@ -1336,8 +1247,7 @@ protected:
 	virtual void BeginPlay();
 
 public:
-
-	/** Initiate a begin play call on this Actor, will handle . */
+	/** Initiate a begin play call on this Actor, will handle calling in the correct order. */
 	void DispatchBeginPlay();
 
 	/** Returns whether an actor has been initialized */
@@ -1370,9 +1280,6 @@ public:
 	UFUNCTION(BlueprintImplementableEvent, BlueprintAuthorityOnly, meta=(DisplayName = "PointDamage"), Category="Game|Damage")
 	void ReceivePointDamage(float Damage, const class UDamageType* DamageType, FVector HitLocation, FVector HitNormal, class UPrimitiveComponent* HitComponent, FName BoneName, FVector ShotFromDirection, class AController* InstigatedBy, AActor* DamageCauser, const FHitResult& HitInfo);
 
-	DEPRECATED(4.14, "Call the updated version of ReceivePointDamage that takes a FHitResult.")
-	void ReceivePointDamage(float Damage, const class UDamageType* DamageType, FVector HitLocation, FVector HitNormal, class UPrimitiveComponent* HitComponent, FName BoneName, FVector ShotFromDirection, class AController* InstigatedBy, AActor* DamageCauser);
-
 	/** Event called every frame */
 	UFUNCTION(BlueprintImplementableEvent, meta=(DisplayName = "Tick"))
 	void ReceiveTick(float DeltaSeconds);
@@ -1496,7 +1403,6 @@ public:
 	UFUNCTION(BlueprintImplementableEvent, meta=(DisplayName = "Hit"), Category="Collision")
 	void ReceiveHit(class UPrimitiveComponent* MyComp, AActor* Other, class UPrimitiveComponent* OtherComp, bool bSelfMoved, FVector HitLocation, FVector HitNormal, FVector NormalImpulse, const FHitResult& Hit);
 
-
 	/** Set the lifespan of this actor. When it expires the object will be destroyed. If requested lifespan is 0, the timer is cleared and the actor will not be destroyed. */
 	UFUNCTION(BlueprintCallable, Category="Utilities", meta=(Keywords = "delete destroy"))
 	virtual void SetLifeSpan( float InLifespan );
@@ -1530,7 +1436,6 @@ public:
 	UPROPERTY(BlueprintAssignable, Category="Game")
 	FActorDestroyedSignature OnDestroyed;
 
-
 	/** Event to notify blueprints this actor is about to be deleted. */
 	UFUNCTION(BlueprintImplementableEvent, meta=(Keywords = "delete", DisplayName = "End Play"))
 	void ReceiveEndPlay(EEndPlayReason::Type EndPlayReason);
@@ -1551,7 +1456,7 @@ public:
 	virtual void PostLoadSubobjects( FObjectInstancingGraph* OuterInstanceGraph ) override;
 	virtual void BeginDestroy() override;
 	virtual bool IsReadyForFinishDestroy() override;
-	virtual bool Rename( const TCHAR* NewName=NULL, UObject* NewOuter=NULL, ERenameFlags Flags=REN_None ) override;
+	virtual bool Rename( const TCHAR* NewName=nullptr, UObject* NewOuter=nullptr, ERenameFlags Flags=REN_None ) override;
 	virtual bool CanBeInCluster() const override;
 	static void AddReferencedObjects(UObject* InThis, FReferenceCollector& Collector);
 #if WITH_EDITOR
@@ -1603,7 +1508,7 @@ public:
 		FActorRootComponentReconstructionData RootComponentData;
 	};
 
-	/* Cached pointer to the transaction annotation data from PostEditUndo to be used in the next RerunConstructionScript */
+	/** Cached pointer to the transaction annotation data from PostEditUndo to be used in the next RerunConstructionScript */
 	TSharedPtr<FActorTransactionAnnotation> CurrentTransactionAnnotation;
 
 	virtual TSharedPtr<ITransactionObjectAnnotation> GetTransactionAnnotation() const override;
@@ -1611,22 +1516,20 @@ public:
 
 	/** @return true if the component is allowed to re-register its components when modified.  False for CDOs or PIE instances. */
 	bool ReregisterComponentsWhenModified() const;
-#endif // WITH_EDITOR
-	//~ End UObject Interface
 
-#if WITH_EDITOR
+	/** Called after an actor has been moved in the editor */
 	virtual void PostEditMove(bool bFinished);
 #endif // WITH_EDITOR
+	//~ End UObject Interface
+
 
-	//-----------------------------------------------------------------------------------------------
-	// PROPERTY REPLICATION
+	//~=============================================================================
+	// Property Replication
 
-	/** fills ReplicatedMovement property */
+	/** Fills ReplicatedMovement property */
 	virtual void GatherCurrentMovement();
 
-	/**
-	 * See if this actor is owned by TestOwner.
-	 */
+	/** See if this actor is owned by TestOwner. */
 	inline bool IsOwnedBy( const AActor* TestOwner ) const
 	{
 		for( const AActor* Arg=this; Arg; Arg=Arg->Owner )
@@ -1696,9 +1599,9 @@ public:
 	}
 #endif
 
-/*-----------------------------------------------------------------------------
-	Relations.
------------------------------------------------------------------------------*/
+
+	//~=============================================================================
+	// Relations
 
 	/** 
 	 * Called by owning level to shift an actor location and all relevant data structures by specified delta
@@ -1753,7 +1656,7 @@ public:
 		return bHiddenEd;
 	}
 
-	// Returns true if this actor is hidden in the editor viewports.
+	/** Returns true if this actor is hidden in the editor viewports. */
 	UFUNCTION(BlueprintCallable, Category = "Editor Scripting | Actor Editing")
 	bool IsHiddenEd() const;
 
@@ -1786,10 +1689,10 @@ public:
 	/** @return	Returns true if this actor is allowed to be attached to the given actor */
 	virtual bool EditorCanAttachTo(const AActor* InParent, FText& OutReason) const;
 
-	// Called before editor copy, true allow export
+	/** Called before editor copy, true allow export */
 	virtual bool ShouldExport() { return true; }
 
-	// Called before editor paste, true allow import
+	/** Called before editor paste, true allow import */
 	virtual bool ShouldImport(FString* ActorPropString, bool IsMovingLevel) { return true; }
 
 	/** Called by InputKey when an unhandled key is pressed with a selected actor */
@@ -1942,9 +1845,9 @@ public:
 	bool IsActorTickEnabled() const;
 
 	/** 
-	* Sets the tick interval of this actor's primary tick function. Will not enable a disabled tick function. Takes effect on next tick. 
-	* @param TickInterval	The rate at which this actor should be ticking
-	*/
+	 * Sets the tick interval of this actor's primary tick function. Will not enable a disabled tick function. Takes effect on next tick. 
+	 * @param TickInterval	The rate at which this actor should be ticking
+	 */
 	UFUNCTION(BlueprintCallable, Category="Utilities")
 	void SetActorTickInterval(float TickInterval);
 
@@ -1970,10 +1873,10 @@ public:
 	/** Called when the lifespan of an actor expires (if he has one). */
 	virtual void LifeSpanExpired();
 
-	// Always called immediately before properties are received from the remote.
+	/** Always called immediately before properties are received from the remote. */
 	virtual void PreNetReceive() override;
 	
-	// Always called immediately after properties are received from the remote.
+	/** Always called immediately after properties are received from the remote. */
 	virtual void PostNetReceive() override;
 
 	/** IsNameStableForNetworking means an object can be referred to its path name (relative to outer) over the network */
@@ -1985,7 +1888,7 @@ public:
 	/** Returns a list of sub-objects that have stable names for networking */
 	virtual void GetSubobjectsWithStableNamesForNetworking(TArray<UObject*> &ObjList) override;
 
-	// Always called immediately after spawning and reading in replicated properties
+	/** Always called immediately after spawning and reading in replicated properties */
 	virtual void PostNetInit();
 
 	/** ReplicatedMovement struct replication event */
@@ -2027,7 +1930,8 @@ public:
 	 */
 	virtual bool CheckStillInWorld();
 
-	//--------------------------------------------------------------------------------------
+
+	//~=============================================================================
 	// Actor overlap tracking
 	
 	/**
@@ -2062,10 +1966,11 @@ public:
 	/** See if the root component has Mobility of EComponentMobility::Movable */
 	bool IsRootComponentMovable() const;
 
-	//--------------------------------------------------------------------------------------
+
+	//~=============================================================================
 	// Actor ticking
 
-	/** accessor for the value of bCanEverTick */
+	/** Accessor for the value of bCanEverTick */
 	FORCEINLINE bool CanEverTick() const { return PrimaryActorTick.bCanEverTick; }
 
 	/** 
@@ -2079,11 +1984,11 @@ public:
 	/** If true, actor is ticked even if TickType==LEVELTICK_ViewportsOnly	 */
 	virtual bool ShouldTickIfViewportsOnly() const;
 
-	//--------------------------------------------------------------------------------------
+
+	//~=============================================================================
 	// Actor relevancy determination
 
 protected:
-
 	/**
 	 * Determines whether or not the distance between the given SrcLocation and the Actor's location
 	 * is within the net relevancy distance. Actors outside relevancy distance may not be replicated.
@@ -2094,7 +1999,6 @@ protected:
 	bool IsWithinNetRelevancyDistance(const FVector& SrcLocation) const;
 	
 public:	
-
 	/** 
 	  * @param RealViewer - is the "controlling net object" associated with the client for which network relevancy is being checked (typically player controller)
 	  * @param ViewTarget - is the Actor being used as the point of view for the RealViewer
@@ -2137,14 +2041,12 @@ public:
 	/** Called immediately before gameplay begins. */
 	virtual void PreInitializeComponents();
 
-	// Allow actors to initialize themselves on the C++ side
+	/** Allow actors to initialize themselves on the C++ side */
 	virtual void PostInitializeComponents();
 
 	/**
 	 * Adds a controlling matinee actor for use during matinee playback
 	 * @param InMatineeActor	The matinee actor which controls this actor
-	 * @todo UE4 would be nice to move this out of Actor to MatineeInterface, but it needs variable, and I can't declare variable in interface
-	 *	do we still need this?
 	 */
 	void AddControllingMatineeActor( AMatineeActor& InMatineeActor );
 
@@ -2157,10 +2059,10 @@ public:
 	/** Dispatches ReceiveHit virtual and OnComponentHit delegate */
 	virtual void DispatchPhysicsCollisionHit(const struct FRigidBodyCollisionInfo& MyInfo, const struct FRigidBodyCollisionInfo& OtherInfo, const FCollisionImpactData& RigidCollisionData);
 	
-	/** @return the actor responsible for replication, if any.  Typically the player controller */
+	/** @return the actor responsible for replication, if any. Typically the player controller */
 	virtual const AActor* GetNetOwner() const;
 
-	/** @return the owning UPlayer (if any) of this actor. This will be a local player, a net connection, or NULL. */
+	/** @return the owning UPlayer (if any) of this actor. This will be a local player, a net connection, or nullptr. */
 	virtual class UPlayer* GetNetOwningPlayer();
 
 	/**
@@ -2184,13 +2086,14 @@ public:
 	ENetMode GetNetMode() const;
 
 	/**
-	* Test whether net mode is the given mode.
-	* In optimized non-editor builds this can be more efficient than GetNetMode()
-	* because it can check the static build flags without considering PIE.
-	*/
+	 * Test whether net mode is the given mode.
+	 * In optimized non-editor builds this can be more efficient than GetNetMode()
+	 * because it can check the static build flags without considering PIE.
+	 */
 	bool IsNetMode(ENetMode Mode) const;
 
-	class UNetDriver * GetNetDriver() const;
+	/** Returns the net driver that this actor is bound to, may be null */
+	class UNetDriver* GetNetDriver() const;
 
 	/** Puts actor in dormant networking state */
 	UFUNCTION(BlueprintAuthorityOnly, BlueprintCallable, Category = "Networking")
@@ -2263,12 +2166,14 @@ public:
 	/** Debug rendering to visualize the component tree for this actor. */
 	void DrawDebugComponents(FColor const& BaseColor=FColor::White) const;
 
+	/** Called to mark all components as pending kill when the actor is being destroyed */
 	virtual void MarkComponentsAsPendingKill();
 
-	/** Returns true if this actor has begun the destruction process.
-	 *  This is set to true in UWorld::DestroyActor, after the network connection has been closed but before any other shutdown has been performed.
-	 *	@return true if this actor has begun destruction, or if this actor has been destroyed already.
-	 **/
+	/**
+	 * Returns true if this actor has begun the destruction process.
+	 * This is set to true in UWorld::DestroyActor, after the network connection has been closed but before any other shutdown has been performed.
+	 * @return true if this actor has begun destruction, or if this actor has been destroyed already.
+	 */
 	inline bool IsPendingKillPending() const
 	{
 		return bActorIsBeingDestroyed || IsPendingKill();
@@ -2284,16 +2189,16 @@ public:
 	virtual void InvalidateLightingCacheDetailed(bool bTranslationOnly);
 
 	/**
-	  * Used for adding actors to levels or teleporting them to a new location.
-	  * The result of this function is independent of the actor's current location and rotation.
-	  * If the actor doesn't fit exactly at the location specified, tries to slightly move it out of walls and such if bNoCheck is false.
-	  *
-	  * @param DestLocation The target destination point
-	  * @param DestRotation The target rotation at the destination
-	  * @param bIsATest is true if this is a test movement, which shouldn't cause any notifications (used by AI pathfinding, for example)
-	  * @param bNoCheck is true if we should skip checking for encroachment in the world or other actors
-	  * @return true if the actor has been successfully moved, or false if it couldn't fit.
-	  */
+	 * Used for adding actors to levels or teleporting them to a new location.
+	 * The result of this function is independent of the actor's current location and rotation.
+	 * If the actor doesn't fit exactly at the location specified, tries to slightly move it out of walls and such if bNoCheck is false.
+	 *
+	 * @param DestLocation The target destination point
+	 * @param DestRotation The target rotation at the destination
+	 * @param bIsATest is true if this is a test movement, which shouldn't cause any notifications (used by AI pathfinding, for example)
+	 * @param bNoCheck is true if we should skip checking for encroachment in the world or other actors
+	 * @return true if the actor has been successfully moved, or false if it couldn't fit.
+	 */
 	virtual bool TeleportTo( const FVector& DestLocation, const FRotator& DestRotation, bool bIsATest=false, bool bNoCheck=false );
 
 	/**
@@ -2310,13 +2215,13 @@ public:
 	virtual void TeleportSucceeded(bool bIsATest) {}
 
 	/**
-	 *  Trace a ray against the Components of this Actor and return the first blocking hit
-	 *  @param  OutHit          First blocking hit found
-	 *  @param  Start           Start location of the ray
-	 *  @param  End             End location of the ray
-	 *  @param  TraceChannel    The 'channel' that this ray is in, used to determine which components to hit
-	 *  @param  Params          Additional parameters used for the trace
-	 *  @return TRUE if a blocking hit is found
+	 * Trace a ray against the Components of this Actor and return the first blocking hit
+	 * @param  OutHit          First blocking hit found
+	 * @param  Start           Start location of the ray
+	 * @param  End             End location of the ray
+	 * @param  TraceChannel    The 'channel' that this ray is in, used to determine which components to hit
+	 * @param  Params          Additional parameters used for the trace
+	 * @return TRUE if a blocking hit is found
 	 */
 	bool ActorLineTraceSingle(struct FHitResult& OutHit, const FVector& Start, const FVector& End, ECollisionChannel TraceChannel, const struct FCollisionQueryParams& Params);
 
@@ -2332,12 +2237,9 @@ public:
 	 * @return		Success if returns > 0.f, if returns 0.f, it is either not convex or inside of the point
 	 *				If returns < 0.f, this Actor does not have any primitive with collision
 	 */
-	float ActorGetDistanceToCollision(const FVector& Point, ECollisionChannel TraceChannel, FVector& ClosestPointOnCollision, UPrimitiveComponent** OutPrimitiveComponent = NULL) const;
+	float ActorGetDistanceToCollision(const FVector& Point, ECollisionChannel TraceChannel, FVector& ClosestPointOnCollision, UPrimitiveComponent** OutPrimitiveComponent = nullptr) const;
 
-	/**
-	 * Returns true if this actor is contained by TestLevel.
-	 * @todo seamless: update once Actor->Outer != Level
-	 */
+	/** Returns true if this actor is contained by TestLevel. */
 	bool IsInLevel(const class ULevel *TestLevel) const;
 
 	/** Return the ULevel that this Actor is part of. */
@@ -2352,23 +2254,26 @@ public:
 	/** Overridable function called whenever this actor is being removed from a level */
 	virtual void EndPlay(const EEndPlayReason::Type EndPlayReason);
 
-	/** iterates up the Base chain to see whether or not this Actor is based on the given Actor
+	/**
+	 * Iterates up the Base chain to see whether or not this Actor is based on the given Actor
 	 * @param Other the Actor to test for
 	 * @return true if this Actor is based on Other Actor
 	 */
 	virtual bool IsBasedOnActor(const AActor* Other) const;
 	
-
-	/** iterates up the Base chain to see whether or not this Actor is attached to the given Actor
-	* @param Other the Actor to test for
-	* @return true if this Actor is attached on Other Actor
-	*/
+	/**
+	 * Iterates up the Base chain to see whether or not this Actor is attached to the given Actor
+	 * @param Other the Actor to test for
+	 * @return true if this Actor is attached on Other Actor
+	 */
 	virtual bool IsAttachedTo( const AActor* Other ) const;
 
 	/** Get the extent used when placing this actor in the editor, used for 'pulling back' hit. */
 	FVector GetPlacementExtent() const;
 
-	// Blueprint 
+
+	//~=============================================================================
+	// Blueprint
 
 #if WITH_EDITOR
 	/** Find all FRandomStream structs in this ACtor and generate new random seeds for them. */
@@ -2420,21 +2325,9 @@ public:
 	UActorComponent* CreateComponentFromTemplate(UActorComponent* Template, const FName InName = NAME_None );
 	UActorComponent* CreateComponentFromTemplateData(const struct FBlueprintCookedComponentInstancingData* TemplateData, const FName InName = NAME_None);
 
-	DEPRECATED(4.11, "Use CreateComponentFromTemplate that takes a FName instead of a FString")
-	UActorComponent* CreateComponentFromTemplate(UActorComponent* Template, const FString& InName);
-
 	/** Destroys the constructed components. */
 	void DestroyConstructedComponents();
 
-#if !(UE_BUILD_SHIPPING || UE_BUILD_TEST)
-	private:
-	// this is the old name of the tick function. We just want to avoid mistakes with an attempt to override this
-	virtual void Tick( float DeltaTime, enum ELevelTick TickType ) final
-	{
-		check(0);
-	}
-#endif
-
 protected:
 	/**
 	 * Virtual call chain to register all tick functions for the actor class hierarchy
@@ -2445,9 +2338,7 @@ protected:
 	/** Runs UserConstructionScript, delays component registration until it's complete. */
 	void ProcessUserConstructionScript();
 
-	/**
-	* Checks components for validity, implemented in AActor
-	*/
+	/** Checks components for validity, implemented in AActor */
 	bool CheckActorComponents();
 
 	/** Called after instancing a new Blueprint Component from either a template or cooked data. */
@@ -2458,7 +2349,7 @@ public:
 	/** Checks for and resolve any name conflicts prior to instancing a new Blueprint Component. */
 	void CheckComponentInstanceName(const FName InName);
 
-	/** Walk up the attachment chain from RootComponent until we encounter a different actor, and return it. If we are not attached to a component in a different actor, returns NULL */
+	/** Walk up the attachment chain from RootComponent until we encounter a different actor, and return it. If we are not attached to a component in a different actor, returns nullptr */
 	UFUNCTION(BlueprintPure, Category = "Utilities")
 	AActor* GetAttachParentActor() const;
 
@@ -2483,25 +2374,25 @@ public:
 	/** Call ReceiveHit, as well as delegates on Actor and Component */
 	void DispatchBlockingHit(UPrimitiveComponent* MyComp, UPrimitiveComponent* OtherComp, bool bSelfMoved, FHitResult const& Hit);
 
-	/** called when the actor falls out of the world 'safely' (below KillZ and such) */
+	/** Called when the actor falls out of the world 'safely' (below KillZ and such) */
 	virtual void FellOutOfWorld(const class UDamageType& dmgType);
 
-	/** called when the Actor is outside the hard limit on world bounds */
+	/** Called when the Actor is outside the hard limit on world bounds */
 	virtual void OutsideWorldBounds();
 
 	/** 
-	 *	Returns the world space bounding box of all components in this Actor.
-	 *	@param bNonColliding Indicates that you want to include non-colliding components in the bounding box
+	 * Returns the world space bounding box of all components in this Actor.
+	 * @param bNonColliding Indicates that you want to include non-colliding components in the bounding box
 	 */
 	virtual FBox GetComponentsBoundingBox(bool bNonColliding = false) const;
 
 	/** 
-	 *	Calculates the actor space bounding box of all components in this Actor.  This is slower than GetComponentsBoundingBox(), because the local bounds of the components are not cached -- they are recalculated every time this function is called.
-	 *	@param bNonColliding Indicates that you want to include non-colliding components in the bounding box
+	 * Calculates the actor space bounding box of all components in this Actor.  This is slower than GetComponentsBoundingBox(), because the local bounds of the components are not cached -- they are recalculated every time this function is called.
+	 * @param bNonColliding Indicates that you want to include non-colliding components in the bounding box
 	 */
 	virtual FBox CalculateComponentsBoundingBoxInLocalSpace(bool bNonColliding = false) const;
 
-	/* Get half-height/radius of a big axis-aligned cylinder around this actors registered colliding components, or all registered components if bNonColliding is false. */
+	/** Get half-height/radius of a big axis-aligned cylinder around this actors registered colliding components, or all registered components if bNonColliding is false. */
 	virtual void GetComponentsBoundingCylinder(float& CollisionRadius, float& CollisionHalfHeight, bool bNonColliding = false) const;
 
 	/**
@@ -2528,8 +2419,9 @@ public:
 	 */
 	virtual void TornOff();
 
+
 	//~=============================================================================
-	// Collision functions.
+	// Collision/Physics functions.
  
 	/** 
 	 * Get Collision Response to the Channel that entered for all components
@@ -2540,15 +2432,13 @@ public:
 	 */
 	virtual ECollisionResponse GetComponentsCollisionResponseToChannel(ECollisionChannel Channel) const;
 
-	//~=============================================================================
-	// Physics
-
 	/** Stop all simulation from all components in this actor */
 	void DisableComponentsSimulatePhysics();
 
-public:
-	/** @return WorldSettings for the World the actor is in
-	 - if you'd like to know UWorld this placed actor (not dynamic spawned actor) belong to, use GetTypedOuter<UWorld>() **/
+	/**
+	 * @return WorldSettings for the World the actor is in
+	 * If you'd like to know UWorld this placed actor (not dynamic spawned actor) belong to, use GetTypedOuter<UWorld>()
+	 */
 	class AWorldSettings* GetWorldSettings() const;
 
 	/**
@@ -2557,7 +2447,8 @@ public:
 	 */
 	virtual bool CanBeBaseForCharacter(class APawn* Pawn) const;
 
-	/** Apply damage to this actor.
+	/**
+	 * Apply damage to this actor.
 	 * @see https://www.unrealengine.com/blog/damage-in-ue4
 	 * @param DamageAmount		How much damage to apply
 	 * @param DamageEvent		Data package that fully describes the damage received.
@@ -2569,13 +2460,13 @@ public:
 
 protected:
 	virtual float InternalTakeRadialDamage(float Damage, struct FRadialDamageEvent const& RadialDamageEvent, class AController* EventInstigator, AActor* DamageCauser);
-	virtual float InternalTakePointDamage(float Damage, struct FPointDamageEvent const& RadialDamageEvent, class AController* EventInstigator, AActor* DamageCauser);
+	virtual float InternalTakePointDamage(float Damage, struct FPointDamageEvent const& PointDamageEvent, class AController* EventInstigator, AActor* DamageCauser);
 public:
 
-	/* Called when this actor becomes the given PlayerController's ViewTarget. Triggers the Blueprint event K2_OnBecomeViewTarget. */
+	/** Called when this actor becomes the given PlayerController's ViewTarget. Triggers the Blueprint event K2_OnBecomeViewTarget. */
 	virtual void BecomeViewTarget( class APlayerController* PC );
 
-	/* Called when this actor is no longer the given PlayerController's ViewTarget. Also triggers the Blueprint event K2_OnEndViewTarget. */
+	/*( Called when this actor is no longer the given PlayerController's ViewTarget. Also triggers the Blueprint event K2_OnEndViewTarget. */
 	virtual void EndViewTarget( class APlayerController* PC );
 
 	/** Event called when this Actor becomes the view target for the given PlayerController. */
@@ -2587,20 +2478,20 @@ public:
 	void K2_OnEndViewTarget( class APlayerController* PC );
 
 	/**
-	 *	Calculate camera view point, when viewing this actor.
+	 * Calculate camera view point, when viewing this actor.
 	 *
 	 * @param	DeltaTime	Delta time seconds since last update
 	 * @param	OutResult	Camera configuration
 	 */
 	virtual void CalcCamera(float DeltaTime, struct FMinimalViewInfo& OutResult);
 
-	// Returns true if the actor contains an active camera component
+	/** Returns true if the actor contains an active camera component */
 	virtual bool HasActiveCameraComponent() const;
 
-	// Returns true if the actor contains an active locked to HMD camera component
+	/** Returns true if the actor contains an active locked to HMD camera component */
 	virtual bool HasActivePawnControlCameraComponent() const;
 
-	// Returns the human readable string representation of an object.
+	/** Returns the human readable string representation of an object. */
 	virtual FString GetHumanReadableName() const;
 
 	/** Reset actor to initial state - used when restarting level without reloading. */
@@ -2659,10 +2550,10 @@ public:
 	virtual void ForceNetUpdate();
 
 	/**
-	 *	Calls PrestreamTextures() for all the actor's meshcomponents.
-	 *	@param Seconds - Number of seconds to force all mip-levels to be resident
-	 *	@param bEnableStreaming	- Whether to start (true) or stop (false) streaming
-	 *	@param CinematicTextureGroups - Bitfield indicating which texture groups that use extra high-resolution mips
+	 * Calls PrestreamTextures() for all the actor's meshcomponents.
+	 * @param Seconds - Number of seconds to force all mip-levels to be resident
+	 * @param bEnableStreaming	- Whether to start (true) or stop (false) streaming
+	 * @param CinematicTextureGroups - Bitfield indicating which texture groups that use extra high-resolution mips
 	 */
 	UFUNCTION(BlueprintCallable, Category = "Rendering")
 	virtual void PrestreamTextures( float Seconds, bool bEnableStreaming, int32 CinematicTextureGroups = 0 );
@@ -2681,18 +2572,18 @@ public:
 	virtual void GetActorEyesViewPoint( FVector& OutLocation, FRotator& OutRotation ) const;
 
 	/**
+	 * Returns the optimal location to fire weapons at this actor
 	 * @param RequestedBy - the Actor requesting the target location
-	 * @return the optimal location to fire weapons at this actor
 	 */
-	virtual FVector GetTargetLocation(AActor* RequestedBy = NULL) const;
+	virtual FVector GetTargetLocation(AActor* RequestedBy = nullptr) const;
 
-	/** 
-	  * Hook to allow actors to render HUD overlays for themselves.  Called from AHUD::DrawActorOverlays(). 
-	  * @param PC is the PlayerController on whose view this overlay is rendered
-	  * @param Canvas is the Canvas on which to draw the overlay
-	  * @param CameraPosition Position of Camera
-	  * @param CameraDir direction camera is pointing in.
-	  */
+	/**
+	 * Hook to allow actors to render HUD overlays for themselves.  Called from AHUD::DrawActorOverlays(). 
+	 * @param PC is the PlayerController on whose view this overlay is rendered
+	 * @param Canvas is the Canvas on which to draw the overlay
+	 * @param CameraPosition Position of Camera
+	 * @param CameraDir direction camera is pointing in.
+	 */
 	virtual void PostRenderFor(class APlayerController* PC, class UCanvas* Canvas, FVector CameraPosition, FVector CameraDir);
 
 	/** whether this Actor is in the persistent level, i.e. not a sublevel */
@@ -2717,12 +2608,14 @@ public:
 	UFUNCTION(BlueprintCallable, Category = "Actor", meta = (ComponentClass = "ActorComponent"), meta = (DeterminesOutputType = "ComponentClass"))
 	UActorComponent* GetComponentByClass(TSubclassOf<UActorComponent> ComponentClass) const;
 
-	/* Gets all the components that inherit from the given class.
-	Currently returns an array of UActorComponent which must be cast to the correct type. */
+	/**
+	 * Gets all the components that inherit from the given class.
+	 * Currently returns an array of UActorComponent which must be cast to the correct type.
+	 */
 	UFUNCTION(BlueprintCallable, Category = "Actor", meta = (ComponentClass = "ActorComponent"), meta = (DeterminesOutputType = "ComponentClass"))
 	TArray<UActorComponent*> GetComponentsByClass(TSubclassOf<UActorComponent> ComponentClass) const;
 
-	/* Gets all the components that inherit from the given class with a given tag. */
+	/** Gets all the components that inherit from the given class with a given tag. */
 	UFUNCTION(BlueprintCallable, Category = "Actor", meta = (ComponentClass = "ActorComponent"), meta = (DeterminesOutputType = "ComponentClass"))
 	TArray<UActorComponent*> GetComponentsByTag(TSubclassOf<UActorComponent> ComponentClass, FName Tag) const;
 
@@ -2845,33 +2738,34 @@ public:
 		return OwnedComponents;
 	}
 
-	/** Puts a component in to the OwnedComponents array of the Actor.
-	 *  The Component must be owned by the Actor or else it will assert
-	 *  In general this should not need to be called directly by anything other than UActorComponent functions
+	/**
+	 * Puts a component in to the OwnedComponents array of the Actor.
+	 * The Component must be owned by the Actor or else it will assert
+	 * In general this should not need to be called directly by anything other than UActorComponent functions
 	 */
 	void AddOwnedComponent(UActorComponent* Component);
 
-	/** Removes a component from the OwnedComponents array of the Actor.
-	 *  In general this should not need to be called directly by anything other than UActorComponent functions
+	/**
+	 * Removes a component from the OwnedComponents array of the Actor.
+	 * In general this should not need to be called directly by anything other than UActorComponent functions
 	 */
 	void RemoveOwnedComponent(UActorComponent* Component);
 
 #if DO_CHECK
-	// Utility function for validating that a component is correctly in its Owner's OwnedComponents array
+	/** Utility function for validating that a component is correctly in its Owner's OwnedComponents array */
 	bool OwnsComponent(UActorComponent* Component) const;
 #endif
 
-	/** Force the Actor to clear and rebuild its OwnedComponents array by evaluating all children (recursively) and locating components
-	 *  In general this should not need to be called directly, but can sometimes be necessary as part of undo/redo code paths.
+	/**
+	 * Force the Actor to clear and rebuild its OwnedComponents array by evaluating all children (recursively) and locating components
+	 * In general this should not need to be called directly, but can sometimes be necessary as part of undo/redo code paths.
 	 */
 	void ResetOwnedComponents();
 
-	/** Called when the replicated state of a component changes to update the Actor's cached ReplicatedComponents array
-	 */
+	/** Called when the replicated state of a component changes to update the Actor's cached ReplicatedComponents array */
 	void UpdateReplicatedComponent(UActorComponent* Component);
 
-	/** Completely synchronizes the replicated components array so that it contains exactly the number of replicated components currently owned
-	 */
+	/** Completely synchronizes the replicated components array so that it contains exactly the number of replicated components currently owned */
 	void UpdateAllReplicatedComponents();
 
 	/** Returns whether replication is enabled or not. */
@@ -2880,40 +2774,35 @@ public:
 		return bReplicates;
 	}
 
-	/** Returns a constant reference to the replicated components set
-	 */
-	const TSet<UActorComponent*>& GetReplicatedComponents() const 
+	/** Returns a constant reference to the replicated components set */
+	const TArray<UActorComponent*>& GetReplicatedComponents() const
 	{ 
 		return ReplicatedComponents; 
 	}
 
 private:
 	/**
-	 * All ActorComponents owned by this Actor.
+	 * All ActorComponents owned by this Actor. Stored as a Set as actors may have a large number of components
 	 * @see GetComponents()
 	 */
 	TSet<UActorComponent*> OwnedComponents;
 
-	/** Set of replicated components. */
-	TSet<UActorComponent*> ReplicatedComponents;
-
 #if WITH_EDITOR
 	/** Maps natively-constructed components to properties that reference them. */
 	TMultiMap<FName, UObjectProperty*> NativeConstructedComponentToPropertyMap;
 #endif
 
-public:
-
-	/** Array of ActorComponents that are created by blueprints and serialized per-instance. */
-	UPROPERTY(TextExportTransient, NonTransactional)
-	TArray<UActorComponent*> BlueprintCreatedComponents;
+	/** Set of replicated components, stored as an array to save space as this is generally not very large */
+	TArray<UActorComponent*> ReplicatedComponents;
 
-private:
 	/** Array of ActorComponents that have been added by the user on a per-instance basis. */
 	UPROPERTY(Instanced)
 	TArray<UActorComponent*> InstanceComponents;
 
 public:
+	/** Array of ActorComponents that are created by blueprints and serialized per-instance. */
+	UPROPERTY(TextExportTransient, NonTransactional)
+	TArray<UActorComponent*> BlueprintCreatedComponents;
 
 	/** Adds a component to the instance components array */
 	void AddInstanceComponent(UActorComponent* Component);
@@ -2927,19 +2816,43 @@ public:
 	/** Returns the instance components array */
 	const TArray<UActorComponent*>& GetInstanceComponents() const;
 
-public:
+
 	//~=============================================================================
-	// Navigation related functions
-	// 
+	// Navigation/AI related functions
+	
+	/**
+	 * Trigger a noise caused by a given Pawn, at a given location.
+	 * Note that the NoiseInstigator Pawn MUST have a PawnNoiseEmitterComponent for the noise to be detected by a PawnSensingComponent.
+	 * Senders of MakeNoise should have an Instigator if they are not pawns, or pass a NoiseInstigator.
+	 *
+	 * @param Loudness The relative loudness of this noise. Usual range is 0 (no noise) to 1 (full volume). If MaxRange is used, this scales the max range, otherwise it affects the hearing range specified by the sensor.
+	 * @param NoiseInstigator Pawn responsible for this noise.  Uses the actor's Instigator if NoiseInstigator is null
+	 * @param NoiseLocation Position of noise source.  If zero vector, use the actor's location.
+	 * @param MaxRange Max range at which the sound may be heard. A value of 0 indicates no max range (though perception may have its own range). Loudness scales the range. (Note: not supported for legacy PawnSensingComponent, only for AIPerception)
+	 * @param Tag Identifier for the noise.
+	 */
+	UFUNCTION(BlueprintCallable, BlueprintAuthorityOnly, Category="AI", meta=(BlueprintProtected = "true"))
+	void MakeNoise(float Loudness=1.f, APawn* NoiseInstigator=nullptr, FVector NoiseLocation=FVector::ZeroVector, float MaxRange = 0.f, FName Tag = NAME_None);
+
+	/** Default Implementation of MakeNoise */
+	static void MakeNoiseImpl(AActor* NoiseMaker, float Loudness, APawn* NoiseInstigator, const FVector& NoiseLocation, float MaxRange, FName Tag);
+
+	/** Modifies the global delegate used for handling MakeNoise */
+	static void SetMakeNoiseDelegate(const FMakeNoiseDelegate& NewDelegate);
 
-	/** Check if owned component should be relevant for navigation
-	 *  Allows implementing master switch to disable e.g. collision export in projectiles
+	/**
+	 * Check if owned component should be relevant for navigation
+	 * Allows implementing master switch to disable e.g. collision export in projectiles
 	 */
 	virtual bool IsComponentRelevantForNavigation(UActorComponent* Component) const { return true; }
 
+private:
+	static FMakeNoiseDelegate MakeNoiseDelegate;
+
+public:
 	//~=============================================================================
 	// Debugging functions
-public:
+
 	/**
 	 * Draw important Actor variables on canvas.  HUD will call DisplayDebug() on the current ViewTarget when the ShowDebug exec is used
 	 *
@@ -2955,29 +2868,22 @@ public:
 	 */
 	virtual void DisplayDebug(class UCanvas* Canvas, const class FDebugDisplayInfo& DebugDisplay, float& YL, float& YPos);
 
-	/** Retrieves actor's name used for logging, or string "NULL" if Actor == NULL */
+	/** Retrieves actor's name used for logging, or string "NULL" if Actor is null */
 	static FString GetDebugName(const AActor* Actor) { return Actor ? Actor->GetName() : TEXT("NULL"); }
 
-	//* Sets the friendly actor label and name */
-	void SetActorLabelInternal( const FString& NewActorLabelDirty, bool bMakeGloballyUniqueFName, bool bMarkDirty );
-
-	static FMakeNoiseDelegate MakeNoiseDelegate;
-
-public:
 #if !UE_BUILD_SHIPPING
 	/** Delegate for globally hooking ProccessEvent calls - used by a non-public testing plugin */
 	static FOnProcessEvent ProcessEventDelegate;
 #endif
 
-	static void MakeNoiseImpl(AActor* NoiseMaker, float Loudness, APawn* NoiseInstigator, const FVector& NoiseLocation, float MaxRange, FName Tag);
-	static void SetMakeNoiseDelegate(const FMakeNoiseDelegate& NewDelegate);
-
 	/** A fence to track when the primitive is detached from the scene in the rendering thread. */
 	FRenderCommandFence DetachFence;
 
 private:
+	/** Sets the friendly actor label and name */
+	void SetActorLabelInternal(const FString& NewActorLabelDirty, bool bMakeGloballyUniqueFName, bool bMarkDirty);
 
-	// Helper that already assumes the Hit info is reversed, and avoids creating a temp FHitResult if possible.
+	/** Helper that already assumes the Hit info is reversed, and avoids creating a temp FHitResult if possible. */
 	void InternalDispatchBlockingHit(UPrimitiveComponent* MyComp, UPrimitiveComponent* OtherComp, bool bSelfMoved, FHitResult const& Hit);
 
 	/** Private version without inlining that does *not* check Dedicated server build flags (which should already have been done). */
@@ -2989,8 +2895,6 @@ private:
 	friend struct FMarkActorIsBeingDestroyed;
 	friend struct FActorParentComponentSetter;
 
-private:
-
 	// Static helpers for accessing functions on SceneComponent.
 	// These are templates for no other reason than to delay compilation until USceneComponent is defined.
 
@@ -3049,8 +2953,8 @@ private:
 	}
 
 	/*~
-	* Returns the Up (Z) vector (length 1.0) of the RootComponent, in world space.
-	*/ 
+	 * Returns the Up (Z) vector (length 1.0) of the RootComponent, in world space.
+	 */ 
 	template<class T>
 	static FORCEINLINE FVector TemplateGetActorUpVector(const T* RootComponent)
 	{
@@ -3058,8 +2962,8 @@ private:
 	}
 
 	/*~
-	* Returns the Right (Y) vector (length 1.0) of the RootComponent, in world space.	
-	*/ 
+	 * Returns the Right (Y) vector (length 1.0) of the RootComponent, in world space.	
+	 */ 
 	template<class T>
 	static FORCEINLINE FVector TemplateGetActorRightVector(const T* RootComponent)
 	{
@@ -3100,7 +3004,6 @@ public:
 };
 
 /** Helper function for executing tick functions based on the normal conditions previous found in UActorComponent::ConditionalTick */
-
 template <typename ExecuteTickLambda>
 void FActorComponentTickFunction::ExecuteTickHelper(UActorComponent* Target, bool bTickInEditor, float DeltaTime, ELevelTick TickType, const ExecuteTickLambda& ExecuteTickFunc)
 {
@@ -3269,4 +3172,3 @@ FORCEINLINE_DEBUGGABLE bool AActor::IsNetMode(ENetMode Mode) const
 
 
 
-
