Copying //UE4/Dev-Framework to //UE4/Dev-Main (Source: //UE4/Dev-Framework @ 3716594)

diff --git a/Engine/Source/Runtime/Engine/Classes/GameFramework/Character.h b/Engine/Source/Runtime/Engine/Classes/GameFramework/Character.h
--- a/Engine/Source/Runtime/Engine/Classes/GameFramework/Character.h
+++ b/Engine/Source/Runtime/Engine/Classes/GameFramework/Character.h
@@ -32,13 +32,6 @@ DECLARE_DYNAMIC_MULTICAST_DELEGATE_ThreeParams(FCharacterMovementUpdatedSignatur
 DECLARE_DYNAMIC_MULTICAST_DELEGATE(FCharacterReachedApexSignature);
 DECLARE_DYNAMIC_MULTICAST_DELEGATE_OneParam(FLandedSignature, const FHitResult&, Hit);
 
-//
-// Forward declarations
-//
-class UAnimMontage;
-class UPrimitiveComponent;
-struct FAnimMontageInstance;
-
 /** Replicated data when playing a root motion montage. */
 USTRUCT()
 struct FRepRootMotionMontage
@@ -97,14 +90,14 @@ struct FRepRootMotionMontage
 	void Clear()
 	{
 		bIsActive = false;
-		AnimMontage = NULL;
+		AnimMontage = nullptr;
 		AuthoritativeRootMotion.Clear();
 	}
 
 	/** Is Valid - animation root motion only */
 	bool HasRootMotion() const
 	{
-		return (AnimMontage != NULL);
+		return (AnimMontage != nullptr);
 	}
 };
 
@@ -147,7 +140,7 @@ namespace MovementBaseUtility
 	/** Get the tangential velocity at WorldLocation for the given component. */
 	ENGINE_API FVector GetMovementBaseTangentialVelocity(const UPrimitiveComponent* MovementBase, const FName BoneName, const FVector& WorldLocation);
 
-	/** Get the transforms for the given MovementBase, optionally at the location of a bone. Returns false if MovementBase is NULL, or if BoneName is not a valid bone. */
+	/** Get the transforms for the given MovementBase, optionally at the location of a bone. Returns false if MovementBase is nullptr, or if BoneName is not a valid bone. */
 	ENGINE_API bool GetMovementBaseTransform(const UPrimitiveComponent* MovementBase, const FName BoneName, FVector& OutLocation, FQuat& OutQuat);
 }
 
@@ -200,7 +193,7 @@ struct FBasedMovementInfo
 	/** Return true if the client should have MovementBase, but it hasn't replicated (possibly component has not streamed in). */
 	FORCEINLINE bool IsBaseUnresolved() const
 	{
-		return (MovementBase == NULL) && bServerHasBaseComponent;
+		return (MovementBase == nullptr) && bServerHasBaseComponent;
 	}
 };
 
@@ -213,48 +206,135 @@ struct FBasedMovementInfo
  * @see APawn, UCharacterMovementComponent
  * @see https://docs.unrealengine.com/latest/INT/Gameplay/Framework/Pawn/Character/
  */ 
-
 UCLASS(config=Game, BlueprintType, meta=(ShortTooltip="A character is a type of Pawn that includes the ability to walk around."))
 class ENGINE_API ACharacter : public APawn
 {
 	GENERATED_BODY()
 public:
-	/**
-	 * Default UObject constructor.
-	 */
+	/** Default UObject constructor. */
 	ACharacter(const FObjectInitializer& ObjectInitializer = FObjectInitializer::Get());
 
 	void GetLifetimeReplicatedProps(TArray<FLifetimeProperty>& OutLifetimeProps) const override;
 
 private:
 	/** The main skeletal mesh associated with this Character (optional sub-object). */
-	UPROPERTY(Category = Character, VisibleAnywhere, BlueprintReadOnly, meta = (AllowPrivateAccess = "true"))
-	class USkeletalMeshComponent* Mesh;
+	UPROPERTY(Category=Character, VisibleAnywhere, BlueprintReadOnly, meta=(AllowPrivateAccess = "true"))
+	USkeletalMeshComponent* Mesh;
+
+	/** Movement component used for movement logic in various movement modes (walking, falling, etc), containing relevant settings and functions to control movement. */
+	UPROPERTY(Category=Character, VisibleAnywhere, BlueprintReadOnly, meta=(AllowPrivateAccess = "true"))
+	UCharacterMovementComponent* CharacterMovement;
+
+	/** The CapsuleComponent being used for movement collision (by CharacterMovement). Always treated as being vertically aligned in simple collision check functions. */
+	UPROPERTY(Category=Character, VisibleAnywhere, BlueprintReadOnly, meta=(AllowPrivateAccess = "true"))
+	UCapsuleComponent* CapsuleComponent;
 
 #if WITH_EDITORONLY_DATA
+	/** Component shown in the editor only to indicate character facing */
 	UPROPERTY()
-	class UArrowComponent* ArrowComponent;
+	UArrowComponent* ArrowComponent;
 #endif
 
-	/** Movement component used for movement logic in various movement modes (walking, falling, etc), containing relevant settings and functions to control movement. */
-	UPROPERTY(Category = Character, VisibleAnywhere, BlueprintReadOnly, meta = (AllowPrivateAccess = "true"))
-	class UCharacterMovementComponent* CharacterMovement;
+public:
 
-	/** The CapsuleComponent being used for movement collision (by CharacterMovement). Always treated as being vertically aligned in simple collision check functions. */
-	UPROPERTY(Category = Character, VisibleAnywhere, BlueprintReadOnly, meta = (AllowPrivateAccess = "true"))
-	class UCapsuleComponent* CapsuleComponent;
+	//////////////////////////////////////////////////////////////////////////
+	// Server RPCs that pass through to CharacterMovement (avoids RPC overhead for components).
+	// The base RPC function (eg 'ServerMove') is auto-generated for clients to trigger the call to the server function,
+	// eventually going to the _Implementation function (which we just pass to the CharacterMovementComponent).
+	//////////////////////////////////////////////////////////////////////////
+
+	/** Replicated function sent by client to server - contains client movement and view info. */
+	UFUNCTION(unreliable, server, WithValidation)
+	void ServerMove(float TimeStamp, FVector_NetQuantize10 InAccel, FVector_NetQuantize100 ClientLoc, uint8 CompressedMoveFlags, uint8 ClientRoll, uint32 View, UPrimitiveComponent* ClientMovementBase, FName ClientBaseBoneName, uint8 ClientMovementMode);
+	void ServerMove_Implementation(float TimeStamp, FVector_NetQuantize10 InAccel, FVector_NetQuantize100 ClientLoc, uint8 CompressedMoveFlags, uint8 ClientRoll, uint32 View, UPrimitiveComponent* ClientMovementBase, FName ClientBaseBoneName, uint8 ClientMovementMode);
+	bool ServerMove_Validate(float TimeStamp, FVector_NetQuantize10 InAccel, FVector_NetQuantize100 ClientLoc, uint8 CompressedMoveFlags, uint8 ClientRoll, uint32 View, UPrimitiveComponent* ClientMovementBase, FName ClientBaseBoneName, uint8 ClientMovementMode);
+
+	/**
+	 * Replicated function sent by client to server. Saves bandwidth over ServerMove() by implying that ClientMovementBase and ClientBaseBoneName are null.
+	 * Passes through to CharacterMovement->ServerMove_Implementation() with null base params.
+	 */
+	UFUNCTION(unreliable, server, WithValidation)
+	void ServerMoveNoBase(float TimeStamp, FVector_NetQuantize10 InAccel, FVector_NetQuantize100 ClientLoc, uint8 CompressedMoveFlags, uint8 ClientRoll, uint32 View, uint8 ClientMovementMode);
+	void ServerMoveNoBase_Implementation(float TimeStamp, FVector_NetQuantize10 InAccel, FVector_NetQuantize100 ClientLoc, uint8 CompressedMoveFlags, uint8 ClientRoll, uint32 View, uint8 ClientMovementMode);
+	bool ServerMoveNoBase_Validate(float TimeStamp, FVector_NetQuantize10 InAccel, FVector_NetQuantize100 ClientLoc, uint8 CompressedMoveFlags, uint8 ClientRoll, uint32 View, uint8 ClientMovementMode);
+
+	/** Replicated function sent by client to server - contains client movement and view info for two moves. */
+	UFUNCTION(unreliable, server, WithValidation)
+	void ServerMoveDual(float TimeStamp0, FVector_NetQuantize10 InAccel0, uint8 PendingFlags, uint32 View0, float TimeStamp, FVector_NetQuantize10 InAccel, FVector_NetQuantize100 ClientLoc, uint8 NewFlags, uint8 ClientRoll, uint32 View, UPrimitiveComponent* ClientMovementBase, FName ClientBaseBoneName, uint8 ClientMovementMode);
+	void ServerMoveDual_Implementation(float TimeStamp0, FVector_NetQuantize10 InAccel0, uint8 PendingFlags, uint32 View0, float TimeStamp, FVector_NetQuantize10 InAccel, FVector_NetQuantize100 ClientLoc, uint8 NewFlags, uint8 ClientRoll, uint32 View, UPrimitiveComponent* ClientMovementBase, FName ClientBaseBoneName, uint8 ClientMovementMode);
+	bool ServerMoveDual_Validate(float TimeStamp0, FVector_NetQuantize10 InAccel0, uint8 PendingFlags, uint32 View0, float TimeStamp, FVector_NetQuantize10 InAccel, FVector_NetQuantize100 ClientLoc, uint8 NewFlags, uint8 ClientRoll, uint32 View, UPrimitiveComponent* ClientMovementBase, FName ClientBaseBoneName, uint8 ClientMovementMode);
+
+	/** Replicated function sent by client to server - contains client movement and view info for two moves. */
+	UFUNCTION(unreliable, server, WithValidation)
+	void ServerMoveDualNoBase(float TimeStamp0, FVector_NetQuantize10 InAccel0, uint8 PendingFlags, uint32 View0, float TimeStamp, FVector_NetQuantize10 InAccel, FVector_NetQuantize100 ClientLoc, uint8 NewFlags, uint8 ClientRoll, uint32 View, uint8 ClientMovementMode);
+	void ServerMoveDualNoBase_Implementation(float TimeStamp0, FVector_NetQuantize10 InAccel0, uint8 PendingFlags, uint32 View0, float TimeStamp, FVector_NetQuantize10 InAccel, FVector_NetQuantize100 ClientLoc, uint8 NewFlags, uint8 ClientRoll, uint32 View, uint8 ClientMovementMode);
+	bool ServerMoveDualNoBase_Validate(float TimeStamp0, FVector_NetQuantize10 InAccel0, uint8 PendingFlags, uint32 View0, float TimeStamp, FVector_NetQuantize10 InAccel, FVector_NetQuantize100 ClientLoc, uint8 NewFlags, uint8 ClientRoll, uint32 View, uint8 ClientMovementMode);
+
+	/** Replicated function sent by client to server - contains client movement and view info for two moves. First move is non root motion, second is root motion. */
+	UFUNCTION(unreliable, server, WithValidation)
+	void ServerMoveDualHybridRootMotion(float TimeStamp0, FVector_NetQuantize10 InAccel0, uint8 PendingFlags, uint32 View0, float TimeStamp, FVector_NetQuantize10 InAccel, FVector_NetQuantize100 ClientLoc, uint8 NewFlags, uint8 ClientRoll, uint32 View, UPrimitiveComponent* ClientMovementBase, FName ClientBaseBoneName, uint8 ClientMovementMode);
+	void ServerMoveDualHybridRootMotion_Implementation(float TimeStamp0, FVector_NetQuantize10 InAccel0, uint8 PendingFlags, uint32 View0, float TimeStamp, FVector_NetQuantize10 InAccel, FVector_NetQuantize100 ClientLoc, uint8 NewFlags, uint8 ClientRoll, uint32 View, UPrimitiveComponent* ClientMovementBase, FName ClientBaseBoneName, uint8 ClientMovementMode);
+	bool ServerMoveDualHybridRootMotion_Validate(float TimeStamp0, FVector_NetQuantize10 InAccel0, uint8 PendingFlags, uint32 View0, float TimeStamp, FVector_NetQuantize10 InAccel, FVector_NetQuantize100 ClientLoc, uint8 NewFlags, uint8 ClientRoll, uint32 View, UPrimitiveComponent* ClientMovementBase, FName ClientBaseBoneName, uint8 ClientMovementMode);
+
+	/* Resending an (important) old move. Process it if not already processed. */
+	UFUNCTION(unreliable, server, WithValidation)
+	void ServerMoveOld(float OldTimeStamp, FVector_NetQuantize10 OldAccel, uint8 OldMoveFlags);
+	void ServerMoveOld_Implementation(float OldTimeStamp, FVector_NetQuantize10 OldAccel, uint8 OldMoveFlags);
+	bool ServerMoveOld_Validate(float OldTimeStamp, FVector_NetQuantize10 OldAccel, uint8 OldMoveFlags);
+
+	//////////////////////////////////////////////////////////////////////////
+	// Client RPCS that pass through to CharacterMovement (avoids RPC overhead for components).
+	//////////////////////////////////////////////////////////////////////////
+
+	/** If no client adjustment is needed after processing received ServerMove(), ack the good move so client can remove it from SavedMoves */
+	UFUNCTION(unreliable, client)
+	void ClientAckGoodMove(float TimeStamp);
+	void ClientAckGoodMove_Implementation(float TimeStamp);
+
+	/** Replicate position correction to client, associated with a timestamped servermove.  Client will replay subsequent moves after applying adjustment.  */
+	UFUNCTION(unreliable, client)
+	void ClientAdjustPosition(float TimeStamp, FVector NewLoc, FVector NewVel, UPrimitiveComponent* NewBase, FName NewBaseBoneName, bool bHasBase, bool bBaseRelativePosition, uint8 ServerMovementMode);
+	void ClientAdjustPosition_Implementation(float TimeStamp, FVector NewLoc, FVector NewVel, UPrimitiveComponent* NewBase, FName NewBaseBoneName, bool bHasBase, bool bBaseRelativePosition, uint8 ServerMovementMode);
+
+	/* Bandwidth saving version, when velocity is zeroed */
+	UFUNCTION(unreliable, client)
+	void ClientVeryShortAdjustPosition(float TimeStamp, FVector NewLoc, UPrimitiveComponent* NewBase, FName NewBaseBoneName, bool bHasBase, bool bBaseRelativePosition, uint8 ServerMovementMode);
+	void ClientVeryShortAdjustPosition_Implementation(float TimeStamp, FVector NewLoc, UPrimitiveComponent* NewBase, FName NewBaseBoneName, bool bHasBase, bool bBaseRelativePosition, uint8 ServerMovementMode);
+
+	/** Replicate position correction to client when using root motion for movement. (animation root motion specific) */
+	UFUNCTION(unreliable, client)
+	void ClientAdjustRootMotionPosition(float TimeStamp, float ServerMontageTrackPosition, FVector ServerLoc, FVector_NetQuantizeNormal ServerRotation, float ServerVelZ, UPrimitiveComponent* ServerBase, FName ServerBoneName, bool bHasBase, bool bBaseRelativePosition, uint8 ServerMovementMode);
+	void ClientAdjustRootMotionPosition_Implementation(float TimeStamp, float ServerMontageTrackPosition, FVector ServerLoc, FVector_NetQuantizeNormal ServerRotation, float ServerVelZ, UPrimitiveComponent* ServerBase, FName ServerBoneName, bool bHasBase, bool bBaseRelativePosition, uint8 ServerMovementMode);
+
+	/** Replicate root motion source correction to client when using root motion for movement. */
+	UFUNCTION(unreliable, client)
+	void ClientAdjustRootMotionSourcePosition(float TimeStamp, FRootMotionSourceGroup ServerRootMotion, bool bHasAnimRootMotion, float ServerMontageTrackPosition, FVector ServerLoc, FVector_NetQuantizeNormal ServerRotation, float ServerVelZ, UPrimitiveComponent* ServerBase, FName ServerBoneName, bool bHasBase, bool bBaseRelativePosition, uint8 ServerMovementMode);
+	void ClientAdjustRootMotionSourcePosition_Implementation(float TimeStamp, FRootMotionSourceGroup ServerRootMotion, bool bHasAnimRootMotion, float ServerMontageTrackPosition, FVector ServerLoc, FVector_NetQuantizeNormal ServerRotation, float ServerVelZ, UPrimitiveComponent* ServerBase, FName ServerBoneName, bool bHasBase, bool bBaseRelativePosition, uint8 ServerMovementMode);
 
 public:
+	/** Returns Mesh subobject **/
+	class USkeletalMeshComponent* GetMesh() const { return Mesh; }
 
 	/** Name of the MeshComponent. Use this name if you want to prevent creation of the component (with ObjectInitializer.DoNotCreateDefaultSubobject). */
 	static FName MeshComponentName;
 
+	/** Returns CharacterMovement subobject **/
+	class UCharacterMovementComponent* GetCharacterMovement() const { return CharacterMovement; }
+
 	/** Name of the CharacterMovement component. Use this name if you want to use a different class (with ObjectInitializer.SetDefaultSubobjectClass). */
 	static FName CharacterMovementComponentName;
 
+	/** Returns CapsuleComponent subobject **/
+	class UCapsuleComponent* GetCapsuleComponent() const { return CapsuleComponent; }
+
 	/** Name of the CapsuleComponent. */
 	static FName CapsuleComponentName;
 
+#if WITH_EDITORONLY_DATA
+	/** Returns ArrowComponent subobject **/
+	class UArrowComponent* GetArrowComponent() const { return ArrowComponent; }
+#endif
+
 	/** Sets the component the Character is walking on, used by CharacterMovement walking movement to be able to follow dynamic objects. */
 	virtual void SetBase(UPrimitiveComponent* NewBase, const FName BoneName = NAME_None, bool bNotifyActor=true);
 	
@@ -263,11 +343,10 @@ public:
 	 * This is automatically called during initialization; call this at runtime if you intend to change the default mesh offset from the capsule.
 	 * @see GetBaseTranslationOffset(), GetBaseRotationOffset()
 	 */
-	UFUNCTION(BlueprintCallable, Category="Pawn|Character")
+	UFUNCTION(BlueprintCallable, Category=Character)
 	virtual void CacheInitialMeshOffset(FVector MeshRelativeLocation, FRotator MeshRelativeRotation);
 
 protected:
-
 	/** Info about our current movement base (object we are standing on). */
 	UPROPERTY()
 	struct FBasedMovementInfo BasedMovement;
@@ -286,7 +365,7 @@ public:
 	virtual void OnRep_ReplicatedBasedMovement();
 
 	/** Set whether this actor's movement replicates to network clients. */
-	UFUNCTION(BlueprintCallable, Category = "Replication")
+	UFUNCTION(BlueprintCallable, Category=Replication)
 	virtual void SetReplicateMovement(bool bInReplicateMovement) override;
 
 protected:
@@ -313,8 +392,7 @@ protected:
 	UPROPERTY()
 	bool bInBaseReplication;
 
-public:	
-
+public:
 	/** Accessor for ReplicatedServerLastTransformUpdateTimeStamp. */
 	FORCEINLINE float GetReplicatedServerLastTransformUpdateTimeStamp() const { return ReplicatedServerLastTransformUpdateTimeStamp; }
 
@@ -331,20 +409,16 @@ public:
 	uint8 GetReplicatedMovementMode() const { return ReplicatedMovementMode; }
 
 	/** Get the saved translation offset of mesh. This is how much extra offset is applied from the center of the capsule. */
-	UFUNCTION(BlueprintCallable, Category="Pawn|Character")
+	UFUNCTION(BlueprintCallable, Category=Character)
 	FVector GetBaseTranslationOffset() const { return BaseTranslationOffset; }
 
 	/** Get the saved rotation offset of mesh. This is how much extra rotation is applied from the capsule rotation. */
 	virtual FQuat GetBaseRotationOffset() const { return BaseRotationOffset; }
 
 	/** Get the saved rotation offset of mesh. This is how much extra rotation is applied from the capsule rotation. */
-	UFUNCTION(BlueprintCallable, Category="Pawn|Character", meta=(DisplayName="GetBaseRotationOffset"))
+	UFUNCTION(BlueprintCallable, Category=Character, meta=(DisplayName="GetBaseRotationOffset"))
 	FRotator GetBaseRotationOffsetRotator() const { return GetBaseRotationOffset().Rotator(); }
 
-	//~ Begin INavAgentInterface Interface
-	virtual FVector GetNavAgentLocation() const override;
-	//~ End INavAgentInterface Interface
-
 	/** Default crouched eye height */
 	UPROPERTY(EditAnywhere, BlueprintReadWrite, Category=Camera)
 	float CrouchedEyeHeight;
@@ -358,7 +432,7 @@ public:
 	virtual void OnRep_IsCrouched();
 
 	/** When true, player wants to jump */
-	UPROPERTY(BlueprintReadOnly, Category="Pawn|Character")
+	UPROPERTY(BlueprintReadOnly, Category=Character)
 	uint32 bPressedJump:1;
 
 	/** When true, applying updates to network client (replaying saved moves for a locally controlled character) */
@@ -388,6 +462,10 @@ public:
 	UPROPERTY(Transient)
 	uint32 bServerMoveIgnoreRootMotion:1;
 
+	/** Tracks whether or not the character was already jumping last frame. */
+	UPROPERTY(VisibleInstanceOnly, BlueprintReadOnly, Transient, Category=Character)
+	uint32 bWasJumping : 1;
+
 	/** 
 	 * Jump key Held Time.
 	 * This is the time that the player has held the jump key, in seconds.
@@ -419,16 +497,9 @@ public:
      * When providing overrides for these methods, it's recommended to either manually
      * increment / reset this value, or call the Super:: method.
      */
-    UPROPERTY(VisibleInstanceOnly, BlueprintReadOnly, Category="Character")
+    UPROPERTY(VisibleInstanceOnly, BlueprintReadOnly, Category=Character)
     int32 JumpCurrentCount;
 
-	DEPRECATED(4.14, "This value is no longer used.")
-	uint32 bJumpMaxCountExceeded:1;
-
-	// Tracks whether or not the character was already jumping last frame.
-	UPROPERTY(VisibleInstanceOnly, BlueprintReadOnly, Transient, Category="Character")
-	uint32 bWasJumping:1;
-
 	//~ Begin AActor Interface.
 	virtual void BeginPlay() override;
 	virtual void ClearCrossLevelReferences() override;
@@ -447,6 +518,10 @@ public:
 		return AActor::FindComponentByClass<T>();
 	}
 
+	//~ Begin INavAgentInterface Interface
+	virtual FVector GetNavAgentLocation() const override;
+	//~ End INavAgentInterface Interface
+
 	//~ Begin APawn Interface.
 	virtual void PostInitializeComponents() override;
 	virtual UPawnMovementComponent* GetMovementComponent() const override;
@@ -474,7 +549,7 @@ public:
 	 * as on a button up event), otherwise the character will carry on receiving the 
 	 * velocity until JumpKeyHoldTime is reached.
 	 */
-	UFUNCTION(BlueprintCallable, Category="Pawn|Character")
+	UFUNCTION(BlueprintCallable, Category=Character)
 	virtual void Jump();
 
 	/** 
@@ -483,7 +558,7 @@ public:
 	 * jump Z-velocity. If this is not called, then jump z-velocity will be applied
 	 * until JumpMaxHoldTime is reached.
 	 */
-	UFUNCTION(BlueprintCallable, Category="Pawn|Character")
+	UFUNCTION(BlueprintCallable, Category=Character)
 	virtual void StopJumping();
 
 	/**
@@ -493,11 +568,10 @@ public:
 	 * 
 	 * @Return Whether the character can jump in the current state. 
 	 */
-	UFUNCTION(BlueprintCallable, Category="Pawn|Character")
+	UFUNCTION(BlueprintCallable, Category=Character)
 	bool CanJump() const;
 
 protected:
-
 	/**
 	 * Customizable event to check if the character can jump in the current state.
 	 * Default implementation returns true if the character is on the ground and not crouching,
@@ -509,49 +583,41 @@ protected:
 	 *
 	 * @Return Whether the character can jump in the current state. 
 	 */
-
-	UFUNCTION(BlueprintNativeEvent, Category="Pawn|Character", meta=(DisplayName="CanJump"))
+	UFUNCTION(BlueprintNativeEvent, Category=Character, meta=(DisplayName="CanJump"))
 	bool CanJumpInternal() const;
 	virtual bool CanJumpInternal_Implementation() const;
 
-	DEPRECATED(4.14, "This function is deprecated. Please use ResetJumpState instead.")
-	void CheckResetJumpCount()
-	{
-		ResetJumpState();
-	}
-
+	/** Marks character as not trying to jump */
 	void ResetJumpState();
 
 public:
-
 	/**
 	 * True if jump is actively providing a force, such as when the jump key is held and the time it has been held is less than JumpMaxHoldTime.
 	 * @see CharacterMovement->IsFalling
 	 */
-	UFUNCTION(BlueprintCallable, Category="Pawn|Character")
+	UFUNCTION(BlueprintCallable, Category=Character)
 	virtual bool IsJumpProvidingForce() const;
 
 	/** Play Animation Montage on the character mesh **/
 	UFUNCTION(BlueprintCallable, Category=Animation)
 	virtual float PlayAnimMontage(class UAnimMontage* AnimMontage, float InPlayRate = 1.f, FName StartSectionName = NAME_None);
 
-	/** Stop Animation Montage. If NULL, it will stop what's currently active. The Blend Out Time is taken from the montage asset that is being stopped. **/
+	/** Stop Animation Montage. If nullptr, it will stop what's currently active. The Blend Out Time is taken from the montage asset that is being stopped. **/
 	UFUNCTION(BlueprintCallable, Category=Animation)
-	virtual void StopAnimMontage(class UAnimMontage* AnimMontage = NULL);
+	virtual void StopAnimMontage(class UAnimMontage* AnimMontage = nullptr);
 
 	/** Return current playing Montage **/
 	UFUNCTION(BlueprintCallable, Category=Animation)
 	class UAnimMontage* GetCurrentMontage();
 
-public:
-
-	/** Set a pending launch velocity on the Character. This velocity will be processed on the next CharacterMovementComponent tick,
-	  * and will set it to the "falling" state. Triggers the OnLaunched event.
-	  * @PARAM LaunchVelocity is the velocity to impart to the Character
-	  * @PARAM bXYOverride if true replace the XY part of the Character's velocity instead of adding to it.
-	  * @PARAM bZOverride if true replace the Z component of the Character's velocity instead of adding to it.
-	  */
-	UFUNCTION(BlueprintCallable, Category="Pawn|Character")
+	/**
+	 * Set a pending launch velocity on the Character. This velocity will be processed on the next CharacterMovementComponent tick,
+	 * and will set it to the "falling" state. Triggers the OnLaunched event.
+	 * @PARAM LaunchVelocity is the velocity to impart to the Character
+	 * @PARAM bXYOverride if true replace the XY part of the Character's velocity instead of adding to it.
+	 * @PARAM bZOverride if true replace the Z component of the Character's velocity instead of adding to it.
+	 */
+	UFUNCTION(BlueprintCallable, Category=Character)
 	virtual void LaunchCharacter(FVector LaunchVelocity, bool bXYOverride, bool bZOverride);
 
 	/** Let blueprint know that we were launched */
@@ -559,7 +625,7 @@ public:
 	void OnLaunched(FVector LaunchVelocity, bool bXYOverride, bool bZOverride);
 
 	/** Event fired when the character has just started jumping */
-	UFUNCTION(BlueprintNativeEvent, Category="Pawn|Character")
+	UFUNCTION(BlueprintNativeEvent, Category=Character)
 	void OnJumped();
 	virtual void OnJumped_Implementation();
 
@@ -570,7 +636,7 @@ public:
 	virtual void NotifyJumpApex();
 
 	/** Broadcast when Character's jump reaches its apex. Needs CharacterMovement->bNotifyApex = true */
-	UPROPERTY(BlueprintAssignable, Category="Pawn|Character")
+	UPROPERTY(BlueprintAssignable, Category=Character)
 	FCharacterReachedApexSignature OnReachedJumpApex;
 
 	/**
@@ -584,23 +650,23 @@ public:
 	virtual void Landed(const FHitResult& Hit);
 
 	/**
-	* Called upon landing when falling, to perform actions based on the Hit result.
-	* Note that movement mode is still "Falling" during this event. Current Velocity value is the velocity at the time of landing.
-	* Consider OnMovementModeChanged() as well, as that can be used once the movement mode changes to the new mode (most likely Walking).
-	*
-	* @param Hit Result describing the landing that resulted in a valid landing spot.
-	* @see OnMovementModeChanged()
-	*/
+	 * Called upon landing when falling, to perform actions based on the Hit result.
+	 * Note that movement mode is still "Falling" during this event. Current Velocity value is the velocity at the time of landing.
+	 * Consider OnMovementModeChanged() as well, as that can be used once the movement mode changes to the new mode (most likely Walking).
+	 *
+	 * @param Hit Result describing the landing that resulted in a valid landing spot.
+	 * @see OnMovementModeChanged()
+	 */
 	FLandedSignature LandedDelegate;
 
 	/**
-	* Called upon landing when falling, to perform actions based on the Hit result.
-	* Note that movement mode is still "Falling" during this event. Current Velocity value is the velocity at the time of landing.
-	* Consider OnMovementModeChanged() as well, as that can be used once the movement mode changes to the new mode (most likely Walking).
-	*
-	* @param Hit Result describing the landing that resulted in a valid landing spot.
-	* @see OnMovementModeChanged()
-	*/
+	 * Called upon landing when falling, to perform actions based on the Hit result.
+	 * Note that movement mode is still "Falling" during this event. Current Velocity value is the velocity at the time of landing.
+	 * Consider OnMovementModeChanged() as well, as that can be used once the movement mode changes to the new mode (most likely Walking).
+	 *
+	 * @param Hit Result describing the landing that resulted in a valid landing spot.
+	 * @see OnMovementModeChanged()
+	 */
 	UFUNCTION(BlueprintImplementableEvent)
 	void OnLanded(const FHitResult& Hit);
 
@@ -614,12 +680,14 @@ public:
 	 * @param  PreviousLocation	Previous character location before movement off the ledge.
 	 * @param  TimeTick	Time delta of movement update resulting in moving off the ledge.
 	 */
-	UFUNCTION(BlueprintNativeEvent, Category="Pawn|Character")
+	UFUNCTION(BlueprintNativeEvent, Category=Character)
 	void OnWalkingOffLedge(const FVector& PreviousFloorImpactNormal, const FVector& PreviousFloorContactNormal, const FVector& PreviousLocation, float TimeDelta);
 	virtual void OnWalkingOffLedge_Implementation(const FVector& PreviousFloorImpactNormal, const FVector& PreviousFloorContactNormal, const FVector& PreviousLocation, float TimeDelta);
 
-	/** Called when pawn's movement is blocked
-		@PARAM Impact describes the blocking hit. */
+	/**
+	 * Called when pawn's movement is blocked
+	 * @param Impact describes the blocking hit.
+	 */
 	virtual void MoveBlockedBy(const FHitResult& Impact) {};
 
 	/**
@@ -628,7 +696,7 @@ public:
 	 * @see IsCrouched
 	 * @see CharacterMovement->WantsToCrouch
 	 */
-	UFUNCTION(BlueprintCallable, Category="Pawn|Character", meta=(HidePin="bClientSimulation"))
+	UFUNCTION(BlueprintCallable, Category=Character, meta=(HidePin="bClientSimulation"))
 	virtual void Crouch(bool bClientSimulation = false);
 
 	/**
@@ -637,7 +705,7 @@ public:
 	 * @see IsCrouched
 	 * @see CharacterMovement->WantsToCrouch
 	 */
-	UFUNCTION(BlueprintCallable, Category="Pawn|Character", meta=(HidePin="bClientSimulation"))
+	UFUNCTION(BlueprintCallable, Category=Character, meta=(HidePin="bClientSimulation"))
 	virtual void UnCrouch(bool bClientSimulation = false);
 
 	/** @return true if this character is currently able to crouch (and is not currently crouched) */
@@ -655,7 +723,7 @@ public:
 	 * @param	HalfHeightAdjust		difference between default collision half-height, and actual crouched capsule half-height.
 	 * @param	ScaledHalfHeightAdjust	difference after component scale is taken in to account.
 	 */
-	UFUNCTION(BlueprintImplementableEvent, meta=(DisplayName = "OnEndCrouch"))
+	UFUNCTION(BlueprintImplementableEvent, meta=(DisplayName="OnEndCrouch"))
 	void K2_OnEndCrouch(float HalfHeightAdjust, float ScaledHalfHeightAdjust);
 
 	/**
@@ -670,7 +738,7 @@ public:
 	 * @param	HalfHeightAdjust		difference between default collision half-height, and actual crouched capsule half-height.
 	 * @param	ScaledHalfHeightAdjust	difference after component scale is taken in to account.
 	 */
-	UFUNCTION(BlueprintImplementableEvent, meta=(DisplayName = "OnStartCrouch"))
+	UFUNCTION(BlueprintImplementableEvent, meta=(DisplayName="OnStartCrouch"))
 	void K2_OnStartCrouch(float HalfHeightAdjust, float ScaledHalfHeightAdjust);
 
 	/**
@@ -681,7 +749,7 @@ public:
 	virtual void OnMovementModeChanged(EMovementMode PrevMovementMode, uint8 PreviousCustomMode = 0);
 
 	/** Multicast delegate for MovementMode changing. */
-	UPROPERTY(BlueprintAssignable, Category="Pawn|Character")
+	UPROPERTY(BlueprintAssignable, Category=Character)
 	FMovementModeChangedSignature MovementModeChangedDelegate;
 
 	/**
@@ -691,7 +759,7 @@ public:
 	 * @param	PrevCustomMode		Custom mode before the change (applicable if PrevMovementMode is Custom)
 	 * @param	NewCustomMode		New custom mode (applicable if NewMovementMode is Custom)
 	 */
-	UFUNCTION(BlueprintImplementableEvent, meta=(DisplayName = "OnMovementModeChanged"))
+	UFUNCTION(BlueprintImplementableEvent, meta=(DisplayName="OnMovementModeChanged"))
 	void K2_OnMovementModeChanged(EMovementMode PrevMovementMode, EMovementMode NewMovementMode, uint8 PrevCustomMode, uint8 NewCustomMode);
 
 	/**
@@ -699,7 +767,7 @@ public:
 	 * @note C++ code should override UCharacterMovementComponent::PhysCustom() instead.
 	 * @see UCharacterMovementComponent::PhysCustom()
 	 */
-	UFUNCTION(BlueprintImplementableEvent, meta=(DisplayName = "UpdateCustomMovement"))
+	UFUNCTION(BlueprintImplementableEvent, meta=(DisplayName="UpdateCustomMovement"))
 	void K2_UpdateCustomMovement(float DeltaTime);
 
 	/**
@@ -712,12 +780,10 @@ public:
 	 * @param	InitialLocation		Location at the start of the update. May be different than the current location if movement occurred.
 	 * @param	InitialVelocity		Velocity at the start of the update. May be different than the current velocity.
 	 */
-	UPROPERTY(BlueprintAssignable, Category="Pawn|Character")
+	UPROPERTY(BlueprintAssignable, Category=Character)
 	FCharacterMovementUpdatedSignature OnCharacterMovementUpdated;
 
-	/** 
-	 * Returns true if the Landed() event should be called. Used by CharacterMovement to prevent notifications while playing back network moves.
-	 */
+	/** Returns true if the Landed() event should be called. Used by CharacterMovement to prevent notifications while playing back network moves. */
 	virtual bool ShouldNotifyLanded(const struct FHitResult& Hit);
 
 	/** Trigger jump if jump button has been pressed. */
@@ -736,33 +802,31 @@ public:
 	 */
 	virtual float GetJumpMaxHoldTime() const;
 
-public:
-
-	UFUNCTION(reliable, client)
+	UFUNCTION(Reliable, Client)
 	void ClientCheatWalk();
 	virtual void ClientCheatWalk_Implementation();
 
-	UFUNCTION(reliable, client)
+	UFUNCTION(Reliable, Client)
 	void ClientCheatFly();
 	virtual void ClientCheatFly_Implementation();
 
-	UFUNCTION(reliable, client)
+	UFUNCTION(Reliable, Client)
 	void ClientCheatGhost();
 	virtual void ClientCheatGhost_Implementation();
 
-	UFUNCTION(reliable, client)
+	UFUNCTION(Reliable, Client)
 	void RootMotionDebugClientPrintOnScreen(const FString& InString);
 	virtual void RootMotionDebugClientPrintOnScreen_Implementation(const FString& InString);
 
 	// Root Motion
-public:
+
 	/** 
-	 *  For LocallyControlled Autonomous clients. 
-	 *  During a PerformMovement() after root motion is prepared, we save it off into this and
-	 *  then record it into our SavedMoves.
-	 *  During SavedMove playback we use it as our "Previous Move" SavedRootMotion which includes
-	 *  last received root motion from the Server
-	 **/
+	 * For LocallyControlled Autonomous clients. 
+	 * During a PerformMovement() after root motion is prepared, we save it off into this and
+	 * then record it into our SavedMoves.
+	 * During SavedMove playback we use it as our "Previous Move" SavedRootMotion which includes
+	 * last received root motion from the Server
+	 */
 	UPROPERTY(Transient)
 	FRootMotionSourceGroup SavedRootMotion;
 
@@ -776,18 +840,16 @@ public:
 	
 	/** Find usable root motion replicated move from our buffer.
 	 * Goes through the buffer back in time, to find the first move that clears 'CanUseRootMotionRepMove' below.
-	 * Returns index of that move or INDEX_NONE otherwise. */
+	 * Returns index of that move or INDEX_NONE otherwise.
+	 */
 	int32 FindRootMotionRepMove(const FAnimMontageInstance& ClientMontageInstance) const;
 
-	/** true if buffered move is usable to teleport client back to. */
+	/** True if buffered move is usable to teleport client back to. */
 	bool CanUseRootMotionRepMove(const FSimulatedRootMotionReplicatedMove& RootMotionRepMove, const FAnimMontageInstance& ClientMontageInstance) const;
 
 	/** Restore actor to an old buffered move. */
 	bool RestoreReplicatedMove(const FSimulatedRootMotionReplicatedMove& RootMotionRepMove);
 	
-	DEPRECATED(4.11, "UpdateSimulatedPosition() is deprecated and is not used by engine code. Use OnUpdateSimulatedPosition() instead.")
-	virtual void UpdateSimulatedPosition(const FVector& Location, const FRotator& NewRotation);
-
 	/**
 	 * Called on client after position update is received to respond to the new location and rotation.
 	 * Actual change in location is expected to occur in CharacterMovement->SmoothCorrection(), after which this occurs.
@@ -809,129 +871,33 @@ public:
 	/** Get FAnimMontageInstance playing RootMotion */
 	FAnimMontageInstance * GetRootMotionAnimMontageInstance() const;
 
-	/** true if we are playing Root Motion right now */
+	/** True if we are playing Root Motion right now */
 	UFUNCTION(BlueprintCallable, Category=Animation)
 	bool IsPlayingRootMotion() const;
 
-	/** true if we are playing Root Motion right now, through a Montage with RootMotionMode == ERootMotionMode::RootMotionFromMontagesOnly.
-	 * This means code path for networked root motion is enabled. */
-	UFUNCTION(BlueprintCallable, Category = Animation)
+	/**
+	 * True if we are playing Root Motion right now, through a Montage with RootMotionMode == ERootMotionMode::RootMotionFromMontagesOnly.
+	 * This means code path for networked root motion is enabled.
+	 */
+	UFUNCTION(BlueprintCallable, Category=Animation)
 	bool IsPlayingNetworkedRootMotionMontage() const;
 
 	/** Sets scale to apply to root motion translation on this Character */
 	void SetAnimRootMotionTranslationScale(float InAnimRootMotionTranslationScale = 1.f);
 
 	/** Returns current value of AnimRootMotionScale */
-	UFUNCTION(BlueprintCallable, Category = Animation)
+	UFUNCTION(BlueprintCallable, Category=Animation)
 	float GetAnimRootMotionTranslationScale() const;
 
-	/** Called on the actor right before replication occurs.
-	* Only called on Server, and for autonomous proxies if recording a Client Replay. */
+	/**
+	 * Called on the actor right before replication occurs.
+	 * Only called on Server, and for autonomous proxies if recording a Client Replay.
+	 */
 	virtual void PreReplication(IRepChangedPropertyTracker & ChangedPropertyTracker) override;
 
-	/** Called on the actor right before replication occurs.
-	* Called for everyone when recording a Client Replay, including Simulated Proxies. */
-	virtual void PreReplicationForReplay(IRepChangedPropertyTracker & ChangedPropertyTracker) override;
-public:
-
-	//////////////////////////////////////////////////////////////////////////
-	// Server RPCs that pass through to CharacterMovement (avoids RPC overhead for components).
-	// The base RPC function (eg 'ServerMove') is auto-generated for clients to trigger the call to the server function,
-	// eventually going to the _Implementation function (which we just pass to the CharacterMovementComponent).
-	//////////////////////////////////////////////////////////////////////////
-
-	/** Replicated function sent by client to server - contains client movement and view info. */
-	UFUNCTION(unreliable, server, WithValidation)
-	void ServerMove(float TimeStamp, FVector_NetQuantize10 InAccel, FVector_NetQuantize100 ClientLoc, uint8 CompressedMoveFlags, uint8 ClientRoll, uint32 View, UPrimitiveComponent* ClientMovementBase, FName ClientBaseBoneName, uint8 ClientMovementMode);
-	void ServerMove_Implementation(float TimeStamp, FVector_NetQuantize10 InAccel, FVector_NetQuantize100 ClientLoc, uint8 CompressedMoveFlags, uint8 ClientRoll, uint32 View, UPrimitiveComponent* ClientMovementBase, FName ClientBaseBoneName, uint8 ClientMovementMode);
-	bool ServerMove_Validate(float TimeStamp, FVector_NetQuantize10 InAccel, FVector_NetQuantize100 ClientLoc, uint8 CompressedMoveFlags, uint8 ClientRoll, uint32 View, UPrimitiveComponent* ClientMovementBase, FName ClientBaseBoneName, uint8 ClientMovementMode);
-
 	/**
-	 * Replicated function sent by client to server. Saves bandwidth over ServerMove() by implying that ClientMovementBase and ClientBaseBoneName are null.
-	 * Passes through to CharacterMovement->ServerMove_Implementation() with null base params.
+	 * Called on the actor right before replication occurs.
+	 * Called for everyone when recording a Client Replay, including Simulated Proxies.
 	 */
-	UFUNCTION(unreliable, server, WithValidation)
-	void ServerMoveNoBase(float TimeStamp, FVector_NetQuantize10 InAccel, FVector_NetQuantize100 ClientLoc, uint8 CompressedMoveFlags, uint8 ClientRoll, uint32 View, uint8 ClientMovementMode);
-	void ServerMoveNoBase_Implementation(float TimeStamp, FVector_NetQuantize10 InAccel, FVector_NetQuantize100 ClientLoc, uint8 CompressedMoveFlags, uint8 ClientRoll, uint32 View, uint8 ClientMovementMode);
-	bool ServerMoveNoBase_Validate(float TimeStamp, FVector_NetQuantize10 InAccel, FVector_NetQuantize100 ClientLoc, uint8 CompressedMoveFlags, uint8 ClientRoll, uint32 View, uint8 ClientMovementMode);
-
-	/** Replicated function sent by client to server - contains client movement and view info for two moves. */
-	UFUNCTION(unreliable, server, WithValidation)
-	void ServerMoveDual(float TimeStamp0, FVector_NetQuantize10 InAccel0, uint8 PendingFlags, uint32 View0, float TimeStamp, FVector_NetQuantize10 InAccel, FVector_NetQuantize100 ClientLoc, uint8 NewFlags, uint8 ClientRoll, uint32 View, UPrimitiveComponent* ClientMovementBase, FName ClientBaseBoneName, uint8 ClientMovementMode);
-	void ServerMoveDual_Implementation(float TimeStamp0, FVector_NetQuantize10 InAccel0, uint8 PendingFlags, uint32 View0, float TimeStamp, FVector_NetQuantize10 InAccel, FVector_NetQuantize100 ClientLoc, uint8 NewFlags, uint8 ClientRoll, uint32 View, UPrimitiveComponent* ClientMovementBase, FName ClientBaseBoneName, uint8 ClientMovementMode);
-	bool ServerMoveDual_Validate(float TimeStamp0, FVector_NetQuantize10 InAccel0, uint8 PendingFlags, uint32 View0, float TimeStamp, FVector_NetQuantize10 InAccel, FVector_NetQuantize100 ClientLoc, uint8 NewFlags, uint8 ClientRoll, uint32 View, UPrimitiveComponent* ClientMovementBase, FName ClientBaseBoneName, uint8 ClientMovementMode);
-
-	/** Replicated function sent by client to server - contains client movement and view info for two moves. */
-	UFUNCTION(unreliable, server, WithValidation)
-	void ServerMoveDualNoBase(float TimeStamp0, FVector_NetQuantize10 InAccel0, uint8 PendingFlags, uint32 View0, float TimeStamp, FVector_NetQuantize10 InAccel, FVector_NetQuantize100 ClientLoc, uint8 NewFlags, uint8 ClientRoll, uint32 View, uint8 ClientMovementMode);
-	void ServerMoveDualNoBase_Implementation(float TimeStamp0, FVector_NetQuantize10 InAccel0, uint8 PendingFlags, uint32 View0, float TimeStamp, FVector_NetQuantize10 InAccel, FVector_NetQuantize100 ClientLoc, uint8 NewFlags, uint8 ClientRoll, uint32 View, uint8 ClientMovementMode);
-	bool ServerMoveDualNoBase_Validate(float TimeStamp0, FVector_NetQuantize10 InAccel0, uint8 PendingFlags, uint32 View0, float TimeStamp, FVector_NetQuantize10 InAccel, FVector_NetQuantize100 ClientLoc, uint8 NewFlags, uint8 ClientRoll, uint32 View, uint8 ClientMovementMode);
-
-	/** Replicated function sent by client to server - contains client movement and view info for two moves. First move is non root motion, second is root motion. */
-	UFUNCTION(unreliable, server, WithValidation)
-	void ServerMoveDualHybridRootMotion(float TimeStamp0, FVector_NetQuantize10 InAccel0, uint8 PendingFlags, uint32 View0, float TimeStamp, FVector_NetQuantize10 InAccel, FVector_NetQuantize100 ClientLoc, uint8 NewFlags, uint8 ClientRoll, uint32 View, UPrimitiveComponent* ClientMovementBase, FName ClientBaseBoneName, uint8 ClientMovementMode);
-	void ServerMoveDualHybridRootMotion_Implementation(float TimeStamp0, FVector_NetQuantize10 InAccel0, uint8 PendingFlags, uint32 View0, float TimeStamp, FVector_NetQuantize10 InAccel, FVector_NetQuantize100 ClientLoc, uint8 NewFlags, uint8 ClientRoll, uint32 View, UPrimitiveComponent* ClientMovementBase, FName ClientBaseBoneName, uint8 ClientMovementMode);
-	bool ServerMoveDualHybridRootMotion_Validate(float TimeStamp0, FVector_NetQuantize10 InAccel0, uint8 PendingFlags, uint32 View0, float TimeStamp, FVector_NetQuantize10 InAccel, FVector_NetQuantize100 ClientLoc, uint8 NewFlags, uint8 ClientRoll, uint32 View, UPrimitiveComponent* ClientMovementBase, FName ClientBaseBoneName, uint8 ClientMovementMode);
-
-	/* Resending an (important) old move. Process it if not already processed. */
-	UFUNCTION(unreliable, server, WithValidation)
-	void ServerMoveOld(float OldTimeStamp, FVector_NetQuantize10 OldAccel, uint8 OldMoveFlags);
-	void ServerMoveOld_Implementation(float OldTimeStamp, FVector_NetQuantize10 OldAccel, uint8 OldMoveFlags);
-	bool ServerMoveOld_Validate(float OldTimeStamp, FVector_NetQuantize10 OldAccel, uint8 OldMoveFlags);
-
-	//////////////////////////////////////////////////////////////////////////
-	// Client RPCS that pass through to CharacterMovement (avoids RPC overhead for components).
-	//////////////////////////////////////////////////////////////////////////
-
-	/** If no client adjustment is needed after processing received ServerMove(), ack the good move so client can remove it from SavedMoves */
-	UFUNCTION(unreliable, client)
-	void ClientAckGoodMove(float TimeStamp);
-	void ClientAckGoodMove_Implementation(float TimeStamp);
-
-	/** Replicate position correction to client, associated with a timestamped servermove.  Client will replay subsequent moves after applying adjustment.  */
-	UFUNCTION(unreliable, client)
-	void ClientAdjustPosition(float TimeStamp, FVector NewLoc, FVector NewVel, UPrimitiveComponent* NewBase, FName NewBaseBoneName, bool bHasBase, bool bBaseRelativePosition, uint8 ServerMovementMode);
-	void ClientAdjustPosition_Implementation(float TimeStamp, FVector NewLoc, FVector NewVel, UPrimitiveComponent* NewBase, FName NewBaseBoneName, bool bHasBase, bool bBaseRelativePosition, uint8 ServerMovementMode);
-
-	/* Bandwidth saving version, when velocity is zeroed */
-	UFUNCTION(unreliable, client)
-	void ClientVeryShortAdjustPosition(float TimeStamp, FVector NewLoc, UPrimitiveComponent* NewBase, FName NewBaseBoneName, bool bHasBase, bool bBaseRelativePosition, uint8 ServerMovementMode);
-	void ClientVeryShortAdjustPosition_Implementation(float TimeStamp, FVector NewLoc, UPrimitiveComponent* NewBase, FName NewBaseBoneName, bool bHasBase, bool bBaseRelativePosition, uint8 ServerMovementMode);
-
-	/** Replicate position correction to client when using root motion for movement. (animation root motion specific) */
-	UFUNCTION(unreliable, client)
-	void ClientAdjustRootMotionPosition(float TimeStamp, float ServerMontageTrackPosition, FVector ServerLoc, FVector_NetQuantizeNormal ServerRotation, float ServerVelZ, UPrimitiveComponent* ServerBase, FName ServerBoneName, bool bHasBase, bool bBaseRelativePosition, uint8 ServerMovementMode);
-	void ClientAdjustRootMotionPosition_Implementation(float TimeStamp, float ServerMontageTrackPosition, FVector ServerLoc, FVector_NetQuantizeNormal ServerRotation, float ServerVelZ, UPrimitiveComponent* ServerBase, FName ServerBoneName, bool bHasBase, bool bBaseRelativePosition, uint8 ServerMovementMode);
-
-	/** Replicate root motion source correction to client when using root motion for movement. */
-	UFUNCTION(unreliable, client)
-	void ClientAdjustRootMotionSourcePosition(float TimeStamp, FRootMotionSourceGroup ServerRootMotion, bool bHasAnimRootMotion, float ServerMontageTrackPosition, FVector ServerLoc, FVector_NetQuantizeNormal ServerRotation, float ServerVelZ, UPrimitiveComponent* ServerBase, FName ServerBoneName, bool bHasBase, bool bBaseRelativePosition, uint8 ServerMovementMode);
-	void ClientAdjustRootMotionSourcePosition_Implementation(float TimeStamp, FRootMotionSourceGroup ServerRootMotion, bool bHasAnimRootMotion, float ServerMontageTrackPosition, FVector ServerLoc, FVector_NetQuantizeNormal ServerRotation, float ServerVelZ, UPrimitiveComponent* ServerBase, FName ServerBoneName, bool bHasBase, bool bBaseRelativePosition, uint8 ServerMovementMode);
-
-public:
-	/** Returns Mesh subobject **/
-	class USkeletalMeshComponent* GetMesh() const;
-#if WITH_EDITORONLY_DATA
-	/** Returns ArrowComponent subobject **/
-	class UArrowComponent* GetArrowComponent() const;
-#endif
-	/** Returns CharacterMovement subobject **/
-	class UCharacterMovementComponent* GetCharacterMovement() const;
-	/** Returns CapsuleComponent subobject **/
-	class UCapsuleComponent* GetCapsuleComponent() const;
+	virtual void PreReplicationForReplay(IRepChangedPropertyTracker & ChangedPropertyTracker) override;
 };
-
-
-//////////////////////////////////////////////////////////////////////////
-// Character inlines
-
-/** Returns Mesh subobject **/
-FORCEINLINE USkeletalMeshComponent* ACharacter::GetMesh() const { return Mesh; }
-#if WITH_EDITORONLY_DATA
-/** Returns ArrowComponent subobject **/
-FORCEINLINE UArrowComponent* ACharacter::GetArrowComponent() const { return ArrowComponent; }
-#endif
-/** Returns CharacterMovement subobject **/
-FORCEINLINE UCharacterMovementComponent* ACharacter::GetCharacterMovement() const { return CharacterMovement; }
-/** Returns CapsuleComponent subobject **/
-FORCEINLINE UCapsuleComponent* ACharacter::GetCapsuleComponent() const { return CapsuleComponent; }
diff --git a/Engine/Source/Runtime/Engine/Classes/GameFramework/CharacterMovementComponent.h b/Engine/Source/Runtime/Engine/Classes/GameFramework/CharacterMovementComponent.h
--- a/Engine/Source/Runtime/Engine/Classes/GameFramework/CharacterMovementComponent.h
+++ b/Engine/Source/Runtime/Engine/Classes/GameFramework/CharacterMovementComponent.h
@@ -1317,7 +1317,7 @@ public:
 	 *	Compute the max jump height based on the JumpZVelocity velocity and gravity.
 	 *	This does take into account the CharacterOwner's MaxJumpHoldTime.
 	 */
-	UFUNCTION(BlueprintCallable, Category="Pawn|Components|CharacterMovemet")
+	UFUNCTION(BlueprintCallable, Category="Pawn|Components|CharacterMovement")
 	virtual float GetMaxJumpHeightWithJumpTime() const;
 
 	/** @return Maximum acceleration for the current state. */
diff --git a/Engine/Source/Runtime/Engine/Classes/GameFramework/Controller.h b/Engine/Source/Runtime/Engine/Classes/GameFramework/Controller.h
--- a/Engine/Source/Runtime/Engine/Classes/GameFramework/Controller.h
+++ b/Engine/Source/Runtime/Engine/Classes/GameFramework/Controller.h
@@ -11,14 +11,14 @@
 
 class ACharacter;
 class APawn;
+class APlayerState;
 class FDebugDisplayInfo;
 class UDamageType;
 class UPathFollowingComponent;
 
 UDELEGATE(BlueprintAuthorityOnly)
-DECLARE_DYNAMIC_MULTICAST_DELEGATE_FourParams( FInstigatedAnyDamageSignature, float, Damage, const class UDamageType*, DamageType, class AActor*, DamagedActor, class AActor*, DamageCauser );
+DECLARE_DYNAMIC_MULTICAST_DELEGATE_FourParams( FInstigatedAnyDamageSignature, float, Damage, const UDamageType*, DamageType, AActor*, DamagedActor, AActor*, DamageCauser );
 
-//~=============================================================================
 /**
  * Controllers are non-physical actors that can possess a Pawn to control
  * its actions.  PlayerControllers are used by human players to control pawns, while
@@ -36,11 +36,29 @@ DECLARE_DYNAMIC_MULTICAST_DELEGATE_FourParams( FInstigatedAnyDamageSignature, fl
  * 
  * @see https://docs.unrealengine.com/latest/INT/Gameplay/Framework/Controller/
  */
-
 UCLASS(abstract, notplaceable, NotBlueprintable, HideCategories=(Collision,Rendering,"Utilities|Transformation")) 
 class ENGINE_API AController : public AActor, public INavAgentInterface
 {
-	GENERATED_UCLASS_BODY()
+	GENERATED_BODY()
+
+public:
+	/** Default Constructor */
+	AController(const FObjectInitializer& ObjectInitializer = FObjectInitializer::Get());
+
+	/** PlayerState containing replicated information about the player using this controller (only exists for players, not NPCs). */
+	UPROPERTY(replicatedUsing = OnRep_PlayerState, BlueprintReadOnly, Category=Controller)
+	APlayerState* PlayerState;
+
+	/** Actor marking where this controller spawned in. */
+	TWeakObjectPtr<class AActor> StartSpot;
+
+	/** Called when the controller has instigated damage in any way */
+	UPROPERTY(BlueprintAssignable)
+	FInstigatedAnyDamageSignature OnInstigatedAnyDamage;
+
+	/** Current gameplay state this controller is in */
+	UPROPERTY()
+	FName StateName;
 
 private:
 	/** Pawn currently being controlled by this controller.  Use Pawn.Possess() to take control of a pawn */
@@ -53,39 +71,15 @@ private:
 	 */
 	TWeakObjectPtr< APawn > OldPawn;
 
-	/** Character currently being controlled by this controller.  Value is same as Pawn if the controlled pawn is a character, otherwise NULL */
+	/** Character currently being controlled by this controller.  Value is same as Pawn if the controlled pawn is a character, otherwise nullptr */
 	UPROPERTY()
-	class ACharacter* Character;
+	ACharacter* Character;
 
-public:
-	/** PlayerState containing replicated information about the player using this controller (only exists for players, not NPCs). */
-	UPROPERTY(replicatedUsing=OnRep_PlayerState, BlueprintReadOnly, Category="Controller")
-	class APlayerState* PlayerState;
-
-private:
 	/** Component to give controllers a transform and enable attachment if desired. */
 	UPROPERTY()
-	class USceneComponent* TransformComponent;
-
-public:
-
-	/**
-	  * Get the control rotation. This is the full aim rotation, which may be different than a camera orientation (for example in a third person view),
-	  * and may differ from the rotation of the controlled Pawn (which may choose not to visually pitch or roll, for example).
-	  */
-	UFUNCTION(BlueprintCallable, Category="Pawn")
-	virtual FRotator GetControlRotation() const;
-
-	/** Set the control rotation. The RootComponent's rotation will also be updated to match it if RootComponent->bAbsoluteRotation is true. */
-	UFUNCTION(BlueprintCallable, Category="Pawn", meta=(Tooltip="Set the control rotation."))
-	virtual void SetControlRotation(const FRotator& NewRotation);
-
-	/** Set the initial location and rotation of the controller, as well as the control rotation. Typically used when the controller is first created. */
-	UFUNCTION(BlueprintCallable, Category="Pawn")
-	virtual void SetInitialLocationAndRotation(const FVector& NewLocation, const FRotator& NewRotation);
+	USceneComponent* TransformComponent;
 
 protected:
-
 	/** The control rotation of the Controller. See GetControlRotation. */
 	UPROPERTY()
 	FRotator ControlRotation;
@@ -110,10 +104,10 @@ protected:
 	uint8 IgnoreLookInput;
 
 	/**
-	  * Physically attach the Controller to the specified Pawn, so that our position reflects theirs.
-	  * The attachment persists during possession of the pawn. The Controller's rotation continues to match the ControlRotation.
-	  * Attempting to attach to a NULL Pawn will call DetachFromPawn() instead.
-	  */
+	 * Physically attach the Controller to the specified Pawn, so that our position reflects theirs.
+	 * The attachment persists during possession of the pawn. The Controller's rotation continues to match the ControlRotation.
+	 * Attempting to attach to a nullptr Pawn will call DetachFromPawn() instead.
+	 */
 	virtual void AttachToPawn(APawn* InPawn);
 
 	/** Detach the RootComponent from its parent, but only if bAttachToPawn is true and it was attached to a Pawn.	 */
@@ -125,16 +119,10 @@ protected:
 	/** Remove dependency that makes us tick before the given Pawn.	 */
 	virtual void RemovePawnTickDependency(APawn* InOldPawn);
 
-public:
-
-	/** Actor marking where this controller spawned in. */
-	TWeakObjectPtr<class AActor> StartSpot;
-
-	//~=============================================================================
-	// CONTROLLER STATE PROPERTIES
+	/** Returns TransformComponent subobject **/
+	class USceneComponent* GetTransformComponent() const { return TransformComponent; }
 
-	UPROPERTY()
-	FName StateName;
+public:
 
 	/** Change the current state to named state */
 	virtual void ChangeState(FName NewState);
@@ -148,7 +136,23 @@ public:
 	
 	/** @return the name of the current state */
 	FName GetStateName() const;
-	
+
+	/**
+	 * Get the control rotation. This is the full aim rotation, which may be different than a camera orientation (for example in a third person view),
+	 * and may differ from the rotation of the controlled Pawn (which may choose not to visually pitch or roll, for example).
+	 */
+	UFUNCTION(BlueprintCallable, Category=Pawn)
+	virtual FRotator GetControlRotation() const;
+
+	/** Set the control rotation. The RootComponent's rotation will also be updated to match it if RootComponent->bAbsoluteRotation is true. */
+	UFUNCTION(BlueprintCallable, Category=Pawn, meta=(Tooltip="Set the control rotation."))
+	virtual void SetControlRotation(const FRotator& NewRotation);
+
+	/** Set the initial location and rotation of the controller, as well as the control rotation. Typically used when the controller is first created. */
+	UFUNCTION(BlueprintCallable, Category=Pawn)
+	virtual void SetInitialLocationAndRotation(const FVector& NewLocation, const FRotator& NewRotation);
+
+
 	/**
 	 * Checks line to center and top of other actor
 	 * @param Other is the actor whose visibility is being checked.
@@ -156,7 +160,7 @@ public:
 	 * @param bAlternateChecks used only in AIController implementation
 	 * @return true if controller's pawn can see Other actor.
 	 */
-	UFUNCTION(BlueprintCallable, Category="Controller")
+	UFUNCTION(BlueprintCallable, Category=Controller)
 	virtual bool LineOfSightTo(const class AActor* Other, FVector ViewPoint = FVector(ForceInit), bool bAlternateChecks = false) const;
 
 	/** Replication Notification Callbacks */
@@ -168,23 +172,21 @@ public:
 
 	/** DEPRECATED! Use the standard "Cast To" node instead. Casts this Controller to a Player Controller, if possible. */
 	DEPRECATED(4.11, "CastToPlayerController has been replaced by the standard Cast To node.")
-	UFUNCTION(BlueprintCallable, Category = "Pawn", meta = (DeprecatedFunction, DeprecationMessage = "Use standard Cast To node instead."))
+	UFUNCTION(BlueprintCallable, Category=Pawn, meta=(DeprecatedFunction, DeprecationMessage="Use standard Cast To node instead."))
 	class APlayerController* CastToPlayerController();
 
 	/** Replicated function to set the pawn location and rotation, allowing server to force (ex. teleports). */
-	UFUNCTION(reliable, client)
+	UFUNCTION(Reliable, Client)
 	void ClientSetLocation(FVector NewLocation, FRotator NewRotation);
 
 	/** Replicated function to set the pawn rotation, allowing the server to force. */
-	UFUNCTION(reliable, client)
+	UFUNCTION(Reliable, Client)
 	void ClientSetRotation(FRotator NewRotation, bool bResetCamera = false);
 
 	/** Return the Pawn that is currently 'controlled' by this PlayerController */
-	UFUNCTION(BlueprintCallable, Category="Pawn", meta=(DisplayName="Get Controlled Pawn"))
+	UFUNCTION(BlueprintCallable, Category=Pawn, meta=(DisplayName="Get Controlled Pawn"))
 	APawn* K2_GetPawn() const;
 
-public:
-
 	virtual void GetLifetimeReplicatedProps(TArray<FLifetimeProperty>& OutLifetimeProps) const override;
 
 	//~ Begin AActor Interface
@@ -194,7 +196,7 @@ public:
 	virtual void GetActorEyesViewPoint( FVector& out_Location, FRotator& out_Rotation ) const override;
 	virtual FString GetHumanReadableName() const override;
 
-	/* Overridden to create the player replication info and perform other mundane initialization tasks. */
+	/** Overridden to create the player replication info and perform other mundane initialization tasks. */
 	virtual void PostInitializeComponents() override;
 
 	virtual void Reset() override;
@@ -214,23 +216,23 @@ public:
 	void SetPawnFromRep(APawn* InPawn);
 
 	/** Get the actor the controller is looking at */
-	UFUNCTION(BlueprintCallable, Category="Pawn")
+	UFUNCTION(BlueprintCallable, Category=Pawn)
 	virtual AActor* GetViewTarget() const;
 
 	/** Get the desired pawn target rotation */
-	UFUNCTION(BlueprintCallable, Category="Pawn")
+	UFUNCTION(BlueprintCallable, Category=Pawn)
 	virtual FRotator GetDesiredRotation() const;
 
 	/** Returns whether this Controller is a PlayerController.  */
-	UFUNCTION(BlueprintCallable, Category="Pawn")
+	UFUNCTION(BlueprintCallable, Category=Pawn)
 	bool IsPlayerController() const;
 
 	/** Returns whether this Controller is a locally controlled PlayerController.  */
-	UFUNCTION(BlueprintCallable, Category="Pawn")
+	UFUNCTION(BlueprintCallable, Category=Pawn)
 	bool IsLocalPlayerController() const;
 
 	/** Returns whether this Controller is a local controller.	 */
-	UFUNCTION(BlueprintCallable, Category="Pawn")
+	UFUNCTION(BlueprintCallable, Category=Pawn)
 	virtual bool IsLocalController() const;
 
 	/** Called from Destroyed().  Cleans up PlayerState. */
@@ -242,13 +244,11 @@ public:
 	 * @param InPawn The Pawn to be possessed.
 	 * @see HasAuthority()
 	 */
-	UFUNCTION(BlueprintCallable, BlueprintAuthorityOnly, Category="Pawn", meta=(Keywords="set controller"))
+	UFUNCTION(BlueprintCallable, BlueprintAuthorityOnly, Category=Pawn, meta=(Keywords="set controller"))
 	virtual void Possess(APawn* InPawn);
 
-	/**
-	 * Called to unpossess our pawn for any reason that is not the pawn being destroyed (destruction handled by PawnDestroyed()).
-	 */
-	UFUNCTION(BlueprintCallable, Category="Pawn", meta=(Keywords="set controller"))
+	/** Called to unpossess our pawn for any reason that is not the pawn being destroyed (destruction handled by PawnDestroyed()). */
+	UFUNCTION(BlueprintCallable, Category=Pawn, meta=(Keywords="set controller"))
 	virtual void UnPossess();
 
 	/**
@@ -268,7 +268,7 @@ public:
 	 * @param EndGameFocus Actor to set as the view target on end game
 	 * @param bIsWinner true if this controller is on winning team
 	 */
-	virtual void GameHasEnded(class AActor* EndGameFocus = NULL, bool bIsWinner = false);
+	virtual void GameHasEnded(class AActor* EndGameFocus = nullptr, bool bIsWinner = false);
 
 	/**
 	 * Returns Player's Point of View
@@ -277,7 +277,7 @@ public:
 	 *
 	 * @output	out_Location, view location of player
 	 * @output	out_rotation, view rotation of player
-	*/
+	 */
 	virtual void GetPlayerViewPoint( FVector& Location, FRotator& Rotation ) const;
 
 	/** GameMode failed to spawn pawn for me. */
@@ -291,77 +291,67 @@ public:
 	virtual bool IsFollowingAPath() const override;
 	//~ End INavAgentInterface Interface
 
-	/** prepares path following component */
+	/** Prepares path following component */
 	virtual void InitNavigationControl(UPathFollowingComponent*& PathFollowingComp);
 
-	/** If controller has any navigation-related components then this function 
-	 *	makes them update their cached data */
+	/** If controller has any navigation-related components then this function makes them update their cached data */
 	virtual void UpdateNavigationComponents();
 
 	/** Aborts the move the controller is currently performing */
-	UFUNCTION(BlueprintCallable, Category = "AI|Navigation")
+	UFUNCTION(BlueprintCallable, Category="AI|Navigation")
 	virtual void StopMovement();
 
 	/**
-	* Locks or unlocks movement input, consecutive calls stack up and require the same amount of calls to undo, or can all be undone using ResetIgnoreMoveInput.
-	* @param bNewMoveInput	If true, move input is ignored. If false, input is not ignored.
-	*/
-	UFUNCTION(BlueprintCallable, Category="Input")
+	 * Locks or unlocks movement input, consecutive calls stack up and require the same amount of calls to undo, or can all be undone using ResetIgnoreMoveInput.
+	 * @param bNewMoveInput	If true, move input is ignored. If false, input is not ignored.
+	 */
+	UFUNCTION(BlueprintCallable, Category=Input)
 	virtual void SetIgnoreMoveInput(bool bNewMoveInput);
 
 	/** Stops ignoring move input by resetting the ignore move input state. */
-	UFUNCTION(BlueprintCallable, Category="Input", meta = (Keywords = "ClearIgnoreMoveInput"))
+	UFUNCTION(BlueprintCallable, Category=Input, meta=(Keywords = "ClearIgnoreMoveInput"))
 	virtual void ResetIgnoreMoveInput();
 
 	/** Returns true if movement input is ignored. */
-	UFUNCTION(BlueprintCallable, Category="Input")
+	UFUNCTION(BlueprintCallable, Category=Input)
 	virtual bool IsMoveInputIgnored() const;
 
 	/**
 	* Locks or unlocks look input, consecutive calls stack up and require the same amount of calls to undo, or can all be undone using ResetIgnoreLookInput.
 	* @param bNewLookInput	If true, look input is ignored. If false, input is not ignored.
 	*/
-	UFUNCTION(BlueprintCallable, Category="Input")
+	UFUNCTION(BlueprintCallable, Category=Input)
 	virtual void SetIgnoreLookInput(bool bNewLookInput);
 
 	/** Stops ignoring look input by resetting the ignore look input state. */
-	UFUNCTION(BlueprintCallable, Category="Input", meta=(Keywords="ClearIgnoreLookInput"))
+	UFUNCTION(BlueprintCallable, Category=Input, meta=(Keywords="ClearIgnoreLookInput"))
 	virtual void ResetIgnoreLookInput();
 
 	/** Returns true if look input is ignored. */
-	UFUNCTION(BlueprintCallable, Category="Input")
+	UFUNCTION(BlueprintCallable, Category=Input)
 	virtual bool IsLookInputIgnored() const;
 
 	/** Reset move and look input ignore flags. */
-	UFUNCTION(BlueprintCallable, Category="Input")
+	UFUNCTION(BlueprintCallable, Category=Input)
 	virtual void ResetIgnoreInputFlags();
 
+	/** Called when the level this controller is in is unloaded via streaming. */
+	virtual void CurrentLevelUnloaded();
+
 protected:
 	/** State entered when inactive (no possessed pawn, not spectating, etc). */
 	virtual void BeginInactiveState();
 
-	/** State entered when inactive (no possessed pawn, not spectating, etc). */
+	/** Called when leaving the inactive state */
 	virtual void EndInactiveState();
 
 	/** Event when this controller instigates ANY damage */
 	UFUNCTION(BlueprintImplementableEvent, BlueprintAuthorityOnly)
 	void ReceiveInstigatedAnyDamage(float Damage, const class UDamageType* DamageType, class AActor* DamagedActor, class AActor* DamageCauser);
 
-public:
-	/** Called when the level this controller is in is unloaded via streaming. */
-	virtual void CurrentLevelUnloaded();
-
 private:
 	// Hidden functions that don't make sense to use on this class.
 	HIDE_ACTOR_TRANSFORM_FUNCTIONS();
-
-	/** Called when the controller has instigated damage in any way */
-	UPROPERTY(BlueprintAssignable)
-	FInstigatedAnyDamageSignature OnInstigatedAnyDamage;
-
-protected:
-	/** Returns TransformComponent subobject **/
-	class USceneComponent* GetTransformComponent() const { return TransformComponent; }
 };
 
 
diff --git a/Engine/Source/Runtime/Engine/Classes/GameFramework/GameModeBase.h b/Engine/Source/Runtime/Engine/Classes/GameFramework/GameModeBase.h
--- a/Engine/Source/Runtime/Engine/Classes/GameFramework/GameModeBase.h
+++ b/Engine/Source/Runtime/Engine/Classes/GameFramework/GameModeBase.h
@@ -79,7 +79,7 @@ public:
 	virtual TSubclassOf<AGameSession> GetGameSessionClass() const;
 
 	/** Returns default pawn class for given controller */
-	UFUNCTION(BlueprintNativeEvent, Category=Classes)
+	UFUNCTION(BlueprintCallable, BlueprintNativeEvent, Category=Classes)
 	UClass* GetDefaultPawnClassForController(AController* InController);
 
 	/** Class of GameSession, which handles login approval and online game interface */
diff --git a/Engine/Source/Runtime/Engine/Classes/GameFramework/Pawn.h b/Engine/Source/Runtime/Engine/Classes/GameFramework/Pawn.h
--- a/Engine/Source/Runtime/Engine/Classes/GameFramework/Pawn.h
+++ b/Engine/Source/Runtime/Engine/Classes/GameFramework/Pawn.h
@@ -37,21 +37,20 @@ UCLASS(config=Game, BlueprintType, Blueprintable, hideCategories=(Navigation), m
 class ENGINE_API APawn : public AActor, public INavAgentInterface
 {
 	GENERATED_BODY()
+
 public:
-	/**
-	* Default UObject constructor.
-	*/
+	/** Default UObject constructor. */
 	APawn(const FObjectInitializer& ObjectInitializer = FObjectInitializer::Get());
 
 	virtual void GetLifetimeReplicatedProps(TArray<FLifetimeProperty>& OutLifetimeProps) const override;
 	virtual void PreReplication( IRepChangedPropertyTracker & ChangedPropertyTracker ) override;
 
 	/** Return our PawnMovementComponent, if we have one. By default, returns the first PawnMovementComponent found. Native classes that create their own movement component should override this method for more efficiency. */
-	UFUNCTION(BlueprintCallable, meta=(Tooltip="Return our PawnMovementComponent, if we have one."), Category="Pawn")
+	UFUNCTION(BlueprintCallable, meta=(Tooltip="Return our PawnMovementComponent, if we have one."), Category=Pawn)
 	virtual UPawnMovementComponent* GetMovementComponent() const;
 
 	/** Return PrimitiveComponent we are based on (standing on, attached to, and moving on). */
-	virtual UPrimitiveComponent* GetMovementBase() const { return NULL; }
+	virtual UPrimitiveComponent* GetMovementBase() const { return nullptr; }
 
 	/** If true, this Pawn's pitch will be updated to match the Controller's ControlRotation pitch, if controlled by a PlayerController. */
 	UPROPERTY(EditAnywhere, BlueprintReadWrite, Category=Pawn)
@@ -74,11 +73,13 @@ public:
 	uint32 bCanAffectNavigationGeneration : 1;
 
 private:
-	/* Whether this Pawn's input handling is enabled.  Pawn must still be possessed to get input even if this is true */
+	/** Whether this Pawn's input handling is enabled.  Pawn must still be possessed to get input even if this is true */
 	uint32 bInputEnabled:1;
 
-public:
+	/** Used to prevent re-entry of OutsideWorldBounds event. */
+	uint32 bProcessingOutsideWorldBounds : 1;
 
+public:
 	/** Base eye height above collision center. */
 	UPROPERTY(EditAnywhere, BlueprintReadWrite, Category=Camera)
 	float BaseEyeHeight;
@@ -98,14 +99,10 @@ public:
 	UPROPERTY(EditAnywhere, Category=Pawn)
 	EAutoPossessAI AutoPossessAI;
 
-	/**
-	 * Default class to use when pawn is controlled by AI.
-	 */
+	/** Default class to use when pawn is controlled by AI. */
 	UPROPERTY(EditAnywhere, BlueprintReadWrite, meta=(DisplayName="AI Controller Class"), Category=Pawn)
 	TSubclassOf<AController> AIControllerClass;
 
-public:
-
 	/**
 	 * Return our PawnNoiseEmitterComponent, if any. Default implementation returns the first PawnNoiseEmitterComponent found in the components array.
 	 * If one isn't found, then it tries to find one on the Pawn's current Controller.
@@ -120,12 +117,12 @@ public:
 	 * @param NoiseLocation - Position of noise source.  If zero vector, use the actor's location.
 	 * @param bUseNoiseMakerLocation - If true, use the location of the NoiseMaker rather than NoiseLocation.  If false, use NoiseLocation.
 	 * @param NoiseMaker - Which actor is the source of the noise.  Not to be confused with the Noise Instigator, which is responsible for the noise (and is the pawn on which this function is called).  If not specified, the pawn instigating the noise will be used as the NoiseMaker
-	*/
+	 */
 	UFUNCTION(BlueprintCallable, BlueprintAuthorityOnly, Category=AI)
-	void PawnMakeNoise(float Loudness, FVector NoiseLocation, bool bUseNoiseMakerLocation = true, AActor* NoiseMaker = NULL);
+	void PawnMakeNoise(float Loudness, FVector NoiseLocation, bool bUseNoiseMakerLocation = true, AActor* NoiseMaker = nullptr);
 
 	/** If Pawn is possessed by a player, points to his playerstate.  Needed for network play as controllers are not replicated to clients. */
-	UPROPERTY(replicatedUsing=OnRep_PlayerState, BlueprintReadOnly, Category="Pawn")
+	UPROPERTY(replicatedUsing=OnRep_PlayerState, BlueprintReadOnly, Category=Pawn)
 	APlayerState* PlayerState;
 
 	/** Replicated so we can see where remote clients are looking. */
@@ -159,33 +156,32 @@ public:
 	 */
 	void SetRemoteViewPitch(float NewRemoteViewPitch);
 
-	/** Called when our Controller no longer possesses us.	 */
+	/** Called when our Controller no longer possesses us. */
 	virtual void UnPossessed();
 
-	/** Return Physics Volume for this Pawn **/
+	/** Return Physics Volume for this Pawn */
 	virtual APhysicsVolume* GetPawnPhysicsVolume() const;
 
 	/** Gets the owning actor of the Movement Base Component on which the pawn is standing. */
-	UFUNCTION(BlueprintPure, Category="Pawn")
+	UFUNCTION(BlueprintPure, Category=Pawn)
 	static AActor* GetMovementBaseActor(const APawn* Pawn);
 
-	virtual bool IsBasedOnActor(const AActor* Other) const override;
-
+	/** Return true if yaw is within AllowedYawError of desired yaw */
 	virtual bool ReachedDesiredRotation();
 
 	/** @return The half-height of the default Pawn, scaled by the component scale. By default returns the half-height of the RootComponent, regardless of whether it is registered or collidable. */
 	virtual float GetDefaultHalfHeight() const;
 
 	/** See if this actor is currently being controlled */
-	UFUNCTION(BlueprintCallable, Category="Pawn")
+	UFUNCTION(BlueprintCallable, Category=Pawn)
 	bool IsControlled() const;
 
 	/** Returns controller for this actor. */
-	UFUNCTION(BlueprintCallable, Category="Pawn")
+	UFUNCTION(BlueprintCallable, Category=Pawn)
 	AController* GetController() const;
 
 	/** Get the rotation of the Controller, often the 'view' rotation of this Pawn. */
-	UFUNCTION(BlueprintCallable, Category="Pawn")
+	UFUNCTION(BlueprintCallable, Category=Pawn)
 	FRotator GetControlRotation() const;
 
 	/** Called when Controller is replicated */
@@ -196,9 +192,6 @@ public:
 	UFUNCTION()
 	virtual void OnRep_PlayerState();
 
-	/** used to prevent re-entry of OutsideWorldBounds event. */
-	uint32 bProcessingOutsideWorldBounds:1;
-
 	//~ Begin AActor Interface.
 	virtual FVector GetVelocity() const override;
 	virtual void Reset() override;
@@ -224,32 +217,37 @@ public:
 	virtual void EnableInput(APlayerController* PlayerController) override;
 	virtual void DisableInput(APlayerController* PlayerController) override;
 	virtual void TeleportSucceeded(bool bIsATest) override;
+	virtual bool IsBasedOnActor(const AActor* Other) const override;
 
 	/** Overridden to defer to the RootComponent's CanCharacterStepUpOn setting if it is explicitly Yes or No. If set to Owner, will return Super::CanBeBaseForCharacter(). */
 	virtual bool CanBeBaseForCharacter(APawn* APawn) const override;
 	//~ End AActor Interface
 
-	/** Use SetCanAffectNavigationGeneration to change this value at runtime.
-	 *	Note that calling this function at runtime will result in any navigation change only if runtime navigation generation is enabled. */
+	/**
+	 * Use SetCanAffectNavigationGeneration to change this value at runtime.
+	 * Note that calling this function at runtime will result in any navigation change only if runtime navigation generation is enabled.
+	 */
 	UFUNCTION(BlueprintCallable, Category="AI|Navigation", meta=(AdvancedDisplay="bForceUpdate"))
 	void SetCanAffectNavigationGeneration(bool bNewValue, bool bForceUpdate = false);
 
-	/** update all components relevant for navigation generators to match bCanAffectNavigationGeneration flag */
+	/** Update all components relevant for navigation generators to match bCanAffectNavigationGeneration flag */
 	virtual void UpdateNavigationRelevance() {}
 
 	//~ Begin INavAgentInterface Interface
 	virtual const FNavAgentProperties& GetNavAgentPropertiesRef() const override;
 	/** Basically retrieved pawn's location on navmesh */
-	UFUNCTION(BlueprintCallable, Category="Pawn")
+	UFUNCTION(BlueprintCallable, Category=Pawn)
 	virtual FVector GetNavAgentLocation() const override { return GetActorLocation() - FVector(0.f, 0.f, BaseEyeHeight); }
 	virtual void GetMoveGoalReachTest(const AActor* MovingActor, const FVector& MoveOffset, FVector& GoalOffset, float& GoalRadius, float& GoalHalfHeight) const override;
 	//~ End INavAgentInterface Interface
 
-	/** updates MovementComponent's parameters used by navigation system */
+	/** Updates MovementComponent's parameters used by navigation system */
 	void UpdateNavAgent();
 
-	/** @return true if we are in a state to take damage (checked at the start of TakeDamage.
-	*   Subclasses may check this as well if they override TakeDamage and don't want to potentially trigger TakeDamage actions by checking if it returns zero in the super class. */
+	/**
+	 * Return true if we are in a state to take damage (checked at the start of TakeDamage.
+	 * Subclasses may check this as well if they override TakeDamage and don't want to potentially trigger TakeDamage actions by checking if it returns zero in the super class.
+	 */
 	virtual bool ShouldTakeDamage(float Damage, FDamageEvent const& DamageEvent, AController* EventInstigator, AActor* DamageCauser) const;
 
 #if WITH_EDITOR
@@ -266,28 +264,29 @@ public:
 	/** Set BaseEyeHeight based on current state. */
 	virtual void RecalculateBaseEyeHeight();
 
+	/** Whether this Pawn's input handling is enabled.  Pawn must still be possessed to get input even if this is true */
 	bool InputEnabled() const { return bInputEnabled; }
 
 	/** 
 	 * Called when this Pawn is possessed. Only called on the server (or in standalone).
-	 *	@param C The controller possessing this pawn
+	 *@param C The controller possessing this pawn
 	 */
 	virtual void PossessedBy(AController* NewController);
 
 	/** Event called when the Pawn is possessed by a Controller (normally only occurs on the server/standalone). */
-	UFUNCTION(BlueprintImplementableEvent, meta=(DisplayName = "Possessed"))
+	UFUNCTION(BlueprintImplementableEvent, meta=(DisplayName= "Possessed"))
 	void ReceivePossessed(AController* NewController);
 
 	/** Event called when the Pawn is no longer possessed by a Controller. */
-	UFUNCTION(BlueprintImplementableEvent, meta=(DisplayName = "Unpossessed"))
+	UFUNCTION(BlueprintImplementableEvent, meta=(DisplayName= "Unpossessed"))
 	void ReceiveUnpossessed(AController* OldController);
 
 	/** @return true if controlled by a local (not network) Controller.	 */
-	UFUNCTION(BlueprintPure, Category="Pawn")
+	UFUNCTION(BlueprintPure, Category=Pawn)
 	virtual bool IsLocallyControlled() const;
   
 	/** @return true if controlled by a human player (possessed by a PlayerController).	 */
-	UFUNCTION(BlueprintPure, Category="Pawn")
+	UFUNCTION(BlueprintPure, Category=Pawn)
 	virtual bool IsPlayerControlled() const;
 	
 	/**
@@ -304,10 +303,10 @@ public:
 	 * If we have a controller, by default we aim at the player's 'eyes' direction
 	 * that is by default the Pawn rotation for AI, and camera (crosshair) rotation for human players.
 	 */
-	UFUNCTION(BlueprintCallable, Category="Pawn")
+	UFUNCTION(BlueprintCallable, Category=Pawn)
 	virtual FRotator GetBaseAimRotation() const;
 
-	/** return true if player is viewing this Pawn in FreeCam */
+	/** @return true if player is viewing this Pawn in FreeCam */
 	virtual bool InFreeCam() const;
 
 	/** Tell client that the Pawn is begin restarted. Calls Restart(). */
@@ -317,15 +316,15 @@ public:
 	virtual void FaceRotation(FRotator NewControlRotation, float DeltaTime = 0.f);
 
 	/** Call this function to detach safely pawn from its controller, knowing that we will be destroyed soon.	 */
-	UFUNCTION(BlueprintCallable, Category="Pawn", meta=(Keywords = "Delete"))
+	UFUNCTION(BlueprintCallable, Category=Pawn, meta=(Keywords = "Delete"))
 	virtual void DetachFromControllerPendingDestroy();
 
 	/** Spawn default controller for this Pawn, and get possessed by it. */
-	UFUNCTION(BlueprintCallable, Category="Pawn")
+	UFUNCTION(BlueprintCallable, Category=Pawn)
 	virtual void SpawnDefaultController();
 
 protected:
-	/** Get the controller instigating the damage. If the damage is caused by the world and the supplied controller is NULL or is this pawn's controller, uses LastHitBy as the instigator. */
+	/** Get the controller instigating the damage. If the damage is caused by the world and the supplied controller is nullptr or is this pawn's controller, uses LastHitBy as the instigator. */
 	virtual AController* GetDamageInstigator(AController* InstigatedBy, const UDamageType& DamageType) const;
 
 	/** Creates an InputComponent that can be used for custom input bindings. Called upon possession by a PlayerController. Return null if you don't want one. */
@@ -413,7 +412,6 @@ public:
 	virtual bool IsMoveInputIgnored() const;
 
 protected:
-
 	/**
 	 * Accumulated control input vector, stored in world space. This is the pending input, which is cleared (zeroed) once consumed.
 	 * @see GetPendingMovementInputVector(), AddMovementInput()
@@ -441,31 +439,19 @@ public:
 	/** Internal function meant for use only within Pawn or by a PawnMovementComponent. LastControlInputVector is updated with initial value of ControlInputVector. Returns ControlInputVector and resets it to zero. */
 	FVector Internal_ConsumeMovementInputVector();
 
-public:
-
 	/** Add an Actor to ignore by Pawn's movement collision */
 	void MoveIgnoreActorAdd(AActor* ActorToIgnore);
 
 	/** Remove an Actor to ignore by Pawn's movement collision */
 	void MoveIgnoreActorRemove(AActor* ActorToIgnore);
 
-public:
-
 	// DEPRECATED FUNCTIONS
 
-	/** Deprecated, misleading name and redundant */
-	DEPRECATED(4.14, "ClientSetRotation is deprecated. Call ClientSetRotation on the PlayerController directly")
-	virtual void ClientSetRotation(FRotator NewRotation);
-
 	/** (Deprecated) Launch Character with LaunchVelocity  */
 	DEPRECATED(4.8, "LaunchPawn is deprecated. For Characters, use LaunchCharacter() instead.")
 	UFUNCTION(BlueprintCallable, Category="Pawn", meta=(DeprecatedFunction, DeprecationMessage="Use Character.LaunchCharacter instead"))
 	void LaunchPawn(FVector LaunchVelocity, bool bXYOverride, bool bZOverride);
 
-	/** (Deprecated) Return the input vector in world space. */
-	DEPRECATED(4.5, "GetMovementInputVector() has been deprecated, use either GetPendingMovementInputVector() or GetLastMovementInputVector().")
-	FVector GetMovementInputVector() const;
-
 	/** (Deprecated) Return the input vector in world space. */
 	UFUNCTION(BlueprintCallable, Category="Pawn|Input", meta=(DeprecatedFunction, DisplayName="GetMovementInputVector", DeprecationMessage="GetMovementInputVector has been deprecated, use either GetPendingMovementInputVector or GetLastMovementInputVector"))
 	FVector K2_GetMovementInputVector() const;	
