Copying //UE4/Dev-Framework to //UE4/Dev-Main (Source: //UE4/Dev-Framework @ 3716594)

diff --git a/Engine/Source/Runtime/Engine/Classes/GameFramework/PlayerController.h b/Engine/Source/Runtime/Engine/Classes/GameFramework/PlayerController.h
--- a/Engine/Source/Runtime/Engine/Classes/GameFramework/PlayerController.h
+++ b/Engine/Source/Runtime/Engine/Classes/GameFramework/PlayerController.h
@@ -25,13 +25,17 @@
 #include "PlayerController.generated.h"
 
 class ACameraActor;
+class AHUD;
 class APawn;
+class APlayerCameraManager;
 class ASpectatorPawn;
 class FDebugDisplayInfo;
 class UActorChannel;
+class UCheatManager;
 class UGameViewportClient;
 class UInterpTrackInstDirector;
 class ULocalMessage;
+class UNetConnection;
 class UPlayer;
 class UPrimitiveComponent;
 struct FActiveHapticFeedbackEffect;
@@ -112,10 +116,6 @@ struct ENGINE_API FInputModeUIOnly : public FInputModeDataBase
 	/** Widget to focus */
 	FInputModeUIOnly& SetWidgetToFocus(TSharedPtr<SWidget> InWidgetToFocus) { WidgetToFocus = InWidgetToFocus; return *this; }
 
-	/** Whether to lock the mouse to the viewport */
-	DEPRECATED(4.13, "Mouse locking behavior is now controlled by an enum. Please use SetLockMouseToViewportBehavior(...) instead.")
-	FInputModeUIOnly& SetLockMouseToViewport(bool InLockMouseToViewport) { return SetLockMouseToViewportBehavior( InLockMouseToViewport ? EMouseLockMode::LockOnCapture : EMouseLockMode::DoNotLock ); }
-
 	/** Sets the mouse locking behavior of the viewport */
 	FInputModeUIOnly& SetLockMouseToViewportBehavior(EMouseLockMode InMouseLockMode) { MouseLockMode = InMouseLockMode; return *this; }
 
@@ -179,10 +179,6 @@ struct ENGINE_API FInputModeGameAndUI : public FInputModeDataBase
 	/** Widget to focus */
 	FInputModeGameAndUI& SetWidgetToFocus(TSharedPtr<SWidget> InWidgetToFocus) { WidgetToFocus = InWidgetToFocus; return *this; }
 
-	/** Whether to lock the mouse to the viewport */
-	DEPRECATED(4.13, "Mouse locking behavior is now controlled by an enum. Please use SetLockMouseToViewportBehavior(...) instead.")
-	FInputModeGameAndUI& SetLockMouseToViewport(bool InLockMouseToViewport) { return SetLockMouseToViewportBehavior( InLockMouseToViewport ? EMouseLockMode::LockOnCapture : EMouseLockMode::DoNotLock ); }
-
 	/** Sets the mouse locking behavior of the viewport */
 	FInputModeGameAndUI& SetLockMouseToViewportBehavior(EMouseLockMode InMouseLockMode) { MouseLockMode = InMouseLockMode; return *this; }
 
@@ -221,7 +217,6 @@ protected:
 };
 
 
-//~=============================================================================
 /**
  * PlayerControllers are used by human players to control Pawns.
  *
@@ -234,47 +229,41 @@ protected:
  *
  * @see https://docs.unrealengine.com/latest/INT/Gameplay/Framework/Controller/PlayerController/
  */
-
 UCLASS(config=Game, BlueprintType, Blueprintable, meta=(ShortTooltip="A Player Controller is an actor responsible for controlling a Pawn used by the player."))
 class ENGINE_API APlayerController : public AController
 {
-	GENERATED_UCLASS_BODY()
-	
+	GENERATED_BODY()
+
+public:
+	/** Default Constructor */
+	APlayerController(const FObjectInitializer& ObjectInitializer = FObjectInitializer::Get());
+
 	/** UPlayer associated with this PlayerController.  Could be a local player or a net connection. */
 	UPROPERTY()
-	class UPlayer* Player;
-
-	/** when true, reduces connect timeout from InitialConnectionTimeOut to ConnectionTimeout.  
-		Set once initial level load is complete (client may be unresponsive during level loading). */
-	uint32 bShortConnectTimeOut:1;
+	UPlayer* Player;
 
 	/** Used in net games so client can acknowledge it possessed a specific pawn. */
 	UPROPERTY()
-	class APawn* AcknowledgedPawn;
+	APawn* AcknowledgedPawn;
 
 	/** Director track that's currently possessing this player controller, or none if not possessed. */
-	UPROPERTY(transient)
-	class UInterpTrackInstDirector* ControllingDirTrackInst;
-
-	/** last used FOV based multiplier to distance to an object when determining if it exceeds the object's cull distance
-	 * @note: only valid for local player
-	 */
-	float LocalPlayerCachedLODDistanceFactor;
+	UPROPERTY(Transient)
+	UInterpTrackInstDirector* ControllingDirTrackInst;
 
 	/** Heads up display associated with this PlayerController. */
 	UPROPERTY()
-	class AHUD* MyHUD;
+	AHUD* MyHUD;
 
 	// ******************************************************************************
 	// Camera/view related variables
 
 	/** Camera manager associated with this Player Controller. */
 	UPROPERTY(BlueprintReadOnly, Category=PlayerController)
-	class APlayerCameraManager* PlayerCameraManager;
+	APlayerCameraManager* PlayerCameraManager;
 
 	/** PlayerCamera class should be set for each game, otherwise Engine.PlayerCameraManager is used */
 	UPROPERTY(EditAnywhere, BlueprintReadOnly, Category=PlayerController)
-	TSubclassOf<class APlayerCameraManager> PlayerCameraManagerClass;
+	TSubclassOf<APlayerCameraManager> PlayerCameraManagerClass;
 
 	/** 
 	 * True to allow this player controller to manage the camera target for you,
@@ -288,13 +277,19 @@ class ENGINE_API APlayerController : public AController
 	UPROPERTY(replicated)
 	FRotator TargetViewRotation; 
 
-	/**  Smoothed version of TargetViewRotation to remove jerkiness from intermittent replication updates. */
+	/** Smoothed version of TargetViewRotation to remove jerkiness from intermittent replication updates. */
 	FRotator BlendedTargetViewRotation;
 
 	/** Interp speed for blending remote view rotation for smoother client updates */
 	UPROPERTY(EditAnywhere, BlueprintReadWrite, Category=PlayerController)
 	float SmoothTargetViewRotationSpeed;
 
+	/**
+	 * Last used FOV based multiplier to distance to an object when determining if it exceeds the object's cull distance
+	 * @note: only valid for local player
+	 */
+	float LocalPlayerCachedLODDistanceFactor;
+
 	/** The actors which the camera shouldn't see - e.g. used to hide actors which the camera penetrates */
 	UPROPERTY()
 	TArray<class AActor*> HiddenActors;
@@ -323,25 +318,26 @@ class ENGINE_API APlayerController : public AController
 	int32 ClientCap;
 	
 	/** Object that manages "cheat" commands.  Not instantiated in shipping builds. */
-	UPROPERTY(transient, BlueprintReadOnly, Category="Cheat Manager")
-	class UCheatManager* CheatManager;
+	UPROPERTY(Transient, BlueprintReadOnly, Category="Cheat Manager")
+	UCheatManager* CheatManager;
 	
 	/** Class of my CheatManager.  The Cheat Manager is not created in shipping builds */
 	UPROPERTY(EditDefaultsOnly, BlueprintReadOnly, Category="Cheat Manager")
-	TSubclassOf<class UCheatManager> CheatClass;
+	TSubclassOf<UCheatManager> CheatClass;
 
 	/** Object that manages player input. */
-	UPROPERTY(transient)
+	UPROPERTY(Transient)
 	class UPlayerInput* PlayerInput;    
 	
-	UPROPERTY(transient)
+	UPROPERTY(Transient)
 	TArray<FActiveForceFeedbackEffect> ActiveForceFeedbackEffects;
 
 #if !(UE_BUILD_SHIPPING || UE_BUILD_TEST)
-	/** For debbugging, shows the last force feeback effects that played */
+	/** For debugging, shows the last force feeback effects that played */
 	TArray<FForceFeedbackEffectHistoryEntry> ForceFeedbackEffectHistoryEntries;
 #endif
 
+	/** Map of active force feedback effects */
 	TMap<int32, FDynamicForceFeedbackDetails> DynamicForceFeedbacks;
 
 	/** Currently playing haptic effects for both the left and right hand */
@@ -349,9 +345,18 @@ class ENGINE_API APlayerController : public AController
 	TSharedPtr<struct FActiveHapticFeedbackEffect> ActiveHapticEffect_Right;
 	TSharedPtr<struct FActiveHapticFeedbackEffect> ActiveHapticEffect_Gun;
 
-	/** list of names of levels the server is in the middle of sending us for a PrepareMapChange() call */
+	/** Currently active force feedback weights */
+	FForceFeedbackValues ForceFeedbackValues;
+
+	/** List of names of levels the server is in the middle of sending us for a PrepareMapChange() call */
 	TArray<FName> PendingMapChangeLevelNames;
 
+	/**
+	 * When true, reduces connect timeout from InitialConnectionTimeOut to ConnectionTimeout.  
+	 * Set once initial level load is complete (client may be unresponsive during level loading).
+	 */
+	uint32 bShortConnectTimeOut:1;
+
 	/** Is this player currently in cinematic mode?  Prevents rotation/movement/firing/etc */
 	uint32 bCinematicMode:1;
 	
@@ -373,7 +378,8 @@ class ENGINE_API APlayerController : public AController
 	UFUNCTION(client, reliable, Category=PlayerController)
 	void ClientSetSpectatorWaiting(bool bWaiting);
 
-	/** index identifying players using the same base connection (splitscreen clients)
+	/**
+	 * Index identifying players using the same base connection (splitscreen clients)
 	 * Used by netcode to match replicated PlayerControllers to the correct splitscreen viewport and child connection
 	 * replicated via special internal code, not through normal variable replication
 	 */
@@ -383,19 +389,20 @@ class ENGINE_API APlayerController : public AController
 	/** List of muted players in various categories */
 	FPlayerMuteList MuteList;
 
-	/** this is set on the OLD PlayerController when performing a swap over a network connection
+	/**
+	 * This is set on the OLD PlayerController when performing a swap over a network connection
 	 * so we know what connection we're waiting on acknowledgment from to finish destroying this PC
 	 * (or when the connection is closed)
 	 * @see GameModeBase::SwapPlayerControllers()
 	 */
 	UPROPERTY(DuplicateTransient)
-	class UNetConnection* PendingSwapConnection;
+	UNetConnection* PendingSwapConnection;
 
-	/** The net connection this controller is communicating on, NULL for local players on server */
+	/** The net connection this controller is communicating on, nullptr for local players on server */
 	UPROPERTY(DuplicateTransient)
-	class UNetConnection* NetConnection;
+	UNetConnection* NetConnection;
 
-	// Input axes values, accumulated each tick.
+	/** Input axes values, accumulated each tick. */
 	FRotator RotationInput;
 
 	/** Yaw input speed scaling */
@@ -437,12 +444,15 @@ class ENGINE_API APlayerController : public AController
 	UPROPERTY(config)
 	float ForceFeedbackScale;
 
+	/** List of keys that will cause click events to be forwarded, default to left click */
 	UPROPERTY(EditAnywhere, BlueprintReadWrite, Category=MouseInterface, meta=(EditCondition="bEnableClickEvents"))
 	TArray<FKey> ClickEventKeys;
 
+	/** Type of mouse cursor to show by default */
 	UPROPERTY(EditAnywhere, BlueprintReadOnly, Category=MouseInterface)
 	TEnumAsByte<EMouseCursor::Type> DefaultMouseCursor;
 
+	/** Currently visible mouse cursor */
 	UPROPERTY(BlueprintReadWrite, Category=MouseInterface)
 	TEnumAsByte<EMouseCursor::Type> CurrentMouseCursor;
 
@@ -454,13 +464,20 @@ class ENGINE_API APlayerController : public AController
 	UPROPERTY(BlueprintReadWrite, Category=MouseInterface)
 	TEnumAsByte<ECollisionChannel> CurrentClickTraceChannel;
 
+	/** Distance to trace when computing click events */
 	UPROPERTY(EditAnywhere, BlueprintReadWrite, Category=MouseInterface, meta=(DisplayName="Trace Distance"))
 	float HitResultTraceDistance;
 
-	FForceFeedbackValues ForceFeedbackValues;
+	/** Counter for this players seamless travels (used along with the below value, to restrict ServerNotifyLoadedWorld) */
+	UPROPERTY()
+	uint16 SeamlessTravelCount;
+
+	/** The value of SeamlessTravelCount, upon the last call to GameModeBase::HandleSeamlessTravelPlayer; used to detect seamless travel */
+	UPROPERTY()
+	uint16 LastCompletedSeamlessTravelCount;
 
 public:
-	/**  Enables cheats within the game */
+	/** Enables cheats within the game */
 	UFUNCTION(exec)
 	virtual void EnableCheats();
 
@@ -505,7 +522,7 @@ public:
 	UFUNCTION(exec)
 	virtual void Pause();
 
-	/** Trys to set the player's name to the given name. */
+	/** Tries to set the player's name to the given name. */
 	UFUNCTION(exec)
 	virtual void SetName(const FString& S);
 
@@ -514,7 +531,7 @@ public:
 	virtual void SwitchLevel(const FString& URL);
 
 	/** 
-	 * called to notify the server when the client has loaded a new world via seamless traveling
+	 * Called to notify the server when the client has loaded a new world via seamless traveling
 	 * @param WorldPackageName the name of the world package that was loaded
 	 * @param bFinalDest whether this world is the destination map for the travel (i.e. not the transition level)
 	 */
@@ -542,6 +559,7 @@ public:
 	/** Reset move and look input ignore flags to defaults */
 	virtual void ResetIgnoreInputFlags() override;
 
+	/** Returns hit results from doing a collision query at a certain location on the screen */
 	bool GetHitResultAtScreenPosition(const FVector2D ScreenPosition, const ECollisionChannel TraceChannel, const FCollisionQueryParams& CollisionQueryParams, FHitResult& HitResult) const;
 	bool GetHitResultAtScreenPosition(const FVector2D ScreenPosition, const ECollisionChannel TraceChannel, bool bTraceComplex, FHitResult& HitResult) const;
 	bool GetHitResultAtScreenPosition(const FVector2D ScreenPosition, const ETraceTypeQuery TraceChannel, bool bTraceComplex, FHitResult& HitResult) const;
@@ -550,34 +568,38 @@ public:
 	UFUNCTION(BlueprintCallable, Category="Game|Player", meta=(DeprecatedFunction, DeprecationMessage = "Use new GetHitResultUnderCursorByChannel or GetHitResultUnderCursorForObject", TraceChannel=ECC_Visibility, bTraceComplex=true))
 	bool GetHitResultUnderCursor(ECollisionChannel TraceChannel, bool bTraceComplex, FHitResult& HitResult) const;
 
+	/** Performs a collision query under the mouse cursor, looking on a trace channel */
 	UFUNCTION(BlueprintCallable, Category="Game|Player", meta=(bTraceComplex=true))
 	bool GetHitResultUnderCursorByChannel(ETraceTypeQuery TraceChannel, bool bTraceComplex, FHitResult& HitResult) const;
 
+	/** Performs a collision query under the mouse cursor, looking for object types */
 	UFUNCTION(BlueprintCallable, Category="Game|Player", meta=(bTraceComplex=true))
 	bool GetHitResultUnderCursorForObjects(const TArray<TEnumAsByte<EObjectTypeQuery> > & ObjectTypes, bool bTraceComplex, FHitResult& HitResult) const;
 
 	UFUNCTION(BlueprintCallable, Category="Game|Player", meta=(DeprecatedFunction, DeprecationMessage = "Use new GetHitResultUnderFingerByChannel or GetHitResultUnderFingerForObject", TraceChannel=ECC_Visibility, bTraceComplex=true))
 	bool GetHitResultUnderFinger(ETouchIndex::Type FingerIndex, ECollisionChannel TraceChannel, bool bTraceComplex, FHitResult& HitResult) const;
 
+	/** Performs a collision query under the finger, looking on a trace channel */
 	UFUNCTION(BlueprintCallable, Category="Game|Player", meta=(bTraceComplex=true))
 	bool GetHitResultUnderFingerByChannel(ETouchIndex::Type FingerIndex, ETraceTypeQuery TraceChannel, bool bTraceComplex, FHitResult& HitResult) const;
 
+	/** Performs a collision query under the finger, looking for object types */
 	UFUNCTION(BlueprintCallable, Category="Game|Player", meta=(bTraceComplex=true))
 	bool GetHitResultUnderFingerForObjects(ETouchIndex::Type FingerIndex, const  TArray<TEnumAsByte<EObjectTypeQuery> > & ObjectTypes, bool bTraceComplex, FHitResult& HitResult) const;
 
 	/** Convert current mouse 2D position to World Space 3D position and direction. Returns false if unable to determine value. **/
-	UFUNCTION(BlueprintCallable, Category = "Game|Player", meta = (DisplayName = "ConvertMouseLocationToWorldSpace", Keywords = "deproject"))
+	UFUNCTION(BlueprintCallable, Category="Game|Player", meta = (DisplayName = "ConvertMouseLocationToWorldSpace", Keywords = "deproject"))
 	bool DeprojectMousePositionToWorld(FVector& WorldLocation, FVector& WorldDirection) const;
 
 	/** Convert 2D screen position to World Space 3D position and direction. Returns false if unable to determine value. **/
-	UFUNCTION(BlueprintCallable, Category = "Game|Player", meta = (DisplayName = "ConvertScreenLocationToWorldSpace", Keywords = "deproject"))
+	UFUNCTION(BlueprintCallable, Category="Game|Player", meta = (DisplayName = "ConvertScreenLocationToWorldSpace", Keywords = "deproject"))
 	bool DeprojectScreenPositionToWorld(float ScreenX, float ScreenY, FVector& WorldLocation, FVector& WorldDirection) const;
 
 	/**
 	 * Convert a World Space 3D position into a 2D Screen Space position.
 	 * @return true if the world coordinate was successfully projected to the screen.
 	 */
-	UFUNCTION(BlueprintCallable, Category = "Game|Player", meta = (DisplayName = "ConvertWorldLocationToScreenLocation", Keywords = "project"))
+	UFUNCTION(BlueprintCallable, Category="Game|Player", meta = (DisplayName = "ConvertWorldLocationToScreenLocation", Keywords = "project"))
 	bool ProjectWorldLocationToScreen(FVector WorldLocation, FVector2D& ScreenLocation, bool bPlayerViewportRelative = false) const;
 
 	/**
@@ -589,20 +611,21 @@ public:
 	/** Positions the mouse cursor in screen space, in pixels. */
 	UFUNCTION( BlueprintCallable, Category="Game|Player", meta = (DisplayName = "SetMousePosition", Keywords = "mouse" ))
 	void SetMouseLocation( const int X, const int Y );
+
 	/**
-	  * Updates the rotation of player, based on ControlRotation after RotationInput has been applied.
-	  * This may then be modified by the PlayerCamera, and is passed to Pawn->FaceRotation().
-	  */
+	 * Updates the rotation of player, based on ControlRotation after RotationInput has been applied.
+	 * This may then be modified by the PlayerCamera, and is passed to Pawn->FaceRotation().
+	 */
 	virtual void UpdateRotation(float DeltaTime);
 
 protected:
 	/** Pawn has been possessed, so changing state to NAME_Playing. Start it walking and begin playing with it. */
 	virtual void BeginPlayingState();
 
-	// Leave playing state. */ 
+	/** Leave playing state. */ 
 	virtual void EndPlayingState();
 
-	/** overridden to return that player controllers are capable of RPCs */
+	/** Overridden to return that player controllers are capable of RPCs */
 	virtual bool HasNetOwner() const override;
 
 public:
@@ -613,10 +636,11 @@ public:
 	/** Notify player of change to level */
 	void LevelStreamingStatusChanged(class ULevelStreaming* LevelObject, bool bNewShouldBeLoaded, bool bNewShouldBeVisible, bool bNewShouldBlockOnLoad, int32 LODIndex);
 
-	/** used to wait until a map change can be prepared when one was already in progress */
+	/** Used to wait until a map change can be prepared when one was already in progress */
 	virtual void DelayedPrepareMapChange();
 
-	/** Called on client during seamless level transitions to get the list of Actors that should be moved into the new level
+	/**
+	 * Called on client during seamless level transitions to get the list of Actors that should be moved into the new level
 	 * PlayerControllers, Role < ROLE_Authority Actors, and any non-Actors that are inside an Actor that is in the list
 	 * (i.e. Object.Outer == Actor in the list)
 	 * are all automatically moved regardless of whether they're included here
@@ -628,13 +652,15 @@ public:
 	 */
 	virtual void GetSeamlessTravelActorList(bool bToEntry, TArray<class AActor*>& ActorList);
 
-	/** Called when seamless traveling and we are being replaced by the specified PC
+	/**
+	 * Called when seamless traveling and we are being replaced by the specified PC
 	 * clean up any persistent state (post process chains on LocalPlayers, for example)
 	 * (not called if PlayerController is the same for the from and to GameModes)
 	 */
 	virtual void SeamlessTravelTo(class APlayerController* NewPC);
 
-	/** Called when seamless traveling and the specified PC is being replaced by this one
+	/**
+	 * Called when seamless traveling and the specified PC is being replaced by this one
 	 * copy over data that should persist
 	 * (not called if PlayerController is the same for the from and to GameModes)
 	 */
@@ -754,7 +780,8 @@ public:
 	UFUNCTION(Reliable, Client)	
 	void ClientCommitMapChange();
 
-	/** Tells the client to block until all pending level streaming actions are complete
+	/**
+	 * Tells the client to block until all pending level streaming actions are complete
 	 * happens at the end of the tick
 	 * primarily used to force update the client ASAP at join time
 	 */
@@ -780,11 +807,11 @@ public:
 	UFUNCTION(Reliable, Client)
 	void ClientGotoState(FName NewState);
 
-	/** calls IgnoreLookInput on client */
+	/** Calls IgnoreLookInput on client */
 	UFUNCTION(Reliable, Client)
 	void ClientIgnoreLookInput(bool bIgnore);
 
-	/** calls IgnoreMoveInput on client */
+	/** Calls IgnoreMoveInput on client */
 	UFUNCTION(Reliable, Client)
 	void ClientIgnoreMoveInput(bool bIgnore);
 
@@ -808,7 +835,7 @@ public:
 	 * @param Space - Animation play area
 	 * @param CustomPlaySpace - Matrix used when Space = CAPS_UserDefined
 	 */
-	UFUNCTION(unreliable, client, BlueprintCallable, Category = "Game|Feedback")
+	UFUNCTION(unreliable, client, BlueprintCallable, Category="Game|Feedback")
 	void ClientPlayCameraAnim(class UCameraAnim* AnimToPlay, float Scale=1.f, float Rate=1.f, float BlendInTime=0.f, float BlendOutTime=0.f, bool bLoop=false, bool bRandomStartTime=false, ECameraAnimPlaySpace::Type Space=ECameraAnimPlaySpace::CameraLocal, FRotator CustomPlaySpace=FRotator::ZeroRotator);
 
 	/** 
@@ -875,7 +902,8 @@ public:
 	UFUNCTION(Reliable, Client)
 	void ClientSetBlockOnAsyncLoading();
 
-	/** Tell client to fade camera
+	/**
+	 * Tell client to fade camera
 	 * @Param bEnableFading - true if we should apply FadeColor/FadeAmount to the screen
 	 * @Param FadeColor - Color to fade to
 	 * @Param FadeAlpha - Amount of fading to apply
@@ -909,7 +937,7 @@ public:
 
 	/** Set the client's class of HUD and spawns a new instance of it. If there was already a HUD active, it is destroyed. */
 	UFUNCTION(BlueprintCallable, Category="HUD", Reliable, Client)
-	void ClientSetHUD(TSubclassOf<class AHUD> NewHUDClass);
+	void ClientSetHUD(TSubclassOf<AHUD> NewHUDClass);
 
 	/** Helper to get the size of the HUD canvas for this player controller.  Returns 0 if there is no HUD */
 	UFUNCTION(BlueprintCallable, Category="HUD")
@@ -917,14 +945,14 @@ public:
 
 	/** Gets the HUD currently being used by this player controller */
 	UFUNCTION(BlueprintCallable, Category="HUD")
-	class AHUD* GetHUD() const;
+	AHUD* GetHUD() const;
 
 	/**
 	 * Sets the Widget for the Mouse Cursor to display 
 	 * @param Cursor - the cursor to set the widget for
 	 * @param CursorWidget - the widget to set the cursor to
 	 */
-	UFUNCTION(BlueprintCallable, Category = "UI")
+	UFUNCTION(BlueprintCallable, Category="UI")
 	void SetMouseCursorWidget(EMouseCursor::Type Cursor, class UUserWidget* CursorWidget);
 
 	/** Set the view target
@@ -990,38 +1018,38 @@ public:
 	void PlayDynamicForceFeedback(float Intensity, float Duration, bool bAffectsLeftLarge, bool bAffectsLeftSmall, bool bAffectsRightLarge, bool bAffectsRightSmall, TEnumAsByte<EDynamicForceFeedbackAction::Type> Action, FLatentActionInfo LatentInfo);
 
 	/**
-	* Play a haptic feedback curve on the player's controller
-	* @param	HapticEffect			The haptic effect to play
-	* @param	Hand					Which hand to play the effect on
-	* @param	Scale					Scale between 0.0 and 1.0 on the intensity of playback
-	*/
-	UFUNCTION(BlueprintCallable, Category = "Game|Feedback")
+	 * Play a haptic feedback curve on the player's controller
+	 * @param	HapticEffect			The haptic effect to play
+	 * @param	Hand					Which hand to play the effect on
+	 * @param	Scale					Scale between 0.0 and 1.0 on the intensity of playback
+	 */
+	UFUNCTION(BlueprintCallable, Category="Game|Feedback")
 	void PlayHapticEffect(class UHapticFeedbackEffect_Base* HapticEffect, EControllerHand Hand, float Scale = 1.f,  bool bLoop = false);
 
 	/**
-	* Stops a playing haptic feedback curve
-	* @param	HapticEffect			The haptic effect to stop
-	* @param	Hand					Which hand to stop the effect for
-	*/
-	UFUNCTION(BlueprintCallable, Category = "Game|Feedback")
+	 * Stops a playing haptic feedback curve
+	 * @param	HapticEffect			The haptic effect to stop
+	 * @param	Hand					Which hand to stop the effect for
+	 */
+	UFUNCTION(BlueprintCallable, Category="Game|Feedback")
 	void StopHapticEffect(EControllerHand Hand);
 
 	/**
-	* Sets the value of the haptics for the specified hand directly, using frequency and amplitude.  NOTE:  If a curve is already
-	* playing for this hand, it will be cancelled in favour of the specified values.
-	*
-	* @param	Frequency				The normalized frequency [0.0, 1.0] to play through the haptics system
-	* @param	Amplitude				The normalized amplitude [0.0, 1.0] to set the haptic feedback to
-	* @param	Hand					Which hand to play the effect on
-	*/
-	UFUNCTION(BlueprintCallable, Category = "Game|Feedback")
+	 * Sets the value of the haptics for the specified hand directly, using frequency and amplitude.  NOTE:  If a curve is already
+	 * playing for this hand, it will be cancelled in favour of the specified values.
+	 *
+	 * @param	Frequency				The normalized frequency [0.0, 1.0] to play through the haptics system
+	 * @param	Amplitude				The normalized amplitude [0.0, 1.0] to set the haptic feedback to
+	 * @param	Hand					Which hand to play the effect on
+	 */
+	UFUNCTION(BlueprintCallable, Category="Game|Feedback")
 	void SetHapticsByValue(const float Frequency, const float Amplitude, EControllerHand Hand);
 	
 	/**
-	* Sets the light color of the player's controller
-	* @param	Color					The color for the light to be
-	*/
-	UFUNCTION(BlueprintCallable, Category = "Game|Feedback")
+	 * Sets the light color of the player's controller
+	 * @param	Color					The color for the light to be
+	 */
+	UFUNCTION(BlueprintCallable, Category="Game|Feedback")
 	void SetControllerLightColor(FColor Color);
 
 	/**
@@ -1090,7 +1118,7 @@ public:
 
 	/** send client localized message id */
 	UFUNCTION(Reliable, Client)
-	void ClientReceiveLocalizedMessage(TSubclassOf<ULocalMessage> Message, int32 Switch = 0, class APlayerState* RelatedPlayerState_1 = NULL, class APlayerState* RelatedPlayerState_2 = NULL, class UObject* OptionalObject = NULL);
+	void ClientReceiveLocalizedMessage(TSubclassOf<ULocalMessage> Message, int32 Switch = 0, class APlayerState* RelatedPlayerState_1 = nullptr, class APlayerState* RelatedPlayerState_2 = nullptr, class UObject* OptionalObject = nullptr);
 
 	/** acknowledge possession of pawn */
 	UFUNCTION(reliable, server, WithValidation)
@@ -1258,15 +1286,12 @@ public:
 	virtual void ActivateTouchInterface(class UTouchInterface* NewTouchInterface);
 
 	/** Set the virtual joystick visibility. */
-	UFUNCTION(BlueprintCallable, Category = "Game|Player")
+	UFUNCTION(BlueprintCallable, Category="Game|Player")
 	virtual void SetVirtualJoystickVisibility(bool bVisible);
 
-public:
-
 	/** Setup an input mode. */
 	virtual void SetInputMode(const FInputModeDataBase& InData);
 
-public:
 	/**
 	 * Change Camera mode
 	 * @param	New camera mode to set
@@ -1274,7 +1299,8 @@ public:
 	UFUNCTION(exec)
 	virtual void Camera(FName NewMode);
 
-	/** Set the view target blending with variable control
+	/**
+	 * Set the view target blending with variable control
 	 * @param NewViewTarget - new actor to set as view target
 	 * @param BlendTime - time taken to blend
 	 * @param BlendFunc - Cubic, Linear etc functions for blending
@@ -1308,16 +1334,16 @@ protected:
 	uint32 bCinemaDisableInputMove:1;
 	uint32 bCinemaDisableInputLook:1;
 
+	/** Whether we fully tick when the game is paused, if our tick function is allowed to do so. If false, we do a minimal update during the tick. */
+	UPROPERTY(EditAnywhere, BlueprintReadWrite, Category=PlayerController)
+	uint32 bShouldPerformFullTickWhenPaused : 1;
+
 private:
 	/* Whether the PlayerController's input handling is enabled. */
 	uint32 bInputEnabled:1;
 
 protected:
 
-	/** Whether we fully tick when the game is paused, if our tick function is allowed to do so. If false, we do a minimal update during the tick. */
-	UPROPERTY(EditAnywhere, BlueprintReadWrite, Category=PlayerController)
-	uint32 bShouldPerformFullTickWhenPaused:1;
-
 	/** The virtual touch interface */
 	TSharedPtr<class SVirtualJoystick> VirtualJoystick;
 
@@ -1342,11 +1368,19 @@ public:
 	/** Removes given inputcomponent from the input stack (regardless of if it's the top, actually). */
 	bool PopInputComponent(UInputComponent* Input);
 
+	/** Flushes the current key state. */
 	virtual void FlushPressedKeys();
 
+	/** Handles a key press */
 	virtual bool InputKey(FKey Key, EInputEvent EventType, float AmountDepressed, bool bGamepad);
+
+	/** Handles a touch screen action */
 	virtual bool InputTouch(uint32 Handle, ETouchType::Type Type, const FVector2D& TouchLocation, FDateTime DeviceTimestamp, uint32 TouchpadIndex);
+
+	/** Handles a controller axis input */
 	virtual bool InputAxis(FKey Key, float Delta, float DeltaTime, int32 NumSamples, bool bGamepad);
+
+	/** Handles motion control */
 	virtual bool InputMotion(const FVector& Tilt, const FVector& RotationRate, const FVector& Gravity, const FVector& Acceleration);
 
 	/** Associate a new UPlayer with this PlayerController. */
@@ -1384,6 +1418,7 @@ public:
 
 	//~ Begin UObject Interface
 	virtual void PostLoad() override;
+	virtual void GetLifetimeReplicatedProps(TArray< FLifetimeProperty > & OutLifetimeProps) const override;
 	//~ End UObject Interface
 
 	//~ Begin AActor Interface
@@ -1417,7 +1452,7 @@ protected:
 
 public:
 	//~ Begin AController Interface
-	virtual void GameHasEnded(class AActor* EndGameFocus = NULL, bool bIsWinner = false) override;
+	virtual void GameHasEnded(class AActor* EndGameFocus = nullptr, bool bIsWinner = false) override;
 	virtual bool IsLocalController() const override;
 	virtual void GetPlayerViewPoint(FVector& out_Location, FRotator& out_Rotation) const override;
 	virtual void SetInitialLocationAndRotation(const FVector& NewLocation, const FRotator& NewRotation) override;
@@ -1429,9 +1464,8 @@ public:
 	virtual void SetPawn(APawn* InPawn) override;
 	//~ End AController Interface
 
-	virtual void GetLifetimeReplicatedProps(TArray< FLifetimeProperty > & OutLifetimeProps) const override;
-
-	/** called on the server when the client sends a message indicating it was unable to initialize an Actor channel,
+	/**
+	 * Called on the server when the client sends a message indicating it was unable to initialize an Actor channel,
 	 * most commonly because the desired Actor's archetype couldn't be serialized
 	 * the default is to do nothing (Actor simply won't exist on the client), but this function gives the game code
 	 * an opportunity to try to correct the problem
@@ -1461,12 +1495,12 @@ public:
 
 	/**
 	 * Sets the Matinee director track instance that's currently possessing this player controller
-	 * @param   NewControllingDirector    The director track instance that's now controlling this player controller (or NULL for none)
-	 * @param	bClientSimulatingViewTarget	True to allow clients to simulate their own camera cuts (ignored if NewControllingDirector is NULL).
+	 * @param   NewControllingDirector    The director track instance that's now controlling this player controller (or nullptr for none)
+	 * @param	bClientSimulatingViewTarget	True to allow clients to simulate their own camera cuts (ignored if NewControllingDirector is nullptr).
 	 */
 	void SetControllingDirector(UInterpTrackInstDirector* NewControllingDirector, bool bClientSimulatingViewTarget);
 
-	/** Returns the Matinee director track that's currently possessing this player controller, or NULL for none */
+	/** Returns the Matinee director track that's currently possessing this player controller, or nullptr for none */
 	UInterpTrackInstDirector* GetControllingDirector();
 
 	/** spawn cameras for servers and owning players */
@@ -1488,7 +1522,7 @@ public:
 	/**
 	 * Clear any overrides that have been applied to audio listener
 	 */
-	UFUNCTION(BlueprintCallable, Category = "Game|Audio")
+	UFUNCTION(BlueprintCallable, Category="Game|Audio")
 	void ClearAudioListenerOverride();
 
 protected:
@@ -1501,15 +1535,14 @@ protected:
 	/** Currently overridden rotation of audio listener */
 	FRotator AudioListenerRotationOverride;
 
-protected:
 	/** Internal. */
 	void TickPlayerInput(const float DeltaSeconds, const bool bGamePaused);
 	virtual void ProcessPlayerInput(const float DeltaTime, const bool bGamePaused);
 	virtual void BuildInputStack(TArray<UInputComponent*>& InputStack);
 	void ProcessForceFeedbackAndHaptics(const float DeltaTime, const bool bGamePaused);
 	virtual bool IsInViewportClient(UGameViewportClient* ViewportClient) const;
-
 	virtual int32 GetInputIndex() const;
+	virtual ACameraActor* GetAutoActivateCameraForPlayer() const;
 
 	/** Allows the PlayerController to set up custom input bindings. */
 	virtual void SetupInputComponent();
@@ -1592,7 +1625,7 @@ public:
 	 * @param	OutSplitscreenPlayerIndex	receives the index [into the player's local GamePlayers array] for this player, if playing splitscreen.
 	 * @return	true if this player is playing splitscreen.
 	 */
-	bool IsSplitscreenPlayer(int32* OutSplitscreenPlayerIndex = NULL) const;
+	bool IsSplitscreenPlayer(int32* OutSplitscreenPlayerIndex = nullptr) const;
 
 	/**
 	 * Wrapper for determining whether this player is the first player on their console.
@@ -1617,7 +1650,6 @@ public:
 	/** Update the camera manager; this is called after all actors have been ticked.	 */ 
 	virtual void UpdateCameraManager(float DeltaSeconds);
 
-public:
 	/**
 	 * This function will be called to notify the player controller that the world has received its game class. In the case of a client
 	 * we need to initialize the Input System here.
@@ -1673,34 +1705,31 @@ public:
 	UFUNCTION(Reliable, Client)
 	void OnServerStartedVisualLogger(bool bIsLogging);
 
-protected:
-
-	virtual ACameraActor* GetAutoActivateCameraForPlayer() const;
-
-public:
-
+	/** Returns true if the mouse cursor should be shown */
 	virtual bool ShouldShowMouseCursor() const;
+
+	/** Returns the current mouse cursor, or None */
 	virtual EMouseCursor::Type GetMouseCursor() const;
 
-public:
 	// Spectating
 
-	/** Get the Pawn used when spectating. NULL when not spectating. */
-	UFUNCTION(BlueprintCallable, Category = "Pawn")
-	class ASpectatorPawn* GetSpectatorPawn() const { return SpectatorPawn; }
+	/** Get the Pawn used when spectating. nullptr when not spectating. */
+	UFUNCTION(BlueprintCallable, Category=Pawn)
+	ASpectatorPawn* GetSpectatorPawn() const { return SpectatorPawn; }
 
-	/** Returns the first of GetPawn() or GetSpectatorPawn() that is not NULL, or NULL otherwise. */
-	class APawn* GetPawnOrSpectator() const;
+	/** Returns the first of GetPawn() or GetSpectatorPawn() that is not nullptr, or nullptr otherwise. */
+	APawn* GetPawnOrSpectator() const;
 
 	/** Called to notify the controller that the spectator class has been received. */
-	virtual void ReceivedSpectatorClass(TSubclassOf<class ASpectatorPawn> SpectatorClass);
+	virtual void ReceivedSpectatorClass(TSubclassOf<ASpectatorPawn> SpectatorClass);
 
-	/** Returns the location the PlayerController is focused on.
+	/**
+	 * Returns the location the PlayerController is focused on.
 	 *  If there is a possessed Pawn, returns the Pawn's location.
 	 *  If there is a spectator Pawn, returns that Pawn's location.
 	 *  Otherwise, returns the PlayerController's spawn location (usually the last known Pawn location after it has died).
 	 */
-	UFUNCTION(BlueprintCallable, Category = "Pawn")
+	UFUNCTION(BlueprintCallable, Category=Pawn)
 	virtual FVector GetFocalLocation() const;
 
 protected:
@@ -1711,18 +1740,25 @@ protected:
 	virtual void EndSpectatingState();
 
 	/** Set the spectator pawn. Will also call AttachToPawn() using the new spectator. */
-	virtual void SetSpectatorPawn(class ASpectatorPawn* NewSpectatorPawn);
+	virtual void SetSpectatorPawn(ASpectatorPawn* NewSpectatorPawn);
 
 	/** Spawn a SpectatorPawn to use as a spectator and initialize it. By default it is spawned at the PC's current location and rotation. */
 	virtual ASpectatorPawn* SpawnSpectatorPawn();
 
-	/** Destroys the SpectatorPawn and sets it to NULL. */
+	/** Destroys the SpectatorPawn and sets it to nullptr. */
 	virtual void DestroySpectatorPawn();
 
 private:
-	/** The pawn used when spectating (NULL if not spectating). */
+	/** The pawn used when spectating (nullptr if not spectating). */
 	UPROPERTY()
-	class ASpectatorPawn* SpectatorPawn;
+	ASpectatorPawn* SpectatorPawn;
+
+	/** Used to delay calling ClientRestart() again when it hasn't been appropriately acknowledged. */
+	float		LastRetryPlayerTime;
+
+	/** Set during SpawnActor once and never again to indicate the intent of this controller instance (SERVER ONLY) */
+	UPROPERTY()
+	mutable bool bIsLocalPlayerController;
 
 protected:
 	/** The location used internally when there is no pawn or spectator, to know where to spawn the spectator or focus the camera on death. */
@@ -1733,11 +1769,9 @@ protected:
 	virtual void SetSpawnLocation(const FVector& NewLocation);
 
 public:
-
-	// Get the location used when initially created, or when changing states when there is no pawn or spectator.
+	/** Get the location used when initially created, or when changing states when there is no pawn or spectator. */
 	FVector GetSpawnLocation() const { return SpawnLocation; }
 
-public:
 	/** Called after this PlayerController's viewport/net connection is associated with this player controller. */
 	virtual void ReceivedPlayer();
 
@@ -1772,22 +1806,4 @@ public:
 	 * Designate this player controller as local (public for GameModeBase to use, not expected to be called anywhere else)
 	 */
 	void SetAsLocalPlayerController() { bIsLocalPlayerController = true; }
-
-private:
-	/** Used to delay calling ClientRestart() again when it hasn't been appropriately acknowledged. */
-	float		LastRetryPlayerTime;
-
-	/** Set during SpawnActor once and never again to indicate the intent of this controller instance (SERVER ONLY) */
-	UPROPERTY()
-	mutable bool bIsLocalPlayerController;
-
-public:
-	/** Counter for this players seamless travels (used along with the below value, to restrict ServerNotifyLoadedWorld) */
-	UPROPERTY()
-	uint16		SeamlessTravelCount;
-
-	/** The value of SeamlessTravelCount, upon the last call to GameModeBase::HandleSeamlessTravelPlayer; used to detect seamless travel */
-	UPROPERTY()
-	uint16		LastCompletedSeamlessTravelCount;
-
 };
diff --git a/Engine/Source/Runtime/Engine/Classes/GameFramework/PlayerInput.h b/Engine/Source/Runtime/Engine/Classes/GameFramework/PlayerInput.h
--- a/Engine/Source/Runtime/Engine/Classes/GameFramework/PlayerInput.h
+++ b/Engine/Source/Runtime/Engine/Classes/GameFramework/PlayerInput.h
@@ -526,7 +526,7 @@ private:
 	 * Given raw keystate value, returns the "massaged" value. Override for any custom behavior,
 	 * such as input changes dependent on a particular game state.
  	 */
-	float MassageAxisInput(FKey Key, float RawValue);
+	virtual float MassageAxisInput(FKey Key, float RawValue);
 
 	/** Process non-axes keystates */
 	void ProcessNonAxesKeys(FKey Inkey, FKeyState* KeyState);
diff --git a/Engine/Source/Runtime/Engine/Classes/Kismet/KismetMathLibrary.h b/Engine/Source/Runtime/Engine/Classes/Kismet/KismetMathLibrary.h
--- a/Engine/Source/Runtime/Engine/Classes/Kismet/KismetMathLibrary.h
+++ b/Engine/Source/Runtime/Engine/Classes/Kismet/KismetMathLibrary.h
@@ -617,17 +617,7 @@ class ENGINE_API UKismetMathLibrary : public UBlueprintFunctionLibrary
 	UFUNCTION(BlueprintPure, Category="Math|Float")
 	static float Lerp(float A, float B, float Alpha);
 	
-	/**
-	* Returns the fraction (alpha) of the range B-A that corresponds to Value, e.g.,
-	*	inputs A = 0, B = 8, Value = 3 : outputs Return Value = 3/8, indicating Value is 3/8 from A to B 
-	*	inputs A = 8, B = 0, Value = 3 : outputs Return Value = 5/8, indicating Value is 5/8 from A to B
-	* Named InverseLerp because Lerp( A, B, InverseLerp(A, B, Value) ) == Value
-	* @param A The "from" value this float could be, usually but not necessarily a minimum. Returned as 0.
-	* @param B The "to" value this float could be, usually but not necessarily a maximum. Returned as 1.
-	* @param Value A value intended to be normalized relative to B-A
-	* @return A normalized alpha value considering A and B.
-	*/
-	UFUNCTION(BlueprintPure, meta = (Keywords = "percentage normalize range"), Category = "Math|Float")
+	DEPRECATED(4.19, "Use NormalizeToRange instead")
 	static float InverseLerp(float A, float B, float Value);
 
 	/** Easeing  between A and B using a specified easing function */
diff --git a/Engine/Source/Runtime/Engine/Classes/Kismet/KismetSystemLibrary.h b/Engine/Source/Runtime/Engine/Classes/Kismet/KismetSystemLibrary.h
--- a/Engine/Source/Runtime/Engine/Classes/Kismet/KismetSystemLibrary.h
+++ b/Engine/Source/Runtime/Engine/Classes/Kismet/KismetSystemLibrary.h
@@ -281,7 +281,7 @@ class ENGINE_API UKismetSystemLibrary : public UBlueprintFunctionLibrary
 	 * @return	The literal string
 	 */
 	UFUNCTION(BlueprintPure, Category="Utilities|String", meta=(BlueprintThreadSafe))
-	static FString MakeLiteralString(const FString& Value);
+	static FString MakeLiteralString(FString Value);
 
 	/**
 	 * Creates a literal FText
@@ -1764,7 +1764,7 @@ FORCEINLINE uint8 UKismetSystemLibrary::MakeLiteralByte(uint8 Value)
 	return Value;
 }
 
-FORCEINLINE FString UKismetSystemLibrary::MakeLiteralString(const FString& Value)
+FORCEINLINE FString UKismetSystemLibrary::MakeLiteralString(FString Value)
 {
 	return Value;
 }
diff --git a/Engine/Source/Runtime/Engine/Classes/Materials/Material.h b/Engine/Source/Runtime/Engine/Classes/Materials/Material.h
--- a/Engine/Source/Runtime/Engine/Classes/Materials/Material.h
+++ b/Engine/Source/Runtime/Engine/Classes/Materials/Material.h
@@ -153,7 +153,7 @@ struct FMaterialInput
 	 * Note that this is the only member which is not derived from the output currently connected. 
 	 */
 	UPROPERTY()
-	FString InputName;
+	FName InputName;
 
 	UPROPERTY()
 	int32 Mask;
diff --git a/Engine/Source/Runtime/Engine/Classes/Materials/MaterialExpression.h b/Engine/Source/Runtime/Engine/Classes/Materials/MaterialExpression.h
--- a/Engine/Source/Runtime/Engine/Classes/Materials/MaterialExpression.h
+++ b/Engine/Source/Runtime/Engine/Classes/Materials/MaterialExpression.h
@@ -36,7 +36,7 @@ struct FExpressionInput
 	 * Note that this is the only member which is not derived from the output currently connected. 
 	 */
 	UPROPERTY()
-	FString InputName;
+	FName InputName;
 
 	UPROPERTY()
 	int32 Mask;
@@ -73,7 +73,7 @@ USTRUCT(noexport)
 struct FExpressionOutput
 {
 	UPROPERTY()
-	FString OutputName;
+	FName OutputName;
 
 	UPROPERTY()
 	int32 Mask;
@@ -246,7 +246,7 @@ class ENGINE_API UMaterialExpression : public UObject
 	virtual TArray<FExpressionOutput>& GetOutputs();
 	virtual const TArray<FExpressionInput*> GetInputs();
 	virtual FExpressionInput* GetInput(int32 InputIndex);
-	virtual FString GetInputName(int32 InputIndex) const;
+	virtual FName GetInputName(int32 InputIndex) const;
 	virtual bool IsInputConnectionRequired(int32 InputIndex) const;
 #if WITH_EDITOR
 	virtual uint32 GetInputType(int32 InputIndex);
diff --git a/Engine/Source/Runtime/Engine/Classes/Materials/MaterialExpressionBlendMaterialAttributes.h b/Engine/Source/Runtime/Engine/Classes/Materials/MaterialExpressionBlendMaterialAttributes.h
--- a/Engine/Source/Runtime/Engine/Classes/Materials/MaterialExpressionBlendMaterialAttributes.h
+++ b/Engine/Source/Runtime/Engine/Classes/Materials/MaterialExpressionBlendMaterialAttributes.h
@@ -48,7 +48,7 @@ class UMaterialExpressionBlendMaterialAttributes : public UMaterialExpression
 #endif
 	virtual const TArray<FExpressionInput*> GetInputs()override;
 	virtual FExpressionInput* GetInput(int32 InputIndex)override;
-	virtual FString GetInputName(int32 InputIndex) const override;
+	virtual FName GetInputName(int32 InputIndex) const override;
 	virtual bool IsInputConnectionRequired(int32 InputIndex) const override {return true;}
 #if WITH_EDITOR
 	virtual bool IsResultMaterialAttributes(int32 OutputIndex) override {return true;}
diff --git a/Engine/Source/Runtime/Engine/Classes/Materials/MaterialExpressionBreakMaterialAttributes.h b/Engine/Source/Runtime/Engine/Classes/Materials/MaterialExpressionBreakMaterialAttributes.h
--- a/Engine/Source/Runtime/Engine/Classes/Materials/MaterialExpressionBreakMaterialAttributes.h
+++ b/Engine/Source/Runtime/Engine/Classes/Materials/MaterialExpressionBreakMaterialAttributes.h
@@ -28,7 +28,7 @@ class UMaterialExpressionBreakMaterialAttributes : public UMaterialExpression
 #endif
 	virtual const TArray<FExpressionInput*> GetInputs()override;
 	virtual FExpressionInput* GetInput(int32 InputIndex)override;
-	virtual FString GetInputName(int32 InputIndex) const override;
+	virtual FName GetInputName(int32 InputIndex) const override;
 	virtual bool IsInputConnectionRequired(int32 InputIndex) const override;
 #if WITH_EDITOR
 	virtual uint32 GetInputType(int32 InputIndex) override {return MCT_MaterialAttributes;}
diff --git a/Engine/Source/Runtime/Engine/Classes/Materials/MaterialExpressionCustom.h b/Engine/Source/Runtime/Engine/Classes/Materials/MaterialExpressionCustom.h
--- a/Engine/Source/Runtime/Engine/Classes/Materials/MaterialExpressionCustom.h
+++ b/Engine/Source/Runtime/Engine/Classes/Materials/MaterialExpressionCustom.h
@@ -27,7 +27,7 @@ struct FCustomInput
 	GENERATED_USTRUCT_BODY()
 
 	UPROPERTY(EditAnywhere, Category=CustomInput)
-	FString InputName;
+	FName InputName;
 
 	UPROPERTY()
 	FExpressionInput Input;
@@ -66,7 +66,7 @@ class UMaterialExpressionCustom : public UMaterialExpression
 #endif
 	virtual const TArray<FExpressionInput*> GetInputs() override;
 	virtual FExpressionInput* GetInput(int32 InputIndex) override;
-	virtual FString GetInputName(int32 InputIndex) const override;
+	virtual FName GetInputName(int32 InputIndex) const override;
 #if WITH_EDITOR
 	virtual uint32 GetInputType(int32 InputIndex) override {return MCT_Unknown;}
 	virtual uint32 GetOutputType(int32 OutputIndex) override;
diff --git a/Engine/Source/Runtime/Engine/Classes/Materials/MaterialExpressionDepthFade.h b/Engine/Source/Runtime/Engine/Classes/Materials/MaterialExpressionDepthFade.h
--- a/Engine/Source/Runtime/Engine/Classes/Materials/MaterialExpressionDepthFade.h
+++ b/Engine/Source/Runtime/Engine/Classes/Materials/MaterialExpressionDepthFade.h
@@ -35,7 +35,7 @@ class UMaterialExpressionDepthFade : public UMaterialExpression
 
 
 	//~ Begin UMaterialExpression Interface
-	virtual FString GetInputName(int32 InputIndex) const override
+	virtual FName GetInputName(int32 InputIndex) const override
 	{
 		if (InputIndex == 0)
 		{
diff --git a/Engine/Source/Runtime/Engine/Classes/Materials/MaterialExpressionFunctionInput.h b/Engine/Source/Runtime/Engine/Classes/Materials/MaterialExpressionFunctionInput.h
--- a/Engine/Source/Runtime/Engine/Classes/Materials/MaterialExpressionFunctionInput.h
+++ b/Engine/Source/Runtime/Engine/Classes/Materials/MaterialExpressionFunctionInput.h
@@ -39,7 +39,7 @@ class UMaterialExpressionFunctionInput : public UMaterialExpression
 
 	/** The input's name, which will be drawn on the connector in function call expressions that use this function. */
 	UPROPERTY(EditAnywhere, BlueprintReadWrite, Category=MaterialExpressionFunctionInput)
-	FString InputName;
+	FName InputName;
 
 	/** The input's description, which will be used as a tooltip on the connector in function call expressions that use this function. */
 	UPROPERTY(EditAnywhere, Category=MaterialExpressionFunctionInput)
diff --git a/Engine/Source/Runtime/Engine/Classes/Materials/MaterialExpressionFunctionOutput.h b/Engine/Source/Runtime/Engine/Classes/Materials/MaterialExpressionFunctionOutput.h
--- a/Engine/Source/Runtime/Engine/Classes/Materials/MaterialExpressionFunctionOutput.h
+++ b/Engine/Source/Runtime/Engine/Classes/Materials/MaterialExpressionFunctionOutput.h
@@ -19,7 +19,7 @@ class UMaterialExpressionFunctionOutput : public UMaterialExpression
 
 	/** The output's name, which will be drawn on the connector in function call expressions that use this function. */
 	UPROPERTY(EditAnywhere, BlueprintReadWrite, Category = MaterialExpressionFunctionOutput)
-	FString OutputName;
+	FName OutputName;
 
 	/** The output's description, which will be used as a tooltip on the connector in function call expressions that use this function. */
 	UPROPERTY(EditAnywhere, Category=MaterialExpressionFunctionOutput, meta=(MultiLine=true))
@@ -57,9 +57,9 @@ class UMaterialExpressionFunctionOutput : public UMaterialExpression
 	virtual int32 Compile(class FMaterialCompiler* Compiler, int32 OutputIndex) override;
 	virtual void GetCaption(TArray<FString>& OutCaptions) const override;
 #endif // WITH_EDITOR
-	virtual FString GetInputName(int32 InputIndex) const override
+	virtual FName GetInputName(int32 InputIndex) const override
 	{
-		return TEXT("");
+		return NAME_None;
 	}
 #if WITH_EDITOR
 	virtual void GetExpressionToolTip(TArray<FString>& OutToolTip) override;
diff --git a/Engine/Source/Runtime/Engine/Classes/Materials/MaterialExpressionMaterialFunctionCall.h b/Engine/Source/Runtime/Engine/Classes/Materials/MaterialExpressionMaterialFunctionCall.h
--- a/Engine/Source/Runtime/Engine/Classes/Materials/MaterialExpressionMaterialFunctionCall.h
+++ b/Engine/Source/Runtime/Engine/Classes/Materials/MaterialExpressionMaterialFunctionCall.h
@@ -118,7 +118,7 @@ class UMaterialExpressionMaterialFunctionCall : public UMaterialExpression
 #endif
 	virtual const TArray<FExpressionInput*> GetInputs() override;
 	virtual FExpressionInput* GetInput(int32 InputIndex) override;
-	virtual FString GetInputName(int32 InputIndex) const override;
+	virtual FName GetInputName(int32 InputIndex) const override;
 	virtual bool IsInputConnectionRequired(int32 InputIndex) const override;
 #if WITH_EDITOR
 	virtual FString GetDescription() const override;
@@ -133,7 +133,7 @@ class UMaterialExpressionMaterialFunctionCall : public UMaterialExpression
 	//~ End UMaterialExpression Interface
 
 	/** Util to get name of a particular type, optionally with type appended in parenthesis */
-	ENGINE_API FString GetInputNameWithType(int32 InputIndex, bool bWithType) const;
+	ENGINE_API FName GetInputNameWithType(int32 InputIndex, bool bWithType) const;
 
 #if WITH_EDITOR
 	/** 
diff --git a/Engine/Source/Runtime/Engine/Classes/Materials/MaterialExpressionVertexInterpolator.h b/Engine/Source/Runtime/Engine/Classes/Materials/MaterialExpressionVertexInterpolator.h
--- a/Engine/Source/Runtime/Engine/Classes/Materials/MaterialExpressionVertexInterpolator.h
+++ b/Engine/Source/Runtime/Engine/Classes/Materials/MaterialExpressionVertexInterpolator.h
@@ -23,7 +23,7 @@ class UMaterialExpressionVertexInterpolator : public UMaterialExpressionCustomOu
 
 	virtual FExpressionInput* GetInput(int32 InputIndex) override;
 	virtual FString GetFunctionName() const override { return TEXT("VertexInterpolator"); }
-	virtual FString GetInputName(int32 InputIndex) const override { return TEXT("VS"); }
+	virtual FName GetInputName(int32 InputIndex) const override { return TEXT("VS"); }
 
 	int32				InterpolatorIndex;
 	EMaterialValueType	InterpolatedType;
diff --git a/Engine/Source/Runtime/Engine/Classes/VisualLogger/VisualLoggerKismetLibrary.h b/Engine/Source/Runtime/Engine/Classes/VisualLogger/VisualLoggerKismetLibrary.h
--- a/Engine/Source/Runtime/Engine/Classes/VisualLogger/VisualLoggerKismetLibrary.h
+++ b/Engine/Source/Runtime/Engine/Classes/VisualLogger/VisualLoggerKismetLibrary.h
@@ -2,8 +2,6 @@
 
 #pragma once
 
-#include "CoreMinimal.h"
-#include "UObject/ObjectMacros.h"
 #include "Kismet/BlueprintFunctionLibrary.h"
 #include "VisualLoggerKismetLibrary.generated.h"
 
@@ -12,15 +10,25 @@ class UVisualLoggerKismetLibrary : public UBlueprintFunctionLibrary
 {
 	GENERATED_UCLASS_BODY()
 
+	UFUNCTION(BlueprintCallable, Category = "Debug|VisualLogger", meta = (DisplayName = "Enable VisLog Recording", CallableWithoutWorldContext, DevelopmentOnly))
+	static void EnableRecording(bool bEnabled);
+
+	/** Makes SourceOwner log to DestinationOwner's vislog*/
+	UFUNCTION(BlueprintCallable, Category = "Debug|VisualLogger", meta = (CallableWithoutWorldContext, DevelopmentOnly))
+	static void RedirectVislog(UObject* SourceOwner, UObject* DestinationOwner);
+
 	/** Logs simple text string with Visual Logger - recording for Visual Logs has to be enabled to record this data */
-	UFUNCTION(BlueprintCallable, Category = "Debug|VisualLogger", meta = (HidePin = "WorldContextObject", DefaultToSelf = "WorldContextObject", DisplayName = "LogText"))
-	static void LogText(UObject* WorldContextObject, FString Text, FName LogCategory = TEXT("Blueprint Log"));
+	UFUNCTION(BlueprintCallable, Category = "Debug|VisualLogger", meta = (DisplayName = "VisLog Text", AdvancedDisplay = "WorldContextObject, bAddToMessageLog", CallableWithoutWorldContext, DevelopmentOnly, DefaultToSelf = "WorldContextObject"))
+	static void LogText(UObject* WorldContextObject, FString Text, FName LogCategory = TEXT("VisLogBP"), bool bAddToMessageLog = false);
 
 	/** Logs location as sphere with given radius - recording for Visual Logs has to be enabled to record this data */
-	UFUNCTION(BlueprintCallable, Category = "Debug|VisualLogger", meta = (HidePin = "WorldContextObject", DefaultToSelf = "WorldContextObject", DisplayName = "LogLocation"))
-	static void LogLocation(UObject* WorldContextObject, FVector Location, FString Text, FLinearColor ObjectColor = FLinearColor::White, float Radius = 10, FName LogCategory = TEXT("Blueprint Log"));
+	UFUNCTION(BlueprintCallable, Category = "Debug|VisualLogger", meta = (DisplayName = "VisLog Location", AdvancedDisplay = "WorldContextObject, bAddToMessageLog", CallableWithoutWorldContext, DevelopmentOnly, DefaultToSelf = "WorldContextObject"))
+	static void LogLocation(UObject* WorldContextObject, FVector Location, FString Text, FLinearColor ObjectColor = FLinearColor::Blue, float Radius = 10, FName LogCategory = TEXT("VisLogBP"), bool bAddToMessageLog = false);
 
 	/** Logs box shape - recording for Visual Logs has to be enabled to record this data */
-	UFUNCTION(BlueprintCallable, Category = "Debug|VisualLogger", meta = (HidePin = "WorldContextObject", DefaultToSelf = "WorldContextObject", DisplayName = "Log Box Shape"))
-	static void LogBox(UObject* WorldContextObject, FBox BoxShape, FString Text, FLinearColor ObjectColor = FLinearColor::White, FName LogCategory = TEXT("Blueprint Log"));
+	UFUNCTION(BlueprintCallable, Category = "Debug|VisualLogger", meta = (DisplayName = "VisLog Box Shape", AdvancedDisplay = "WorldContextObject, bAddToMessageLog", CallableWithoutWorldContext, DevelopmentOnly, DefaultToSelf = "WorldContextObject"))
+	static void LogBox(UObject* WorldContextObject, FBox BoxShape, FString Text, FLinearColor ObjectColor = FLinearColor::Blue, FName LogCategory = TEXT("VisLogBP"), bool bAddToMessageLog = false);
+
+	UFUNCTION(BlueprintCallable, Category = "Debug|VisualLogger", meta = (DisplayName = "VisLog Segment", AdvancedDisplay = "WorldContextObject, bAddToMessageLog, Thickness", CallableWithoutWorldContext, DevelopmentOnly, DefaultToSelf = "WorldContextObject"))
+	static void LogSegment(UObject* WorldContextObject, const FVector SegmentStart, const FVector SegmentEnd, FString Text, FLinearColor ObjectColor = FLinearColor::Blue, const float Thickness = 0.f, FName CategoryName = TEXT("VisLogBP"), bool bAddToMessageLog = false);
 };
diff --git a/Engine/Source/Runtime/Engine/Private/AI/Navigation/PImplRecastNavMesh.cpp b/Engine/Source/Runtime/Engine/Private/AI/Navigation/PImplRecastNavMesh.cpp
--- a/Engine/Source/Runtime/Engine/Private/AI/Navigation/PImplRecastNavMesh.cpp
+++ b/Engine/Source/Runtime/Engine/Private/AI/Navigation/PImplRecastNavMesh.cpp
@@ -307,7 +307,13 @@ void FPImplRecastNavMesh::Serialize( FArchive& Ar, int32 NavMeshVersion )
 	{
 		TilesToSave.Reserve(DetourNavMesh->getMaxTiles());
 		
-		if (NavMeshOwner->SupportsStreaming() && !IsRunningCommandlet())
+		// if not registered yet will add everything. Worst case it will get dumped later on
+		// Note, this is especially needed during world duplication, where NavigationSystem
+		// gets instantiated after Serialize has already been called on all duplicated 
+		// NavigationData instances 
+		if (NavMeshOwner->IsRegistered() 
+			&& NavMeshOwner->SupportsStreaming()
+			&& !IsRunningCommandlet())
 		{
 			// We save only tiles that belongs to this level
 			GetNavMeshTilesIn(NavMeshOwner->GetNavigableBoundsInLevel(NavMeshOwner->GetLevel()), TilesToSave);
diff --git a/Engine/Source/Runtime/Engine/Private/AI/Navigation/RecastNavMeshGenerator.cpp b/Engine/Source/Runtime/Engine/Private/AI/Navigation/RecastNavMeshGenerator.cpp
--- a/Engine/Source/Runtime/Engine/Private/AI/Navigation/RecastNavMeshGenerator.cpp
+++ b/Engine/Source/Runtime/Engine/Private/AI/Navigation/RecastNavMeshGenerator.cpp
@@ -3099,7 +3099,7 @@ void FRecastTileGenerator::AddReferencedObjects(FReferenceCollector& Collector)
 	}
 }
 
-static int32 CaclulateMaxTilesCount(const TNavStatArray<FBox>& NavigableAreas, float TileSizeinWorldUnits, float AvgLayersPerGridCell)
+static int32 CalculateMaxTilesCount(const TNavStatArray<FBox>& NavigableAreas, float TileSizeinWorldUnits, float AvgLayersPerGridCell)
 {
 	int32 GridCellsCount = 0;
 	for (FBox AreaBounds : NavigableAreas)
@@ -3377,7 +3377,7 @@ void FRecastNavMeshGenerator::CalcNavMeshProperties(int32& MaxTiles, int32& MaxP
 	int32 MaxRequestedTiles = 0;
 	if (DestNavMesh->IsResizable())
 	{
-		MaxRequestedTiles = CaclulateMaxTilesCount(InclusionBounds, Config.tileSize * Config.cs, AvgLayersPerTile);
+		MaxRequestedTiles = CalculateMaxTilesCount(InclusionBounds, Config.tileSize * Config.cs, AvgLayersPerTile);
 	}
 	else
 	{
@@ -3526,7 +3526,7 @@ void FRecastNavMeshGenerator::OnNavigationBoundsChanged()
 	if (!IsGameStaticNavMesh(DestNavMesh) && DestNavMesh->IsResizable() && DetourMesh)
 	{
 		// Check whether Navmesh size needs to be changed
-		int32 MaxRequestedTiles = CaclulateMaxTilesCount(InclusionBounds, Config.tileSize * Config.cs, AvgLayersPerTile);
+		int32 MaxRequestedTiles = CalculateMaxTilesCount(InclusionBounds, Config.tileSize * Config.cs, AvgLayersPerTile);
 		if (DetourMesh->getMaxTiles() != MaxRequestedTiles)
 		{
 			// Destroy current NavMesh, it will be allocated with a new size on next build request
@@ -4376,7 +4376,7 @@ uint32 FRecastNavMeshGenerator::LogMemUsed() const
 }
 
 #if !(UE_BUILD_SHIPPING || UE_BUILD_TEST) && ENABLE_VISUAL_LOG
-void FRecastNavMeshGenerator::GrabDebugSnapshot(struct FVisualLogEntry* Snapshot, const FBox& BoundingBox, const struct FLogCategoryBase& LogCategory, ELogVerbosity::Type LogVerbosity) const
+void FRecastNavMeshGenerator::GrabDebugSnapshot(struct FVisualLogEntry* Snapshot, const FBox& BoundingBox, const FName& CategoryName, ELogVerbosity::Type LogVerbosity) const
 {
 	const UNavigationSystem* NavSys = UNavigationSystem::GetCurrent(GetWorld());
 	const FNavigationOctree* NavOctree = NavSys ? NavSys->GetNavOctree() : NULL;
@@ -4426,7 +4426,7 @@ void FRecastNavMeshGenerator::GrabDebugSnapshot(struct FVisualLogEntry* Snapshot
 						Indices[IndicesIdx] = CachedGeometry.Indices[IndicesIdx];
 					}
 
-					Snapshot->AddElement(CoordBuffer, Indices, LogCategory.GetCategoryName(), LogVerbosity, FColorList::LightGrey.WithAlpha(255));
+					Snapshot->AddElement(CoordBuffer, Indices, CategoryName, LogVerbosity, FColorList::LightGrey.WithAlpha(255));
 				}
 				else
 				{
@@ -4449,14 +4449,14 @@ void FRecastNavMeshGenerator::GrabDebugSnapshot(struct FVisualLogEntry* Snapshot
 							FBoxNavAreaData Box;
 							AreaMods[i].GetBox(Box);
 
-							Snapshot->AddElement(FBox::BuildAABB(Box.Origin, Box.Extent), FMatrix::Identity, LogCategory.GetCategoryName(), NavAreaVerbosity, PolygonColor.WithAlpha(255));
+							Snapshot->AddElement(FBox::BuildAABB(Box.Origin, Box.Extent), FMatrix::Identity, CategoryName, NavAreaVerbosity, PolygonColor.WithAlpha(255));
 						}
 						else if (AreaMods[i].GetShapeType() == ENavigationShapeType::Cylinder)
 						{
 							FCylinderNavAreaData Cylinder;
 							AreaMods[i].GetCylinder(Cylinder);
 
-							Snapshot->AddElement(Cylinder.Origin, Cylinder.Origin + FVector(0, 0, Cylinder.Height), Cylinder.Radius, LogCategory.GetCategoryName(), NavAreaVerbosity, PolygonColor.WithAlpha(255));
+							Snapshot->AddElement(Cylinder.Origin, Cylinder.Origin + FVector(0, 0, Cylinder.Height), Cylinder.Radius, CategoryName, NavAreaVerbosity, PolygonColor.WithAlpha(255));
 						}
 						else
 						{
@@ -4469,7 +4469,7 @@ void FRecastNavMeshGenerator::GrabDebugSnapshot(struct FVisualLogEntry* Snapshot
 								Verts,
 								Convex.MinZ - NavData->CellHeight, 
 								Convex.MaxZ + NavData->CellHeight, 
-								LogCategory.GetCategoryName(), NavAreaVerbosity, PolygonColor.WithAlpha(255));
+								CategoryName, NavAreaVerbosity, PolygonColor.WithAlpha(255));
 						}
 					}
 				}
diff --git a/Engine/Source/Runtime/Engine/Private/Actor.cpp b/Engine/Source/Runtime/Engine/Private/Actor.cpp
--- a/Engine/Source/Runtime/Engine/Private/Actor.cpp
+++ b/Engine/Source/Runtime/Engine/Private/Actor.cpp
@@ -216,7 +216,7 @@ bool AActor::CheckActorComponents()
 void AActor::ResetOwnedComponents()
 {
 #if WITH_EDITOR
-	// Identify any natively-constructed components referenced by properties that either failed to serialize or came in as NULL.
+	// Identify any natively-constructed components referenced by properties that either failed to serialize or came in as nullptr.
 	if(HasAnyFlags(RF_WasLoaded) && NativeConstructedComponentToPropertyMap.Num() > 0)
 	{
 		for (UActorComponent* Component : OwnedComponents)
@@ -259,7 +259,7 @@ void AActor::ResetOwnedComponents()
 
 			if (Component->GetIsReplicated())
 			{
-				ReplicatedComponents.Add(Component);
+				ReplicatedComponents.AddUnique(Component);
 			}
 		}
 	}, true, RF_NoFlags, EInternalObjectFlags::PendingKill);
@@ -300,7 +300,7 @@ void AActor::AddReferencedObjects(UObject* InThis, FReferenceCollector& Collecto
 UWorld* AActor::GetWorld() const
 {
 	// CDO objects do not belong to a world
-	// If the actors outer is destroyed or unreachable we are shutting down and the world should be NULL
+	// If the actors outer is destroyed or unreachable we are shutting down and the world should be nullptr
 	if (!HasAnyFlags(RF_ClassDefaultObject) && !GetOuter()->HasAnyFlags(RF_BeginDestroyed) && !GetOuter()->IsUnreachable())
 	{
 		if (ULevel* Level = GetLevel())
@@ -349,7 +349,7 @@ bool AActor::TeleportTo( const FVector& DestLocation, const FRotator& DestRotati
 {
 	SCOPE_CYCLE_COUNTER(STAT_TeleportToTime);
 
-	if(RootComponent == NULL)
+	if(RootComponent == nullptr)
 	{
 		return false;
 	}
@@ -563,7 +563,7 @@ void AActor::PostLoad()
 	Super::PostLoad();
 
 	// add ourselves to our Owner's Children array
-	if (Owner != NULL)
+	if (Owner != nullptr)
 	{
 		checkSlow(!Owner->Children.Contains(this));
 		Owner->Children.Add(this);
@@ -840,7 +840,7 @@ bool AActor::Rename( const TCHAR* InName, UObject* NewOuter, ERenameFlags Flags
 
 UNetConnection* AActor::GetNetConnection() const
 {
-	return Owner ? Owner->GetNetConnection() : NULL;
+	return Owner ? Owner->GetNetConnection() : nullptr;
 }
 
 UPlayer* AActor::GetNetOwningPlayer()
@@ -853,7 +853,7 @@ UPlayer* AActor::GetNetOwningPlayer()
 			return Owner->GetNetOwningPlayer();
 		}
 	}
-	return NULL;
+	return nullptr;
 }
 
 bool AActor::DestroyNetworkActorHandled()
@@ -940,7 +940,7 @@ void AActor::PreReplication( IRepChangedPropertyTracker & ChangedPropertyTracker
 	DOREPLIFETIME_ACTIVE_OVERRIDE( AActor, AttachmentReplication, RootComponent && !RootComponent->GetIsReplicated() );
 
 	UBlueprintGeneratedClass* BPClass = Cast<UBlueprintGeneratedClass>(GetClass());
-	if (BPClass != NULL)
+	if (BPClass != nullptr)
 	{
 		BPClass->InstancePreReplication(this, ChangedPropertyTracker);
 	}
@@ -1051,7 +1051,7 @@ void AActor::GetSimpleCollisionCylinder(float& CollisionRadius, float& Collision
 
 bool AActor::IsRootComponentCollisionRegistered() const
 {
-	return RootComponent != NULL && RootComponent->IsRegistered() && RootComponent->IsCollisionEnabled();
+	return RootComponent != nullptr && RootComponent->IsRegistered() && RootComponent->IsCollisionEnabled();
 }
 
 bool AActor::IsAttachedTo(const AActor* Other) const
@@ -1180,7 +1180,7 @@ bool AActor::CheckStillInWorld()
 		return false;
 	}
 	// Check if box has poked outside the world
-	else if( ( RootComponent != NULL ) && ( GetRootComponent()->IsRegistered() == true ) )
+	else if( ( RootComponent != nullptr ) && ( GetRootComponent()->IsRegistered() == true ) )
 	{
 		const FBox&	Box = GetRootComponent()->Bounds.GetBox();
 		if(	Box.Min.X < -HALF_WORLD_MAX || Box.Max.X > HALF_WORLD_MAX ||
@@ -1232,7 +1232,7 @@ void AActor::UpdateOverlaps(bool bDoNotifies)
 	USceneComponent* const RootComp = GetRootComponent();
 	if (RootComp)
 	{
-		RootComp->UpdateOverlaps(NULL, bDoNotifies);
+		RootComp->UpdateOverlaps(nullptr, bDoNotifies);
 	}
 }
 
@@ -1416,7 +1416,7 @@ static void MarkOwnerRelevantComponentsDirty(AActor* TheActor)
 	for (int32 i = 0; i < TheActor->Children.Num(); i++)
 	{
 		AActor* Child = TheActor->Children[i];
-		if (Child != NULL && !Child->IsPendingKill())
+		if (Child != nullptr && !Child->IsPendingKill())
 		{
 			MarkOwnerRelevantComponentsDirty(Child);
 		}
@@ -1449,7 +1449,7 @@ void AActor::SetOwner( AActor *NewOwner )
 {
 	if (Owner != NewOwner && !IsPendingKill())
 	{
-		if (NewOwner != NULL && NewOwner->IsOwnedBy(this))
+		if (NewOwner != nullptr && NewOwner->IsOwnedBy(this))
 		{
 			UE_LOG(LogActor, Error, TEXT("SetOwner(): Failed to set '%s' owner of '%s' because it would cause an Owner loop"), *NewOwner->GetName(), *GetName());
 			return;
@@ -1457,7 +1457,7 @@ void AActor::SetOwner( AActor *NewOwner )
 
 		// Sets this actor's parent to the specified actor.
 		AActor* OldOwner = Owner;
-		if( Owner != NULL )
+		if( Owner != nullptr )
 		{
 			// remove from old owner's Children array
 			verifySlow(Owner->Children.Remove(this) == 1);
@@ -1465,7 +1465,7 @@ void AActor::SetOwner( AActor *NewOwner )
 
 		Owner = NewOwner;
 
-		if( Owner != NULL )
+		if( Owner != nullptr )
 		{
 			// add to new owner's Children array
 			checkSlow(!Owner->Children.Contains(this));
@@ -1479,14 +1479,14 @@ void AActor::SetOwner( AActor *NewOwner )
 
 bool AActor::HasNetOwner() const
 {
-	if (Owner == NULL)
+	if (Owner == nullptr)
 	{
 		// all basic AActors are unable to call RPCs without special AActors as their owners (ie APlayerController)
 		return false;
 	}
 
 	// Find the topmost actor in this owner chain
-	AActor* TopOwner = NULL;
+	AActor* TopOwner = nullptr;
 	for (TopOwner = Owner; TopOwner->Owner; TopOwner = TopOwner->Owner)
 	{
 	}
@@ -1497,18 +1497,11 @@ bool AActor::HasNetOwner() const
 void AActor::K2_AttachRootComponentTo(USceneComponent* InParent, FName InSocketName, EAttachLocation::Type AttachLocationType /*= EAttachLocation::KeepRelativeOffset */, bool bWeldSimulatedBodies /*=true*/)
 {
 PRAGMA_DISABLE_DEPRECATION_WARNINGS
-	AttachRootComponentTo(InParent, InSocketName, AttachLocationType, bWeldSimulatedBodies);
-PRAGMA_ENABLE_DEPRECATION_WARNINGS
-}
-
-void AActor::AttachRootComponentTo(USceneComponent* InParent, FName InSocketName, EAttachLocation::Type AttachLocationType /*= EAttachLocation::KeepRelativeOffset */, bool bWeldSimulatedBodies /*=false*/)
-{
-	if(RootComponent && InParent)
+if (RootComponent && InParent)
 	{
-PRAGMA_DISABLE_DEPRECATION_WARNINGS
 		RootComponent->AttachTo(InParent, InSocketName, AttachLocationType, bWeldSimulatedBodies);
-PRAGMA_ENABLE_DEPRECATION_WARNINGS
 	}
+PRAGMA_ENABLE_DEPRECATION_WARNINGS
 }
 
 void AActor::K2_AttachToComponent(USceneComponent* Parent, FName SocketName, EAttachmentRule LocationRule, EAttachmentRule RotationRule, EAttachmentRule ScaleRule, bool bWeldSimulatedBodies)
@@ -1553,13 +1546,6 @@ void AActor::OnRep_AttachmentReplication()
 }
 
 void AActor::K2_AttachRootComponentToActor(AActor* InParentActor, FName InSocketName /*= NAME_None*/, EAttachLocation::Type AttachLocationType /*= EAttachLocation::KeepRelativeOffset */, bool bWeldSimulatedBodies /*=true*/)
-{
-PRAGMA_DISABLE_DEPRECATION_WARNINGS
-	AttachRootComponentToActor(InParentActor, InSocketName, AttachLocationType, bWeldSimulatedBodies);
-PRAGMA_ENABLE_DEPRECATION_WARNINGS
-}
-
-void AActor::AttachRootComponentToActor(AActor* InParentActor, FName InSocketName /*= NAME_None*/, EAttachLocation::Type AttachLocationType /*= EAttachLocation::KeepRelativeOffset */, bool bWeldSimulatedBodies /*=false*/)
 {
 	if (RootComponent && InParentActor)
 	{
@@ -1628,11 +1614,6 @@ void AActor::DetachFromActor(const FDetachmentTransformRules& DetachmentRules)
 	}
 }
 
-void AActor::DetachSceneComponentsFromParent(USceneComponent* InParentComponent, bool bMaintainWorldPosition)
-{
-	DetachAllSceneComponents(InParentComponent, bMaintainWorldPosition ? FDetachmentTransformRules::KeepWorldTransform : FDetachmentTransformRules::KeepRelativeTransform);
-}
-
 void AActor::DetachAllSceneComponents(USceneComponent* InParentComponent, const FDetachmentTransformRules& DetachmentRules)
 {
 	if (InParentComponent)
@@ -1673,7 +1654,7 @@ FName AActor::GetAttachParentSocketName() const
 void AActor::GetAttachedActors(TArray<class AActor*>& OutActors) const
 {
 	OutActors.Reset();
-	if (RootComponent != NULL)
+	if (RootComponent != nullptr)
 	{
 		// Current set of components to check
 		TInlineComponentArray<USceneComponent*> CompsToCheck;
@@ -1696,7 +1677,7 @@ void AActor::GetAttachedActors(TArray<class AActor*>& OutActors) const
 				CheckedComps.Add(SceneComp);
 
 				AActor* CompOwner = SceneComp->GetOwner();
-				if (CompOwner != NULL)
+				if (CompOwner != nullptr)
 				{
 					if (CompOwner != this)
 					{
@@ -1709,7 +1690,7 @@ void AActor::GetAttachedActors(TArray<class AActor*>& OutActors) const
 						for (USceneComponent* ChildComp : SceneComp->GetAttachChildren())
 						{
 							// Add any we have not explored yet to the set to check
-							if ((ChildComp != NULL) && !CheckedComps.Contains(ChildComp))
+							if ((ChildComp != nullptr) && !CheckedComps.Contains(ChildComp))
 							{
 								CompsToCheck.Push(ChildComp);
 							}
@@ -1749,7 +1730,7 @@ bool AActor::IsInPersistentLevel(bool bIncludeLevelStreamingPersistent) const
 	ULevel* MyLevel = GetLevel();
 	UWorld* World = GetWorld();
 	return ( (MyLevel == World->PersistentLevel) || ( bIncludeLevelStreamingPersistent && World->StreamingLevels.Num() > 0 &&
-														Cast<ULevelStreamingPersistent>(World->StreamingLevels[0]) != NULL &&
+														Cast<ULevelStreamingPersistent>(World->StreamingLevels[0]) != nullptr &&
 														World->StreamingLevels[0]->GetLoadedLevel() == MyLevel ) );
 }
 
@@ -1760,10 +1741,10 @@ bool AActor::IsMatineeControlled() const
 	for(auto It(ControllingMatineeActors.CreateConstIterator()); It; It++)
 	{
 		AMatineeActor* ControllingMatineeActor = *It;
-		if(ControllingMatineeActor != NULL)
+		if(ControllingMatineeActor != nullptr)
 		{
 			UInterpGroupInst* GroupInst = ControllingMatineeActor->FindGroupInst(this);
-			if(GroupInst != NULL)
+			if(GroupInst != nullptr)
 			{
 				if(GroupInst->Group && GroupInst->Group->HasMoveTrack())
 				{
@@ -1782,17 +1763,17 @@ bool AActor::IsMatineeControlled() const
 
 bool AActor::IsRootComponentStatic() const
 {
-	return(RootComponent != NULL && RootComponent->Mobility == EComponentMobility::Static);
+	return(RootComponent != nullptr && RootComponent->Mobility == EComponentMobility::Static);
 }
 
 bool AActor::IsRootComponentStationary() const
 {
-	return(RootComponent != NULL && RootComponent->Mobility == EComponentMobility::Stationary);
+	return(RootComponent != nullptr && RootComponent->Mobility == EComponentMobility::Stationary);
 }
 
 bool AActor::IsRootComponentMovable() const
 {
-	return(RootComponent != NULL && RootComponent->Mobility == EComponentMobility::Movable);
+	return(RootComponent != nullptr && RootComponent->Mobility == EComponentMobility::Movable);
 }
 
 FVector AActor::GetTargetLocation(AActor* RequestedBy) const
@@ -2153,7 +2134,7 @@ float AActor::TakeDamage(float DamageAmount, FDamageEvent const& DamageEvent, AC
 	{
 		ReceiveAnyDamage(ActualDamage, DamageTypeCDO, EventInstigator, DamageCauser);
 		OnTakeAnyDamage.Broadcast(this, ActualDamage, DamageTypeCDO, EventInstigator, DamageCauser);
-		if (EventInstigator != NULL)
+		if (EventInstigator != nullptr)
 		{
 			EventInstigator->InstigatedAnyDamage(ActualDamage, DamageTypeCDO, this, DamageCauser);
 		}
@@ -2195,20 +2176,13 @@ float AActor::InternalTakePointDamage(float Damage, FPointDamageEvent const& Poi
 	return Damage;
 }
 
-// deprecated
-void AActor::ReceivePointDamage(float Damage, const class UDamageType* DamageType, FVector HitLocation, FVector HitNormal, class UPrimitiveComponent* HitComponent, FName BoneName, FVector ShotFromDirection, class AController* InstigatedBy, AActor* DamageCauser)
-{
-	// Call proper version with a default FHitResult.
-	ReceivePointDamage(Damage, DamageType, HitLocation, HitNormal, HitComponent, BoneName, ShotFromDirection, InstigatedBy, DamageCauser, FHitResult());
-}
-
 /** Util to check if prim comp pointer is valid and still alive */
 extern bool IsPrimCompValidAndAlive(UPrimitiveComponent* PrimComp);
 
 /** Used to determine if it is ok to call a notification on this object */
 bool IsActorValidToNotify(AActor* Actor)
 {
-	return (Actor != NULL) && !Actor->IsPendingKill() && !Actor->GetClass()->HasAnyClassFlags(CLASS_NewerVersionExists);
+	return (Actor != nullptr) && !Actor->IsPendingKill() && !Actor->GetClass()->HasAnyClassFlags(CLASS_NewerVersionExists);
 }
 
 void AActor::InternalDispatchBlockingHit(UPrimitiveComponent* MyComp, UPrimitiveComponent* OtherComp, bool bSelfMoved, FHitResult const& Hit)
@@ -2308,7 +2282,7 @@ void AActor::DisplayDebug(UCanvas* Canvas, const FDebugDisplayInfo& DebugDisplay
 		GetComponentsBoundingCylinder(MyRadius, MyHeight);
 		DisplayDebugManager.DrawString(FString::Printf(TEXT("Collision Radius: %f Height: %f"), MyRadius, MyHeight));
 
-		if ( RootComponent == NULL )
+		if ( RootComponent == nullptr )
 		{
 			DisplayDebugManager.DrawString(FString(TEXT("No RootComponent")));
 		}
@@ -2385,7 +2359,7 @@ APawn* AActor::GetInstigator() const
 
 AController* AActor::GetInstigatorController() const
 {
-	return Instigator ? Instigator->Controller : NULL;
+	return Instigator ? Instigator->Controller : nullptr;
 }
 
 void AActor::CalcCamera(float DeltaTime, FMinimalViewInfo& OutResult)
@@ -2511,7 +2485,7 @@ void AActor::AddOwnedComponent(UActorComponent* Component)
 	{
 		if (Component->GetIsReplicated())
 		{
-			ReplicatedComponents.Add(Component);
+			ReplicatedComponents.AddUnique(Component);
 		}
 
 		if (Component->IsCreatedByConstructionScript())
@@ -2558,7 +2532,7 @@ void AActor::UpdateReplicatedComponent(UActorComponent* Component)
 	checkf(Component->GetOwner() == this, TEXT("UE-9568: Component %s being updated for Actor %s"), *Component->GetPathName(), *GetPathName() );
 	if (Component->GetIsReplicated())
 	{
-		ReplicatedComponents.Add(Component);
+		ReplicatedComponents.AddUnique(Component);
 	}
 	else
 	{
@@ -2572,9 +2546,10 @@ void AActor::UpdateAllReplicatedComponents()
 
 	for (UActorComponent* Component : OwnedComponents)
 	{
-		if (Component != NULL)
+		if (Component != nullptr && Component->GetIsReplicated())
 		{
-			UpdateReplicatedComponent(Component);
+			// We reset the array so no need to add unique
+			ReplicatedComponents.Add(Component);
 		}
 	}
 }
@@ -3057,7 +3032,7 @@ void AActor::SetReplicates(bool bInReplicates)
 	{
 		if (bReplicates == false && bInReplicates == true)
 		{
-			if (UWorld* MyWorld = GetWorld())		// GetWorld will return NULL on CDO, FYI
+			if (UWorld* MyWorld = GetWorld())		// GetWorld will return nullptr on CDO, FYI
 			{
 				MyWorld->AddNetworkActor(this);
 			}
@@ -3810,7 +3785,7 @@ int32 AActor::GetFunctionCallspace( UFunction* Function, void* Parameters, FFram
 	bool bIsServer = NetMode == NM_ListenServer || NetMode == NM_DedicatedServer;
 
 	// get the top most function
-	while (Function->GetSuperFunction() != NULL)
+	while (Function->GetSuperFunction() != nullptr)
 	{
 		Function = Function->GetSuperFunction();
 	}
@@ -3856,10 +3831,10 @@ int32 AActor::GetFunctionCallspace( UFunction* Function, void* Parameters, FFram
 	if (Role == ROLE_Authority)
 	{
 		UNetConnection* NetConnection = GetNetConnection();
-		if (NetConnection == NULL)
+		if (NetConnection == nullptr)
 		{
 			UPlayer *ClientPlayer = GetNetOwningPlayer();
-			if (ClientPlayer == NULL)
+			if (ClientPlayer == nullptr)
 			{
 				// Check if a player ever owned this (topmost owner is playercontroller or beacon)
 				if (HasNetOwner())
@@ -3873,7 +3848,7 @@ int32 AActor::GetFunctionCallspace( UFunction* Function, void* Parameters, FFram
 				DEBUG_CALLSPACE(TEXT("GetFunctionCallspace authority non client owner %s %s"), *Function->GetName(), FunctionCallspace::ToString(Callspace));
 				return Callspace;
 			}
-			else if (Cast<ULocalPlayer>(ClientPlayer) != NULL)
+			else if (Cast<ULocalPlayer>(ClientPlayer) != nullptr)
 			{
 				// This is a local player, call locally
 				DEBUG_CALLSPACE(TEXT("GetFunctionCallspace Client local function: %s %s"), *Function->GetName(), FunctionCallspace::ToString(Callspace));
@@ -3915,7 +3890,7 @@ bool AActor::CallRemoteFunction( UFunction* Function, void* Parameters, FOutParm
 	UNetDriver* NetDriver = GetNetDriver();
 	if (NetDriver)
 	{
-		NetDriver->ProcessRemoteFunction(this, Function, Parameters, OutParms, Stack, NULL);
+		NetDriver->ProcessRemoteFunction(this, Function, Parameters, OutParms, Stack, nullptr);
 		return true;
 	}
 
@@ -4095,7 +4070,7 @@ bool AActor::IncrementalRegisterComponents(int32 NumComponentsToRegister)
 	}
 	
 	// Register RootComponent first so all other children components can reliably use it (i.e., call GetLocation) when they register
-	if (RootComponent != NULL && !RootComponent->IsRegistered())
+	if (RootComponent != nullptr && !RootComponent->IsRegistered())
 	{
 #if PERF_TRACK_DETAILED_ASYNC_STATS
 		FScopeCycleCounterUObject ContextScope(RootComponent);
@@ -4170,7 +4145,7 @@ bool AActor::IncrementalRegisterComponents(int32 NumComponentsToRegister)
 
 bool AActor::HasValidRootComponent()
 { 
-	return (RootComponent != NULL && RootComponent->IsRegistered()); 
+	return (RootComponent != nullptr && RootComponent->IsRegistered()); 
 }
 
 void AActor::MarkComponentsAsPendingKill()
@@ -4456,7 +4431,7 @@ UMaterialInstanceDynamic* AActor::MakeMIDForMaterial(class UMaterialInterface* P
 	// OR PrimitiveComponent->CreateAndSetMaterialInstanceDynamicFromMaterial
 	// OR UMaterialInstanceDynamic::Create
 
-	return NULL;
+	return nullptr;
 }
 
 float AActor::GetDistanceTo(const AActor* OtherActor) const
diff --git a/Engine/Source/Runtime/Engine/Private/ActorConstruction.cpp b/Engine/Source/Runtime/Engine/Private/ActorConstruction.cpp
--- a/Engine/Source/Runtime/Engine/Private/ActorConstruction.cpp
+++ b/Engine/Source/Runtime/Engine/Private/ActorConstruction.cpp
@@ -877,11 +877,6 @@ void AActor::FinishAndRegisterComponent(UActorComponent* Component)
 	BlueprintCreatedComponents.Add(Component);
 }
 
-UActorComponent* AActor::CreateComponentFromTemplate(UActorComponent* Template, const FString& InName)
-{
-	return CreateComponentFromTemplate(Template, FName(*InName));
-}
-
 #if !UE_BUILD_SHIPPING
 static TAutoConsoleVariable<int32> CVarLogBlueprintComponentInstanceCalls(
 	TEXT("LogBlueprintComponentInstanceCalls"),
diff --git a/Engine/Source/Runtime/Engine/Private/AssetManager.cpp b/Engine/Source/Runtime/Engine/Private/AssetManager.cpp
--- a/Engine/Source/Runtime/Engine/Private/AssetManager.cpp
+++ b/Engine/Source/Runtime/Engine/Private/AssetManager.cpp
@@ -135,6 +135,7 @@ void UAssetManager::PostInitProperties()
 
 		FEditorDelegates::PreBeginPIE.AddUObject(this, &UAssetManager::PreBeginPIE);
 		FEditorDelegates::EndPIE.AddUObject(this, &UAssetManager::EndPIE);
+		FCoreUObjectDelegates::OnObjectSaved.AddUObject(this, &UAssetManager::OnObjectPreSave);
 
 		// In editor builds guess the type/name if allowed
 		bShouldGuessTypeAndName = Settings.bShouldGuessTypeAndNameInEditor;
@@ -375,14 +376,18 @@ int32 UAssetManager::ScanPathsForPrimaryAssets(FPrimaryAssetType PrimaryAssetTyp
 			if (!ParentClassFromData.IsEmpty())
 			{
 				const FString ClassObjectPath = FPackageName::ExportTextPathToObjectPath(ParentClassFromData);
-				const FString ClassName = FPackageName::ObjectPathToObjectName(ClassObjectPath);
+				const FName ClassName = FName(*FPackageName::ObjectPathToObjectName(ClassObjectPath));
 
 				TArray<FName> ValidNames;
+				ValidNames.Add(ClassName);
 #if WITH_EDITOR
-				// Also check old names
-				ValidNames = FLinkerLoad::FindPreviousNamesForClass(BaseClass->GetPathName(), false);
+				// Check for redirected name
+				FName RedirectedName = FLinkerLoad::FindNewNameForClass(ClassName, false);
+				if (RedirectedName != NAME_None && RedirectedName != ClassName)
+				{
+					ValidNames.Add(RedirectedName);
+				}
 #endif
-				ValidNames.Add(FName(*ClassName));		
 				for (const FName& ValidName : ValidNames)
 				{
 					if (DerivedClassNames.Contains(ValidName))
@@ -2953,6 +2958,23 @@ void UAssetManager::OnInMemoryAssetDeleted(UObject *Object)
 	RemovePrimaryAssetId(PrimaryAssetId);
 }
 
+void UAssetManager::OnObjectPreSave(UObject* Object)
+{
+	// If this is in the asset manager dictionary, make sure it actually has a primary asset id that matches
+	FPrimaryAssetId FoundPrimaryAssetId = GetPrimaryAssetIdForPath(FSoftObjectPath(Object));
+
+	if ((Object->IsAsset() || Object->IsA(UClass::StaticClass())) && FoundPrimaryAssetId.IsValid())
+	{
+		TSharedRef<FPrimaryAssetTypeData>* FoundType = AssetTypeMap.Find(FoundPrimaryAssetId.PrimaryAssetType);
+		FPrimaryAssetId ObjectPrimaryAssetId = Object->GetPrimaryAssetId();
+
+		if (FoundPrimaryAssetId != ObjectPrimaryAssetId && !(*FoundType)->Info.bIsEditorOnly)
+		{
+			UE_LOG(LogAssetManager, Error, TEXT("Registered PrimaryAssetId %s for asset %s does not match object's real id of %s! This will not load properly at runtime!"), *FoundPrimaryAssetId.ToString(), *Object->GetPathName(), *ObjectPrimaryAssetId.ToString());
+		}
+	}
+}
+
 void UAssetManager::OnAssetRenamed(const FAssetData& NewData, const FString& OldPath)
 {
 	// Ignore PIE changes
diff --git a/Engine/Source/Runtime/Engine/Private/BlueprintGeneratedClass.cpp b/Engine/Source/Runtime/Engine/Private/BlueprintGeneratedClass.cpp
--- a/Engine/Source/Runtime/Engine/Private/BlueprintGeneratedClass.cpp
+++ b/Engine/Source/Runtime/Engine/Private/BlueprintGeneratedClass.cpp
@@ -649,7 +649,11 @@ UInheritableComponentHandler* UBlueprintGeneratedClass::GetInheritableComponentH
 	
 	if (InheritableComponentHandler)
 	{
-		InheritableComponentHandler->PreloadAll();
+		if (!GEventDrivenLoaderEnabled || !EVENT_DRIVEN_ASYNC_LOAD_ACTIVE_AT_RUNTIME)
+		{
+			// This preload will not succeed in EDL
+			InheritableComponentHandler->PreloadAll();
+		}	
 	}
 
 	if (!InheritableComponentHandler && bCreateIfNecessary)
@@ -763,6 +767,14 @@ UObject* UBlueprintGeneratedClass::FindArchetype(UClass* ArchetypeClass, const F
 				if (ComponentKey.IsValid())
 				{
 					Archetype = ICH->GetOverridenComponentTemplate(ComponentKey);
+
+					if (GEventDrivenLoaderEnabled && EVENT_DRIVEN_ASYNC_LOAD_ACTIVE_AT_RUNTIME)
+					{
+						if (Archetype && Archetype->HasAnyFlags(RF_NeedLoad))
+						{
+							UE_LOG(LogClass, Fatal, TEXT("%s had RF_NeedLoad when searching for an archetype of %s named %s"), *GetFullNameSafe(Archetype), *GetFullNameSafe(ArchetypeClass), *ArchetypeName.ToString());
+						}
+					}
 				}
 			}
 
@@ -1305,6 +1317,16 @@ void UBlueprintGeneratedClass::GetPreloadDependencies(TArray<UObject*>& OutDeps)
 			}
 		});
 	}
+
+	if (InheritableComponentHandler)
+	{
+		OutDeps.Add(InheritableComponentHandler);
+	}
+
+	if (SimpleConstructionScript)
+	{
+		OutDeps.Add(SimpleConstructionScript);
+	}
 }
 
 bool UBlueprintGeneratedClass::NeedsLoadForServer() const
diff --git a/Engine/Source/Runtime/Engine/Private/Character.cpp b/Engine/Source/Runtime/Engine/Private/Character.cpp
--- a/Engine/Source/Runtime/Engine/Private/Character.cpp
+++ b/Engine/Source/Runtime/Engine/Private/Character.cpp
@@ -130,7 +130,7 @@ void ACharacter::PostInitializeComponents()
 			CharacterMovement->UpdateNavAgent(*CapsuleComponent);
 		}
 
-		if (Controller == NULL && GetNetMode() != NM_Client)
+		if (Controller == nullptr && GetNetMode() != NM_Client)
 		{
 			if (CharacterMovement && CharacterMovement->bRunPhysicsWithNoController)
 			{
@@ -432,7 +432,7 @@ void ACharacter::ApplyDamageMomentum(float DamageTaken, FDamageEvent const& Dama
 	UDamageType const* const DmgTypeCDO = DamageEvent.DamageTypeClass->GetDefaultObject<UDamageType>();
 	float const ImpulseScale = DmgTypeCDO->DamageImpulse;
 
-	if ( (ImpulseScale > 3.f) && (CharacterMovement != NULL) )
+	if ( (ImpulseScale > 3.f) && (CharacterMovement != nullptr) )
 	{
 		FHitResult HitInfo;
 		FVector ImpulseDir;
@@ -461,9 +461,9 @@ void ACharacter::ApplyDamageMomentum(float DamageTaken, FDamageEvent const& Dama
 
 void ACharacter::ClearCrossLevelReferences()
 {
-	if( BasedMovement.MovementBase != NULL && GetOutermost() != BasedMovement.MovementBase->GetOutermost() )
+	if( BasedMovement.MovementBase != nullptr && GetOutermost() != BasedMovement.MovementBase->GetOutermost() )
 	{
-		SetBase( NULL );
+		SetBase( nullptr );
 	}
 
 	Super::ClearCrossLevelReferences();
@@ -631,7 +631,7 @@ namespace MovementBaseUtility
 			return true;
 		}
 
-		// NULL MovementBase
+		// nullptr MovementBase
 		OutLocation = FVector::ZeroVector;
 		OutQuat = FQuat::Identity;
 		return false;
@@ -642,7 +642,7 @@ namespace MovementBaseUtility
 /**	Change the Pawn's base. */
 void ACharacter::SetBase( UPrimitiveComponent* NewBaseComponent, const FName InBoneName, bool bNotifyPawn )
 {
-	// If NewBaseComponent is null, ignore bone name.
+	// If NewBaseComponent is nullptr, ignore bone name.
 	const FName BoneName = (NewBaseComponent ? InBoneName : NAME_None);
 
 	// See what changed.
@@ -652,7 +652,7 @@ void ACharacter::SetBase( UPrimitiveComponent* NewBaseComponent, const FName InB
 	if (bBaseChanged || bBoneChanged)
 	{
 		// Verify no recursion.
-		APawn* Loop = (NewBaseComponent ? Cast<APawn>(NewBaseComponent->GetOwner()) : NULL);
+		APawn* Loop = (NewBaseComponent ? Cast<APawn>(NewBaseComponent->GetOwner()) : nullptr);
 		while (Loop)
 		{
 			if (Loop == this)
@@ -711,7 +711,7 @@ void ACharacter::SetBase( UPrimitiveComponent* NewBaseComponent, const FName InB
 
 			if (Role == ROLE_Authority || Role == ROLE_AutonomousProxy)
 			{
-				BasedMovement.bServerHasBaseComponent = (BasedMovement.MovementBase != NULL); // Also set on proxies for nicer debugging.
+				BasedMovement.bServerHasBaseComponent = (BasedMovement.MovementBase != nullptr); // Also set on proxies for nicer debugging.
 				UE_LOG(LogCharacter, Verbose, TEXT("Setting base on %s for '%s' to '%s'"), Role == ROLE_Authority ? TEXT("Server") : TEXT("AutoProxy"), *GetName(), *GetFullNameSafe(NewBaseComponent));
 			}
 			else
@@ -757,13 +757,13 @@ FVector ACharacter::GetNavAgentLocation() const
 
 void ACharacter::TurnOff()
 {
-	if (CharacterMovement != NULL)
+	if (CharacterMovement != nullptr)
 	{
 		CharacterMovement->StopMovementImmediately();
 		CharacterMovement->DisableMovement();
 	}
 
-	if (GetNetMode() != NM_DedicatedServer && Mesh != NULL)
+	if (GetNetMode() != NM_DedicatedServer && Mesh != nullptr)
 	{
 		Mesh->bPauseAnims = true;
 		if (Mesh->IsSimulatingPhysics())
@@ -794,7 +794,7 @@ void ACharacter::Restart()
 
 void ACharacter::PawnClientRestart()
 {
-	if (CharacterMovement != NULL)
+	if (CharacterMovement != nullptr)
 	{
 		CharacterMovement->StopMovementImmediately();
 		CharacterMovement->ResetPredictionData_Client();
@@ -855,7 +855,7 @@ void ACharacter::BaseChange()
 	if (CharacterMovement && CharacterMovement->MovementMode != MOVE_None)
 	{
 		AActor* ActualMovementBase = GetMovementBaseActor(this);
-		if ((ActualMovementBase != NULL) && !ActualMovementBase->CanBeBaseForCharacter(this))
+		if ((ActualMovementBase != nullptr) && !ActualMovementBase->CanBeBaseForCharacter(this))
 		{
 			CharacterMovement->JumpOff(ActualMovementBase);
 		}
@@ -874,7 +874,7 @@ void ACharacter::DisplayDebug(UCanvas* Canvas, const FDebugDisplayInfo& DebugDis
 		FIndenter PhysicsIndent(Indent);
 
 		FString BaseString;
-		if ( CharacterMovement == NULL || BasedMovement.MovementBase == NULL )
+		if ( CharacterMovement == nullptr || BasedMovement.MovementBase == nullptr )
 		{
 			BaseString = "Not Based";
 		}
@@ -887,7 +887,7 @@ void ACharacter::DisplayDebug(UCanvas* Canvas, const FDebugDisplayInfo& DebugDis
 		FDisplayDebugManager& DisplayDebugManager = Canvas->DisplayDebugManager;
 		DisplayDebugManager.DrawString(FString::Printf(TEXT("RelativeLoc: %s Rot: %s %s"), *BasedMovement.Location.ToCompactString(), *BasedMovement.Rotation.ToCompactString(), *BaseString), Indent);
 
-		if ( CharacterMovement != NULL )
+		if ( CharacterMovement != nullptr )
 		{
 			CharacterMovement->DisplayDebug(Canvas, DebugDisplay, YL, YPos);
 		}
@@ -1324,35 +1324,6 @@ void ACharacter::OnUpdateSimulatedPosition(const FVector& OldLocation, const FQu
 	CharacterMovement->bJustTeleported = true;
 }
 
-// Deprecated, remove
-void ACharacter::UpdateSimulatedPosition(const FVector& NewLocation, const FRotator& NewRotation)
-{
-	// Always consider Location as changed if we were spawned this tick as in that case our replicated Location was set as part of spawning, before PreNetReceive()
-	if( (NewLocation != GetActorLocation()) || (CreationTime == GetWorld()->TimeSeconds) )
-	{
-		FVector FinalLocation = NewLocation;
-
-		// Only need to check for encroachment when teleported without any velocity.
-		// Normal movement pops the character out of geometry anyway, no use doing it before and after (with different rules).
-		bSimGravityDisabled = false;
-		if (CharacterMovement->Velocity.IsZero())
-		{
-			if (GetWorld()->EncroachingBlockingGeometry(this, NewLocation, NewRotation))
-			{
-				bSimGravityDisabled = true;
-			}
-		}
-		
-		// Don't use TeleportTo(), that clears our base.
-		SetActorLocationAndRotation(FinalLocation, NewRotation, false);
-		CharacterMovement->bJustTeleported = true;
-	}
-	else if( NewRotation != GetActorRotation() )
-	{
-		GetRootComponent()->MoveComponent(FVector::ZeroVector, NewRotation, false);
-	}
-}
-
 void ACharacter::PostNetReceiveLocationAndRotation()
 {
 	if( Role == ROLE_SimulatedProxy )
@@ -1511,7 +1482,7 @@ float ACharacter::GetAnimRootMotionTranslationScale() const
 
 float ACharacter::PlayAnimMontage(class UAnimMontage* AnimMontage, float InPlayRate, FName StartSectionName)
 {
-	UAnimInstance * AnimInstance = (Mesh)? Mesh->GetAnimInstance() : NULL; 
+	UAnimInstance * AnimInstance = (Mesh)? Mesh->GetAnimInstance() : nullptr; 
 	if( AnimMontage && AnimInstance )
 	{
 		float const Duration = AnimInstance->Montage_Play(AnimMontage, InPlayRate);
@@ -1533,7 +1504,7 @@ float ACharacter::PlayAnimMontage(class UAnimMontage* AnimMontage, float InPlayR
 
 void ACharacter::StopAnimMontage(class UAnimMontage* AnimMontage)
 {
-	UAnimInstance * AnimInstance = (Mesh)? Mesh->GetAnimInstance() : NULL; 
+	UAnimInstance * AnimInstance = (Mesh)? Mesh->GetAnimInstance() : nullptr; 
 	UAnimMontage * MontageToStop = (AnimMontage)? AnimMontage : GetCurrentMontage();
 	bool bShouldStopMontage =  AnimInstance && MontageToStop && !AnimInstance->Montage_GetIsStopped(MontageToStop);
 
@@ -1545,13 +1516,13 @@ void ACharacter::StopAnimMontage(class UAnimMontage* AnimMontage)
 
 class UAnimMontage * ACharacter::GetCurrentMontage()
 {
-	UAnimInstance * AnimInstance = (Mesh)? Mesh->GetAnimInstance() : NULL; 
+	UAnimInstance * AnimInstance = (Mesh)? Mesh->GetAnimInstance() : nullptr; 
 	if ( AnimInstance )
 	{
 		return AnimInstance->GetCurrentActiveMontage();
 	}
 
-	return NULL;
+	return nullptr;
 }
 
 void ACharacter::ClientCheatWalk_Implementation()
diff --git a/Engine/Source/Runtime/Engine/Private/Components/ActorComponent.cpp b/Engine/Source/Runtime/Engine/Private/Components/ActorComponent.cpp
--- a/Engine/Source/Runtime/Engine/Private/Components/ActorComponent.cpp
+++ b/Engine/Source/Runtime/Engine/Private/Components/ActorComponent.cpp
@@ -1,4 +1,4 @@
-﻿// Copyright 1998-2017 Epic Games, Inc. All Rights Reserved.
+// Copyright 1998-2017 Epic Games, Inc. All Rights Reserved.
 // ActorComponent.cpp: Actor component implementation.
 
 #include "Components/ActorComponent.h"
@@ -1543,7 +1543,7 @@ void UActorComponent::SetTickableWhenPaused(bool bTickableWhenPaused)
 bool UActorComponent::IsOwnerRunningUserConstructionScript() const
 {
 	AActor* MyOwner = GetOwner();
-	return (MyOwner && MyOwner->bRunningUserConstructionScript);
+	return (MyOwner && MyOwner->IsRunningUserConstructionScript());
 }
 
 void UActorComponent::AddAssetUserData(UAssetUserData* InUserData)
diff --git a/Engine/Source/Runtime/Engine/Private/Components/DecalComponent.cpp b/Engine/Source/Runtime/Engine/Private/Components/DecalComponent.cpp
--- a/Engine/Source/Runtime/Engine/Private/Components/DecalComponent.cpp
+++ b/Engine/Source/Runtime/Engine/Private/Components/DecalComponent.cpp
@@ -134,11 +134,12 @@ void UDecalComponent::LifeSpanCallback()
 {
 	DestroyComponent();
 
-	auto* Owner = GetOwner();
-
-	if (bDestroyOwnerAfterFade && Owner && (FadeDuration > 0.0f || FadeStartDelay > 0.0f))
+	if (bDestroyOwnerAfterFade  && (FadeDuration > 0.0f || FadeStartDelay > 0.0f))
 	{
-		Owner->Destroy();
+		if (AActor* Owner = GetOwner())
+		{
+			Owner->Destroy();
+		}
 	}
 }
 
diff --git a/Engine/Source/Runtime/Engine/Private/Components/LightComponent.cpp b/Engine/Source/Runtime/Engine/Private/Components/LightComponent.cpp
--- a/Engine/Source/Runtime/Engine/Private/Components/LightComponent.cpp
+++ b/Engine/Source/Runtime/Engine/Private/Components/LightComponent.cpp
@@ -617,7 +617,8 @@ void ULightComponent::UpdateLightSpriteTexture()
 			bAffectsWorld &&
 			CastShadows &&
 			CastStaticShadows &&
-			PreviewShadowMapChannel == INDEX_NONE)
+			PreviewShadowMapChannel == INDEX_NONE &&
+			(GetWorld() && !GetWorld()->IsPreviewWorld()))
 		{
 			UTexture2D* SpriteTexture = NULL;
 			SpriteTexture = LoadObject<UTexture2D>(NULL, TEXT("/Engine/EditorResources/LightIcons/S_LightError.S_LightError"));
diff --git a/Engine/Source/Runtime/Engine/Private/Components/SkeletalMeshComponent.cpp b/Engine/Source/Runtime/Engine/Private/Components/SkeletalMeshComponent.cpp
--- a/Engine/Source/Runtime/Engine/Private/Components/SkeletalMeshComponent.cpp
+++ b/Engine/Source/Runtime/Engine/Private/Components/SkeletalMeshComponent.cpp
@@ -35,6 +35,8 @@
 #include "Features/IModularFeatures.h"
 #include "Misc/RuntimeErrors.h"
 #include "AnimPhysObjectVersion.h"
+#include "ScopeExit.h"
+
 #define LOCTEXT_NAMESPACE "SkeletalMeshComponent"
 
 TAutoConsoleVariable<int32> CVarUseParallelAnimationEvaluation(TEXT("a.ParallelAnimEvaluation"), 1, TEXT("If 1, animation evaluation will be run across the task graph system. If 0, evaluation will run purely on the game thread"));
@@ -187,13 +189,14 @@ USkeletalMeshComponent::USkeletalMeshComponent(const FObjectInitializer& ObjectI
 
 #endif//#if WITH_APEX_CLOTHING
 
+#if WITH_EDITORONLY_DATA
 	DefaultPlayRate_DEPRECATED = 1.0f;
 	bDefaultPlaying_DEPRECATED = true;
+#endif
 	bEnablePhysicsOnDedicatedServer = UPhysicsSettings::Get()->bSimulateSkeletalMeshOnDedicatedServer;
 	bEnableUpdateRateOptimizations = false;
 	RagdollAggregateThreshold = UPhysicsSettings::Get()->RagdollAggregateThreshold;
 
-	// LastPoseTickTime = -1;
 	LastPoseTickFrame = 0u;
 
 	bHasCustomNavigableGeometry = EHasCustomNavigableGeometry::Yes;
@@ -1994,7 +1997,14 @@ void USkeletalMeshComponent::PostAnimEvaluation(FAnimationEvaluationContext& Eva
 {
 #if DO_CHECK
 	checkf(!bPostEvaluatingAnimation, TEXT("PostAnimEvaluation already in progress, recursion detected for SkeletalMeshComponent [%s], AnimInstance [%s]"), *GetNameSafe(this), *GetNameSafe(EvaluationContext.AnimInstance));
-	TGuardValue<bool> CircularGuard(bPostEvaluatingAnimation, true);
+
+	const bool bWasPostEvaluatingAnimation = bPostEvaluatingAnimation;
+	bPostEvaluatingAnimation = true;
+	ON_SCOPE_EXIT
+	{
+		bPostEvaluatingAnimation = bWasPostEvaluatingAnimation;
+	};
+
 #endif
 
 	SCOPE_CYCLE_COUNTER(STAT_PostAnimEvaluation);
diff --git a/Engine/Source/Runtime/Engine/Private/Curves/RichCurve.cpp b/Engine/Source/Runtime/Engine/Private/Curves/RichCurve.cpp
--- a/Engine/Source/Runtime/Engine/Private/Curves/RichCurve.cpp
+++ b/Engine/Source/Runtime/Engine/Private/Curves/RichCurve.cpp
@@ -847,12 +847,23 @@ void FRichCurve::BakeCurve(float SampleRate, float FirstKeyTime, float LastKeyTi
 		return;
 	}
 
-	for (float Time = FirstKeyTime + SampleRate; Time < LastKeyTime;  )
+	// we need to generate new keys first rather than modifying the
+	// curve directly since that would affect the results of Eval calls
+	TArray<TPair<float, float> > BakedKeys;
+	BakedKeys.Reserve(((LastKeyTime - FirstKeyTime) / SampleRate) - 1);
+
+	// the skip the first and last key unchanged
+	for (float Time = FirstKeyTime + SampleRate; Time < LastKeyTime; )
 	{
-		float Value = Eval(Time);
-		UpdateOrAddKey(Time, Value);
+		const float Value = Eval(Time);
+		BakedKeys.Add(TPair<float, float>(Time, Value));
 		Time += SampleRate;
 	}
+
+	for (auto NewKey : BakedKeys)
+	{
+		UpdateOrAddKey(NewKey.Key, NewKey.Value);
+	}
 }
 
 void FRichCurve::RemoveRedundantKeys(float Tolerance)
