Copying //UE4/Dev-Framework to //UE4/Dev-Main (Source: //UE4/Dev-Framework @ 3716594)

diff --git a/Engine/Source/Runtime/Engine/Private/DataAsset.cpp b/Engine/Source/Runtime/Engine/Private/DataAsset.cpp
--- a/Engine/Source/Runtime/Engine/Private/DataAsset.cpp
+++ b/Engine/Source/Runtime/Engine/Private/DataAsset.cpp
@@ -50,23 +50,49 @@ FPrimaryAssetId UPrimaryDataAsset::GetPrimaryAssetId() const
 {
 	if (HasAnyFlags(RF_ClassDefaultObject))
 	{
-		UClass* SearchNativeClass = GetClass();
+		UClass* BestPrimaryAssetTypeClass = nullptr;
+		UClass* SearchPrimaryAssetTypeClass = GetClass()->GetSuperClass();
 
-		while (SearchNativeClass && !SearchNativeClass->HasAnyClassFlags(CLASS_Native | CLASS_Intrinsic))
+		// If this is a native class or immediate child of PrimaryDataAsset, return invalid as we are a type ourselves
+		if (GetClass()->HasAnyClassFlags(CLASS_Native | CLASS_Intrinsic) || SearchPrimaryAssetTypeClass == UPrimaryDataAsset::StaticClass())
 		{
-			SearchNativeClass = SearchNativeClass->GetSuperClass();
+			return FPrimaryAssetId();
 		}
 
-		if (SearchNativeClass && SearchNativeClass != GetClass())
+		// Starting with parent, look up the hierarchy for a class that is either native, or a blueprint class immediately below PrimaryDataAsset
+		while (SearchPrimaryAssetTypeClass)
 		{
-			// If blueprint, return native class and asset name
-			return FPrimaryAssetId(SearchNativeClass->GetFName(), FPackageName::GetShortFName(GetOutermost()->GetFName()));
+			if (SearchPrimaryAssetTypeClass->GetSuperClass() == UPrimaryDataAsset::StaticClass())
+			{
+				// If our parent is this base class, return this as the best class
+				BestPrimaryAssetTypeClass = SearchPrimaryAssetTypeClass;
+				break;
+			}
+			else if (SearchPrimaryAssetTypeClass->HasAnyClassFlags(CLASS_Native | CLASS_Intrinsic))
+			{
+				// Our parent is the first native class found, use that
+				BestPrimaryAssetTypeClass = SearchPrimaryAssetTypeClass;
+				break;
+			}
+			else
+			{
+				SearchPrimaryAssetTypeClass = SearchPrimaryAssetTypeClass->GetSuperClass();
+			}
 		}
-		
-		// Native CDO, return nothing
+
+		if (BestPrimaryAssetTypeClass)
+		{
+			// If this is a native class use the raw name if it's a blueprint use the package name as it will be missing _C
+			FName PrimaryAssetType = BestPrimaryAssetTypeClass->HasAnyClassFlags(CLASS_Native | CLASS_Intrinsic) ? BestPrimaryAssetTypeClass->GetFName() : FPackageName::GetShortFName(BestPrimaryAssetTypeClass->GetOutermost()->GetFName());
+
+			return FPrimaryAssetId(PrimaryAssetType, FPackageName::GetShortFName(GetOutermost()->GetFName()));
+		}
+
+		// No valid parent class found, return invalid
 		return FPrimaryAssetId();
 	}
 
+
 	// Data assets use Class and ShortName by default, there's no inheritance so class works fine
 	return FPrimaryAssetId(GetClass()->GetFName(), GetFName());
 }
diff --git a/Engine/Source/Runtime/Engine/Private/EdGraph/EdGraphNode.cpp b/Engine/Source/Runtime/Engine/Private/EdGraph/EdGraphNode.cpp
--- a/Engine/Source/Runtime/Engine/Private/EdGraph/EdGraphNode.cpp
+++ b/Engine/Source/Runtime/Engine/Private/EdGraph/EdGraphNode.cpp
@@ -2,6 +2,7 @@
 
 #include "EdGraph/EdGraphNode.h"
 #include "UObject/BlueprintsObjectVersion.h"
+#include "UObject/FrameworkObjectVersion.h"
 #include "EdGraph/EdGraphPin.h"
 #include "Textures/SlateIcon.h"
 #include "EdGraph/EdGraph.h"
@@ -30,8 +31,23 @@ FEdGraphTerminalType FEdGraphTerminalType::FromPinType(const FEdGraphPinType& Pi
 
 FArchive& operator<<(FArchive& Ar, FEdGraphTerminalType& T)
 {
+	Ar.UsingCustomVersion(FFrameworkObjectVersion::GUID);
+
+	if (Ar.CustomVer(FFrameworkObjectVersion::GUID) >= FFrameworkObjectVersion::PinsStoreFName)
+	{
 	Ar << T.TerminalCategory;
 	Ar << T.TerminalSubCategory;
+	}
+	else
+	{
+		FString TerminalCategoryStr;
+		Ar << TerminalCategoryStr;
+		T.TerminalCategory = *TerminalCategoryStr;
+
+		FString TerminalSubCategoryStr;
+		Ar << TerminalSubCategoryStr;
+		T.TerminalSubCategory = *TerminalSubCategoryStr;
+	}
 
 	// See: FArchive& operator<<( FArchive& Ar, FWeakObjectPtr& WeakObjectPtr )
 	// The PinSubCategoryObject should be serialized into the package.
@@ -53,11 +69,22 @@ FArchive& operator<<(FArchive& Ar, FEdGraphTerminalType& T)
 
 FName const FNodeMetadata::DefaultGraphNode(TEXT("DefaultGraphNode"));
 
+#if WITH_EDITOR
+UEdGraphNode::FCreatePinParams::FCreatePinParams(const FEdGraphPinType& PinType)
+	: ContainerType(PinType.ContainerType)
+	, bIsReference(PinType.bIsReference)
+	, bIsConst(PinType.bIsConst)
+	, Index(INDEX_NONE)
+	, ValueTerminalType(PinType.PinValueType)
+{
+}
+#endif // WITH_EDITOR
+
 /////////////////////////////////////////////////////
 // FGraphNodeContextMenuBuilder
 
 FGraphNodeContextMenuBuilder::FGraphNodeContextMenuBuilder(const UEdGraph* InGraph, const UEdGraphNode* InNode, const UEdGraphPin* InPin, FMenuBuilder* InMenuBuilder, bool bInDebuggingMode)
-	: Blueprint(NULL)
+	: Blueprint(nullptr)
 	, Graph(InGraph)
 	, Node(InNode)
 	, Pin(InPin)
@@ -68,7 +95,7 @@ FGraphNodeContextMenuBuilder::FGraphNodeContextMenuBuilder(const UEdGraph* InGra
 	Blueprint = FBlueprintEditorUtils::FindBlueprintForGraph(Graph);
 #endif
 
-	if (Pin != NULL)
+	if (Pin)
 	{
 		Node = Pin->GetOwningNode();
 	}
@@ -151,7 +178,7 @@ void UEdGraphNode::DiffProperties(UClass* StructA, UClass* StructB, UObject* Dat
 	}
 }
 
-UEdGraphPin* UEdGraphNode::CreatePin(EEdGraphPinDirection Dir, const FEdGraphPinType& InPinType, const FString& PinName, int32 Index /*= INDEX_NONE*/)
+UEdGraphPin* UEdGraphNode::CreatePin(EEdGraphPinDirection Dir, const FEdGraphPinType& InPinType, const FName PinName, int32 Index /*= INDEX_NONE*/)
 {
 	UEdGraphPin* NewPin = UEdGraphPin::CreatePin(this);
 	NewPin->PinName = PinName;
@@ -171,56 +198,70 @@ UEdGraphPin* UEdGraphNode::CreatePin(EEdGraphPinDirection Dir, const FEdGraphPin
 	return NewPin;
 }
 
-UEdGraphPin* UEdGraphNode::CreatePin(EEdGraphPinDirection Dir, const FString& PinCategory, const FString& PinSubCategory, UObject* PinSubCategoryObject, bool bIsArray, bool bIsReference, const FString& PinName, bool bIsConst /*= false*/, int32 Index /*= INDEX_NONE*/, bool bIsSet /*= false*/, bool bIsMap /*= false*/, const FEdGraphTerminalType& ValueTerminalType /*= FEdGraphTerminalType()*/)
+UEdGraphPin* UEdGraphNode::CreatePin(EEdGraphPinDirection Dir, const FNameParameterHelper PinCategory, const FNameParameterHelper PinSubCategory, UObject* PinSubCategoryObject, bool bIsArray, bool bIsReference, const FNameParameterHelper PinName, bool bIsConst /*= false*/, int32 Index /*= INDEX_NONE*/, bool bIsSet /*= false*/, bool bIsMap /*= false*/, const FEdGraphTerminalType& ValueTerminalType /*= FEdGraphTerminalType()*/)
 {
+	PRAGMA_DISABLE_DEPRECATION_WARNINGS
 	return CreatePin(Dir, PinCategory, PinSubCategory, PinSubCategoryObject, PinName, FEdGraphPinType::ToPinContainerType(bIsArray, bIsSet, bIsMap), bIsReference, bIsConst, Index, ValueTerminalType);
+	PRAGMA_ENABLE_DEPRECATION_WARNINGS
 }
 
-UEdGraphPin* UEdGraphNode::CreatePin(EEdGraphPinDirection Dir, const FString& PinCategory, const FString& PinSubCategory, UObject* PinSubCategoryObject, const FString& PinName, EPinContainerType PinContainerType /* EPinContainerType::None */, bool bIsReference /* = false */, bool bIsConst /*= false*/, int32 Index /*= INDEX_NONE*/, const FEdGraphTerminalType& ValueTerminalType /*= FEdGraphTerminalType()*/)
+UEdGraphPin* UEdGraphNode::CreatePin(EEdGraphPinDirection Dir, const FNameParameterHelper PinCategory, const FNameParameterHelper PinSubCategory, UObject* PinSubCategoryObject, const FNameParameterHelper PinName, EPinContainerType PinContainerType /* EPinContainerType::None */, bool bIsReference /* = false */, bool bIsConst /*= false*/, int32 Index /*= INDEX_NONE*/, const FEdGraphTerminalType& ValueTerminalType /*= FEdGraphTerminalType()*/)
 {
-	FEdGraphPinType PinType(PinCategory, PinSubCategory, PinSubCategoryObject, PinContainerType, bIsReference, ValueTerminalType);
-	PinType.bIsConst = bIsConst;
+	FCreatePinParams PinParams;
+	PinParams.ContainerType = PinContainerType;
+	PinParams.bIsConst = bIsConst;
+	PinParams.bIsReference = bIsReference;
+	PinParams.Index = Index;
+	PinParams.ValueTerminalType = ValueTerminalType;
 
-	return CreatePin(Dir, PinType, PinName, Index);
+	return CreatePin(Dir, *PinCategory, *PinSubCategory, PinSubCategoryObject, *PinName, PinParams);
 }
 
-UEdGraphPin* UEdGraphNode::FindPin(const FString& PinName, const EEdGraphPinDirection Direction) const
+
+UEdGraphPin* UEdGraphNode::CreatePin(const EEdGraphPinDirection Dir, const FName PinCategory, const FName PinSubCategory, UObject* PinSubCategoryObject, const FName PinName, const FCreatePinParams& PinParams)
 {
-	for(int32 PinIdx=0; PinIdx<Pins.Num(); PinIdx++)
+	FEdGraphPinType PinType(PinCategory, PinSubCategory, PinSubCategoryObject, PinParams.ContainerType, PinParams.bIsReference, PinParams.ValueTerminalType);
+	PinType.bIsConst = PinParams.bIsConst;
+
+	return CreatePin(Dir, PinType, PinName, PinParams.Index);
+}
+
+UEdGraphPin* UEdGraphNode::FindPin(const FName PinName, const EEdGraphPinDirection Direction) const
+{
+	for (UEdGraphPin* Pin : Pins)
 	{
-		if ((Direction == EGPD_MAX || Direction == Pins[PinIdx]->Direction) && (Pins[PinIdx]->PinName == PinName))
+		if ((Direction == EGPD_MAX || Direction == Pin->Direction) && Pin->PinName == PinName)
 		{
-			return Pins[PinIdx];
+			return Pin;
 		}
 	}
 
-	return NULL;
+	return nullptr;
 }
 
-UEdGraphPin* UEdGraphNode::FindPinChecked(const FString& PinName, const EEdGraphPinDirection Direction) const
+UEdGraphPin* UEdGraphNode::FindPin(const TCHAR* const PinName, const EEdGraphPinDirection Direction) const
 {
-	UEdGraphPin* Result = FindPin(PinName, Direction);
-	check(Result != NULL);
-	return Result;
+	const FName PinFName(PinName, FNAME_Find);
+	return (!PinFName.IsNone() ? FindPin(PinFName, Direction) : nullptr);
 }
 
 UEdGraphPin* UEdGraphNode::FindPinById(const FGuid PinId) const
 {
-	for (int32 PinIdx = 0; PinIdx < Pins.Num(); PinIdx++)
+	for (UEdGraphPin* Pin : Pins)
 	{
-		if (Pins[PinIdx]->PinId == PinId)
+		if (Pin->PinId == PinId)
 		{
-			return Pins[PinIdx];
+			return Pin;
 		}
 	}
 
-	return NULL;
+	return nullptr;
 }
 
 UEdGraphPin* UEdGraphNode::FindPinByIdChecked(const FGuid PinId) const
 {
 	UEdGraphPin* Result = FindPinById(PinId);
-	check(Result != NULL);
+	check(Result);
 	return Result;
 }
 
diff --git a/Engine/Source/Runtime/Engine/Private/EdGraph/EdGraphPin.cpp b/Engine/Source/Runtime/Engine/Private/EdGraph/EdGraphPin.cpp
--- a/Engine/Source/Runtime/Engine/Private/EdGraph/EdGraphPin.cpp
+++ b/Engine/Source/Runtime/Engine/Private/EdGraph/EdGraphPin.cpp
@@ -1,4 +1,4 @@
-// Copyright 1998-2017 Epic Games, Inc. All Rights Reserved.
+﻿// Copyright 1998-2017 Epic Games, Inc. All Rights Reserved.
 
 #include "EdGraph/EdGraphPin.h"
 #include "UObject/BlueprintsObjectVersion.h"
@@ -98,7 +98,23 @@ bool FEdGraphPinType::Serialize(FArchive& Ar)
 		return false;
 	}
 
-	Ar << PinCategory;
+	Ar.UsingCustomVersion(FFrameworkObjectVersion::GUID);
+
+	if (Ar.CustomVer(FFrameworkObjectVersion::GUID) >= FFrameworkObjectVersion::PinsStoreFName)
+	{
+		Ar << PinCategory;
+		Ar << PinSubCategory;
+	}
+	else
+	{
+		FString PinCategoryStr;
+		Ar << PinCategoryStr;
+		PinCategory = *PinCategoryStr;
+
+		FString PinSubCategoryStr;
+		Ar << PinSubCategoryStr;
+		PinSubCategory = *PinSubCategoryStr;
+	}
 
 	if (Ar.UE4Ver() < VER_UE4_ADDED_SOFT_OBJECT_PATH)
 	{
@@ -113,8 +129,6 @@ bool FEdGraphPinType::Serialize(FArchive& Ar)
 		}
 	}
 
-	Ar << PinSubCategory;
-
 	// See: FArchive& operator<<( FArchive& Ar, FWeakObjectPtr& WeakObjectPtr )
 	// The PinSubCategoryObject should be serialized into the package.
 	if(!Ar.IsObjectReferenceCollector() || Ar.IsModifyingWeakAndStrongReferences() || Ar.IsPersistent())
@@ -127,7 +141,6 @@ bool FEdGraphPinType::Serialize(FArchive& Ar)
 		}
 	}
 
-	Ar.UsingCustomVersion(FFrameworkObjectVersion::GUID);
 	if (Ar.CustomVer(FFrameworkObjectVersion::GUID) >= FFrameworkObjectVersion::EdGraphPinContainerType)
 	{
 		Ar << ContainerType;
@@ -544,25 +557,25 @@ void UEdGraphPin::AssignByRefPassThroughConnection(UEdGraphPin* InTargetPin)
 	{
 		if (GetOwningNode() != InTargetPin->GetOwningNode())
 		{
-			UE_LOG(LogBlueprint, Warning, TEXT("Pin '%s' is owned by node '%s' and pin '%s' is owned by '%s', they must be owned by the same node!"), *PinName, *GetOwningNode()->GetName(), *InTargetPin->PinName, *InTargetPin->GetOwningNode()->GetName());
+			UE_LOG(LogBlueprint, Warning, TEXT("Pin '%s' is owned by node '%s' and pin '%s' is owned by '%s', they must be owned by the same node!"), *GetName(), *GetOwningNode()->GetName(), *InTargetPin->GetName(), *InTargetPin->GetOwningNode()->GetName());
 		}
 		else if (Direction == InTargetPin->Direction)
 		{
-			UE_LOG(LogBlueprint, Warning, TEXT("Both pin '%s' and pin '%s' on node '%s' go in the same direction, one must be an input and the other an output!"), *PinName, *InTargetPin->PinName, *GetOwningNode()->GetName());
+			UE_LOG(LogBlueprint, Warning, TEXT("Both pin '%s' and pin '%s' on node '%s' go in the same direction, one must be an input and the other an output!"), *GetName(), *InTargetPin->GetName(), *GetOwningNode()->GetName());
 		}
 		else if (!PinType.bIsReference && !InTargetPin->PinType.bIsReference)
 		{
 			UEdGraphPin* InputPin = (Direction == EGPD_Input) ? this : InTargetPin;
 			UEdGraphPin* OutputPin = (Direction == EGPD_Input) ? InTargetPin : this;
-			UE_LOG(LogBlueprint, Warning, TEXT("Input pin '%s' is attempting to by-ref pass-through to output pin '%s' on node '%s', however neither pin is by-ref!"), *InputPin->PinName, *OutputPin->PinName, *InputPin->GetOwningNode()->GetName());
+			UE_LOG(LogBlueprint, Warning, TEXT("Input pin '%s' is attempting to by-ref pass-through to output pin '%s' on node '%s', however neither pin is by-ref!"), *InputPin->GetName(), *OutputPin->GetName(), *InputPin->GetOwningNode()->GetName());
 		}
 		else if (!PinType.bIsReference)
 		{
-			UE_LOG(LogBlueprint, Warning, TEXT("Pin '%s' on node '%s' is not by-ref but it is attempting to pass-through to '%s'"), *PinName, *GetOwningNode()->GetName(), *InTargetPin->PinName);
+			UE_LOG(LogBlueprint, Warning, TEXT("Pin '%s' on node '%s' is not by-ref but it is attempting to pass-through to '%s'"), *GetName(), *GetOwningNode()->GetName(), *InTargetPin->GetName());
 		}
 		else if (!InTargetPin->PinType.bIsReference)
 		{
-			UE_LOG(LogBlueprint, Warning, TEXT("Pin '%s' on node '%s' is not by-ref but it is attempting to pass-through to '%s'"), *InTargetPin->PinName, *InTargetPin->GetOwningNode()->GetName(), *PinName);
+			UE_LOG(LogBlueprint, Warning, TEXT("Pin '%s' on node '%s' is not by-ref but it is attempting to pass-through to '%s'"), *InTargetPin->GetName(), *InTargetPin->GetOwningNode()->GetName(), *GetName());
 		}
 		else
 		{
@@ -618,7 +631,7 @@ FText UEdGraphPin::GetDisplayName() const
 	}
 	else
 	{
-		DisplayName = (!PinFriendlyName.IsEmpty()) ? PinFriendlyName : FText::FromString(PinName);
+		DisplayName = (!PinFriendlyName.IsEmpty()) ? PinFriendlyName : FText::FromName(PinName);
 
 		bool bShowNodesAndPinsUnlocalized = false;
 		GConfig->GetBool( TEXT("Internationalization"), TEXT("ShowNodesAndPinsUnlocalized"), bShowNodesAndPinsUnlocalized, GEditorSettingsIni );
@@ -634,11 +647,11 @@ FText UEdGraphPin::GetDisplayName() const
 const FString UEdGraphPin::GetLinkInfoString( const FString& InFunctionName, const FString& InInfoData, const UEdGraphPin* InToPin ) const
 {
 #if WITH_EDITOR
-	const FString FromPinName = PinName;
+	const FString FromPinName = GetName();
 	const UEdGraphNode* FromPinNode = GetOwningNodeUnchecked();
 	const FString FromPinNodeName = (FromPinNode != nullptr) ? FromPinNode->GetNodeTitle(ENodeTitleType::ListView).ToString() : TEXT("Unknown");
 	
-	const FString ToPinName = InToPin->PinName;
+	const FString ToPinName = InToPin->GetName();
 	const UEdGraphNode* ToPinNode = InToPin->GetOwningNodeUnchecked();
 	const FString ToPinNodeName = (ToPinNode != nullptr) ? ToPinNode->GetNodeTitle(ENodeTitleType::ListView).ToString() : TEXT("Unknown");
 	const FString LinkInfo = FString::Printf( TEXT("UEdGraphPin::%s Pin '%s' on node '%s' %s '%s' on node '%s'"), *InFunctionName, *ToPinName, *ToPinNodeName, *InInfoData, *FromPinName, *FromPinNodeName);
@@ -740,6 +753,7 @@ UEdGraphPin* UEdGraphPin::FindPinCreatedFromDeprecatedPin(UEdGraphPin_Deprecated
 
 bool UEdGraphPin::ExportTextItem(FString& ValueStr, int32 PortFlags) const
 {
+	const UNameProperty* NamePropCDO = GetDefault<UNameProperty>();
 	const UStrProperty* StrPropCDO = GetDefault<UStrProperty>();
 	const UTextProperty* TextPropCDO = GetDefault<UTextProperty>();
 	const UBoolProperty* BoolPropCDO = GetDefault<UBoolProperty>();
@@ -754,7 +768,7 @@ bool UEdGraphPin::ExportTextItem(FString& ValueStr, int32 PortFlags) const
 	if (PinName != DefaultPin.PinName)
 	{
 		ValueStr += PinHelpers::PinNameName + TEXT("=");
-		StrPropCDO->ExportTextItem(ValueStr, &PinName, nullptr, nullptr, PortFlags, nullptr);
+		NamePropCDO->ExportTextItem(ValueStr, &PinName, nullptr, nullptr, PortFlags, nullptr);
 		ValueStr += PinHelpers::ExportTextPropDelimiter;
 	}
 
@@ -914,6 +928,7 @@ bool UEdGraphPin::ImportTextItem(const TCHAR*& Buffer, int32 PortFlags, class UO
 {
 	PinId = FGuid();
 
+	const UNameProperty* NamePropCDO = GetDefault<UNameProperty>();
 	const UStrProperty* StrPropCDO = GetDefault<UStrProperty>();
 	const UTextProperty* TextPropCDO = GetDefault<UTextProperty>();
 	const UBoolProperty* BoolPropCDO = GetDefault<UBoolProperty>();
@@ -976,7 +991,7 @@ bool UEdGraphPin::ImportTextItem(const TCHAR*& Buffer, int32 PortFlags, class UO
 		}
 		else if (PropertyToken == PinHelpers::PinNameName)
 		{
-			Buffer = StrPropCDO->ImportText(Buffer, &PinName, PortFlags, Parent, ErrorText);
+			Buffer = NamePropCDO->ImportText(Buffer, &PinName, PortFlags, Parent, ErrorText);
 			bParseSuccess = (Buffer != nullptr);
 		}
 #if WITH_EDITORONLY_DATA
@@ -1014,9 +1029,8 @@ bool UEdGraphPin::ImportTextItem(const TCHAR*& Buffer, int32 PortFlags, class UO
 		}
 		else if (PropertyToken.StartsWith(PinHelpers::PinTypeName + TEXT(".")))
 		{
-			FString Left;
 			FString Right;
-			if (ensure(PropertyToken.Split(".", &Left, &Right, ESearchCase::CaseSensitive)))
+			if (ensure(PropertyToken.Split(".", nullptr, &Right, ESearchCase::CaseSensitive)))
 			{
 				UProperty* FoundProp = FEdGraphPinType::StaticStruct()->FindPropertyByName(FName(*Right));
 				if (FoundProp)
@@ -1148,7 +1162,7 @@ bool UEdGraphPin::ImportTextItem(const TCHAR*& Buffer, int32 PortFlags, class UO
 
 		if (!bParseSuccess)
 		{
-			ErrorText->Logf(ELogVerbosity::Warning, TEXT("%s: Parse error while importing property values (PinName = %s)."), *GetName(), *PinName);
+			ErrorText->Logf(ELogVerbosity::Warning, TEXT("%s: Parse error while importing property values (PinName = %s)."), *GetName(), *PinName.ToString());
 
 			// Parse error
 			return false;
@@ -1169,7 +1183,7 @@ bool UEdGraphPin::ImportTextItem(const TCHAR*& Buffer, int32 PortFlags, class UO
 
 	if (!PinId.IsValid())
 	{
-		ErrorText->Logf(ELogVerbosity::Warning, TEXT("%s: Invalid PinId after importing property values (PinName = %s)."), *GetName(), *PinName);
+		ErrorText->Logf(ELogVerbosity::Warning, TEXT("%s: Invalid PinId after importing property values (PinName = %s)."), *GetName(), *PinName.ToString());
 
 		// Did not contain a valid PinId
 		return false;
@@ -1289,7 +1303,7 @@ void UEdGraphPin::InitFromDeprecatedPin(class UEdGraphPin_Deprecated* Deprecated
 
 	OwningNode = CastChecked<UEdGraphNode>(DeprecatedPin->GetOuter());
 
-	PinName = DeprecatedPin->PinName;
+	PinName = *DeprecatedPin->PinName;
 
 #if WITH_EDITORONLY_DATA
 	PinFriendlyName = DeprecatedPin->PinFriendlyName;
@@ -1459,11 +1473,22 @@ void UEdGraphPin::DestroyImpl(bool bClearLinks)
 
 bool UEdGraphPin::Serialize(FArchive& Ar)
 {
+	Ar.UsingCustomVersion(FFrameworkObjectVersion::GUID);
+
 	// These properties are in every pin and are unlikely to be removed, so they are native serialized for speed.
 	Ar << OwningNode;
 	Ar << PinId;
 
-	Ar << PinName;
+	if (Ar.CustomVer(FFrameworkObjectVersion::GUID) >= FFrameworkObjectVersion::PinsStoreFName)
+	{
+		Ar << PinName;
+	}
+	else
+	{
+		FString PinNameStr;
+		Ar << PinNameStr;
+		PinName = *PinNameStr;
+	}
 
 #if WITH_EDITORONLY_DATA
 	if (!Ar.IsFilterEditorOnly())
@@ -1544,7 +1569,6 @@ bool UEdGraphPin::Serialize(FArchive& Ar)
 			bWasTrashed = !!(BitField & (1 << (PersistentBits + 3)));
 		}
 
-		Ar.UsingCustomVersion(FFrameworkObjectVersion::GUID);
 		if (Ar.IsLoading() && Ar.CustomVer(FFrameworkObjectVersion::GUID) < FFrameworkObjectVersion::ChangeAssetPinsToString)
 		{
 			bUseBackwardsCompatForEmptyAutogeneratedValue = true;
diff --git a/Engine/Source/Runtime/Engine/Private/EdGraph/EdGraphSchema.cpp b/Engine/Source/Runtime/Engine/Private/EdGraph/EdGraphSchema.cpp
--- a/Engine/Source/Runtime/Engine/Private/EdGraph/EdGraphSchema.cpp
+++ b/Engine/Source/Runtime/Engine/Private/EdGraph/EdGraphSchema.cpp
@@ -596,7 +596,7 @@ void UEdGraphSchema::BreakPinLinks(UEdGraphPin& TargetPin, bool bSendsNodeNotifi
 #endif	//#if WITH_EDITOR
 }
 
-void UEdGraphSchema::BreakSinglePinLink(UEdGraphPin* SourcePin, UEdGraphPin* TargetPin)
+void UEdGraphSchema::BreakSinglePinLink(UEdGraphPin* SourcePin, UEdGraphPin* TargetPin) const
 {
 	SourcePin->BreakLinkTo(TargetPin);
 
@@ -693,7 +693,7 @@ FText UEdGraphSchema::GetPinDisplayName(const UEdGraphPin* Pin) const
 	check(Pin != nullptr);
 	if (Pin->PinFriendlyName.IsEmpty())
 	{
-		ResultPinName = FText::FromString(Pin->PinName);
+		ResultPinName = FText::FromName(Pin->PinName);
 	}
 	else
 	{
diff --git a/Engine/Source/Runtime/Engine/Private/InheritableComponentHandler.cpp b/Engine/Source/Runtime/Engine/Private/InheritableComponentHandler.cpp
--- a/Engine/Source/Runtime/Engine/Private/InheritableComponentHandler.cpp
+++ b/Engine/Source/Runtime/Engine/Private/InheritableComponentHandler.cpp
@@ -79,16 +79,6 @@ void UInheritableComponentHandler::PostLoad()
 	}
 }
 
-void UInheritableComponentHandler::GetPreloadDependencies(TArray<UObject*>& OutDeps)
-{
-	Super::GetPreloadDependencies(OutDeps);
-	for (auto Record : Records)
-	{
-		OutDeps.Add(Record.ComponentTemplate);
-	}
-}
-
-
 #if WITH_EDITOR
 UActorComponent* UInheritableComponentHandler::CreateOverridenComponentTemplate(FComponentKey Key)
 {
diff --git a/Engine/Source/Runtime/Engine/Private/KismetMathLibrary.cpp b/Engine/Source/Runtime/Engine/Private/KismetMathLibrary.cpp
--- a/Engine/Source/Runtime/Engine/Private/KismetMathLibrary.cpp
+++ b/Engine/Source/Runtime/Engine/Private/KismetMathLibrary.cpp
@@ -241,7 +241,14 @@ float UKismetMathLibrary::NormalizeToRange(float Value, float RangeMin, float Ra
 {
 	if (RangeMin == RangeMax)
 	{
-		return RangeMin;
+		if (Value < RangeMin)
+		{
+			return 0.f;
+		}
+		else
+		{
+			return 1.f;
+		}
 	}
 
 	if (RangeMin > RangeMax)
@@ -303,21 +310,7 @@ void UKismetMathLibrary::MinOfByteArray(const TArray<uint8>& ByteArray, int32& I
 
 float UKismetMathLibrary::InverseLerp(float A, float B, float Value)
 {
-	if (FMath::IsNearlyEqual(A, B))
-	{
-		if (Value < A)
-		{
-			return 0;
-		}
-		else
-		{
-			return 1;
-		}
-	}
-	else
-	{
-		return ((Value - A) / (B - A));
-	}
+	return NormalizeToRange(Value, A, B);
 }
 
 float UKismetMathLibrary::Ease(float A, float B, float Alpha, TEnumAsByte<EEasingFunc::Type> EasingFunc, float BlendExp, int32 Steps)
diff --git a/Engine/Source/Runtime/Engine/Private/Level.cpp b/Engine/Source/Runtime/Engine/Private/Level.cpp
--- a/Engine/Source/Runtime/Engine/Private/Level.cpp
+++ b/Engine/Source/Runtime/Engine/Private/Level.cpp
@@ -1927,8 +1927,8 @@ void ULevel::OnLevelScriptBlueprintChanged(ULevelScriptBlueprint* InBlueprint)
 		// Get rid of the old LevelScriptActor
 		if( LevelScriptActor )
 		{
-			LevelScriptActor->Destroy();
-			LevelScriptActor = NULL;
+			LevelScriptActor->MarkPendingKill();
+			LevelScriptActor = nullptr;
 		}
 
 		check( OwningWorld );
diff --git a/Engine/Source/Runtime/Engine/Private/LevelBounds.cpp b/Engine/Source/Runtime/Engine/Private/LevelBounds.cpp
--- a/Engine/Source/Runtime/Engine/Private/LevelBounds.cpp
+++ b/Engine/Source/Runtime/Engine/Private/LevelBounds.cpp
@@ -137,9 +137,14 @@ TStatId ALevelBounds::GetStatId() const
 	RETURN_QUICK_DECLARE_CYCLE_STAT(ALevelBounds, STATGROUP_Tickables);
 }
 
+ETickableTickType ALevelBounds::GetTickableTickType() const
+{
+	return ((GIsEditor && !IsTemplate()) ? ETickableTickType::Conditional : ETickableTickType::Never);
+}
+
 bool ALevelBounds::IsTickable() const
 {
-	if (GIsEditor && bAutoUpdateBounds && !IsTemplate())
+	if (bAutoUpdateBounds)
 	{
 		UWorld* World = GetWorld();
 		return (World && World->WorldType == EWorldType::Editor);
diff --git a/Engine/Source/Runtime/Engine/Private/LevelStreaming.cpp b/Engine/Source/Runtime/Engine/Private/LevelStreaming.cpp
--- a/Engine/Source/Runtime/Engine/Private/LevelStreaming.cpp
+++ b/Engine/Source/Runtime/Engine/Private/LevelStreaming.cpp
@@ -424,7 +424,6 @@ bool ULevelStreaming::RequestLevel(UWorld* PersistentWorld, bool bAllowLevelLoad
 			UWorld* PIELevelWorld = UWorld::DuplicateWorldForPIE(NonPrefixedLevelName, PersistentWorld);
 			if (PIELevelWorld)
 			{
-				PIELevelWorld->PersistentLevel->bAlreadyMovedActors = true; // As we have duplicated the world, the actors will already have been transformed
 				check(PendingUnloadLevel == NULL);
 				SetLoadedLevel(PIELevelWorld->PersistentLevel);
 
@@ -1077,23 +1076,46 @@ ALevelScriptActor* ULevelStreaming::GetLevelScriptActor()
 	return nullptr;
 }
 
-ULevelStreamingKismet* ULevelStreamingKismet::LoadLevelInstance(UObject* WorldContextObject, const FString& LevelName, const FVector& Location, const FRotator& Rotation, bool& bOutSuccess)
-{ 
-	bOutSuccess = false; 
+ULevelStreamingKismet* ULevelStreamingKismet::LoadLevelInstance(UObject* WorldContextObject, const FString LevelName, const FVector Location, const FRotator Rotation, bool& bOutSuccess)
+{
+	bOutSuccess = false;
 	UWorld* const World = GEngine->GetWorldFromContextObject(WorldContextObject, EGetWorldErrorMode::LogAndReturnNull);
-	if (!World) 
+	if (!World)
 	{
 		return nullptr;
 	}
-	
+
 	// Check whether requested map exists, this could be very slow if LevelName is a short package name
 	FString LongPackageName;
 	bOutSuccess = FPackageName::SearchForPackageOnDisk(LevelName, &LongPackageName);
-	if (!bOutSuccess) 
+	if (!bOutSuccess)
+	{
+		return nullptr;
+	}
+
+	return LoadLevelInstance_Internal(World, LongPackageName, Location, Rotation, bOutSuccess);
+}
+
+ULevelStreamingKismet* ULevelStreamingKismet::LoadLevelInstanceBySoftObjectPtr(UObject* WorldContextObject, const TSoftObjectPtr<UWorld> Level, const FVector Location, const FRotator Rotation, bool& bOutSuccess)
+{
+	bOutSuccess = false;
+	UWorld* World = GEngine->GetWorldFromContextObject(WorldContextObject, EGetWorldErrorMode::LogAndReturnNull);
+	if (!World)
+	{
+		return nullptr;
+	}
+
+	// Check whether requested map exists, this could be very slow if LevelName is a short package name
+	if (Level.IsNull())
 	{
 		return nullptr;
 	}
-		 
+
+	return LoadLevelInstance_Internal(World, Level.ToString(), Location, Rotation, bOutSuccess);
+}
+
+ULevelStreamingKismet* ULevelStreamingKismet::LoadLevelInstance_Internal(UWorld* World, const FString& LongPackageName, const FVector Location, const FRotator Rotation, bool& bOutSuccess)
+{
     // Create Unique Name for sub-level package
 	const FString ShortPackageName = FPackageName::GetShortName(LongPackageName);
 	const FString PackagePath = FPackageName::GetLongPackagePath(LongPackageName);
diff --git a/Engine/Source/Runtime/Engine/Private/LevelTick.cpp b/Engine/Source/Runtime/Engine/Private/LevelTick.cpp
--- a/Engine/Source/Runtime/Engine/Private/LevelTick.cpp
+++ b/Engine/Source/Runtime/Engine/Private/LevelTick.cpp
@@ -139,7 +139,8 @@ extern bool GShouldLogOutAFrameOfSetBodyTransform;
 -----------------------------------------------------------------------------*/
 
 /** Static array of tickable objects */
-TArray<FTickableGameObject*> FTickableGameObject::TickableObjects;
+TArray<FTickableObjectBase::FTickableObjectEntry> FTickableGameObject::TickableObjects;
+TArray<FTickableGameObject*> FTickableGameObject::PendingTickableObjects;
 bool FTickableGameObject::bIsTickingObjects = false;
 
 #if LOG_DETAILED_PATHFINDING_STATS
@@ -1184,6 +1185,12 @@ void EndTickDrawEvent(TDrawEvent<FRHICommandList>* TickDrawEvent)
 
 void FTickableGameObject::TickObjects(UWorld* World, const int32 InTickType, const bool bIsPaused, const float DeltaSeconds)
 {
+	for (FTickableGameObject* PendingTickable : PendingTickableObjects)
+	{
+		AddTickableObject(TickableObjects, PendingTickable);
+	}
+	PendingTickableObjects.Empty();
+
 	if (TickableObjects.Num() > 0)
 	{
 		check(!bIsTickingObjects);
@@ -1192,12 +1199,12 @@ void FTickableGameObject::TickObjects(UWorld* World, const int32 InTickType, con
 		bool bNeedsCleanup = false;
 		const ELevelTick TickType = (ELevelTick)InTickType;
 
-		for (int32 i = 0; i < TickableObjects.Num(); ++i)
+		for (const FTickableObjectEntry& TickableEntry : TickableObjects)
 		{
-			if (FTickableGameObject* TickableObject = TickableObjects[i])
+			if (FTickableGameObject* TickableObject = static_cast<FTickableGameObject*>(TickableEntry.TickableObject))
 			{
 				// If it is tickable and in this world
-				if (TickableObject->IsTickable() && (TickableObject->GetTickableGameObjectWorld() == World))
+				if (((TickableEntry.TickType == ETickableTickType::Always) || TickableObject->IsTickable()) && (TickableObject->GetTickableGameObjectWorld() == World))
 				{
 					const bool bIsGameWorld = InTickType == LEVELTICK_All || (World && World->IsGameWorld());
 					// If we are in editor and it is editor tickable, always tick
@@ -1209,7 +1216,7 @@ void FTickableGameObject::TickObjects(UWorld* World, const int32 InTickType, con
 						TickableObject->Tick(DeltaSeconds);
 
 						// In case it was removed during tick
-						if (TickableObjects[i] == nullptr)
+						if (TickableEntry.TickableObject == nullptr)
 						{
 							bNeedsCleanup = true;
 						}
@@ -1224,7 +1231,7 @@ void FTickableGameObject::TickObjects(UWorld* World, const int32 InTickType, con
 
 		if (bNeedsCleanup)
 		{
-			TickableObjects.RemoveAll([](FTickableGameObject* Object) { return Object == nullptr; });
+			TickableObjects.RemoveAll([](const FTickableObjectEntry& Entry) { return Entry.TickableObject == nullptr; });
 		}
 
 		bIsTickingObjects = false;
diff --git a/Engine/Source/Runtime/Engine/Private/Materials/HLSLMaterialTranslator.h b/Engine/Source/Runtime/Engine/Private/Materials/HLSLMaterialTranslator.h
--- a/Engine/Source/Runtime/Engine/Private/Materials/HLSLMaterialTranslator.h
+++ b/Engine/Source/Runtime/Engine/Private/Materials/HLSLMaterialTranslator.h
@@ -5211,46 +5211,47 @@ protected:
 		for( int32 i = 0; i < Custom->Inputs.Num(); i++ )
 		{
 			// skip over unnamed inputs
-			if( Custom->Inputs[i].InputName.Len()==0 )
+			if( Custom->Inputs[i].InputName.IsNone() )
 			{
 				continue;
 			}
 			InputParamDecl += TEXT(",");
+			const FString InputNameStr = Custom->Inputs[i].InputName.ToString();
 			switch(GetParameterType(CompiledInputs[i]))
 			{
 			case MCT_Float:
 			case MCT_Float1:
 				InputParamDecl += TEXT("MaterialFloat ");
-				InputParamDecl += Custom->Inputs[i].InputName;
+				InputParamDecl += InputNameStr;
 				break;
 			case MCT_Float2:
 				InputParamDecl += TEXT("MaterialFloat2 ");
-				InputParamDecl += Custom->Inputs[i].InputName;
+				InputParamDecl += InputNameStr;
 				break;
 			case MCT_Float3:
 				InputParamDecl += TEXT("MaterialFloat3 ");
-				InputParamDecl += Custom->Inputs[i].InputName;
+				InputParamDecl += InputNameStr;
 				break;
 			case MCT_Float4:
 				InputParamDecl += TEXT("MaterialFloat4 ");
-				InputParamDecl += Custom->Inputs[i].InputName;
+				InputParamDecl += InputNameStr;
 				break;
 			case MCT_Texture2D:
 				InputParamDecl += TEXT("Texture2D ");
-				InputParamDecl += Custom->Inputs[i].InputName;
+				InputParamDecl += InputNameStr;
 				InputParamDecl += TEXT(", SamplerState ");
-				InputParamDecl += Custom->Inputs[i].InputName;
+				InputParamDecl += InputNameStr;
 				InputParamDecl += TEXT("Sampler ");
 				break;
 			case MCT_TextureCube:
 				InputParamDecl += TEXT("TextureCube ");
-				InputParamDecl += Custom->Inputs[i].InputName;
+				InputParamDecl += InputNameStr;
 				InputParamDecl += TEXT(", SamplerState ");
-				InputParamDecl += Custom->Inputs[i].InputName;
+				InputParamDecl += InputNameStr;
 				InputParamDecl += TEXT("Sampler ");
 				break;
 			default:
-				return Errorf(TEXT("Bad type %s for %s input %s"),DescribeType(GetParameterType(CompiledInputs[i])), *Custom->Description, *Custom->Inputs[i].InputName);
+				return Errorf(TEXT("Bad type %s for %s input %s"),DescribeType(GetParameterType(CompiledInputs[i])), *Custom->Description, *InputNameStr);
 				break;
 			}
 		}
@@ -5272,7 +5273,7 @@ protected:
 		for( int32 i = 0; i < CompiledInputs.Num(); i++ )
 		{
 			// skip over unnamed inputs
-			if( Custom->Inputs[i].InputName.Len()==0 )
+			if( Custom->Inputs[i].InputName.IsNone() )
 			{
 				continue;
 			}
diff --git a/Engine/Source/Runtime/Engine/Private/Materials/MaterialExpressions.cpp b/Engine/Source/Runtime/Engine/Private/Materials/MaterialExpressions.cpp
--- a/Engine/Source/Runtime/Engine/Private/Materials/MaterialExpressions.cpp
+++ b/Engine/Source/Runtime/Engine/Private/Materials/MaterialExpressions.cpp
@@ -468,9 +468,9 @@ void UMaterialExpression::CopyMaterialExpressions(const TArray<UMaterialExpressi
 	TMap<UMaterialExpression*,UMaterialExpression*> SrcToDestMap;
 
 	// Duplicate source expressions into the editor's material copy buffer.
-	for (int32 SrcExpressionIndex = 0; SrcExpressionIndex < SrcExpressions.Num(); ++SrcExpressionIndex)
+	for( int32 SrcExpressionIndex = 0 ; SrcExpressionIndex < SrcExpressions.Num() ; ++SrcExpressionIndex )
 	{
-		UMaterialExpression* SrcExpression = SrcExpressions[SrcExpressionIndex];
+		UMaterialExpression*	SrcExpression		= SrcExpressions[SrcExpressionIndex];
 		UMaterialExpressionMaterialFunctionCall* FunctionExpression = Cast<UMaterialExpressionMaterialFunctionCall>(SrcExpression);
 		bool bIsValidFunctionExpression = true;
 
@@ -775,7 +775,7 @@ FExpressionInput* UMaterialExpression::GetInput(int32 InputIndex)
 }
 
 
-FString UMaterialExpression::GetInputName(int32 InputIndex) const
+FName UMaterialExpression::GetInputName(int32 InputIndex) const
 {
 	int32 Index = 0;
 	for( TFieldIterator<UStructProperty> InputIt(GetClass(),EFieldIteratorFlags::IncludeSuper,  EFieldIteratorFlags::ExcludeDeprecated) ; InputIt ; ++InputIt )
@@ -788,20 +788,28 @@ FString UMaterialExpression::GetInputName(int32 InputIndex) const
 			if( Index == InputIndex )
 			{
 					FExpressionInput const* Input = StructProp->ContainerPtrToValuePtr<FExpressionInput>(this, ArrayIndex);
-					FString StructName = StructProp->GetFName().ToString();
+
+						if (!Input->InputName.IsNone())
+						{
+							return Input->InputName;
+						}
+						else
+						{
+							FName StructName = StructProp->GetFName();
 
 					if (StructProp->ArrayDim > 1)
 					{
-						StructName += FString::FromInt(ArrayIndex);
+								StructName = *FString::Printf(TEXT("%s_%d"), *StructName.ToString(), ArrayIndex);
 					}
 
-					return (Input->InputName.Len() > 0) ? Input->InputName : StructName;
+							return StructName;
+						}
 			}
 			Index++;
 		}
 	}
 	}
-	return TEXT("");
+	return NAME_None;
 }
 
 #if WITH_EDITOR
@@ -1467,7 +1475,7 @@ FExpressionInput* UMaterialExpressionTextureSample::GetInput(int32 InputIndex)
 
 // this define is only used for the following function
 #define IF_INPUT_RETURN(Item, Name) if(!InputIndex) return Name; --InputIndex
-FString UMaterialExpressionTextureSample::GetInputName(int32 InputIndex) const
+FName UMaterialExpressionTextureSample::GetInputName(int32 InputIndex) const
 {
 	IF_INPUT_RETURN(Coordinates, TEXT("Coordinates"));
 
@@ -4444,16 +4452,16 @@ FExpressionInput* UMaterialExpressionBreakMaterialAttributes::GetInput(int32 Inp
 		return &MaterialAttributes;
 	}
 
-	return NULL;
+	return nullptr;
 }
 
-FString UMaterialExpressionBreakMaterialAttributes::GetInputName(int32 InputIndex) const
+FName UMaterialExpressionBreakMaterialAttributes::GetInputName(int32 InputIndex) const
 {
 	if( 0 == InputIndex )
 	{
-		return NSLOCTEXT("BreakMaterialAttributes", "InputName", "Attr").ToString();
+		return *NSLOCTEXT("BreakMaterialAttributes", "InputName", "Attr").ToString();
 	}
-	return TEXT("");
+	return NAME_None;
 }
 
 bool UMaterialExpressionBreakMaterialAttributes::IsInputConnectionRequired(int32 InputIndex) const
@@ -4550,9 +4558,9 @@ FExpressionInput* UMaterialExpressionGetMaterialAttributes::GetInput(int32 Input
 	return nullptr;
 }
 
-FString UMaterialExpressionGetMaterialAttributes::GetInputName(int32 InputIndex) const
+FName UMaterialExpressionGetMaterialAttributes::GetInputName(int32 InputIndex) const
 {
-	return TEXT("");
+	return NAME_None;
 }
 
 #if WITH_EDITOR
@@ -4622,7 +4630,7 @@ void UMaterialExpressionGetMaterialAttributes::PostEditChangeProperty(FPropertyC
 			// Type changed, update pin names
 			for (int i = 1; i < Outputs.Num(); ++i)
 			{
-				Outputs[i].OutputName = FMaterialAttributeDefinitionMap::GetDisplayName(AttributeGetTypes[i-1]);
+				Outputs[i].OutputName = *FMaterialAttributeDefinitionMap::GetDisplayName(AttributeGetTypes[i-1]);
 			}
 
 			GraphNode->ReconstructNode();
@@ -4641,8 +4649,8 @@ void UMaterialExpressionGetMaterialAttributes::PostLoad()
 
 	for (int i = 1; i < Outputs.Num(); ++i)
 	{
-		FString DisplayName = FMaterialAttributeDefinitionMap::GetDisplayName(AttributeGetTypes[i-1]);
-		if (Outputs[i].OutputName != DisplayName)
+		const FString DisplayName = FMaterialAttributeDefinitionMap::GetDisplayName(AttributeGetTypes[i-1]);
+		if (Outputs[i].OutputName.ToString() != DisplayName)
 		{
 			FString MaterialName;
 			if (Material)
@@ -4654,8 +4662,8 @@ void UMaterialExpressionGetMaterialAttributes::PostLoad()
 				Function->GetName(MaterialName);
 			}
 
-			UE_LOG(LogMaterial, Warning, TEXT("Serialized attribute that no longer exists (%s) for material \"%s\"."), *(Outputs[i].OutputName), *MaterialName);
-			Outputs[i].OutputName = DisplayName;
+			UE_LOG(LogMaterial, Warning, TEXT("Serialized attribute that no longer exists (%s) for material \"%s\"."), *(Outputs[i].OutputName.ToString()), *MaterialName);
+			Outputs[i].OutputName = *DisplayName;
 		}
 	}
 }
@@ -4753,17 +4761,17 @@ FExpressionInput* UMaterialExpressionSetMaterialAttributes::GetInput(int32 Input
 	return &Inputs[InputIndex];
 }
 
-FString UMaterialExpressionSetMaterialAttributes::GetInputName(int32 InputIndex) const
+FName UMaterialExpressionSetMaterialAttributes::GetInputName(int32 InputIndex) const
 {
-	FString Name(TEXT(""));
+	FName Name;
 
 	if (InputIndex == 0)
 	{
-		Name = NSLOCTEXT("SetMaterialAttributes", "InputName", "MaterialAttributes").ToString();
+		Name = *NSLOCTEXT("SetMaterialAttributes", "InputName", "MaterialAttributes").ToString();
 	}
 	else if (InputIndex > 0)
 	{
-		Name = FMaterialAttributeDefinitionMap::GetDisplayName(AttributeSetTypes[InputIndex-1]);
+		Name = *FMaterialAttributeDefinitionMap::GetDisplayName(AttributeSetTypes[InputIndex-1]);
 	}
 
 	return Name;
@@ -4941,9 +4949,9 @@ FExpressionInput* UMaterialExpressionBlendMaterialAttributes::GetInput(int32 Inp
 	return nullptr;
 }
 
-FString UMaterialExpressionBlendMaterialAttributes::GetInputName(int32 InputIndex) const
+FName UMaterialExpressionBlendMaterialAttributes::GetInputName(int32 InputIndex) const
 {
-	FString Name;
+	FName Name;
 
 	switch (InputIndex)
 	{
@@ -5381,9 +5389,9 @@ FExpressionInput* UMaterialExpressionMaterialAttributeLayers::GetInput(int32 Inp
 	return nullptr;
 }
 
-FString UMaterialExpressionMaterialAttributeLayers::GetInputName(int32 InputIndex) const
+FName UMaterialExpressionMaterialAttributeLayers::GetInputName(int32 InputIndex) const
 {
-	return TEXT("");
+	return NAME_None;
 }
 
 #if WITH_EDITOR
@@ -6020,7 +6028,7 @@ void UMaterialExpressionStaticSwitchParameter::GetCaption(TArray<FString>& OutCa
 }
 #endif // WITH_EDITOR
 
-FString UMaterialExpressionStaticSwitchParameter::GetInputName(int32 InputIndex) const
+FName UMaterialExpressionStaticSwitchParameter::GetInputName(int32 InputIndex) const
 {
 	if (InputIndex == 0)
 	{
@@ -6188,7 +6196,7 @@ void UMaterialExpressionStaticSwitch::GetCaption(TArray<FString>& OutCaptions) c
 }
 #endif // WITH_EDITOR
 
-FString UMaterialExpressionStaticSwitch::GetInputName(int32 InputIndex) const
+FName UMaterialExpressionStaticSwitch::GetInputName(int32 InputIndex) const
 {
 	if (InputIndex == 0)
 	{
@@ -6276,7 +6284,7 @@ void UMaterialExpressionPreviousFrameSwitch::GetExpressionToolTip(TArray<FString
 }
 #endif // WITH_EDITOR
 
-FString UMaterialExpressionPreviousFrameSwitch::GetInputName(int32 InputIndex) const
+FName UMaterialExpressionPreviousFrameSwitch::GetInputName(int32 InputIndex) const
 {
 	if (InputIndex == 0)
 	{
@@ -6369,16 +6377,14 @@ FExpressionInput* UMaterialExpressionQualitySwitch::GetInput(int32 InputIndex)
 	return &Inputs[InputIndex - 1];
 }
 
-FString UMaterialExpressionQualitySwitch::GetInputName(int32 InputIndex) const
+FName UMaterialExpressionQualitySwitch::GetInputName(int32 InputIndex) const
 {
 	if (InputIndex == 0)
 	{
 		return TEXT("Default");
 	}
 
-	FString QualityLevelName;
-	GetMaterialQualityLevelName((EMaterialQualityLevel::Type)(InputIndex - 1), QualityLevelName);
-	return QualityLevelName;
+	return GetMaterialQualityLevelFName((EMaterialQualityLevel::Type)(InputIndex - 1));
 }
 
 bool UMaterialExpressionQualitySwitch::IsInputConnectionRequired(int32 InputIndex) const
@@ -6483,14 +6489,14 @@ FExpressionInput* UMaterialExpressionFeatureLevelSwitch::GetInput(int32 InputInd
 	return &Inputs[InputIndex - 1];
 }
 
-FString UMaterialExpressionFeatureLevelSwitch::GetInputName(int32 InputIndex) const
+FName UMaterialExpressionFeatureLevelSwitch::GetInputName(int32 InputIndex) const
 {
 	if (InputIndex == 0)
 	{
 		return TEXT("Default");
 	}
 
-	FString FeatureLevelName;
+	FName FeatureLevelName;
 	GetFeatureLevelName((ERHIFeatureLevel::Type)(InputIndex - 1), FeatureLevelName);
 	return FeatureLevelName;
 }
@@ -7305,15 +7311,16 @@ void UMaterialExpressionSceneDepth::GetCaption(TArray<FString>& OutCaptions) con
 }
 #endif // WITH_EDITOR
 
-FString UMaterialExpressionSceneDepth::GetInputName(int32 InputIndex) const
+FName UMaterialExpressionSceneDepth::GetInputName(int32 InputIndex) const
 {
 	if(InputIndex == 0)
 	{
 		// Display the current InputMode enum's display name.
 		UByteProperty* InputModeProperty = FindField<UByteProperty>( UMaterialExpressionSceneDepth::StaticClass(), "InputMode" );
-		return InputModeProperty->Enum->GetNameStringByValue((int64)InputMode.GetValue());
+		// Can't use GetNameByValue as GetNameStringByValue does name mangling that GetNameByValue does not
+		return *InputModeProperty->Enum->GetNameStringByValue((int64)InputMode.GetValue());
 	}
-	return TEXT("");
+	return NAME_None;
 }
 
 
@@ -7612,15 +7619,16 @@ void UMaterialExpressionLogarithm10::GetExpressionToolTip(TArray<FString>& OutTo
 }
 #endif // WITH_EDITOR
 
-FString UMaterialExpressionSceneColor::GetInputName(int32 InputIndex) const
+FName UMaterialExpressionSceneColor::GetInputName(int32 InputIndex) const
 {
 	if(InputIndex == 0)
 	{
 		// Display the current InputMode enum's display name.
 		UByteProperty* InputModeProperty = FindField<UByteProperty>( UMaterialExpressionSceneColor::StaticClass(), "InputMode" );
-		return InputModeProperty->Enum->GetNameStringByValue((int64)InputMode.GetValue());
+		// Can't use GetNameByValue as GetNameStringByValue does name mangling that GetNameByValue does not
+		return *InputModeProperty->Enum->GetNameStringByValue((int64)InputMode.GetValue());
 	}
-	return TEXT("");
+	return NAME_None;
 }
 
 UMaterialExpressionIf::UMaterialExpressionIf(const FObjectInitializer& ObjectInitializer)
@@ -8754,7 +8762,7 @@ int32 UMaterialExpressionCustom::Compile(class FMaterialCompiler* Compiler, int3
 	for( int32 i=0;i<Inputs.Num();i++ )
 	{
 		// skip over unnamed inputs
-		if( Inputs[i].InputName.Len()==0 )
+		if( Inputs[i].InputName.IsNone() )
 		{
 			CompiledInputs.Add(INDEX_NONE);
 		}
@@ -8762,7 +8770,7 @@ int32 UMaterialExpressionCustom::Compile(class FMaterialCompiler* Compiler, int3
 		{
 			if(!Inputs[i].Input.GetTracedInput().Expression)
 			{
-				return Compiler->Errorf(TEXT("Custom material %s missing input %d (%s)"), *Description, i+1, *Inputs[i].InputName);
+				return Compiler->Errorf(TEXT("Custom material %s missing input %d (%s)"), *Description, i+1, *Inputs[i].InputName.ToString());
 			}
 			int32 InputCode = Inputs[i].Input.Compile(Compiler);
 			if( InputCode < 0 )
@@ -8802,13 +8810,13 @@ FExpressionInput* UMaterialExpressionCustom::GetInput(int32 InputIndex)
 	return NULL;
 }
 
-FString UMaterialExpressionCustom::GetInputName(int32 InputIndex) const
+FName UMaterialExpressionCustom::GetInputName(int32 InputIndex) const
 {
 	if( InputIndex < Inputs.Num() )
 	{
 		return Inputs[InputIndex].InputName;
 	}
-	return TEXT("");
+	return NAME_None;
 }
 
 #if WITH_EDITOR
@@ -8816,25 +8824,26 @@ void UMaterialExpressionCustom::PostEditChangeProperty(FPropertyChangedEvent& Pr
 {
 	// strip any spaces from input name
 	UProperty* PropertyThatChanged = PropertyChangedEvent.Property;
-	if( PropertyThatChanged && PropertyThatChanged->GetFName() == FName(TEXT("InputName")) )
+	if( PropertyThatChanged && PropertyThatChanged->GetFName() == GET_MEMBER_NAME_CHECKED(FCustomInput, InputName))
+	{
+		for( FCustomInput& Input : Inputs )
 	{
-		for( int32 i=0;i<Inputs.Num();i++ )
+			FString InputName = Input.InputName.ToString();
+			if (InputName.ReplaceInline(TEXT(" "),TEXT("")) > 0)
 		{
-			Inputs[i].InputName.ReplaceInline(TEXT(" "),TEXT(""));
+				Input.InputName = *InputName;
+			}
 		}
 	}
 
-	if (PropertyChangedEvent.MemberProperty)
+	if (PropertyChangedEvent.MemberProperty && GraphNode)
 	{
 		const FName PropertyName = PropertyChangedEvent.MemberProperty->GetFName();
 		if (PropertyName == GET_MEMBER_NAME_CHECKED(UMaterialExpressionCustom, Inputs))
 		{
-			if (GraphNode)
-			{
 				GraphNode->ReconstructNode();
 			}
 		}
-	}
 
 	Super::PostEditChangeProperty(PropertyChangedEvent);
 }
@@ -9315,7 +9324,7 @@ static const FFunctionExpressionInput* FindInputById(const FGuid& Id, const TArr
 }
 
 /** Finds an input in the passed in array with a matching name. */
-static const FFunctionExpressionInput* FindInputByName(const FString& Name, const TArray<FFunctionExpressionInput>& Inputs)
+static const FFunctionExpressionInput* FindInputByName(const FName& Name, const TArray<FFunctionExpressionInput>& Inputs)
 {
 	for (int32 InputIndex = 0; InputIndex < Inputs.Num(); InputIndex++)
 	{
@@ -9357,7 +9366,7 @@ static int32 FindOutputIndexById(const FGuid& Id, const TArray<FFunctionExpressi
 }
 
 /** Finds an output in the passed in array with a matching name. */
-static int32 FindOutputIndexByName(const FString& Name, const TArray<FFunctionExpressionOutput>& Outputs)
+static int32 FindOutputIndexByName(const FName& Name, const TArray<FFunctionExpressionOutput>& Outputs)
 {
 	for (int32 OutputIndex = 0; OutputIndex < Outputs.Num(); OutputIndex++)
 	{
@@ -9467,15 +9476,15 @@ int32 UMaterialFunction::Compile(FMaterialCompiler* Compiler, const FFunctionExp
 
 	if (ValidateFunctionUsage(Compiler, Output))
 	{
-		if (Output.ExpressionOutput->A.GetTracedInput().Expression)
-		{
-			// Compile the given function output
-			ReturnValue = Output.ExpressionOutput->A.Compile(Compiler);
-		}
-		else
-		{
-			ReturnValue = Compiler->Errorf(TEXT("Missing function output connection '%s'"), *Output.ExpressionOutput->OutputName);
-		}
+	if (Output.ExpressionOutput->A.GetTracedInput().Expression)
+	{
+		// Compile the given function output
+		ReturnValue = Output.ExpressionOutput->A.Compile(Compiler);
+	}
+	else
+	{
+		ReturnValue = Compiler->Errorf(TEXT("Missing function output connection '%s'"), *Output.ExpressionOutput->OutputName.ToString());
+	}
 	}
 
 	return ReturnValue;
@@ -9564,7 +9573,7 @@ bool UMaterialFunction::IsDependent(UMaterialFunctionInterface* OtherFunction)
 		if (MaterialFunctionExpression && MaterialFunctionExpression->MaterialFunction)
 		{
 			// Recurse to handle nesting
-			bIsDependent = bIsDependent
+			bIsDependent = bIsDependent 
 				|| MaterialFunctionExpression->MaterialFunction->GetReentrantFlag()
 				|| MaterialFunctionExpression->MaterialFunction->IsDependent(OtherFunction);
 		}
@@ -9590,7 +9599,7 @@ void UMaterialFunction::AppendReferencedTextures(TArray<UTexture*>& InOutTexture
 {
 	for (UMaterialExpression* CurrentExpression : FunctionExpressions)
 	{
-		if (CurrentExpression)
+		if(CurrentExpression)
 		{
 			UTexture* ReferencedTexture = CurrentExpression->GetReferencedTexture();
 
@@ -10154,23 +10163,23 @@ static const TCHAR* GetInputTypeName(uint8 InputType)
 	return TypeNames[InputType];
 }
 
-FString UMaterialExpressionMaterialFunctionCall::GetInputNameWithType(int32 InputIndex, bool bWithType) const
+FName UMaterialExpressionMaterialFunctionCall::GetInputNameWithType(int32 InputIndex, bool bWithType) const
 {
 	if (InputIndex < FunctionInputs.Num())
 	{
-		if (FunctionInputs[InputIndex].ExpressionInput != NULL && bWithType)
+		if (FunctionInputs[InputIndex].ExpressionInput != nullptr && bWithType)
 		{
-			return FunctionInputs[InputIndex].Input.InputName + TEXT(" (") + GetInputTypeName(FunctionInputs[InputIndex].ExpressionInput->InputType) + TEXT(")");
+			return *FString::Printf(TEXT("%s (%s)"), *FunctionInputs[InputIndex].Input.InputName.ToString(), GetInputTypeName(FunctionInputs[InputIndex].ExpressionInput->InputType));
 		}
 		else
 		{
 			return FunctionInputs[InputIndex].Input.InputName;
 		}
 	}
-	return TEXT("");
+	return NAME_None;
 }
 
-FString UMaterialExpressionMaterialFunctionCall::GetInputName(int32 InputIndex) const
+FName UMaterialExpressionMaterialFunctionCall::GetInputName(int32 InputIndex) const
 {
 	return GetInputNameWithType(InputIndex, true);
 }
@@ -10394,7 +10403,7 @@ void UMaterialExpressionMaterialFunctionCall::UpdateFromFunctionResource(bool bR
 			{
 				// Maintain the input connection if an input with matching Id is found, but propagate the new name
 				// This way function inputs names can be changed without affecting material connections
-				const FString TempInputName = CurrentInput.Input.InputName;
+				const FName TempInputName = CurrentInput.Input.InputName;
 				CurrentInput.Input = OriginalInput->Input;
 				CurrentInput.Input.InputName = TempInputName;
 			}
@@ -10615,11 +10624,11 @@ void UMaterialExpressionFunctionInput::PostEditImport()
 	ConditionallyGenerateId(true);
 }
 
-FString InputNameBackup;
+FName InputNameBackup;
 
 void UMaterialExpressionFunctionInput::PreEditChange(UProperty* PropertyAboutToChange)
 {
-	if (PropertyAboutToChange && PropertyAboutToChange->GetFName() == FName(TEXT("InputName")))
+	if (PropertyAboutToChange && PropertyAboutToChange->GetFName() == GET_MEMBER_NAME_CHECKED(UMaterialExpressionFunctionInput, InputName))
 	{
 		InputNameBackup = InputName;
 	}
@@ -10629,7 +10638,7 @@ void UMaterialExpressionFunctionInput::PreEditChange(UProperty* PropertyAboutToC
 void UMaterialExpressionFunctionInput::PostEditChangeProperty(FPropertyChangedEvent& PropertyChangedEvent)
 {
 	UProperty* PropertyThatChanged = PropertyChangedEvent.Property;
-	if (PropertyThatChanged && PropertyThatChanged->GetFName() == FName(TEXT("InputName")))
+	if (PropertyThatChanged && PropertyThatChanged->GetFName() == GET_MEMBER_NAME_CHECKED(UMaterialExpressionFunctionInput, InputName))
 	{
 		if (Material)
 		{
@@ -10663,7 +10672,7 @@ void UMaterialExpressionFunctionInput::GetCaption(TArray<FString>& OutCaptions)
 		TEXT("MaterialAttributes")
 	};
 	check(InputType < FunctionInput_MAX);
-	OutCaptions.Add(FString(TEXT("Input ")) + InputName + TEXT(" (") + TypeNames[InputType] + TEXT(")"));
+	OutCaptions.Add(FString(TEXT("Input ")) + InputName.ToString() + TEXT(" (") + TypeNames[InputType] + TEXT(")"));
 }
 
 void UMaterialExpressionFunctionInput::GetExpressionToolTip(TArray<FString>& OutToolTip) 
@@ -10697,7 +10706,7 @@ int32 UMaterialExpressionFunctionInput::CompilePreviewValue(FMaterialCompiler* C
 		case FunctionInput_Texture2D:
 		case FunctionInput_TextureCube:
 		case FunctionInput_StaticBool:
-			return Compiler->Errorf(TEXT("Missing Preview connection for function input '%s'"), *InputName);
+			return Compiler->Errorf(TEXT("Missing Preview connection for function input '%s'"), *InputName.ToString());
 		default:
 			return Compiler->Errorf(TEXT("Unknown input type"));
 		}
@@ -10767,7 +10776,7 @@ int32 UMaterialExpressionFunctionInput::Compile(class FMaterialCompiler* Compile
 		}
 		else
 		{
-			return Compiler->Errorf(TEXT("Missing function input '%s'"), *InputName);
+			return Compiler->Errorf(TEXT("Missing function input '%s'"), *InputName.ToString());
 		}
 	}
 }
@@ -10793,7 +10802,7 @@ void UMaterialExpressionFunctionInput::ValidateName()
 	{
 		int32 InputNameIndex = 0;
 		bool bResultNameIndexValid = true;
-		FString PotentialInputName;
+		FName PotentialInputName;
 
 		// Find an available unique name
 		do 
@@ -10801,7 +10810,7 @@ void UMaterialExpressionFunctionInput::ValidateName()
 			PotentialInputName = InputName;
 			if (InputNameIndex != 0)
 			{
-				PotentialInputName += FString::FromInt(InputNameIndex);
+				PotentialInputName = *FString::Printf(TEXT("%s%d"), *InputName.ToString(), InputNameIndex);
 			}
 
 			bResultNameIndexValid = true;
@@ -10920,12 +10929,12 @@ void UMaterialExpressionFunctionOutput::PostEditImport()
 }
 #endif	//#if WITH_EDITOR
 
-FString OutputNameBackup;
+FName OutputNameBackup;
 
 #if WITH_EDITOR
 void UMaterialExpressionFunctionOutput::PreEditChange(UProperty* PropertyAboutToChange)
 {
-	if (PropertyAboutToChange && PropertyAboutToChange->GetFName() == FName(TEXT("OutputName")))
+	if (PropertyAboutToChange && PropertyAboutToChange->GetFName() == GET_MEMBER_NAME_CHECKED(UMaterialExpressionFunctionOutput, OutputName))
 	{
 		OutputNameBackup = OutputName;
 	}
@@ -10935,7 +10944,7 @@ void UMaterialExpressionFunctionOutput::PreEditChange(UProperty* PropertyAboutTo
 void UMaterialExpressionFunctionOutput::PostEditChangeProperty(FPropertyChangedEvent& PropertyChangedEvent)
 {
 	UProperty* PropertyThatChanged = PropertyChangedEvent.Property;
-	if (PropertyThatChanged && PropertyThatChanged->GetFName() == FName(TEXT("OutputName")))
+	if (PropertyThatChanged && PropertyThatChanged->GetFName() == GET_MEMBER_NAME_CHECKED(UMaterialExpressionFunctionOutput, OutputName))
 	{
 		if (Material)
 		{
@@ -10956,7 +10965,7 @@ void UMaterialExpressionFunctionOutput::PostEditChangeProperty(FPropertyChangedE
 
 void UMaterialExpressionFunctionOutput::GetCaption(TArray<FString>& OutCaptions) const
 {
-	OutCaptions.Add(FString(TEXT("Output ")) + OutputName);
+	OutCaptions.Add(FString(TEXT("Output ")) + OutputName.ToString());
 }
 
 void UMaterialExpressionFunctionOutput::GetExpressionToolTip(TArray<FString>& OutToolTip) 
@@ -10974,7 +10983,7 @@ int32 UMaterialExpressionFunctionOutput::Compile(class FMaterialCompiler* Compil
 {
 	if (!A.GetTracedInput().Expression)
 	{
-		return Compiler->Errorf(TEXT("Missing function output '%s'"), *OutputName);
+		return Compiler->Errorf(TEXT("Missing function output '%s'"), *OutputName.ToString());
 	}
 	return A.Compile(Compiler);
 }
@@ -10994,7 +11003,7 @@ void UMaterialExpressionFunctionOutput::ValidateName()
 	{
 		int32 OutputNameIndex = 0;
 		bool bResultNameIndexValid = true;
-		FString PotentialOutputName;
+		FName PotentialOutputName;
 
 		// Find an available unique name
 		do 
@@ -11002,7 +11011,7 @@ void UMaterialExpressionFunctionOutput::ValidateName()
 			PotentialOutputName = OutputName;
 			if (OutputNameIndex != 0)
 			{
-				PotentialOutputName += FString::FromInt(OutputNameIndex);
+				PotentialOutputName = *FString::Printf(TEXT("%s%d"), *OutputName.ToString(), OutputNameIndex);
 			}
 
 			bResultNameIndexValid = true;
diff --git a/Engine/Source/Runtime/Engine/Private/Materials/MaterialShared.cpp b/Engine/Source/Runtime/Engine/Private/Materials/MaterialShared.cpp
--- a/Engine/Source/Runtime/Engine/Private/Materials/MaterialShared.cpp
+++ b/Engine/Source/Runtime/Engine/Private/Materials/MaterialShared.cpp
@@ -7,6 +7,7 @@
 #include "MaterialShared.h"
 #include "Stats/StatsMisc.h"
 #include "UObject/CoreObjectVersion.h"
+#include "UObject/FrameworkObjectVersion.h"
 #include "Misc/App.h"
 #include "UObject/UObjectHash.h"
 #include "LocalVertexFactory.h"
@@ -59,6 +60,12 @@ void GetMaterialQualityLevelName(EMaterialQualityLevel::Type InQualityLevel, FSt
 	MaterialQualityLevelNames[(int32)InQualityLevel].ToString(OutName);
 }
 
+FName GetMaterialQualityLevelFName(EMaterialQualityLevel::Type InQualityLevel)
+{
+	check(InQualityLevel < ARRAY_COUNT(MaterialQualityLevelNames));
+	return MaterialQualityLevelNames[(int32)InQualityLevel];
+}
+
 static inline SIZE_T AddShaderSize(FShader* Shader, TSet<FShaderResourceId>& UniqueShaderResourceIds)
 {
 	SIZE_T ResourceSize = 0;
@@ -144,6 +151,7 @@ FExpressionInput FExpressionInput::GetTracedInput() const
 static bool SerializeExpressionInput(FArchive& Ar, FExpressionInput& Input)
 {
 	Ar.UsingCustomVersion(FCoreObjectVersion::GUID);
+	Ar.UsingCustomVersion(FFrameworkObjectVersion::GUID);
 
 	if (Ar.CustomVer(FCoreObjectVersion::GUID) < FCoreObjectVersion::MaterialInputNativeSerialize)
 	{
@@ -157,7 +165,16 @@ static bool SerializeExpressionInput(FArchive& Ar, FExpressionInput& Input)
 	}
 #endif
 	Ar << Input.OutputIndex;
-	Ar << Input.InputName;
+	if (Ar.CustomVer(FFrameworkObjectVersion::GUID) >= FFrameworkObjectVersion::PinsStoreFName)
+	{
+		Ar << Input.InputName;
+	}
+	else
+	{
+		FString InputNameStr;
+		Ar << InputNameStr;
+		Input.InputName = *InputNameStr;
+	}
 	Ar << Input.Mask;
 	Ar << Input.MaskR;
 	Ar << Input.MaskG;
diff --git a/Engine/Source/Runtime/Engine/Private/NavMeshBoundsVolume.cpp b/Engine/Source/Runtime/Engine/Private/NavMeshBoundsVolume.cpp
--- a/Engine/Source/Runtime/Engine/Private/NavMeshBoundsVolume.cpp
+++ b/Engine/Source/Runtime/Engine/Private/NavMeshBoundsVolume.cpp
@@ -26,16 +26,30 @@ void ANavMeshBoundsVolume::PostEditChangeProperty( struct FPropertyChangedEvent&
 	UNavigationSystem* NavSys = UNavigationSystem::GetCurrent(GetWorld());
 	if (GIsEditor && NavSys)
 	{
-		if (PropertyChangedEvent.Property == NULL ||
-			PropertyChangedEvent.Property->GetFName() == GET_MEMBER_NAME_CHECKED(ABrush, BrushBuilder) ||
-			(PropertyChangedEvent.MemberProperty && 
-			 PropertyChangedEvent.MemberProperty->GetFName() == GET_MEMBER_NAME_CHECKED(ANavMeshBoundsVolume, SupportedAgents)))
+		const FName PropName = (PropertyChangedEvent.Property != nullptr) ? PropertyChangedEvent.Property->GetFName() : FName();
+		const FName MemberName = (PropertyChangedEvent.MemberProperty != nullptr) ? PropertyChangedEvent.MemberProperty->GetFName() : FName();
+
+		if (PropName == GET_MEMBER_NAME_CHECKED(ABrush, BrushBuilder)
+			|| MemberName == GET_MEMBER_NAME_CHECKED(ANavMeshBoundsVolume, SupportedAgents)
+			|| MemberName == GET_MEMBER_NAME_CHECKED(USceneComponent, RelativeLocation)
+			|| MemberName == GET_MEMBER_NAME_CHECKED(USceneComponent, RelativeRotation)
+			|| MemberName == GET_MEMBER_NAME_CHECKED(USceneComponent, RelativeScale3D))
 		{
 			NavSys->OnNavigationBoundsUpdated(this);
 		}
 	}
 }
 
+void ANavMeshBoundsVolume::PostEditUndo()
+{
+	Super::PostEditUndo();
+	UNavigationSystem* NavSys = UNavigationSystem::GetCurrent(GetWorld());
+	if (GIsEditor && NavSys)
+	{
+		NavSys->OnNavigationBoundsUpdated(this);
+	}
+}
+
 #endif // WITH_EDITOR
 
 void ANavMeshBoundsVolume::PostRegisterAllComponents() 
diff --git a/Engine/Source/Runtime/Engine/Private/Pawn.cpp b/Engine/Source/Runtime/Engine/Private/Pawn.cpp
--- a/Engine/Source/Runtime/Engine/Private/Pawn.cpp
+++ b/Engine/Source/Runtime/Engine/Private/Pawn.cpp
@@ -42,7 +42,7 @@ APawn::APawn(const FObjectInitializer& ObjectInitializer)
 
 	if (HasAnyFlags(RF_ClassDefaultObject) && GetClass() == APawn::StaticClass())
 	{
-		AIControllerClass = LoadClass<AController>(NULL, *((UEngine*)(UEngine::StaticClass()->GetDefaultObject()))->AIControllerClassName.ToString(), NULL, LOAD_None, NULL);
+		AIControllerClass = LoadClass<AController>(nullptr, *((UEngine*)(UEngine::StaticClass()->GetDefaultObject()))->AIControllerClassName.ToString(), nullptr, LOAD_None, nullptr);
 	}
 	else
 	{
@@ -107,7 +107,7 @@ void APawn::PostInitializeComponents()
 
 		// Automatically add Controller to AI Pawns if we are allowed to.
 		if (AutoPossessPlayer == EAutoReceiveInput::Disabled
-			&& AutoPossessAI != EAutoPossessAI::Disabled && Controller == NULL && GetNetMode() != NM_Client
+			&& AutoPossessAI != EAutoPossessAI::Disabled && Controller == nullptr && GetNetMode() != NM_Client
 #if WITH_EDITOR
 			&& (GIsEditor == false || GetWorld()->IsGameWorld())
 #endif // WITH_EDITOR
@@ -151,7 +151,7 @@ void APawn::UpdateNavAgent()
 {
 	UPawnMovementComponent* MovementComponent = GetMovementComponent();
 	//// Update Nav Agent props with collision component's setup if it's not set yet
-	if (RootComponent != NULL && MovementComponent != NULL && MovementComponent->ShouldUpdateNavAgentWithOwnersCollision())
+	if (RootComponent != nullptr && MovementComponent != nullptr && MovementComponent->ShouldUpdateNavAgentWithOwnersCollision())
 	{
 		RootComponent->UpdateBounds();
 		MovementComponent->UpdateNavAgent(*this);
@@ -176,12 +176,12 @@ void APawn::PawnStartFire(uint8 FireModeNum) {}
 
 AActor* APawn::GetMovementBaseActor(const APawn* Pawn)
 {
-	if (Pawn != NULL && Pawn->GetMovementBase())
+	if (Pawn != nullptr && Pawn->GetMovementBase())
 	{
 		return Pawn->GetMovementBase()->GetOwner();
 	}
 
-	return NULL;
+	return nullptr;
 }
 
 bool APawn::CanBeBaseForCharacter(class APawn* APawn) const
@@ -277,7 +277,7 @@ FVector APawn::GetPawnViewLocation() const
 
 FRotator APawn::GetViewRotation() const
 {
-	if (Controller != NULL)
+	if (Controller != nullptr)
 	{
 		return Controller->GetControlRotation();
 	}
@@ -299,11 +299,11 @@ FRotator APawn::GetViewRotation() const
 
 void APawn::SpawnDefaultController()
 {
-	if ( Controller != NULL || GetNetMode() == NM_Client)
+	if ( Controller != nullptr || GetNetMode() == NM_Client)
 	{
 		return;
 	}
-	if ( AIControllerClass != NULL )
+	if ( AIControllerClass != nullptr )
 	{
 		FActorSpawnParameters SpawnInfo;
 		SpawnInfo.Instigator = Instigator;
@@ -311,7 +311,7 @@ void APawn::SpawnDefaultController()
 		SpawnInfo.OverrideLevel = GetLevel();
 		SpawnInfo.ObjectFlags |= RF_Transient;	// We never want to save AI controllers into a map
 		AController* NewController = GetWorld()->SpawnActor<AController>(AIControllerClass, GetActorLocation(), GetActorRotation(), SpawnInfo);
-		if (NewController != NULL)
+		if (NewController != nullptr)
 		{
 			// if successful will result in setting this->Controller 
 			// as part of possession mechanics
@@ -364,7 +364,7 @@ void APawn::PawnClientRestart()
 		}
 
 		// Set up player input component, if there isn't one already.
-		if (InputComponent == NULL)
+		if (InputComponent == nullptr)
 		{
 			InputComponent = CreatePlayerInputComponent();
 			if (InputComponent)
@@ -436,7 +436,7 @@ float APawn::TakeDamage(float Damage, FDamageEvent const& DamageEvent, AControll
 bool APawn::IsControlled() const
 {
 	APlayerController* const PC = Cast<APlayerController>(Controller);
-	return(PC != NULL);
+	return(PC != nullptr);
 }
 
 FRotator APawn::GetControlRotation() const
@@ -446,7 +446,7 @@ FRotator APawn::GetControlRotation() const
 
 void APawn::OnRep_Controller()
 {
-	if ( (Controller != NULL) && (Controller->GetPawn() == NULL) )
+	if ( (Controller != nullptr) && (Controller->GetPawn() == nullptr) )
 	{
 		// This ensures that AController::OnRep_Pawn is called. Since we cant ensure replication order of APawn::Controller and AController::Pawn,
 		// if APawn::Controller is repped first, it will set AController::Pawn locally. When AController::Pawn is repped, the rep value will not
@@ -457,7 +457,7 @@ void APawn::OnRep_Controller()
 		Controller->SetPawnFromRep(this); 
 
 		APlayerController* const PC = Cast<APlayerController>(Controller);
-		if ( (PC != NULL) && PC->bAutoManageActiveCameraTarget && (PC->PlayerCameraManager->ViewTarget.Target == Controller) )
+		if ( (PC != nullptr) && PC->bAutoManageActiveCameraTarget && (PC->PlayerCameraManager->ViewTarget.Target == Controller) )
 		{
 			PC->AutoManageActiveCameraTarget(this);
 		}
@@ -473,7 +473,7 @@ void APawn::PossessedBy(AController* NewController)
 	Controller = NewController;
 	ForceNetUpdate();
 
-	if (Controller->PlayerState != NULL)
+	if (Controller->PlayerState != nullptr)
 	{
 		PlayerState = Controller->PlayerState;
 	}
@@ -504,9 +504,9 @@ void APawn::UnPossessed()
 
 	ForceNetUpdate();
 
-	PlayerState = NULL;
-	SetOwner(NULL);
-	Controller = NULL;
+	PlayerState = nullptr;
+	SetOwner(nullptr);
+	Controller = nullptr;
 
 	// Unregister input component if we created one
 	DestroyPlayerInputComponent();
@@ -543,7 +543,7 @@ class UPlayer* APawn::GetNetOwningPlayer()
 		if (Controller)
 		{
 			APlayerController* PC = Cast<APlayerController>(Controller);
-			return PC ? PC->Player : NULL;
+			return PC ? PC->Player : nullptr;
 		}
 	}
 
@@ -562,7 +562,7 @@ void APawn::DestroyPlayerInputComponent()
 	if (InputComponent)
 	{
 		InputComponent->DestroyComponent();
-		InputComponent = NULL;
+		InputComponent = nullptr;
 	}
 }
 
@@ -598,19 +598,12 @@ FVector APawn::GetLastMovementInputVector() const
 	return LastControlInputVector;
 }
 
-// TODO: deprecated, remove
-FVector APawn::GetMovementInputVector() const
-{
-	return GetPendingMovementInputVector();
-}
-
 // TODO: deprecated, remove
 FVector APawn::K2_GetMovementInputVector() const
 {
 	return GetPendingMovementInputVector();
 }
 
-
 FVector APawn::ConsumeMovementInputVector()
 {
 	UPawnMovementComponent* MovementComponent = GetMovementComponent();
@@ -706,7 +699,7 @@ void APawn::RecalculateBaseEyeHeight()
 
 void APawn::Reset()
 {
-	if ( (Controller == NULL) || (Controller->PlayerState != NULL) )
+	if ( (Controller == nullptr) || (Controller->PlayerState != nullptr) )
 	{
 		DetachFromControllerPendingDestroy();
 		Destroy();
@@ -729,7 +722,7 @@ void APawn::DisplayDebug(class UCanvas* Canvas, const FDebugDisplayInfo& DebugDi
 // 	{
 // 		UPrimitiveComponent* const MyComp = It.Key().Get();
 // 		UPrimitiveComponent* const OtherComp = It.Value().Get();
-// 		AActor* OtherActor = OtherComp ? OtherComp->GetOwner() : NULL;
+// 		AActor* OtherActor = OtherComp ? OtherComp->GetOwner() : nullptr;
 // 
 // 		YL = HUD->Canvas->DrawText(FString::Printf(TEXT("TOUCHING my %s to %s's %s"), *GetNameSafe(MyComp), *GetNameSafe(OtherActor), *GetNameSafe(OtherComp)));
 // 		YPos += YL;
@@ -737,7 +730,7 @@ void APawn::DisplayDebug(class UCanvas* Canvas, const FDebugDisplayInfo& DebugDi
 // 	}
 
 	FDisplayDebugManager& DisplayDebugManager = Canvas->DisplayDebugManager;
-	if ( PlayerState == NULL )
+	if ( PlayerState == nullptr )
 	{
 		DisplayDebugManager.DrawString(TEXT("NO PlayerState"));
 	}
@@ -756,7 +749,7 @@ void APawn::DisplayDebug(class UCanvas* Canvas, const FDebugDisplayInfo& DebugDi
 	}
 
 	// Controller
-	if ( Controller == NULL )
+	if ( Controller == nullptr )
 	{
 		DisplayDebugManager.SetDrawColor(FColor(255, 0, 0));
 		DisplayDebugManager.DrawString(TEXT("NO Controller"));
@@ -781,7 +774,7 @@ FRotator APawn::GetBaseAimRotation() const
 	// that is by default Controller.Rotation for AI, and camera (crosshair) rotation for human players.
 	FVector POVLoc;
 	FRotator POVRot;
-	if( Controller != NULL && !InFreeCam() )
+	if( Controller != nullptr && !InFreeCam() )
 	{
 		Controller->GetPlayerViewPoint(POVLoc, POVRot);
 		return POVRot;
@@ -805,7 +798,7 @@ bool APawn::InFreeCam() const
 	const APlayerController* PC = Cast<const APlayerController>(Controller);
 	static FName NAME_FreeCam =  FName(TEXT("FreeCam"));
 	static FName NAME_FreeCamDefault = FName(TEXT("FreeCam_Default"));
-	return (PC != NULL && PC->PlayerCameraManager != NULL && (PC->PlayerCameraManager->CameraStyle == NAME_FreeCam || PC->PlayerCameraManager->CameraStyle == NAME_FreeCamDefault) );
+	return (PC != nullptr && PC->PlayerCameraManager != nullptr && (PC->PlayerCameraManager->CameraStyle == NAME_FreeCam || PC->PlayerCameraManager->CameraStyle == NAME_FreeCamDefault) );
 }
 
 void APawn::OutsideWorldBounds()
@@ -814,7 +807,7 @@ void APawn::OutsideWorldBounds()
 	{
 		bProcessingOutsideWorldBounds = true;
 		// AI pawns on the server just destroy
-		if (Role == ROLE_Authority && Cast<APlayerController>(Controller) == NULL)
+		if (Role == ROLE_Authority && Cast<APlayerController>(Controller) == nullptr)
 		{
 			Destroy();
 		}
@@ -829,14 +822,6 @@ void APawn::OutsideWorldBounds()
 	}
 }
 
-void APawn::ClientSetRotation( FRotator NewRotation )
-{
-	if ( Controller != NULL )
-	{
-		Controller->ClientSetRotation(NewRotation);
-	}
-}
-
 void APawn::FaceRotation(FRotator NewControlRotation, float DeltaTime)
 {
 	// Only if we actually are going to use any component of rotation.
@@ -872,24 +857,24 @@ void APawn::FaceRotation(FRotator NewControlRotation, float DeltaTime)
 
 void APawn::DetachFromControllerPendingDestroy()
 {
-	if ( Controller != NULL && Controller->GetPawn() == this )
+	if ( Controller != nullptr && Controller->GetPawn() == this )
 	{
 		Controller->PawnPendingDestroy(this);
-		if (Controller != NULL)
+		if (Controller != nullptr)
 		{
 			Controller->UnPossess();
-			Controller = NULL;
+			Controller = nullptr;
 		}
 	}
 }
 
 AController* APawn::GetDamageInstigator(AController* InstigatedBy, const UDamageType& DamageType) const
 {
-	if ( (InstigatedBy != NULL) && (InstigatedBy != Controller) )
+	if ( (InstigatedBy != nullptr) && (InstigatedBy != Controller) )
 	{
 		return InstigatedBy;
 	}
-	else if ( DamageType.bCausedByWorld && (LastHitBy != NULL) )
+	else if ( DamageType.bCausedByWorld && (LastHitBy != nullptr) )
 	{
 		return LastHitBy;
 	}
@@ -924,7 +909,7 @@ void APawn::PostEditChangeProperty(struct FPropertyChangedEvent& PropertyChanged
 
 void APawn::EnableInput(class APlayerController* PlayerController)
 {
-	if (PlayerController == Controller || PlayerController == NULL)
+	if (PlayerController == Controller || PlayerController == nullptr)
 	{
 		bInputEnabled = true;
 	}
@@ -936,7 +921,7 @@ void APawn::EnableInput(class APlayerController* PlayerController)
 
 void APawn::DisableInput(class APlayerController* PlayerController)
 {
-	if (PlayerController == Controller || PlayerController == NULL)
+	if (PlayerController == Controller || PlayerController == nullptr)
 	{
 		bInputEnabled = false;
 	}
@@ -1020,7 +1005,7 @@ void APawn::PostNetReceiveLocationAndRotation()
 bool APawn::IsBasedOnActor(const AActor* Other) const
 {
 	UPrimitiveComponent * MovementBase = GetMovementBase();
-	AActor* MovementBaseActor = MovementBase ? MovementBase->GetOwner() : NULL;
+	AActor* MovementBaseActor = MovementBase ? MovementBase->GetOwner() : nullptr;
 
 	if (MovementBaseActor && MovementBaseActor == Other)
 	{
@@ -1046,7 +1031,7 @@ bool APawn::IsNetRelevantFor(const AActor* RealViewer, const AActor* ViewTarget,
 	else
 	{
 		UPrimitiveComponent* MovementBase = GetMovementBase();
-		AActor* BaseActor = MovementBase ? MovementBase->GetOwner() : NULL;
+		AActor* BaseActor = MovementBase ? MovementBase->GetOwner() : nullptr;
 		if ( MovementBase && BaseActor && GetMovementComponent() && ((Cast<const USkeletalMeshComponent>(MovementBase)) || (BaseActor == GetOwner())) )
 		{
 			return BaseActor->IsNetRelevantFor(RealViewer, ViewTarget, SrcLocation);
@@ -1097,7 +1082,7 @@ void APawn::MoveIgnoreActorRemove(AActor* ActorToIgnore)
 
 void APawn::PawnMakeNoise(float Loudness, FVector NoiseLocation, bool bUseNoiseMakerLocation, AActor* NoiseMaker)
 {
-	if (NoiseMaker == NULL)
+	if (NoiseMaker == nullptr)
 	{
 		NoiseMaker = this;
 	}
diff --git a/Engine/Source/Runtime/Engine/Private/StreamableManager.cpp b/Engine/Source/Runtime/Engine/Private/StreamableManager.cpp
--- a/Engine/Source/Runtime/Engine/Private/StreamableManager.cpp
+++ b/Engine/Source/Runtime/Engine/Private/StreamableManager.cpp
@@ -77,9 +77,9 @@ public:
 		// When DelegatesToCall falls out of scope it may delete the referenced handles
 	}
 
-	virtual bool IsTickable() const override
+	virtual ETickableTickType GetTickableTickType() const override
 	{
-		return true;
+		return ETickableTickType::Always;
 	}
 
 	virtual bool IsTickableWhenPaused() const override
@@ -452,14 +452,18 @@ void FStreamableHandle::CompleteLoad()
 		ExecuteDelegate(CompleteDelegate, AsShared());
 		UnbindDelegates();
 
-		// Update any meta handles that are still active
-		for (TWeakPtr<FStreamableHandle> WeakHandle : ParentHandles)
+		if (ParentHandles.Num() > 0)
 		{
-			TSharedPtr<FStreamableHandle> Handle = WeakHandle.Pin();
-
-			if (Handle.IsValid())
+			// Update any meta handles that are still active. Copy the array first as elements may be removed from original while iterating
+			TArray<TWeakPtr<FStreamableHandle>> ParentHandlesCopy = ParentHandles;
+			for (TWeakPtr<FStreamableHandle> WeakHandle : ParentHandlesCopy)
 			{
-				Handle->UpdateCombinedHandle();
+				TSharedPtr<FStreamableHandle> Handle = WeakHandle.Pin();
+
+				if (Handle.IsValid())
+				{
+					Handle->UpdateCombinedHandle();
+				}
 			}
 		}
 	}
@@ -1034,6 +1038,12 @@ void FStreamableManager::FindInMemory( FSoftObjectPath& InOutTargetName, struct
 	UE_LOG(LogStreamableManager, Verbose, TEXT("     Searching in memory for %s"), *InOutTargetName.ToString());
 	Existing->Target = StaticFindObject(UObject::StaticClass(), nullptr, *InOutTargetName.ToString());
 
+	if (Existing->Target && Existing->Target->HasAnyInternalFlags(EInternalObjectFlags::AsyncLoading))
+	{
+		// This can get called from PostLoad on async loaded objects, if it is we do not want to return partially loaded objects and instead want to register for their full load
+		Existing->Target = nullptr;
+	}
+
 	UObjectRedirector* Redir = Cast<UObjectRedirector>(Existing->Target);
 	
 	FRedirectedPath RedirectedPath;
diff --git a/Engine/Source/Runtime/Engine/Private/UserDefinedStruct.cpp b/Engine/Source/Runtime/Engine/Private/UserDefinedStruct.cpp
--- a/Engine/Source/Runtime/Engine/Private/UserDefinedStruct.cpp
+++ b/Engine/Source/Runtime/Engine/Private/UserDefinedStruct.cpp
@@ -44,7 +44,7 @@ namespace
 		{
 			for (const FStructVariableDescription& StructVariableDesc : UDSEditorData->VariablesDescriptions)
 			{
-				static const FString TextCategory = TEXT("text"); // Must match UEdGraphSchema_K2::PC_Text
+				static const FName TextCategory = TEXT("text"); // Must match UEdGraphSchema_K2::PC_Text
 				if (StructVariableDesc.Category == TextCategory)
 				{
 					FText StructVariableValue;
diff --git a/Engine/Source/Runtime/Engine/Private/VisualLogger/VisualLogger.cpp b/Engine/Source/Runtime/Engine/Private/VisualLogger/VisualLogger.cpp
--- a/Engine/Source/Runtime/Engine/Private/VisualLogger/VisualLogger.cpp
+++ b/Engine/Source/Runtime/Engine/Private/VisualLogger/VisualLogger.cpp
@@ -51,7 +51,7 @@ namespace
 TMap<const UWorld*, FVisualLogger::RedirectionMapType> FVisualLogger::WorldToRedirectionMap;
 int32 FVisualLogger::bIsRecording = false;
 
-bool FVisualLogger::CheckVisualLogInputInternal(const UObject* Object, const FLogCategoryBase& Category, ELogVerbosity::Type Verbosity, UWorld **World, FVisualLogEntry **CurrentEntry)
+bool FVisualLogger::CheckVisualLogInputInternal(const UObject* Object, const FName& CategoryName, ELogVerbosity::Type Verbosity, UWorld **World, FVisualLogEntry **CurrentEntry)
 {
 	if (FVisualLogger::IsRecording() == false || !Object || !GEngine || GEngine->bDisableAILogging || Object->HasAnyFlags(RF_ClassDefaultObject))
 	{
@@ -59,7 +59,6 @@ bool FVisualLogger::CheckVisualLogInputInternal(const UObject* Object, const FLo
 	}
 
 	FVisualLogger& VisualLogger = FVisualLogger::Get();
-	const FName CategoryName = Category.GetCategoryName();
 	if (VisualLogger.IsBlockedForAllCategories() && VisualLogger.IsWhiteListed(CategoryName) == false)
 	{
 		return false;
@@ -305,8 +304,8 @@ void FVisualLogger::EventLog(const UObject* Object, const FName EventTag1, const
 	SCOPE_CYCLE_COUNTER(STAT_VisualLog);
 	UWorld *World = nullptr;
 	FVisualLogEntry *CurrentEntry = nullptr;
-	const FLogCategory<ELogVerbosity::Log, ELogVerbosity::Log> Category(*Event.Name);
-	if (CheckVisualLogInputInternal(Object, Category, ELogVerbosity::Log, &World, &CurrentEntry) == false)
+	const FName CategoryName(*Event.Name);
+	if (CheckVisualLogInputInternal(Object, CategoryName, ELogVerbosity::Log, &World, &CurrentEntry) == false)
 	{
 		return;
 	}
@@ -331,11 +330,16 @@ void FVisualLogger::EventLog(const UObject* Object, const FName EventTag1, const
 }
 
 void FVisualLogger::NavigationDataDump(const UObject* Object, const FLogCategoryBase& Category, const ELogVerbosity::Type Verbosity, const FBox& Box)
+{
+	NavigationDataDump(Object, Category.GetCategoryName(), Verbosity, Box);
+}
+
+void FVisualLogger::NavigationDataDump(const UObject* Object, const FName& CategoryName, const ELogVerbosity::Type Verbosity, const FBox& Box)
 {
 	SCOPE_CYCLE_COUNTER(STAT_VisualLog);
 	UWorld* World = nullptr;
 	FVisualLogEntry* CurrentEntry = nullptr;
-	if (CheckVisualLogInputInternal(Object, Category, Verbosity, &World, &CurrentEntry) == false || CurrentEntry == nullptr)
+	if (CheckVisualLogInputInternal(Object, CategoryName, Verbosity, &World, &CurrentEntry) == false || CurrentEntry == nullptr)
 	{
 		return;
 	}
@@ -344,7 +348,7 @@ void FVisualLogger::NavigationDataDump(const UObject* Object, const FLogCategory
 	const FNavDataGenerator* Generator = MainNavData ? MainNavData->GetGenerator() : nullptr;
 	if (Generator)
 	{
-		Generator->GrabDebugSnapshot(CurrentEntry, FMath::IsNearlyZero(Box.GetVolume()) ? MainNavData->GetBounds().ExpandBy(FVector(20,20,20)) : Box, Category, Verbosity);
+		Generator->GrabDebugSnapshot(CurrentEntry, FMath::IsNearlyZero(Box.GetVolume()) ? MainNavData->GetBounds().ExpandBy(FVector(20,20,20)) : Box, CategoryName, Verbosity);
 	}
 }
 
diff --git a/Engine/Source/Runtime/Engine/Private/VisualLogger/VisualLoggerKismetLibrary.cpp b/Engine/Source/Runtime/Engine/Private/VisualLogger/VisualLoggerKismetLibrary.cpp
--- a/Engine/Source/Runtime/Engine/Private/VisualLogger/VisualLoggerKismetLibrary.cpp
+++ b/Engine/Source/Runtime/Engine/Private/VisualLogger/VisualLoggerKismetLibrary.cpp
@@ -3,32 +3,84 @@
 #include "VisualLogger/VisualLoggerKismetLibrary.h"
 #include "EngineDefines.h"
 #include "VisualLogger/VisualLogger.h"
+#include "Logging/MessageLog.h"
+
+namespace
+{
+	const FName NAME_Empty;
+	const FName NAME_BlueprintVLog = TEXT("VisLogBP");
+}
 
 UVisualLoggerKismetLibrary::UVisualLoggerKismetLibrary(const FObjectInitializer& ObjectInitializer)
 	: Super(ObjectInitializer)
 {
 }
 
-void UVisualLoggerKismetLibrary::LogText(UObject* WorldContextObject, FString Text, FName CategoryName)
+void UVisualLoggerKismetLibrary::RedirectVislog(UObject* SourceOwner, UObject* DestinationOwner)
+{
+	if (SourceOwner && DestinationOwner)
+	{
+		REDIRECT_OBJECT_TO_VLOG(SourceOwner, DestinationOwner);
+	}
+}
+
+void UVisualLoggerKismetLibrary::EnableRecording(bool bEnabled)
+{
+#if ENABLE_VISUAL_LOG
+	FVisualLogger::Get().SetIsRecording(bEnabled);
+#endif // ENABLE_VISUAL_LOG
+}
+
+void UVisualLoggerKismetLibrary::LogText(UObject* WorldContextObject, FString Text, FName CategoryName, bool bAddToMessageLog)
+{
+#if ENABLE_VISUAL_LOG
+	const ELogVerbosity::Type DefaultVerbosity = ELogVerbosity::Log;
+	FVisualLogger::CategorizedLogf(WorldContextObject, CategoryName, DefaultVerbosity
+		, TEXT("%s"), *Text);
+#endif
+	if (bAddToMessageLog)
+	{
+		FMessageLog(CategoryName).Info(FText::FromString(FString::Printf(TEXT("LogText: '%s'"), *Text)));
+	}
+}
+
+void UVisualLoggerKismetLibrary::LogLocation(UObject* WorldContextObject, FVector Location, FString Text, FLinearColor Color, float Radius, FName CategoryName, bool bAddToMessageLog)
 {
 #if ENABLE_VISUAL_LOG
 	const ELogVerbosity::Type DefaultVerbosity = ELogVerbosity::Log;
-	FVisualLogger::CategorizedLogf(WorldContextObject, FLogCategoryBase(*CategoryName.ToString(), DefaultVerbosity, DefaultVerbosity), DefaultVerbosity, TEXT("%s"), *Text);
+	FVisualLogger::GeometryShapeLogf(WorldContextObject, CategoryName, DefaultVerbosity
+		, Location, Radius, Color.ToFColor(true), TEXT("%s"), *Text);
 #endif
+	if (bAddToMessageLog)
+	{
+		FMessageLog(CategoryName).Info(FText::FromString(FString::Printf(TEXT("LogLocation: '%s' - Location: (%s)")
+			, *Text, *Location.ToString())));
+	}
 }
 
-void UVisualLoggerKismetLibrary::LogLocation(UObject* WorldContextObject, FVector Location, FString Text, FLinearColor Color, float Radius, FName CategoryName)
+void UVisualLoggerKismetLibrary::LogBox(UObject* WorldContextObject, FBox Box, FString Text, FLinearColor ObjectColor, FName CategoryName, bool bAddToMessageLog)
 {
 #if ENABLE_VISUAL_LOG
 	const ELogVerbosity::Type DefaultVerbosity = ELogVerbosity::Log;
-	FVisualLogger::GeometryShapeLogf(WorldContextObject, FLogCategoryBase(*CategoryName.ToString(), DefaultVerbosity, DefaultVerbosity), DefaultVerbosity, Location, Radius, Color.ToFColor(true), TEXT("%s"), *Text);
+	FVisualLogger::GeometryShapeLogf(WorldContextObject, CategoryName, DefaultVerbosity
+		, Box, FMatrix::Identity, ObjectColor.ToFColor(true), TEXT("%s"), *Text);
 #endif
+	if (bAddToMessageLog)
+	{
+		FMessageLog(CategoryName).Info(FText::FromString(FString::Printf(TEXT("LogBox: '%s' - BoxMin: (%s) | BoxMax: (%s)")
+			, *Text, *Box.Min.ToString(), *Box.Max.ToString())));
+	}
 }
 
-void UVisualLoggerKismetLibrary::LogBox(UObject* WorldContextObject, FBox Box, FString Text, FLinearColor ObjectColor, FName CategoryName)
+void UVisualLoggerKismetLibrary::LogSegment(UObject* WorldContextObject, const FVector SegmentStart, const FVector SegmentEnd, FString Text, FLinearColor ObjectColor, const float Thickness, FName CategoryName, bool bAddToMessageLog)
 {
 #if ENABLE_VISUAL_LOG
 	const ELogVerbosity::Type DefaultVerbosity = ELogVerbosity::Log;
-	FVisualLogger::GeometryShapeLogf(WorldContextObject, FLogCategoryBase(*CategoryName.ToString(), DefaultVerbosity, DefaultVerbosity), DefaultVerbosity, Box, FMatrix::Identity, ObjectColor.ToFColor(true), TEXT("%s"), *Text);
+	FVisualLogger::GeometryShapeLogf(WorldContextObject, CategoryName, DefaultVerbosity, SegmentStart, SegmentEnd, ObjectColor.ToFColor(true), Thickness, TEXT("%s"), *Text);
 #endif
+	if (bAddToMessageLog)
+	{
+		FMessageLog(CategoryName).Info(FText::FromString(FString::Printf(TEXT("Log: '%s' - segment from (%s) to (%s)")
+			, *Text, *SegmentStart.ToString(), *SegmentStart.ToString())));
+	}
 }
diff --git a/Engine/Source/Runtime/Engine/Private/World.cpp b/Engine/Source/Runtime/Engine/Private/World.cpp
--- a/Engine/Source/Runtime/Engine/Private/World.cpp
+++ b/Engine/Source/Runtime/Engine/Private/World.cpp
@@ -1488,9 +1488,8 @@ UWorld* UWorld::CreateWorld(const EWorldType::Type InWorldType, bool bInformEngi
 
 void UWorld::RemoveActor(AActor* Actor, bool bShouldModifyLevel)
 {
-	bool	bSuccessfulRemoval = false;
 	ULevel* CheckLevel = Actor->GetLevel();
-	int32 ActorListIndex = CheckLevel->Actors.Find( Actor );
+	const int32 ActorListIndex = CheckLevel->Actors.Find( Actor );
 	// Search the entire list.
 	if( ActorListIndex != INDEX_NONE )
 	{
@@ -1504,37 +1503,11 @@ void UWorld::RemoveActor(AActor* Actor, bool bShouldModifyLevel)
 			CheckLevel->Actors[ActorListIndex]->Modify();
 		}
 		
-		CheckLevel->Actors[ActorListIndex] = NULL;
-		bSuccessfulRemoval = true;		
+		CheckLevel->Actors[ActorListIndex] = nullptr;
 	}
 
 	// Remove actor from network list
 	RemoveNetworkActor( Actor );
-
-	// TTP 281860: Callstack will hopefully indicate how the actors array ends up without the required default actors
-	check(CheckLevel->Actors.Num() >= 2);
-
-	if ( !bSuccessfulRemoval && !( Actor->GetFlags() & RF_Transactional ) )
-	{
-		//CheckLevel->Actors is a transactional array so it is very likely that non-transactional
-		//actors could be missing from the array if the array was reverted to a state before they
-		//existed. (but they won't be reverted since they are non-transactional)
-		bSuccessfulRemoval = true;
-	}
-
-	if ( !bSuccessfulRemoval )
-	{
-		// TTP 270000: Trying to track down why certain actors aren't in the level actor list when saving.  If we're reinstancing, dump the list
-		{
-			UE_LOG(LogWorld, Log, TEXT("--- Actors Currently in %s ---"), *CheckLevel->GetPathName());
-			for (int32 ActorIdx = 0; ActorIdx < CheckLevel->Actors.Num(); ActorIdx++)
-			{
-				AActor* CurrentActor = CheckLevel->Actors[ActorIdx];
-				UE_LOG(LogWorld, Log, TEXT("  %s"), (CurrentActor ? *CurrentActor->GetPathName() : TEXT("NONE")));
-			}
-		}
-		ensureMsgf(false, TEXT("Could not remove actor %s from world (check level is %s)"), *Actor->GetPathName(), *CheckLevel->GetPathName());
-	}
 }
 
 
@@ -2632,6 +2605,18 @@ public:
 		}
 		return *this;
 	}
+
+	virtual FArchive& operator<<(FSoftObjectPtr& Value) override
+	{
+		// Explicitly do nothing, we don't want to accidentally do PIE fixups
+		return *this;
+	}
+
+	virtual FArchive& operator<<(FSoftObjectPath& Value) override
+	{
+		// Explicitly do nothing, we don't want to accidentally do PIE fixups
+		return *this;
+	}
 };
 
 UWorld* UWorld::DuplicateWorldForPIE(const FString& PackageName, UWorld* OwningWorld)
@@ -2686,6 +2671,7 @@ UWorld* UWorld::DuplicateWorldForPIE(const FString& PackageName, UWorld* OwningW
 	PIELevelPackage->SetPackageFlags(PKG_PlayInEditor);
 	PIELevelPackage->PIEInstanceID = PIEInstanceID;
 	PIELevelPackage->SetGuid( EditorLevelPackage->GetGuid() );
+	PIELevelPackage->MarkAsFullyLoaded();
 
 	ULevel::StreamedLevelsOwningWorld.Add(PIELevelPackage->GetFName(), OwningWorld);
 	UWorld* PIELevelWorld = CastChecked<UWorld>(StaticDuplicateObject(EditorLevelWorld, PIELevelPackage, EditorLevelWorld->GetFName(), RF_AllFlags, nullptr, EDuplicateMode::PIE));
@@ -2709,6 +2695,8 @@ UWorld* UWorld::DuplicateWorldForPIE(const FString& PackageName, UWorld* OwningW
 		ULevel* EditorLevel = EditorLevelWorld->PersistentLevel;
 		ULevel* PIELevel = PIELevelWorld->PersistentLevel;
 
+		// If editor has run construction scripts or applied level offset, we dont do it again
+		PIELevel->bAlreadyMovedActors = EditorLevel->bAlreadyMovedActors;
 		PIELevel->bHasRerunConstructionScripts = EditorLevel->bHasRerunConstructionScripts;
 
 		// Fixup model components. The index buffers have been created for the components in the EditorWorld and the order
@@ -6423,6 +6411,7 @@ static ULevel* DuplicateLevelWithPrefix(ULevel* InLevel, int32 InstanceID )
 	NewPackage->PIEInstanceID = InstanceID;
 	NewPackage->FileName = OriginalPackage->FileName;
 	NewPackage->SetGuid( OriginalPackage->GetGuid() );
+	NewPackage->MarkAsFullyLoaded();
 
 	FSoftObjectPath::AddPIEPackageName(NewPackage->GetFName());
 
diff --git a/Engine/Source/Runtime/Engine/Public/AI/NavDataGenerator.h b/Engine/Source/Runtime/Engine/Public/AI/NavDataGenerator.h
--- a/Engine/Source/Runtime/Engine/Public/AI/NavDataGenerator.h
+++ b/Engine/Source/Runtime/Engine/Public/AI/NavDataGenerator.h
@@ -46,6 +46,9 @@ public:
 
 #if !(UE_BUILD_SHIPPING || UE_BUILD_TEST)
 	virtual void ExportNavigationData(const FString& FileName) const {}
+	virtual void GrabDebugSnapshot(struct FVisualLogEntry* Snapshot, const FBox& BoundingBox, const FName& CategoryName, ELogVerbosity::Type Verbosity) const {}
+
+	DEPRECATED(4.19, "This version of GrabDebugSnapshot has been deprecated. Please use the other version of the function.")
 	virtual void GrabDebugSnapshot(struct FVisualLogEntry* Snapshot, const FBox& BoundingBox, const struct FLogCategoryBase& Category, ELogVerbosity::Type Verbosity) const {}
 #endif
 
diff --git a/Engine/Source/Runtime/Engine/Public/AI/Navigation/RecastNavMeshGenerator.h b/Engine/Source/Runtime/Engine/Public/AI/Navigation/RecastNavMeshGenerator.h
--- a/Engine/Source/Runtime/Engine/Public/AI/Navigation/RecastNavMeshGenerator.h
+++ b/Engine/Source/Runtime/Engine/Public/AI/Navigation/RecastNavMeshGenerator.h
@@ -434,7 +434,7 @@ public:
 	
 #if !(UE_BUILD_SHIPPING || UE_BUILD_TEST) && ENABLE_VISUAL_LOG
 	virtual void ExportNavigationData(const FString& FileName) const override;
-	virtual void GrabDebugSnapshot(struct FVisualLogEntry* Snapshot, const FBox& BoundingBox, const struct FLogCategoryBase& Category, ELogVerbosity::Type Verbosity) const override;
+	virtual void GrabDebugSnapshot(struct FVisualLogEntry* Snapshot, const FBox& BoundingBox, const FName& CategoryName, ELogVerbosity::Type Verbosity) const override;
 #endif
 
 	/** 
diff --git a/Engine/Source/Runtime/Engine/Public/EngineFwd.h b/Engine/Source/Runtime/Engine/Public/EngineFwd.h
--- a/Engine/Source/Runtime/Engine/Public/EngineFwd.h
+++ b/Engine/Source/Runtime/Engine/Public/EngineFwd.h
@@ -169,7 +169,7 @@ class HHitProxy;
 class INavLinkCustomInterface;
 
 enum class ECacheApplyPhase;
-enum class ETeleportType;
+enum class ETeleportType : uint8;
 
 struct FActiveSound;
 struct FActorComponentTickFunction;
diff --git a/Engine/Source/Runtime/Engine/Public/MaterialExpressionIO.h b/Engine/Source/Runtime/Engine/Public/MaterialExpressionIO.h
--- a/Engine/Source/Runtime/Engine/Public/MaterialExpressionIO.h
+++ b/Engine/Source/Runtime/Engine/Public/MaterialExpressionIO.h
@@ -31,7 +31,7 @@ struct FExpressionInput
 	 * Optional name of the input.  
 	 * Note that this is the only member which is not derived from the output currently connected. 
 	 */
-	FString						InputName;
+	FName						InputName;
 
 	int32						Mask,
 								MaskR,
@@ -111,7 +111,7 @@ struct TStructOpsTypeTraits<FExpressionInput>
 
 struct FExpressionOutput
 {
-	FString	OutputName;
+	FName	OutputName;
 	int32	Mask,
 		MaskR,
 		MaskG,
@@ -126,7 +126,7 @@ struct FExpressionOutput
 		MaskA(InMaskA)
 	{}
 
-	FExpressionOutput(FString InOutputName, int32 InMask = 0, int32 InMaskR = 0, int32 InMaskG = 0, int32 InMaskB = 0, int32 InMaskA = 0) :
+	FExpressionOutput(FName InOutputName, int32 InMask = 0, int32 InMaskR = 0, int32 InMaskG = 0, int32 InMaskB = 0, int32 InMaskA = 0) :
 		OutputName(InOutputName),
 		Mask(InMask),
 		MaskR(InMaskR),
diff --git a/Engine/Source/Runtime/Engine/Public/MaterialShared.h b/Engine/Source/Runtime/Engine/Public/MaterialShared.h
--- a/Engine/Source/Runtime/Engine/Public/MaterialShared.h
+++ b/Engine/Source/Runtime/Engine/Public/MaterialShared.h
@@ -61,6 +61,7 @@ DECLARE_LOG_CATEGORY_EXTERN(LogMaterial,Log,Verbose);
 
 /** Creates a string that represents the given quality level. */
 extern void GetMaterialQualityLevelName(EMaterialQualityLevel::Type InMaterialQualityLevel, FString& OutName);
+extern FName GetMaterialQualityLevelFName(EMaterialQualityLevel::Type InMaterialQualityLevel);
 
 inline bool IsSubsurfaceShadingModel(EMaterialShadingModel ShadingModel)
 {
diff --git a/Engine/Source/Runtime/Engine/Public/Tickable.h b/Engine/Source/Runtime/Engine/Public/Tickable.h
--- a/Engine/Source/Runtime/Engine/Public/Tickable.h
+++ b/Engine/Source/Runtime/Engine/Public/Tickable.h
@@ -10,11 +10,80 @@
 #include "CoreMinimal.h"
 #include "Stats/Stats.h"
 
+/**
+ * Enum used to convey whether a tickable object will always tick, conditionally tick, or never tick.
+ */
+enum class ETickableTickType : uint8
+{
+	/** Use IsTickable to determine whether to tick */
+	Conditional,
+
+	/** Always tick the object */
+	Always,
+
+	/** Never tick the object, do not add to tickables array */
+	Never        
+};
+
 /**
  * Base class for tickable objects
  */
 class FTickableObjectBase
 {
+protected:
+	struct FTickableObjectEntry
+	{
+		FTickableObjectBase* TickableObject;
+		ETickableTickType TickType;
+
+		bool operator==(FTickableObjectBase* OtherObject) const { return TickableObject == OtherObject; }
+	};
+
+	static void AddTickableObject(TArray<FTickableObjectEntry>& TickableObjects, FTickableObjectBase* TickableObject)
+	{
+		check(!TickableObjects.Contains(TickableObject));
+		const ETickableTickType TickType = TickableObject->GetTickableTickType();
+		if (TickType != ETickableTickType::Never)
+		{
+			TickableObjects.Add({ TickableObject, TickType });
+		}
+	}
+
+	static void RemoveTickableObject(TArray<FTickableObjectEntry>& TickableObjects, FTickableObjectBase* TickableObject, const bool bIsTickingObjects)
+	{
+		const int32 Pos = TickableObjects.IndexOfByKey(TickableObject);
+
+#if 0 // virtual from destructor doesn't work ... need to rethink how to do warning
+		// ensure that GetTickableTickType did not change over time
+		switch (TickableObject->GetTickableTickType())
+		{
+		case ETickableTickType::Always:
+			ensureMsgf(Pos != INDEX_NONE && TickableObjects[Pos].TickType == ETickableTickType::Always, TEXT("TickType has changed since object was created. Result of GetTickableTickType must be invariant for a given object."));
+			break;
+
+		case ETickableTickType::Conditional:
+			ensureMsgf(Pos != INDEX_NONE && TickableObjects[Pos].TickType == ETickableTickType::Conditional, TEXT("TickType has changed since object was created. Result of GetTickableTickType must be invariant for a given object."));
+			break;
+
+		case ETickableTickType::Never:
+			ensureMsgf(Pos == INDEX_NONE, TEXT("TickType has changed since object was created. Result of GetTickableTickType must be invariant for a given object."));
+			break;
+		}
+#endif
+
+		if (Pos != INDEX_NONE)
+		{
+			if (bIsTickingObjects)
+			{
+				TickableObjects[Pos].TickableObject = nullptr;
+			}
+			else
+			{
+				TickableObjects.RemoveAt(Pos);
+			}
+		}
+	}
+
 public:
 	/**
 	 * Pure virtual that must be overloaded by the inheriting class. It will
@@ -26,14 +95,25 @@ public:
 	virtual void Tick( float DeltaTime ) = 0;
 
 	/**
-	 * Pure virtual that must be overloaded by the inheriting class. It is
+	* Virtual that can be overloaded by the inheriting class. It is
+	* used to determine whether an object will ever be able to tick, and if not,
+	* it will not get added to the tickable objects array. If the tickable tick type
+	* is Conditional then the virtual IsTickable will be called to determine whether
+	* to tick the object on each given frame
+	*
+	* @return	true if object will ever want to be ticked, false otherwise.
+	*/
+	virtual ETickableTickType GetTickableTickType() const { return ETickableTickType::Conditional; }
+
+	/**
+	 * Virtual that can be overloaded by the inheriting class. It is
 	 * used to determine whether an object is ready to be ticked. This is 
 	 * required for example for all UObject derived classes as they might be
 	 * loaded async and therefore won't be ready immediately.
 	 *
-	 * @return	true if class is ready to be ticked, false otherwise.
+	 * @return	true if object is ready to be ticked, false otherwise.
 	 */
-	virtual bool IsTickable() const = 0;
+	virtual bool IsTickable() const { return true; }
 
 	/** return the stat id to use for this tickable **/
 	virtual TStatId GetStatId() const = 0;
@@ -45,8 +125,10 @@ public:
  */
 class ENGINE_API FTickableGameObject : public FTickableObjectBase
 {
+private:
 	/** Static array of tickable objects */
-	static TArray<FTickableGameObject*> TickableObjects;
+	static TArray<FTickableObjectEntry> TickableObjects;
+	static TArray<FTickableGameObject*> PendingTickableObjects;
 	static bool bIsTickingObjects;
 
 public:
@@ -56,8 +138,9 @@ public:
 	 */
 	FTickableGameObject()
 	{
+		check(!PendingTickableObjects.Contains(this));
 		check(!TickableObjects.Contains(this));
-		TickableObjects.Add( this );
+		PendingTickableObjects.Add(this);
 	}
 
 	/**
@@ -65,16 +148,9 @@ public:
 	 */
 	virtual ~FTickableGameObject()
 	{
-		// make sure this tickable object was registered from the game thread
-		const int32 Pos = TickableObjects.Find(this);
-		check(Pos!=INDEX_NONE);
-		if (bIsTickingObjects)
-		{
-			TickableObjects[Pos] = nullptr;
-		}
-		else
+		if (PendingTickableObjects.Remove(this) == 0)
 		{
-			TickableObjects.RemoveAt(Pos);
+			RemoveTickableObject(TickableObjects, this, bIsTickingObjects);
 		}
 	}
 
