Copying //UE4/Dev-Framework to //UE4/Dev-Main (Source: //UE4/Dev-Framework @ 3716594)

diff --git a/Engine/Source/Runtime/Engine/Public/VisualLogger/VisualLogger.h b/Engine/Source/Runtime/Engine/Public/VisualLogger/VisualLogger.h
--- a/Engine/Source/Runtime/Engine/Public/VisualLogger/VisualLogger.h
+++ b/Engine/Source/Runtime/Engine/Public/VisualLogger/VisualLogger.h
@@ -144,28 +144,40 @@ class ENGINE_API FVisualLogger : public FOutputDevice
 public:
 	// Regular text log
 	VARARG_DECL(static void, static void, return, CategorizedLogf, VARARG_NONE, const TCHAR*, VARARG_EXTRA(const UObject* LogOwner) VARARG_EXTRA(const FLogCategoryBase& Category) VARARG_EXTRA(ELogVerbosity::Type Verbosity), VARARG_EXTRA(LogOwner) VARARG_EXTRA(Category) VARARG_EXTRA(Verbosity));
+	VARARG_DECL(static void, static void, return, CategorizedLogf, VARARG_NONE, const TCHAR*, VARARG_EXTRA(const UObject* LogOwner) VARARG_EXTRA(const FName& CategoryName) VARARG_EXTRA(ELogVerbosity::Type Verbosity), VARARG_EXTRA(LogOwner) VARARG_EXTRA(CategoryName) VARARG_EXTRA(Verbosity));
 	// Segment log
 	VARARG_DECL(static void, static void, return, GeometryShapeLogf, VARARG_NONE, const TCHAR*, VARARG_EXTRA(const UObject* LogOwner) VARARG_EXTRA(const FLogCategoryBase& Category) VARARG_EXTRA(ELogVerbosity::Type Verbosity) VARARG_EXTRA(const FVector& Start) VARARG_EXTRA(const FVector& End) VARARG_EXTRA(const FColor& Color) VARARG_EXTRA(const uint16 Thickness), VARARG_EXTRA(LogOwner) VARARG_EXTRA(Category) VARARG_EXTRA(Verbosity) VARARG_EXTRA(Start) VARARG_EXTRA(End) VARARG_EXTRA(Color) VARARG_EXTRA(Thickness));
+	VARARG_DECL(static void, static void, return, GeometryShapeLogf, VARARG_NONE, const TCHAR*, VARARG_EXTRA(const UObject* LogOwner) VARARG_EXTRA(const FName& CategoryName) VARARG_EXTRA(ELogVerbosity::Type Verbosity) VARARG_EXTRA(const FVector& Start) VARARG_EXTRA(const FVector& End) VARARG_EXTRA(const FColor& Color) VARARG_EXTRA(const uint16 Thickness), VARARG_EXTRA(LogOwner) VARARG_EXTRA(CategoryName) VARARG_EXTRA(Verbosity) VARARG_EXTRA(Start) VARARG_EXTRA(End) VARARG_EXTRA(Color) VARARG_EXTRA(Thickness));
 	// Location/Shpere log
 	VARARG_DECL(static void, static void, return, GeometryShapeLogf, VARARG_NONE, const TCHAR*, VARARG_EXTRA(const UObject* LogOwner) VARARG_EXTRA(const FLogCategoryBase& Category) VARARG_EXTRA(ELogVerbosity::Type Verbosity) VARARG_EXTRA(const FVector& Location) VARARG_EXTRA(float Radius) VARARG_EXTRA(const FColor& Color), VARARG_EXTRA(LogOwner) VARARG_EXTRA(Category) VARARG_EXTRA(Verbosity) VARARG_EXTRA(Location) VARARG_EXTRA(Radius) VARARG_EXTRA(Color));
+	VARARG_DECL(static void, static void, return, GeometryShapeLogf, VARARG_NONE, const TCHAR*, VARARG_EXTRA(const UObject* LogOwner) VARARG_EXTRA(const FName& CategoryName) VARARG_EXTRA(ELogVerbosity::Type Verbosity) VARARG_EXTRA(const FVector& Location) VARARG_EXTRA(float Radius) VARARG_EXTRA(const FColor& Color), VARARG_EXTRA(LogOwner) VARARG_EXTRA(CategoryName) VARARG_EXTRA(Verbosity) VARARG_EXTRA(Location) VARARG_EXTRA(Radius) VARARG_EXTRA(Color));
 	// Box log
 	VARARG_DECL(static void, static void, return, GeometryShapeLogf, VARARG_NONE, const TCHAR*, VARARG_EXTRA(const UObject* LogOwner) VARARG_EXTRA(const FLogCategoryBase& Category) VARARG_EXTRA(ELogVerbosity::Type Verbosity) VARARG_EXTRA(const FBox& Box) VARARG_EXTRA(const FMatrix& Matrix) VARARG_EXTRA(const FColor& Color), VARARG_EXTRA(LogOwner) VARARG_EXTRA(Category) VARARG_EXTRA(Verbosity) VARARG_EXTRA(Box) VARARG_EXTRA(Matrix) VARARG_EXTRA(Color));
+	VARARG_DECL(static void, static void, return, GeometryShapeLogf, VARARG_NONE, const TCHAR*, VARARG_EXTRA(const UObject* LogOwner) VARARG_EXTRA(const FName& CategoryName) VARARG_EXTRA(ELogVerbosity::Type Verbosity) VARARG_EXTRA(const FBox& Box) VARARG_EXTRA(const FMatrix& Matrix) VARARG_EXTRA(const FColor& Color), VARARG_EXTRA(LogOwner) VARARG_EXTRA(CategoryName) VARARG_EXTRA(Verbosity) VARARG_EXTRA(Box) VARARG_EXTRA(Matrix) VARARG_EXTRA(Color));
 	// Cone log
 	VARARG_DECL(static void, static void, return, GeometryShapeLogf, VARARG_NONE, const TCHAR*, VARARG_EXTRA(const UObject* LogOwner) VARARG_EXTRA(const FLogCategoryBase& Category) VARARG_EXTRA(ELogVerbosity::Type Verbosity) VARARG_EXTRA(const FVector& Orgin) VARARG_EXTRA(const FVector& Direction) VARARG_EXTRA(const float Length) VARARG_EXTRA(const float Angle) VARARG_EXTRA(const FColor& Color), VARARG_EXTRA(LogOwner) VARARG_EXTRA(Category) VARARG_EXTRA(Verbosity) VARARG_EXTRA(Orgin) VARARG_EXTRA(Direction) VARARG_EXTRA(Length) VARARG_EXTRA(Angle) VARARG_EXTRA(Color));
+	VARARG_DECL(static void, static void, return, GeometryShapeLogf, VARARG_NONE, const TCHAR*, VARARG_EXTRA(const UObject* LogOwner) VARARG_EXTRA(const FName& CategoryName) VARARG_EXTRA(ELogVerbosity::Type Verbosity) VARARG_EXTRA(const FVector& Orgin) VARARG_EXTRA(const FVector& Direction) VARARG_EXTRA(const float Length) VARARG_EXTRA(const float Angle) VARARG_EXTRA(const FColor& Color), VARARG_EXTRA(LogOwner) VARARG_EXTRA(CategoryName) VARARG_EXTRA(Verbosity) VARARG_EXTRA(Orgin) VARARG_EXTRA(Direction) VARARG_EXTRA(Length) VARARG_EXTRA(Angle) VARARG_EXTRA(Color));
 	// Cylinder log
 	VARARG_DECL(static void, static void, return, GeometryShapeLogf, VARARG_NONE, const TCHAR*, VARARG_EXTRA(const UObject* LogOwner) VARARG_EXTRA(const FLogCategoryBase& Category) VARARG_EXTRA(ELogVerbosity::Type Verbosity) VARARG_EXTRA(const FVector& Start) VARARG_EXTRA(const FVector& End) VARARG_EXTRA(const float Radius) VARARG_EXTRA(const FColor& Color), VARARG_EXTRA(LogOwner) VARARG_EXTRA(Category) VARARG_EXTRA(Verbosity) VARARG_EXTRA(Start) VARARG_EXTRA(End) VARARG_EXTRA(Radius) VARARG_EXTRA(Color));
+	VARARG_DECL(static void, static void, return, GeometryShapeLogf, VARARG_NONE, const TCHAR*, VARARG_EXTRA(const UObject* LogOwner) VARARG_EXTRA(const FName& CategoryName) VARARG_EXTRA(ELogVerbosity::Type Verbosity) VARARG_EXTRA(const FVector& Start) VARARG_EXTRA(const FVector& End) VARARG_EXTRA(const float Radius) VARARG_EXTRA(const FColor& Color), VARARG_EXTRA(LogOwner) VARARG_EXTRA(CategoryName) VARARG_EXTRA(Verbosity) VARARG_EXTRA(Start) VARARG_EXTRA(End) VARARG_EXTRA(Radius) VARARG_EXTRA(Color));
 	// Capsule log
 	VARARG_DECL(static void, static void, return, GeometryShapeLogf, VARARG_NONE, const TCHAR*, VARARG_EXTRA(const UObject* LogOwner) VARARG_EXTRA(const FLogCategoryBase& Category) VARARG_EXTRA(ELogVerbosity::Type Verbosity) VARARG_EXTRA(const FVector& Center) VARARG_EXTRA(float HalfHeight) VARARG_EXTRA(float Radius) VARARG_EXTRA(const FQuat & Rotation) VARARG_EXTRA(const FColor& Color), VARARG_EXTRA(LogOwner) VARARG_EXTRA(Category) VARARG_EXTRA(Verbosity) VARARG_EXTRA(Center) VARARG_EXTRA(HalfHeight) VARARG_EXTRA(Radius) VARARG_EXTRA(Rotation) VARARG_EXTRA(Color));
+	VARARG_DECL(static void, static void, return, GeometryShapeLogf, VARARG_NONE, const TCHAR*, VARARG_EXTRA(const UObject* LogOwner) VARARG_EXTRA(const FName& CategoryName) VARARG_EXTRA(ELogVerbosity::Type Verbosity) VARARG_EXTRA(const FVector& Center) VARARG_EXTRA(float HalfHeight) VARARG_EXTRA(float Radius) VARARG_EXTRA(const FQuat & Rotation) VARARG_EXTRA(const FColor& Color), VARARG_EXTRA(LogOwner) VARARG_EXTRA(CategoryName) VARARG_EXTRA(Verbosity) VARARG_EXTRA(Center) VARARG_EXTRA(HalfHeight) VARARG_EXTRA(Radius) VARARG_EXTRA(Rotation) VARARG_EXTRA(Color));
 	// NavArea/Extruded convex log
 	VARARG_DECL(static void, static void, return, NavAreaShapeLogf, VARARG_NONE, const TCHAR*, VARARG_EXTRA(const UObject* LogOwner) VARARG_EXTRA(const FLogCategoryBase& Category) VARARG_EXTRA(ELogVerbosity::Type Verbosity) VARARG_EXTRA(const TArray<FVector> &ConvexPoints) VARARG_EXTRA(float MinZ) VARARG_EXTRA(float MaxZ) VARARG_EXTRA(const FColor& Color), VARARG_EXTRA(LogOwner) VARARG_EXTRA(Category) VARARG_EXTRA(Verbosity) VARARG_EXTRA(ConvexPoints) VARARG_EXTRA(MinZ) VARARG_EXTRA(MaxZ) VARARG_EXTRA(Color));
+	VARARG_DECL(static void, static void, return, NavAreaShapeLogf, VARARG_NONE, const TCHAR*, VARARG_EXTRA(const UObject* LogOwner) VARARG_EXTRA(const FName& CategoryName) VARARG_EXTRA(ELogVerbosity::Type Verbosity) VARARG_EXTRA(const TArray<FVector> &ConvexPoints) VARARG_EXTRA(float MinZ) VARARG_EXTRA(float MaxZ) VARARG_EXTRA(const FColor& Color), VARARG_EXTRA(LogOwner) VARARG_EXTRA(CategoryName) VARARG_EXTRA(Verbosity) VARARG_EXTRA(ConvexPoints) VARARG_EXTRA(MinZ) VARARG_EXTRA(MaxZ) VARARG_EXTRA(Color));
 	// 3d Mesh log
 	VARARG_DECL(static void, static void, return, GeometryShapeLogf, VARARG_NONE, const TCHAR*, VARARG_EXTRA(const UObject* LogOwner) VARARG_EXTRA(const FLogCategoryBase& Category) VARARG_EXTRA(ELogVerbosity::Type Verbosity) VARARG_EXTRA(const TArray<FVector> &Vertices) VARARG_EXTRA(const TArray<int32> &Indices) VARARG_EXTRA(const FColor& Color), VARARG_EXTRA(LogOwner) VARARG_EXTRA(Category) VARARG_EXTRA(Verbosity) VARARG_EXTRA(Vertices) VARARG_EXTRA(Indices) VARARG_EXTRA(Color));
+	VARARG_DECL(static void, static void, return, GeometryShapeLogf, VARARG_NONE, const TCHAR*, VARARG_EXTRA(const UObject* LogOwner) VARARG_EXTRA(const FName& CategoryName) VARARG_EXTRA(ELogVerbosity::Type Verbosity) VARARG_EXTRA(const TArray<FVector> &Vertices) VARARG_EXTRA(const TArray<int32> &Indices) VARARG_EXTRA(const FColor& Color), VARARG_EXTRA(LogOwner) VARARG_EXTRA(CategoryName) VARARG_EXTRA(Verbosity) VARARG_EXTRA(Vertices) VARARG_EXTRA(Indices) VARARG_EXTRA(Color));
 	// 2d Convex shape
 	VARARG_DECL(static void, static void, return, GeometryConvexLogf, VARARG_NONE, const TCHAR*, VARARG_EXTRA(const UObject* LogOwner) VARARG_EXTRA(const FLogCategoryBase& Category) VARARG_EXTRA(ELogVerbosity::Type Verbosity) VARARG_EXTRA(const TArray<FVector> &Points) VARARG_EXTRA(const FColor& Color), VARARG_EXTRA(LogOwner) VARARG_EXTRA(Category) VARARG_EXTRA(Verbosity) VARARG_EXTRA(Points) VARARG_EXTRA(Color));
+	VARARG_DECL(static void, static void, return, GeometryConvexLogf, VARARG_NONE, const TCHAR*, VARARG_EXTRA(const UObject* LogOwner) VARARG_EXTRA(const FName& CategoryName) VARARG_EXTRA(ELogVerbosity::Type Verbosity) VARARG_EXTRA(const TArray<FVector> &Points) VARARG_EXTRA(const FColor& Color), VARARG_EXTRA(LogOwner) VARARG_EXTRA(CategoryName) VARARG_EXTRA(Verbosity) VARARG_EXTRA(Points) VARARG_EXTRA(Color));
 	//Histogram data
 	VARARG_DECL(static void, static void, return, HistogramDataLogf, VARARG_NONE, const TCHAR*, VARARG_EXTRA(const UObject* LogOwner) VARARG_EXTRA(const FLogCategoryBase& Category) VARARG_EXTRA(ELogVerbosity::Type Verbosity) VARARG_EXTRA(FName GraphName) VARARG_EXTRA(FName DataName) VARARG_EXTRA(const FVector2D& Data) VARARG_EXTRA(const FColor& Color), VARARG_EXTRA(LogOwner) VARARG_EXTRA(Category) VARARG_EXTRA(Verbosity) VARARG_EXTRA(GraphName) VARARG_EXTRA(DataName) VARARG_EXTRA(Data) VARARG_EXTRA(Color));
+	VARARG_DECL(static void, static void, return, HistogramDataLogf, VARARG_NONE, const TCHAR*, VARARG_EXTRA(const UObject* LogOwner) VARARG_EXTRA(const FName& CategoryName) VARARG_EXTRA(ELogVerbosity::Type Verbosity) VARARG_EXTRA(FName GraphName) VARARG_EXTRA(FName DataName) VARARG_EXTRA(const FVector2D& Data) VARARG_EXTRA(const FColor& Color), VARARG_EXTRA(LogOwner) VARARG_EXTRA(CategoryName) VARARG_EXTRA(Verbosity) VARARG_EXTRA(GraphName) VARARG_EXTRA(DataName) VARARG_EXTRA(Data) VARARG_EXTRA(Color));
 	// Navigation data debug snapshot
 	static void NavigationDataDump(const UObject* LogOwner, const FLogCategoryBase& Category, const ELogVerbosity::Type Verbosity, const FBox& Box);
+	static void NavigationDataDump(const UObject* LogOwner, const FName& CategoryName, const ELogVerbosity::Type Verbosity, const FBox& Box);
 
 	/** Log events */
 	static void EventLog(const UObject* LogOwner, const FName EventTag1, const FVisualLogEventBase& Event1, const FName EventTag2 = NAME_None, const FName EventTag3 = NAME_None, const FName EventTag4 = NAME_None, const FName EventTag5 = NAME_None, const FName EventTag6 = NAME_None);
@@ -265,7 +277,7 @@ public:
 	const TMap<FName, FVisualLogExtensionInterface*>& GetAllExtensions() const { return AllExtensions; }
 
 	/** internal check for each usage of visual logger */
-	static bool CheckVisualLogInputInternal(const UObject* Object, const FLogCategoryBase& Category, ELogVerbosity::Type Verbosity, UWorld **World, FVisualLogEntry **CurrentEntry);
+	static bool CheckVisualLogInputInternal(const UObject* Object, const FName& CategoryName, ELogVerbosity::Type Verbosity, UWorld **World, FVisualLogEntry **CurrentEntry);
 
 	typedef TMap<UObject*, TArray<TWeakObjectPtr<const UObject> > > RedirectionMapType;
 	static RedirectionMapType& GetRedirectionMap(const UObject* InObject);
@@ -337,6 +349,13 @@ protected:
 // Unfortunately needs to be a #define since it uses GET_VARARGS_RESULT which uses the va_list stuff which operates on the
 // current function, so we can't easily call a function
 #define COLLAPSED_LOGF(SerializeFunc) \
+	SCOPE_CYCLE_COUNTER(STAT_VisualLog); \
+	UWorld *World = nullptr; \
+	FVisualLogEntry *CurrentEntry = nullptr; \
+	if (CheckVisualLogInputInternal(Object, CategoryName, Verbosity, &World, &CurrentEntry) == false) \
+	{ \
+		return; \
+	} \
 	int32	BufferSize	= 1024; \
 	TCHAR*	Buffer		= nullptr; \
 	int32	Result		= -1; \
@@ -363,180 +382,190 @@ protected:
 	SerializeFunc; \
 	FMemory::SystemFree(AllocatedBuffer);
 
+inline 
+VARARG_BODY(void, FVisualLogger::CategorizedLogf, const TCHAR*, VARARG_EXTRA(const UObject* Object) VARARG_EXTRA(const FLogCategoryBase& Category) VARARG_EXTRA(ELogVerbosity::Type Verbosity)) 
+{ 
+	const FName CategoryName = Category.GetCategoryName(); 
+	COLLAPSED_LOGF(
+		CurrentEntry->AddText(Buffer, CategoryName, Verbosity);
+	);
+}
+
 inline
-VARARG_BODY(void, FVisualLogger::CategorizedLogf, const TCHAR*, VARARG_EXTRA(const UObject* Object) VARARG_EXTRA(const FLogCategoryBase& Category) VARARG_EXTRA(ELogVerbosity::Type Verbosity))
+VARARG_BODY(void, FVisualLogger::CategorizedLogf, const TCHAR*, VARARG_EXTRA(const UObject* Object) VARARG_EXTRA(const FName& CategoryName) VARARG_EXTRA(ELogVerbosity::Type Verbosity))
 {
-	SCOPE_CYCLE_COUNTER(STAT_VisualLog);
-	UWorld *World = nullptr;
-	FVisualLogEntry *CurrentEntry = nullptr;
-	if (CheckVisualLogInputInternal(Object, Category, Verbosity, &World, &CurrentEntry) == false)
-	{
-		return;
-	}
-
 	COLLAPSED_LOGF(
-		CurrentEntry->AddText(Buffer, Category.GetCategoryName(), Verbosity);
+		CurrentEntry->AddText(Buffer, CategoryName, Verbosity);
 	);
 }
 
 inline
 VARARG_BODY(void, FVisualLogger::GeometryShapeLogf, const TCHAR*, VARARG_EXTRA(const UObject* Object) VARARG_EXTRA(const FLogCategoryBase& Category) VARARG_EXTRA(ELogVerbosity::Type Verbosity) VARARG_EXTRA(const FVector& Start) VARARG_EXTRA(const FVector& End) VARARG_EXTRA(const FColor& Color) VARARG_EXTRA(const uint16 Thickness))
 {
-	SCOPE_CYCLE_COUNTER(STAT_VisualLog);
-	UWorld *World = nullptr;
-	FVisualLogEntry *CurrentEntry = nullptr;
-	if (CheckVisualLogInputInternal(Object, Category, Verbosity, &World, &CurrentEntry) == false)
-	{
-		return;
-	}
+	const FName CategoryName = Category.GetCategoryName();
+	COLLAPSED_LOGF(
+		CurrentEntry->AddElement(Start, End, CategoryName, Verbosity, Color, Buffer, Thickness);
+	);
+}
 
+inline
+VARARG_BODY(void, FVisualLogger::GeometryShapeLogf, const TCHAR*, VARARG_EXTRA(const UObject* Object) VARARG_EXTRA(const FName& CategoryName) VARARG_EXTRA(ELogVerbosity::Type Verbosity) VARARG_EXTRA(const FVector& Start) VARARG_EXTRA(const FVector& End) VARARG_EXTRA(const FColor& Color) VARARG_EXTRA(const uint16 Thickness))
+{
 	COLLAPSED_LOGF(
-		CurrentEntry->AddElement(Start, End, Category.GetCategoryName(), Verbosity, Color, Buffer, Thickness);
+		CurrentEntry->AddElement(Start, End, CategoryName, Verbosity, Color, Buffer, Thickness);
 	);
 }
 
 inline
 VARARG_BODY(void, FVisualLogger::GeometryShapeLogf, const TCHAR*, VARARG_EXTRA(const UObject* Object) VARARG_EXTRA(const FLogCategoryBase& Category) VARARG_EXTRA(ELogVerbosity::Type Verbosity) VARARG_EXTRA(const FVector& Location) VARARG_EXTRA(float Radius) VARARG_EXTRA(const FColor& Color))
 {
-	SCOPE_CYCLE_COUNTER(STAT_VisualLog);
-	UWorld *World = nullptr;
-	FVisualLogEntry *CurrentEntry = nullptr;
-	if (CheckVisualLogInputInternal(Object, Category, Verbosity, &World, &CurrentEntry) == false)
-	{
-		return;
-	}
+	const FName CategoryName = Category.GetCategoryName();
+	COLLAPSED_LOGF(
+		CurrentEntry->AddElement(Location, CategoryName, Verbosity, Color, Buffer, Radius);
+	);
+}
 
+inline
+VARARG_BODY(void, FVisualLogger::GeometryShapeLogf, const TCHAR*, VARARG_EXTRA(const UObject* Object) VARARG_EXTRA(const FName& CategoryName) VARARG_EXTRA(ELogVerbosity::Type Verbosity) VARARG_EXTRA(const FVector& Location) VARARG_EXTRA(float Radius) VARARG_EXTRA(const FColor& Color))
+{
 	COLLAPSED_LOGF(
-		CurrentEntry->AddElement(Location, Category.GetCategoryName(), Verbosity, Color, Buffer, Radius);
+		CurrentEntry->AddElement(Location, CategoryName, Verbosity, Color, Buffer, Radius);
 	);
 }
 
 inline
 VARARG_BODY(void, FVisualLogger::GeometryShapeLogf, const TCHAR*, VARARG_EXTRA(const UObject* Object) VARARG_EXTRA(const FLogCategoryBase& Category) VARARG_EXTRA(ELogVerbosity::Type Verbosity) VARARG_EXTRA(const FBox& Box) VARARG_EXTRA(const FMatrix& Matrix) VARARG_EXTRA(const FColor& Color))
 {
-	SCOPE_CYCLE_COUNTER(STAT_VisualLog);
-	UWorld *World = nullptr;
-	FVisualLogEntry *CurrentEntry = nullptr;
-	if (CheckVisualLogInputInternal(Object, Category, Verbosity, &World, &CurrentEntry) == false)
-	{
-		return;
-	}
+	const FName CategoryName = Category.GetCategoryName();
+	COLLAPSED_LOGF(
+		CurrentEntry->AddElement(Box, Matrix, CategoryName, Verbosity, Color, Buffer);
+	);
+}
 
+inline
+VARARG_BODY(void, FVisualLogger::GeometryShapeLogf, const TCHAR*, VARARG_EXTRA(const UObject* Object) VARARG_EXTRA(const FName& CategoryName) VARARG_EXTRA(ELogVerbosity::Type Verbosity) VARARG_EXTRA(const FBox& Box) VARARG_EXTRA(const FMatrix& Matrix) VARARG_EXTRA(const FColor& Color))
+{
 	COLLAPSED_LOGF(
-		CurrentEntry->AddElement(Box, Matrix, Category.GetCategoryName(), Verbosity, Color, Buffer);
+		CurrentEntry->AddElement(Box, Matrix, CategoryName, Verbosity, Color, Buffer);
 	);
 }
 
 inline
 VARARG_BODY(void, FVisualLogger::GeometryShapeLogf, const TCHAR*, VARARG_EXTRA(const UObject* Object) VARARG_EXTRA(const FLogCategoryBase& Category) VARARG_EXTRA(ELogVerbosity::Type Verbosity) VARARG_EXTRA(const FVector& Orgin) VARARG_EXTRA(const FVector& Direction) VARARG_EXTRA(const float Length) VARARG_EXTRA(const float Angle)  VARARG_EXTRA(const FColor& Color))
 {
-	SCOPE_CYCLE_COUNTER(STAT_VisualLog);
-	UWorld *World = nullptr;
-	FVisualLogEntry *CurrentEntry = nullptr;
-	if (CheckVisualLogInputInternal(Object, Category, Verbosity, &World, &CurrentEntry) == false)
-	{
-		return;
-	}
+	const FName CategoryName = Category.GetCategoryName();
+	COLLAPSED_LOGF(
+		CurrentEntry->AddElement(Orgin, Direction, Length, Angle, Angle, CategoryName, Verbosity, Color, Buffer);
+	);
+}
 
+inline
+VARARG_BODY(void, FVisualLogger::GeometryShapeLogf, const TCHAR*, VARARG_EXTRA(const UObject* Object) VARARG_EXTRA(const FName& CategoryName) VARARG_EXTRA(ELogVerbosity::Type Verbosity) VARARG_EXTRA(const FVector& Orgin) VARARG_EXTRA(const FVector& Direction) VARARG_EXTRA(const float Length) VARARG_EXTRA(const float Angle)  VARARG_EXTRA(const FColor& Color))
+{
 	COLLAPSED_LOGF(
-		CurrentEntry->AddElement(Orgin, Direction, Length, Angle, Angle, Category.GetCategoryName(), Verbosity, Color, Buffer);
+		CurrentEntry->AddElement(Orgin, Direction, Length, Angle, Angle, CategoryName, Verbosity, Color, Buffer);
 	);
 }
 
 inline
 VARARG_BODY(void, FVisualLogger::GeometryShapeLogf, const TCHAR*, VARARG_EXTRA(const UObject* Object) VARARG_EXTRA(const FLogCategoryBase& Category) VARARG_EXTRA(ELogVerbosity::Type Verbosity) VARARG_EXTRA(const FVector& Start) VARARG_EXTRA(const FVector& End) VARARG_EXTRA(const float Radius) VARARG_EXTRA(const FColor& Color))
 {
-	SCOPE_CYCLE_COUNTER(STAT_VisualLog);
-	UWorld *World = nullptr;
-	FVisualLogEntry *CurrentEntry = nullptr;
-	if (CheckVisualLogInputInternal(Object, Category, Verbosity, &World, &CurrentEntry) == false)
-	{
-		return;
-	}
+	const FName CategoryName = Category.GetCategoryName();
+	COLLAPSED_LOGF(
+		CurrentEntry->AddElement(Start, End, Radius, CategoryName, Verbosity, Color, Buffer);
+	);
+}
 
+inline
+VARARG_BODY(void, FVisualLogger::GeometryShapeLogf, const TCHAR*, VARARG_EXTRA(const UObject* Object) VARARG_EXTRA(const FName& CategoryName) VARARG_EXTRA(ELogVerbosity::Type Verbosity) VARARG_EXTRA(const FVector& Start) VARARG_EXTRA(const FVector& End) VARARG_EXTRA(const float Radius) VARARG_EXTRA(const FColor& Color))
+{
 	COLLAPSED_LOGF(
-		CurrentEntry->AddElement(Start, End, Radius, Category.GetCategoryName(), Verbosity, Color, Buffer);
+		CurrentEntry->AddElement(Start, End, Radius, CategoryName, Verbosity, Color, Buffer);
 	);
 }
 
 inline
 VARARG_BODY(void, FVisualLogger::GeometryShapeLogf, const TCHAR*, VARARG_EXTRA(const UObject* Object) VARARG_EXTRA(const FLogCategoryBase& Category) VARARG_EXTRA(ELogVerbosity::Type Verbosity) VARARG_EXTRA(const FVector& Center) VARARG_EXTRA(float HalfHeight) VARARG_EXTRA(float Radius) VARARG_EXTRA(const FQuat& Rotation) VARARG_EXTRA(const FColor& Color))
 {
-	SCOPE_CYCLE_COUNTER(STAT_VisualLog);
-	UWorld *World = nullptr;
-	FVisualLogEntry *CurrentEntry = nullptr;
-	if (CheckVisualLogInputInternal(Object, Category, Verbosity, &World, &CurrentEntry) == false)
-	{
-		return;
-	}
+	const FName CategoryName = Category.GetCategoryName();
+	COLLAPSED_LOGF(
+		CurrentEntry->AddElement(Center, HalfHeight, Radius, Rotation, CategoryName, Verbosity, Color, Buffer);
+	);
+}
 
+inline
+VARARG_BODY(void, FVisualLogger::GeometryShapeLogf, const TCHAR*, VARARG_EXTRA(const UObject* Object) VARARG_EXTRA(const FName& CategoryName) VARARG_EXTRA(ELogVerbosity::Type Verbosity) VARARG_EXTRA(const FVector& Center) VARARG_EXTRA(float HalfHeight) VARARG_EXTRA(float Radius) VARARG_EXTRA(const FQuat& Rotation) VARARG_EXTRA(const FColor& Color))
+{
 	COLLAPSED_LOGF(
-		CurrentEntry->AddElement(Center, HalfHeight, Radius, Rotation, Category.GetCategoryName(), Verbosity, Color, Buffer);
+		CurrentEntry->AddElement(Center, HalfHeight, Radius, Rotation, CategoryName, Verbosity, Color, Buffer);
 	);
 }
 
 inline
 VARARG_BODY(void, FVisualLogger::NavAreaShapeLogf, const TCHAR*, VARARG_EXTRA(const UObject* Object) VARARG_EXTRA(const FLogCategoryBase& Category) VARARG_EXTRA(ELogVerbosity::Type Verbosity) VARARG_EXTRA(const TArray<FVector> &ConvexPoints) VARARG_EXTRA(float MinZ) VARARG_EXTRA(float MaxZ) VARARG_EXTRA(const FColor& Color))
 {
-	SCOPE_CYCLE_COUNTER(STAT_VisualLog);
-	UWorld *World = nullptr;
-	FVisualLogEntry *CurrentEntry = nullptr;
-	if (CheckVisualLogInputInternal(Object, Category, Verbosity, &World, &CurrentEntry) == false)
-	{
-		return;
-	}
+	const FName CategoryName = Category.GetCategoryName();
+	COLLAPSED_LOGF(
+		CurrentEntry->AddElement(ConvexPoints, MinZ, MaxZ, CategoryName, Verbosity, Color, Buffer);
+	);
+}
 
+inline
+VARARG_BODY(void, FVisualLogger::NavAreaShapeLogf, const TCHAR*, VARARG_EXTRA(const UObject* Object) VARARG_EXTRA(const FName& CategoryName) VARARG_EXTRA(ELogVerbosity::Type Verbosity) VARARG_EXTRA(const TArray<FVector> &ConvexPoints) VARARG_EXTRA(float MinZ) VARARG_EXTRA(float MaxZ) VARARG_EXTRA(const FColor& Color))
+{
 	COLLAPSED_LOGF(
-		CurrentEntry->AddElement(ConvexPoints, MinZ, MaxZ, Category.GetCategoryName(), Verbosity, Color, Buffer);
+		CurrentEntry->AddElement(ConvexPoints, MinZ, MaxZ, CategoryName, Verbosity, Color, Buffer);
 	);
 }
 
 inline
 VARARG_BODY(void, FVisualLogger::GeometryShapeLogf, const TCHAR*, VARARG_EXTRA(const UObject* Object) VARARG_EXTRA(const FLogCategoryBase& Category) VARARG_EXTRA(ELogVerbosity::Type Verbosity) VARARG_EXTRA(const TArray<FVector> &Vertices) VARARG_EXTRA(const TArray<int32> &Indices) VARARG_EXTRA(const FColor& Color))
 {
-	SCOPE_CYCLE_COUNTER(STAT_VisualLog);
-	UWorld *World = nullptr;
-	FVisualLogEntry *CurrentEntry = nullptr;
-	if (CheckVisualLogInputInternal(Object, Category, Verbosity, &World, &CurrentEntry) == false)
-	{
-		return;
-	}
+	const FName CategoryName = Category.GetCategoryName();
+	COLLAPSED_LOGF(
+		CurrentEntry->AddElement(Vertices, Indices, CategoryName, Verbosity, Color, Buffer);
+	);
+}
 
+inline
+VARARG_BODY(void, FVisualLogger::GeometryShapeLogf, const TCHAR*, VARARG_EXTRA(const UObject* Object) VARARG_EXTRA(const FName& CategoryName) VARARG_EXTRA(ELogVerbosity::Type Verbosity) VARARG_EXTRA(const TArray<FVector> &Vertices) VARARG_EXTRA(const TArray<int32> &Indices) VARARG_EXTRA(const FColor& Color))
+{
 	COLLAPSED_LOGF(
-		CurrentEntry->AddElement(Vertices, Indices, Category.GetCategoryName(), Verbosity, Color, Buffer);
+		CurrentEntry->AddElement(Vertices, Indices, CategoryName, Verbosity, Color, Buffer);
 	);
 }
 
 inline
 VARARG_BODY(void, FVisualLogger::GeometryConvexLogf, const TCHAR*, VARARG_EXTRA(const UObject* Object) VARARG_EXTRA(const FLogCategoryBase& Category) VARARG_EXTRA(ELogVerbosity::Type Verbosity) VARARG_EXTRA(const TArray<FVector> &Points) VARARG_EXTRA(const FColor& Color))
 {
-	SCOPE_CYCLE_COUNTER(STAT_VisualLog);
-	UWorld *World = nullptr;
-	FVisualLogEntry *CurrentEntry = nullptr;
-	if (CheckVisualLogInputInternal(Object, Category, Verbosity, &World, &CurrentEntry) == false)
-	{
-		return;
-	}
-
+	const FName CategoryName = Category.GetCategoryName();
 	COLLAPSED_LOGF(
-		CurrentEntry->AddConvexElement(Points, Category.GetCategoryName(), Verbosity, Color, Buffer);
+		CurrentEntry->AddConvexElement(Points, CategoryName, Verbosity, Color, Buffer);
 	);
 }
 
+inline
+VARARG_BODY(void, FVisualLogger::GeometryConvexLogf, const TCHAR*, VARARG_EXTRA(const UObject* Object) VARARG_EXTRA(const FName& CategoryName) VARARG_EXTRA(ELogVerbosity::Type Verbosity) VARARG_EXTRA(const TArray<FVector> &Points) VARARG_EXTRA(const FColor& Color))
+{
+	COLLAPSED_LOGF(
+		CurrentEntry->AddConvexElement(Points, CategoryName, Verbosity, Color, Buffer);
+	);
+}
 
 inline
 VARARG_BODY(void, FVisualLogger::HistogramDataLogf, const TCHAR*, VARARG_EXTRA(const UObject* Object) VARARG_EXTRA(const FLogCategoryBase& Category) VARARG_EXTRA(ELogVerbosity::Type Verbosity) VARARG_EXTRA(FName GraphName) VARARG_EXTRA(FName DataName) VARARG_EXTRA(const FVector2D& Data) VARARG_EXTRA(const FColor& Color))
 {
-	SCOPE_CYCLE_COUNTER(STAT_VisualLog);
-	UWorld *World = nullptr;
-	FVisualLogEntry *CurrentEntry = nullptr;
-	if (CheckVisualLogInputInternal(Object, Category, Verbosity, &World, &CurrentEntry) == false)
-	{
-		return;
-	}
+	const FName CategoryName = Category.GetCategoryName();
+	COLLAPSED_LOGF(
+		CurrentEntry->AddHistogramData(Data, CategoryName, Verbosity, GraphName, DataName);
+	);
+}
 
+inline
+VARARG_BODY(void, FVisualLogger::HistogramDataLogf, const TCHAR*, VARARG_EXTRA(const UObject* Object) VARARG_EXTRA(const FName& CategoryName) VARARG_EXTRA(ELogVerbosity::Type Verbosity) VARARG_EXTRA(FName GraphName) VARARG_EXTRA(FName DataName) VARARG_EXTRA(const FVector2D& Data) VARARG_EXTRA(const FColor& Color))
+{
 	COLLAPSED_LOGF(
-		CurrentEntry->AddHistogramData(Data, Category.GetCategoryName(), Verbosity, GraphName, DataName);
+		CurrentEntry->AddHistogramData(Data, CategoryName, Verbosity, GraphName, DataName);
 	);
 }
 
diff --git a/Engine/Source/Runtime/Landscape/Private/LandscapeGrass.cpp b/Engine/Source/Runtime/Landscape/Private/LandscapeGrass.cpp
--- a/Engine/Source/Runtime/Landscape/Private/LandscapeGrass.cpp
+++ b/Engine/Source/Runtime/Landscape/Private/LandscapeGrass.cpp
@@ -1011,9 +1011,9 @@ FExpressionInput* UMaterialExpressionLandscapeGrassOutput::GetInput(int32 InputI
 	return &GrassTypes[InputIndex].Input;
 }
 
-FString UMaterialExpressionLandscapeGrassOutput::GetInputName(int32 InputIndex) const
+FName UMaterialExpressionLandscapeGrassOutput::GetInputName(int32 InputIndex) const
 {
-	return GrassTypes[InputIndex].Name.ToString();
+	return GrassTypes[InputIndex].Name;
 }
 
 bool UMaterialExpressionLandscapeGrassOutput::NeedsLoadForClient() const
diff --git a/Engine/Source/Runtime/Landscape/Private/Materials/MaterialExpressionLandscapeLayerBlend.cpp b/Engine/Source/Runtime/Landscape/Private/Materials/MaterialExpressionLandscapeLayerBlend.cpp
--- a/Engine/Source/Runtime/Landscape/Private/Materials/MaterialExpressionLandscapeLayerBlend.cpp
+++ b/Engine/Source/Runtime/Landscape/Private/Materials/MaterialExpressionLandscapeLayerBlend.cpp
@@ -96,24 +96,24 @@ FExpressionInput* UMaterialExpressionLandscapeLayerBlend::GetInput(int32 InputIn
 }
 
 
-FString UMaterialExpressionLandscapeLayerBlend::GetInputName(int32 InputIndex) const
+FName UMaterialExpressionLandscapeLayerBlend::GetInputName(int32 InputIndex) const
 {
 	int32 Idx = 0;
 	for (int32 LayerIdx = 0; LayerIdx<Layers.Num(); LayerIdx++)
 	{
 		if (InputIndex == Idx++)
 		{
-			return FString::Printf(TEXT("Layer %s"), *Layers[LayerIdx].LayerName.ToString());
+			return *FString::Printf(TEXT("Layer %s"), *Layers[LayerIdx].LayerName.ToString());
 		}
 		if (Layers[LayerIdx].BlendType == LB_HeightBlend)
 		{
 			if (InputIndex == Idx++)
 			{
-				return FString::Printf(TEXT("Height %s"), *Layers[LayerIdx].LayerName.ToString());
+				return *FString::Printf(TEXT("Height %s"), *Layers[LayerIdx].LayerName.ToString());
 			}
 		}
 	}
-	return TEXT("");
+	return NAME_None;
 }
 
 #if WITH_EDITOR
diff --git a/Engine/Source/Runtime/LevelSequence/Private/LevelSequenceBindingReference.cpp b/Engine/Source/Runtime/LevelSequence/Private/LevelSequenceBindingReference.cpp
--- a/Engine/Source/Runtime/LevelSequence/Private/LevelSequenceBindingReference.cpp
+++ b/Engine/Source/Runtime/LevelSequence/Private/LevelSequenceBindingReference.cpp
@@ -50,7 +50,18 @@ UObject* FLevelSequenceBindingReference::Resolve(UObject* InContext) const
 		FSoftObjectPath TempPath = ExternalObjectPath;
 
 #if WITH_EDITORONLY_DATA
-		TempPath.FixupForPIE();
+		int32 ContextPlayInEditorID = InContext ? InContext->GetOutermost()->PIEInstanceID : INDEX_NONE;
+
+		if (ContextPlayInEditorID != INDEX_NONE)
+		{
+			// We have an override PIE id, so set the global before entering
+			TGuardValue<int32> PIEGuard(GPlayInEditorID, ContextPlayInEditorID);
+			TempPath.FixupForPIE();
+		}
+		else
+		{
+			TempPath.FixupForPIE();
+		}
 #endif
 
 		return TempPath.ResolveObject();
diff --git a/Engine/Source/Runtime/NetworkReplayStreaming/HttpNetworkReplayStreaming/Private/HttpNetworkReplayStreaming.cpp b/Engine/Source/Runtime/NetworkReplayStreaming/HttpNetworkReplayStreaming/Private/HttpNetworkReplayStreaming.cpp
--- a/Engine/Source/Runtime/NetworkReplayStreaming/HttpNetworkReplayStreaming/Private/HttpNetworkReplayStreaming.cpp
+++ b/Engine/Source/Runtime/NetworkReplayStreaming/HttpNetworkReplayStreaming/Private/HttpNetworkReplayStreaming.cpp
@@ -2344,16 +2344,6 @@ void FHttpNetworkReplayStreamingFactory::Tick( float DeltaTime )
 	}
 }
 
-bool FHttpNetworkReplayStreamingFactory::IsTickable() const
-{
-	return true;
-}
-
-bool FHttpNetworkReplayStreamingFactory::IsTickableWhenPaused() const
-{
-	return true;
-}
-
 TStatId FHttpNetworkReplayStreamingFactory::GetStatId() const
 {
 	RETURN_QUICK_DECLARE_CYCLE_STAT( FHttpNetworkReplayStreamingFactory, STATGROUP_Tickables );
diff --git a/Engine/Source/Runtime/NetworkReplayStreaming/HttpNetworkReplayStreaming/Public/HttpNetworkReplayStreaming.h b/Engine/Source/Runtime/NetworkReplayStreaming/HttpNetworkReplayStreaming/Public/HttpNetworkReplayStreaming.h
--- a/Engine/Source/Runtime/NetworkReplayStreaming/HttpNetworkReplayStreaming/Public/HttpNetworkReplayStreaming.h
+++ b/Engine/Source/Runtime/NetworkReplayStreaming/HttpNetworkReplayStreaming/Public/HttpNetworkReplayStreaming.h
@@ -349,9 +349,9 @@ public:
 
 	/** FTickableGameObject */
 	virtual void Tick( float DeltaTime ) override;
-	virtual bool IsTickable() const override;
+	virtual ETickableTickType GetTickableTickType() const override { return ETickableTickType::Always; }
 	virtual TStatId GetStatId() const override;
-	bool IsTickableWhenPaused() const override;
+	bool IsTickableWhenPaused() const override { return true; }
 
 	TArray< TSharedPtr< FHttpNetworkReplayStreamer > > HttpStreamers;
 };
diff --git a/Engine/Source/Runtime/UMG/Private/WidgetBlueprintGeneratedClass.cpp b/Engine/Source/Runtime/UMG/Private/WidgetBlueprintGeneratedClass.cpp
--- a/Engine/Source/Runtime/UMG/Private/WidgetBlueprintGeneratedClass.cpp
+++ b/Engine/Source/Runtime/UMG/Private/WidgetBlueprintGeneratedClass.cpp
@@ -482,7 +482,7 @@ void UWidgetBlueprintGeneratedClass::InitializeTemplate(const ITargetPlatform* T
 			UUserWidget* OldArchetype = FindObject<UUserWidget>(WidgetTemplatePackage, TEXT("WidgetArchetype"));
 			if ( OldArchetype )
 			{
-				const ERenameFlags RenFlags = REN_DontCreateRedirectors | REN_NonTransactional | REN_DoNotDirty;
+				const ERenameFlags RenFlags = REN_DontCreateRedirectors | REN_NonTransactional | REN_DoNotDirty | REN_ForceNoResetLoaders;
 
 				FString TransientArchetypeString = FString::Printf(TEXT("OLD_TEMPLATE_%s"), *OldArchetype->GetName());
 				FName TransientArchetypeName = MakeUniqueObjectName(GetTransientPackage(), OldArchetype->GetClass(), FName(*TransientArchetypeString));
diff --git a/Engine/Source/Runtime/UMG/Private/WidgetBlueprintLibrary.cpp b/Engine/Source/Runtime/UMG/Private/WidgetBlueprintLibrary.cpp
--- a/Engine/Source/Runtime/UMG/Private/WidgetBlueprintLibrary.cpp
+++ b/Engine/Source/Runtime/UMG/Private/WidgetBlueprintLibrary.cpp
@@ -81,9 +81,9 @@ void UWidgetBlueprintLibrary::SetInputMode_UIOnly(APlayerController* Target, UWi
 	SetInputMode_UIOnlyEx(Target, InWidgetToFocus, bLockMouseToViewport ? EMouseLockMode::LockOnCapture : EMouseLockMode::DoNotLock);
 }
 
-void UWidgetBlueprintLibrary::SetInputMode_UIOnlyEx(APlayerController* Target, UWidget* InWidgetToFocus, EMouseLockMode InMouseLockMode)
+void UWidgetBlueprintLibrary::SetInputMode_UIOnlyEx(APlayerController* PlayerController, UWidget* InWidgetToFocus, EMouseLockMode InMouseLockMode)
 {
-	if (Target != nullptr)
+	if (PlayerController != nullptr)
 	{
 		FInputModeUIOnly InputMode;
 		InputMode.SetLockMouseToViewportBehavior(InMouseLockMode);
@@ -92,12 +92,12 @@ void UWidgetBlueprintLibrary::SetInputMode_UIOnlyEx(APlayerController* Target, U
 		{
 			InputMode.SetWidgetToFocus(InWidgetToFocus->TakeWidget());
 		}
-		Target->SetInputMode(InputMode);
+		PlayerController->SetInputMode(InputMode);
 	}
 	#if WITH_EDITOR 
 	else
 	{ 
-		FMessageLog("PIE").Error(LOCTEXT("UMG WidgetBlueprint Library: SetInputMode_UIOnly", "SetInputMode_UIOnly expects a valid player controller as target"));
+		FMessageLog("PIE").Error(LOCTEXT("UMG WidgetBlueprint Library: SetInputMode_UIOnly", "SetInputMode_UIOnly expects a valid player controller as 'PlayerController' target"));
 	}
 	#endif // WITH_EDITOR
 }
@@ -107,9 +107,9 @@ void UWidgetBlueprintLibrary::SetInputMode_GameAndUI(APlayerController* Target,
 	SetInputMode_GameAndUIEx(Target, InWidgetToFocus, bLockMouseToViewport ? EMouseLockMode::LockOnCapture : EMouseLockMode::DoNotLock, bHideCursorDuringCapture);
 }
 
-void UWidgetBlueprintLibrary::SetInputMode_GameAndUIEx(APlayerController* Target, UWidget* InWidgetToFocus, EMouseLockMode InMouseLockMode, bool bHideCursorDuringCapture)
+void UWidgetBlueprintLibrary::SetInputMode_GameAndUIEx(APlayerController* PlayerController, UWidget* InWidgetToFocus, EMouseLockMode InMouseLockMode, bool bHideCursorDuringCapture)
 {
-	if (Target != nullptr)
+	if (PlayerController != nullptr)
 	{
 		FInputModeGameAndUI InputMode;
 		InputMode.SetLockMouseToViewportBehavior(InMouseLockMode);
@@ -119,28 +119,28 @@ void UWidgetBlueprintLibrary::SetInputMode_GameAndUIEx(APlayerController* Target
 		{
 			InputMode.SetWidgetToFocus(InWidgetToFocus->TakeWidget());
 		}
-		Target->SetInputMode(InputMode);
+		PlayerController->SetInputMode(InputMode);
 	}
 	#if WITH_EDITOR 
 	else
 	{
-		FMessageLog("PIE").Error(LOCTEXT("UMG WidgetBlueprint Library: SetInputMode_GameAndUI", "SetInputMode_GameAndUI expects a valid player controller as target"));	
+		FMessageLog("PIE").Error(LOCTEXT("UMG WidgetBlueprint Library: SetInputMode_GameAndUI", "SetInputMode_GameAndUI expects a valid player controller as 'PlayerController' target"));	
 	}
 	#endif // WITH_EDITOR
 	
 }
 
-void UWidgetBlueprintLibrary::SetInputMode_GameOnly(APlayerController* Target)
+void UWidgetBlueprintLibrary::SetInputMode_GameOnly(APlayerController* PlayerController)
 {
-	if (Target != nullptr)
+	if (PlayerController != nullptr)
 	{
 		FInputModeGameOnly InputMode;
-		Target->SetInputMode(InputMode);
+		PlayerController->SetInputMode(InputMode);
 	}
 	#if WITH_EDITOR 
 	else
 	{
-		FMessageLog("PIE").Error(LOCTEXT("UMG WidgetBlueprint Library: SetInputMode_GameOnly", "SetInputMode_GameOnly expects a valid player controller as target"));
+		FMessageLog("PIE").Error(LOCTEXT("UMG WidgetBlueprint Library: SetInputMode_GameOnly", "SetInputMode_GameOnly expects a valid player controller as 'PlayerController' target"));
 	}
 	#endif // WITH_EDITOR
 }
diff --git a/Engine/Source/Runtime/UMG/Public/Blueprint/WidgetBlueprintLibrary.h b/Engine/Source/Runtime/UMG/Public/Blueprint/WidgetBlueprintLibrary.h
--- a/Engine/Source/Runtime/UMG/Public/Blueprint/WidgetBlueprintLibrary.h
+++ b/Engine/Source/Runtime/UMG/Public/Blueprint/WidgetBlueprintLibrary.h
@@ -47,7 +47,7 @@ public:
 
 	/** Setup an input mode that allows only the UI to respond to user input. */
 	UFUNCTION(BlueprintCallable, BlueprintCosmetic, Category="Input", meta=(DisplayName="Set Input Mode UI Only"))
-	static void SetInputMode_UIOnlyEx(APlayerController* Target, UWidget* InWidgetToFocus = nullptr, EMouseLockMode InMouseLockMode = EMouseLockMode::DoNotLock);
+	static void SetInputMode_UIOnlyEx(APlayerController* PlayerController, UWidget* InWidgetToFocus = nullptr, EMouseLockMode InMouseLockMode = EMouseLockMode::DoNotLock);
 
 	/** Setup an input mode that allows only the UI to respond to user input, and if the UI doesn't handle it player input / player controller gets a chance. */
 	DEPRECATED(4.13, "Locking the mouse to the viewport is now controlled by an enum. Call SetInputMode_GameAndUIEx instead")
@@ -56,11 +56,11 @@ public:
 
 	/** Setup an input mode that allows only the UI to respond to user input, and if the UI doesn't handle it player input / player controller gets a chance. */
 	UFUNCTION(BlueprintCallable, BlueprintCosmetic, Category="Input", meta = (DisplayName = "Set Input Mode Game And UI"))
-	static void SetInputMode_GameAndUIEx(APlayerController* Target, UWidget* InWidgetToFocus = nullptr, EMouseLockMode InMouseLockMode = EMouseLockMode::DoNotLock, bool bHideCursorDuringCapture = true);
+	static void SetInputMode_GameAndUIEx(APlayerController* PlayerController, UWidget* InWidgetToFocus = nullptr, EMouseLockMode InMouseLockMode = EMouseLockMode::DoNotLock, bool bHideCursorDuringCapture = true);
 
 	/** Setup an input mode that allows only player input / player controller to respond to user input. */
 	UFUNCTION(BlueprintCallable, BlueprintCosmetic, Category="Input")
-	static void SetInputMode_GameOnly(APlayerController* Target);
+	static void SetInputMode_GameOnly(APlayerController* PlayerController);
 
 	/** */
 	UFUNCTION(BlueprintCallable, BlueprintCosmetic, Category="Focus")
