Fixed issues due to merging 4.18 with refactor made in Dev-Rendering

diff --git a/Engine/Source/Runtime/Engine/Private/InstancedStaticMesh.cpp b/Engine/Source/Runtime/Engine/Private/InstancedStaticMesh.cpp
--- a/Engine/Source/Runtime/Engine/Private/InstancedStaticMesh.cpp
+++ b/Engine/Source/Runtime/Engine/Private/InstancedStaticMesh.cpp
@@ -98,7 +98,10 @@ void FStaticMeshInstanceBuffer::UpdateInstanceData(UInstancedStaticMeshComponent
 	int32 NumRenderInstances = InComponent->PerInstanceSMData.Num() - InComponent->RemovedInstances.Num();
 
 	// Allocate the vertex data storage type.
-	InstanceData->AllocateInstances(NumRenderInstances, false);
+	if (NumRenderInstances >= InstanceData->GetNumInstances())
+	{
+		InstanceData->AllocateInstances(NumRenderInstances, false);
+	}
 
 	NumInstances = InstanceData->GetNumInstances();
 
@@ -245,38 +248,9 @@ void FStaticMeshInstanceBuffer::InitRHI()
 	QUICK_SCOPE_CYCLE_COUNTER(STAT_FStaticMeshInstanceBuffer_InitRHI);
 	
 	auto AccessFlags = IsDynamic ? BUF_Dynamic : BUF_Static;
-	{
-		FResourceArrayInterface* ResourceArray = InstanceData->GetOriginResourceArray();
-		if (ResourceArray->GetResourceDataSize())
-		{
-			// Create the vertex buffer.
-			FRHIResourceCreateInfo CreateInfo(ResourceArray);
-			InstanceOriginBuffer.VertexBufferRHI = RHICreateVertexBuffer(ResourceArray->GetResourceDataSize(), AccessFlags | BUF_ShaderResource, CreateInfo);
-			InstanceOriginSRV = RHICreateShaderResourceView(InstanceOriginBuffer.VertexBufferRHI, 16, PF_A32B32G32R32F);
-		}
-	}
-	{
-		FResourceArrayInterface* ResourceArray = InstanceData->GetTransformResourceArray();
-		if (ResourceArray->GetResourceDataSize())
-		{
-			// Create the vertex buffer.
-			FRHIResourceCreateInfo CreateInfo(ResourceArray);
-			InstanceTransformBuffer.VertexBufferRHI = RHICreateVertexBuffer(ResourceArray->GetResourceDataSize(), AccessFlags | BUF_ShaderResource, CreateInfo);
-			InstanceTransformSRV = RHICreateShaderResourceView(InstanceTransformBuffer.VertexBufferRHI, InstanceData->GetTranslationUsesHalfs() ? 8 : 16, InstanceData->GetTranslationUsesHalfs() ? PF_FloatRGBA : PF_A32B32G32R32F);
-		}
-	}
-	{
-		// TODO: possibility over allocated the vertex buffer when we support partial update for when working in the editor
-		// Create the vertex buffer.
-		FResourceArrayInterface* ResourceArray = InstanceData->GetLightMapResourceArray();
-		if (ResourceArray->GetResourceDataSize())
-		{
-			// Create the vertex buffer.
-			FRHIResourceCreateInfo CreateInfo(ResourceArray);
-			InstanceLightmapBuffer.VertexBufferRHI = RHICreateVertexBuffer(ResourceArray->GetResourceDataSize(), AccessFlags | BUF_ShaderResource, CreateInfo);
-			InstanceLightmapSRV = RHICreateShaderResourceView(InstanceTransformBuffer.VertexBufferRHI, 8, PF_R16G16B16A16_SNORM);
-		}
-	}
+	CreateVertexBuffer(InstanceData->GetOriginResourceArray(), AccessFlags | BUF_ShaderResource, 16, PF_A32B32G32R32F, InstanceOriginBuffer.VertexBufferRHI, InstanceOriginSRV);
+	CreateVertexBuffer(InstanceData->GetTransformResourceArray(), AccessFlags | BUF_ShaderResource, InstanceData->GetTranslationUsesHalfs() ? 8 : 16, InstanceData->GetTranslationUsesHalfs() ? PF_FloatRGBA : PF_A32B32G32R32F, InstanceTransformBuffer.VertexBufferRHI, InstanceTransformSRV);
+	CreateVertexBuffer(InstanceData->GetLightMapResourceArray(), AccessFlags | BUF_ShaderResource, 8, PF_R16G16B16A16_SNORM, InstanceLightmapBuffer.VertexBufferRHI, InstanceLightmapSRV);
 }
 
 void FStaticMeshInstanceBuffer::ReleaseRHI()
@@ -302,6 +276,38 @@ void FStaticMeshInstanceBuffer::ReleaseResource()
 	InstanceLightmapBuffer.ReleaseResource();
 }
 
+void FStaticMeshInstanceBuffer::CreateVertexBuffer(FResourceArrayInterface* InResourceArray, uint32 InUsage, uint32 InStride, uint8 InFormat, FVertexBufferRHIRef& OutVertexBufferRHI, FShaderResourceViewRHIRef& OutInstanceSRV)
+{
+	if (InResourceArray->GetResourceDataSize() > 0)
+	{
+		// TODO: possibility over allocated the vertex buffer when we support partial update for when working in the editor
+		FRHIResourceCreateInfo CreateInfo(InResourceArray);
+		OutVertexBufferRHI = RHICreateVertexBuffer(InResourceArray->GetResourceDataSize(), InUsage, CreateInfo);
+		OutInstanceSRV = RHICreateShaderResourceView(OutVertexBufferRHI, InStride, InFormat);
+	}
+}
+
+bool FStaticMeshInstanceBuffer::UpdateDynamicVertexBuffer(int32 StartingIndex, uint32 InstanceCount, uint32 Stride, FVertexBufferRHIRef VetexBuffer, FResourceArrayInterface* ResourceArray)
+{
+	check(IsDynamic);
+
+	auto VertexBufferRHI = VetexBuffer;
+
+	uint32 UpdateSize = InstanceCount * Stride;
+	uint32 UpdateOffset = StartingIndex * Stride;
+
+	if (UpdateOffset + UpdateSize <= VertexBufferRHI->GetSize()) // we can only update the buffer
+	{
+		void* VertexBufferData = RHILockVertexBuffer(VertexBufferRHI, UpdateOffset, UpdateSize, RLM_WriteOnly);
+
+		FMemory::Memcpy(VertexBufferData, (uint8*)ResourceArray->GetResourceData() + UpdateOffset, UpdateSize);
+		RHIUnlockVertexBuffer(VertexBufferRHI);
+		return true;
+	}
+
+	return false;
+}
+
 void FStaticMeshInstanceBuffer::UpdateRHIVertexBuffer(const TSet<int32>& InIndexList)
 {
 	check(IsInRenderingThread());
@@ -335,80 +341,20 @@ void FStaticMeshInstanceBuffer::UpdateRHIVertexBuffer(int32 StartingIndex, uint3
 		return;
 	}
 
-	bool forcedUpdate = false;
-	if (IsDynamic)
-	{
-		{
-			uint32 Stride = InstanceData->GetOriginStride();
-			auto VertexBufferRHI = InstanceOriginBuffer.VertexBufferRHI;
-
-			uint32 UpdateSize = InstanceCount * Stride;
-			uint32 UpdateOffset = StartingIndex * Stride;
-
-			if (UpdateOffset + UpdateSize <= VertexBufferRHI->GetSize()) // we can only update the buffer
-			{
-				void* VertexBufferData = RHILockVertexBuffer(VertexBufferRHI, UpdateOffset, UpdateSize, RLM_WriteOnly);
-				FResourceArrayInterface* ResourceArray = InstanceData->GetOriginResourceArray();
+	bool ForcedUpdate = true;
 
-				FMemory::Memcpy(VertexBufferData, (uint8*)ResourceArray->GetResourceData() + UpdateOffset, UpdateSize);
-				RHIUnlockVertexBuffer(VertexBufferRHI);
-			}
-			else // not enough space, so must recreate the RHI resource with proper size
-			{
-				forcedUpdate = true;
-			}
-		}
-
-		if(!forcedUpdate)
-		{
-			uint32 Stride = InstanceData->GetTransformStride();
-			auto VertexBufferRHI = InstanceTransformBuffer.VertexBufferRHI;
-
-			uint32 UpdateSize = InstanceCount * Stride;
-			uint32 UpdateOffset = StartingIndex * Stride;
-
-			if (UpdateOffset + UpdateSize <= VertexBufferRHI->GetSize()) // we can only update the buffer
-			{
-				void* VertexBufferData = RHILockVertexBuffer(VertexBufferRHI, UpdateOffset, UpdateSize, RLM_WriteOnly);
-				FResourceArrayInterface* ResourceArray = InstanceData->GetTransformResourceArray();
-
-				FMemory::Memcpy(VertexBufferData, (uint8*)ResourceArray->GetResourceData() + UpdateOffset, UpdateSize);
-				RHIUnlockVertexBuffer(VertexBufferRHI);
-			}
-			else // not enough space, so must recreate the RHI resource with proper size
-			{
-				forcedUpdate = true;
-			}
-		}
 
-		if (!forcedUpdate)
+	if (IsDynamic)
+	{
+		if (UpdateDynamicVertexBuffer(StartingIndex, InstanceCount, InstanceData->GetOriginStride(), InstanceOriginBuffer.VertexBufferRHI, InstanceData->GetOriginResourceArray()))
 		{
-			uint32 Stride = InstanceData->GetLightMapStride();
-			auto VertexBufferRHI = InstanceLightmapBuffer.VertexBufferRHI;
-
-			uint32 UpdateSize = InstanceCount * Stride;
-			uint32 UpdateOffset = StartingIndex * Stride;
-
-			if (UpdateOffset + UpdateSize <= VertexBufferRHI->GetSize()) // we can only update the buffer
-			{
-				void* VertexBufferData = RHILockVertexBuffer(VertexBufferRHI, UpdateOffset, UpdateSize, RLM_WriteOnly);
-				FResourceArrayInterface* ResourceArray = InstanceData->GetLightMapResourceArray();
-
-				FMemory::Memcpy(VertexBufferData, (uint8*)ResourceArray->GetResourceData() + UpdateOffset, UpdateSize);
-				RHIUnlockVertexBuffer(VertexBufferRHI);
-			}
-			else // not enough space, so must recreate the RHI resource with proper size
-			{
-				forcedUpdate = true;
-			}
+			UpdateDynamicVertexBuffer(StartingIndex, InstanceCount, InstanceData->GetTransformStride(), InstanceTransformBuffer.VertexBufferRHI, InstanceData->GetTransformResourceArray());
+			UpdateDynamicVertexBuffer(StartingIndex, InstanceCount, InstanceData->GetLightMapStride(), InstanceLightmapBuffer.VertexBufferRHI, InstanceData->GetLightMapResourceArray());
+			ForcedUpdate = false;
 		}
 	}
-	else
-	{
-		forcedUpdate = true;
-	}
 	
-	if(forcedUpdate)// in non dynamic mode we have to recreate the RHI from scratch
+	if(ForcedUpdate)// in non dynamic mode we have to recreate the RHI from scratch
 	{
 		UpdateRHI();
 	}
@@ -1969,12 +1915,11 @@ void UInstancedStaticMeshComponent::InitPerInstanceRenderData(bool InitializeFro
 		InstancingRandomSeed = FMath::Rand();
 	}
 
-	UWorld* World = GetWorld();
-
-	ERHIFeatureLevel::Type FeatureLevel = World != nullptr ? World->FeatureLevel : GMaxRHIFeatureLevel;
-
 	if (!PerInstanceRenderData.IsValid())
 	{
+		UWorld* World = GetWorld();
+		ERHIFeatureLevel::Type FeatureLevel = World != nullptr ? World->FeatureLevel : GMaxRHIFeatureLevel;
+
 		bool IsDynamic = UseDynamicInstanceBuffer;
 
 #if WITH_EDITOR
diff --git a/Engine/Source/Runtime/Engine/Private/InstancedStaticMesh.h b/Engine/Source/Runtime/Engine/Private/InstancedStaticMesh.h
--- a/Engine/Source/Runtime/Engine/Private/InstancedStaticMesh.h
+++ b/Engine/Source/Runtime/Engine/Private/InstancedStaticMesh.h
@@ -170,6 +170,9 @@ private:
 	void AllocateData(FStaticMeshInstanceData& Other, bool bRequestsCPUAccess);
 
 	void UpdateRHIVertexBuffer(int32 StartingIndex, uint32 InstanceCount);
+
+	void CreateVertexBuffer(FResourceArrayInterface* InResourceArray, uint32 InUsage, uint32 InStride, uint8 InFormat, FVertexBufferRHIRef& OutVertexBufferRHI, FShaderResourceViewRHIRef& OutInstanceSRV);
+	bool UpdateDynamicVertexBuffer(int32 StartingIndex, uint32 InstanceCount, uint32 Stride, FVertexBufferRHIRef VetexBuffer, FResourceArrayInterface* ResourceArray);
 	
 	bool ComponentRequestsCPUAccess(UInstancedStaticMeshComponent* InComponent);
 };
diff --git a/Engine/Source/Runtime/Engine/Public/StaticMeshResources.h b/Engine/Source/Runtime/Engine/Public/StaticMeshResources.h
--- a/Engine/Source/Runtime/Engine/Public/StaticMeshResources.h
+++ b/Engine/Source/Runtime/Engine/Public/StaticMeshResources.h
@@ -722,7 +722,6 @@ class FStaticMeshInstanceData
 		F InstanceTransform1[4];
 		F InstanceTransform2[4];
 		F InstanceTransform3[4];
-		bool IsUsed;
 
 		friend FArchive& operator<<(FArchive& Ar, FInstanceTransformMatrix& V)
 		{
@@ -807,6 +806,7 @@ public:
 
 	void AllocateInstances(int32 InNumInstances, bool DestroyExistingInstances)
 	{
+		int32 Delta = InNumInstances - NumInstances;
 		NumInstances = InNumInstances;
 
 		if (DestroyExistingInstances)
@@ -826,6 +826,14 @@ public:
 
 		InstanceTransformData->ResizeBuffer(NumInstances);
 		InstanceTransformDataPtr = InstanceTransformData->GetDataPointer();
+
+		if (Delta > 0)
+		{
+			for (int32 i = 0; i < Delta; ++i)
+			{
+				InstancesUsage.Add(false);
+			}
+		}
 	}
 
 	FORCEINLINE int32 IsValidIndex(int32 Index) const
@@ -885,36 +893,10 @@ public:
 
 	FORCEINLINE int32 GetNextAvailableInstanceIndex() const
 	{
-		for (int32 i = 0; i < NumInstances; ++i)
-		{
-			if (bUseHalfFloat)
-			{
-				FInstanceTransformMatrix<FFloat16>* ElementData = reinterpret_cast<FInstanceTransformMatrix<FFloat16>*>(InstanceTransformData[i].GetDataPointer());
-				check((void*)((&ElementData[i]) + 1) <= (void*)(InstanceTransformDataPtr + InstanceTransformData->GetResourceSize()));
-				check((void*)((&ElementData[i]) + 0) >= (void*)(InstanceTransformDataPtr));
-
-				if (!ElementData->IsUsed)
-				{
-					return i;
-				}
-			}
-			else
-			{
-				FInstanceTransformMatrix<float>* ElementData = reinterpret_cast<FInstanceTransformMatrix<float>*>(InstanceTransformData[i].GetDataPointer());
-				check((void*)((&ElementData[i]) + 1) <= (void*)(InstanceTransformDataPtr + InstanceTransformData->GetResourceSize()));
-				check((void*)((&ElementData[i]) + 0) >= (void*)(InstanceTransformDataPtr));
-
-				if (!ElementData->IsUsed)
-				{
-					return i;
-				}
-			}
-		}
-
-		return INDEX_NONE;
+		return InstancesUsage.Find(false);
 	}
 
-	FORCEINLINE void SetInstance(int32 InstanceIndex, const FMatrix& Transform, float RandomInstanceID)
+	void SetInstance(int32 InstanceIndex, const FMatrix& Transform, float RandomInstanceID)
 	{
 		FVector4 Origin(Transform.M[3][0], Transform.M[3][1], Transform.M[3][2], RandomInstanceID);
 		SetInstanceOriginInternal(InstanceIndex, Origin);
@@ -926,17 +908,19 @@ public:
 
 		if (bUseHalfFloat)
 		{
-			SetInstanceTransformInternal<FFloat16>(InstanceIndex, InstanceTransform, true);
+			SetInstanceTransformInternal<FFloat16>(InstanceIndex, InstanceTransform);
 		}
 		else
 		{
-			SetInstanceTransformInternal<float>(InstanceIndex, InstanceTransform, true);
+			SetInstanceTransformInternal<float>(InstanceIndex, InstanceTransform);
 		}
 
 		SetInstanceLightMapDataInternal(InstanceIndex, FVector4(0, 0, 0, 0));
+
+		InstancesUsage[InstanceIndex] = true;
 	}
 	
-	FORCEINLINE void SetInstance(int32 InstanceIndex, const FMatrix& Transform, float RandomInstanceID, const FVector2D& LightmapUVBias, const FVector2D& ShadowmapUVBias)
+	void SetInstance(int32 InstanceIndex, const FMatrix& Transform, float RandomInstanceID, const FVector2D& LightmapUVBias, const FVector2D& ShadowmapUVBias)
 	{
 		FVector4 Origin(Transform.M[3][0], Transform.M[3][1], Transform.M[3][2], RandomInstanceID);
 		SetInstanceOriginInternal(InstanceIndex, Origin);
@@ -948,14 +932,16 @@ public:
 
 		if (bUseHalfFloat)
 		{
-			SetInstanceTransformInternal<FFloat16>(InstanceIndex, InstanceTransform, true);
+			SetInstanceTransformInternal<FFloat16>(InstanceIndex, InstanceTransform);
 		}
 		else
 		{
-			SetInstanceTransformInternal<float>(InstanceIndex, InstanceTransform, true);
+			SetInstanceTransformInternal<float>(InstanceIndex, InstanceTransform);
 		}
 
 		SetInstanceLightMapDataInternal(InstanceIndex, FVector4(LightmapUVBias.X, LightmapUVBias.Y, ShadowmapUVBias.X, ShadowmapUVBias.Y));
+
+		InstancesUsage[InstanceIndex] = true;
 	}
 	
 	FORCEINLINE void NullifyInstance(int32 InstanceIndex)
@@ -969,14 +955,16 @@ public:
 
 		if (bUseHalfFloat)
 		{
-			SetInstanceTransformInternal<FFloat16>(InstanceIndex, InstanceTransform, false);
+			SetInstanceTransformInternal<FFloat16>(InstanceIndex, InstanceTransform);
 		}
 		else
 		{
-			SetInstanceTransformInternal<float>(InstanceIndex, InstanceTransform, false);
+			SetInstanceTransformInternal<float>(InstanceIndex, InstanceTransform);
 		}
 
 		SetInstanceLightMapDataInternal(InstanceIndex, FVector4(0, 0, 0, 0));
+
+		InstancesUsage[InstanceIndex] = false;
 	}
 
 	FORCEINLINE void SetInstanceEditorData(int32 InstanceIndex, FColor HitProxyColor, bool bSelected)
@@ -988,7 +976,7 @@ public:
 			InstanceTransform[0][3] = ((float)HitProxyColor.R) + (bSelected ? 256.f : 0.0f);
 			InstanceTransform[1][3] = (float)HitProxyColor.G;
 			InstanceTransform[2][3] = (float)HitProxyColor.B;
-			SetInstanceTransformInternal<FFloat16>(InstanceIndex, InstanceTransform, true);
+			SetInstanceTransformInternal<FFloat16>(InstanceIndex, InstanceTransform);
 		}
 		else
 		{
@@ -996,8 +984,10 @@ public:
 			InstanceTransform[0][3] = ((float)HitProxyColor.R) + (bSelected ? 256.f : 0.0f);
 			InstanceTransform[1][3] = (float)HitProxyColor.G;
 			InstanceTransform[2][3] = (float)HitProxyColor.B;
-			SetInstanceTransformInternal<float>(InstanceIndex, InstanceTransform, true);
+			SetInstanceTransformInternal<float>(InstanceIndex, InstanceTransform);
 		}
+
+		InstancesUsage[InstanceIndex] = true;
 	}
 
 	FORCEINLINE void SwapInstance(int32 Index1, int32 Index2)
@@ -1150,7 +1140,7 @@ private:
 	}
 
 	template<typename T>
-	FORCEINLINE void SetInstanceTransformInternal(int32 InstanceIndex, FVector4(Transform)[3], bool InIsUsed) const
+	void SetInstanceTransformInternal(int32 InstanceIndex, FVector4(Transform)[3]) const
 	{
 		FInstanceTransformMatrix<T>* ElementData = reinterpret_cast<FInstanceTransformMatrix<T>*>(InstanceTransformDataPtr);
 		check((void*)((&ElementData[InstanceIndex]) + 1) <= (void*)(InstanceTransformDataPtr + InstanceTransformData->GetResourceSize()));
@@ -1170,8 +1160,6 @@ private:
 		ElementData[InstanceIndex].InstanceTransform3[1] = Transform[2][1];
 		ElementData[InstanceIndex].InstanceTransform3[2] = Transform[2][2];
 		ElementData[InstanceIndex].InstanceTransform3[3] = Transform[2][3];
-
-		ElementData->IsUsed = InIsUsed;
 	}
 
 	FORCEINLINE void SetInstanceOriginInternal(int32 InstanceIndex, const FVector4& Origin) const
@@ -1204,6 +1192,8 @@ private:
 	FStaticMeshVertexDataInterface* InstanceLightmapData = nullptr;
 	uint8* InstanceLightmapDataPtr = nullptr;	
 
+	TBitArray<> InstancesUsage;
+
 	int32 NumInstances = 0;
 	const bool bUseHalfFloat = false;
 	const bool bNeedsCPUAccess = false;
diff --git a/Engine/Source/Runtime/Foliage/Private/InstancedFoliage.cpp b/Engine/Source/Runtime/Foliage/Private/InstancedFoliage.cpp
--- a/Engine/Source/Runtime/Foliage/Private/InstancedFoliage.cpp
+++ b/Engine/Source/Runtime/Foliage/Private/InstancedFoliage.cpp
@@ -1035,6 +1035,7 @@ void FFoliageMeshInfo::AddInstance(AInstancedFoliageActor* InIFA, const UFoliage
 	}
 	else
 	{
+		Component->InitPerInstanceRenderData(false);
 		Component->InvalidateLightingCache();
 	}
 
