Added back removed FORCEINLINE by mistake

diff --git a/Engine/Source/Runtime/Engine/Public/StaticMeshResources.h b/Engine/Source/Runtime/Engine/Public/StaticMeshResources.h
--- a/Engine/Source/Runtime/Engine/Public/StaticMeshResources.h
+++ b/Engine/Source/Runtime/Engine/Public/StaticMeshResources.h
@@ -896,7 +896,7 @@ public:
 		return InstancesUsage.Find(false);
 	}
 
-	void SetInstance(int32 InstanceIndex, const FMatrix& Transform, float RandomInstanceID)
+	FORCEINLINE void SetInstance(int32 InstanceIndex, const FMatrix& Transform, float RandomInstanceID)
 	{
 		FVector4 Origin(Transform.M[3][0], Transform.M[3][1], Transform.M[3][2], RandomInstanceID);
 		SetInstanceOriginInternal(InstanceIndex, Origin);
@@ -920,7 +920,7 @@ public:
 		InstancesUsage[InstanceIndex] = true;
 	}
 	
-	void SetInstance(int32 InstanceIndex, const FMatrix& Transform, float RandomInstanceID, const FVector2D& LightmapUVBias, const FVector2D& ShadowmapUVBias)
+	FORCEINLINE void SetInstance(int32 InstanceIndex, const FMatrix& Transform, float RandomInstanceID, const FVector2D& LightmapUVBias, const FVector2D& ShadowmapUVBias)
 	{
 		FVector4 Origin(Transform.M[3][0], Transform.M[3][1], Transform.M[3][2], RandomInstanceID);
 		SetInstanceOriginInternal(InstanceIndex, Origin);
@@ -1140,7 +1140,7 @@ private:
 	}
 
 	template<typename T>
-	void SetInstanceTransformInternal(int32 InstanceIndex, FVector4(Transform)[3]) const
+	FORCEINLINE void SetInstanceTransformInternal(int32 InstanceIndex, FVector4(Transform)[3]) const
 	{
 		FInstanceTransformMatrix<T>* ElementData = reinterpret_cast<FInstanceTransformMatrix<T>*>(InstanceTransformDataPtr);
 		check((void*)((&ElementData[InstanceIndex]) + 1) <= (void*)(InstanceTransformDataPtr + InstanceTransformData->GetResourceSize()));
