[NO NON-PII COMMIT MESSAGE]

diff --git a/Engine/Source/Runtime/Engine/Private/InstancedStaticMesh.cpp b/Engine/Source/Runtime/Engine/Private/InstancedStaticMesh.cpp
--- a/Engine/Source/Runtime/Engine/Private/InstancedStaticMesh.cpp
+++ b/Engine/Source/Runtime/Engine/Private/InstancedStaticMesh.cpp
@@ -127,8 +127,11 @@ void FStaticMeshInstanceBuffer::UpdateInstanceData(UInstancedStaticMeshComponent
 			if (DestInstanceIndex == INDEX_NONE)
 			{
 				DestInstanceIndex = InstanceData->GetNextAvailableInstanceIndex();
-				check(DestInstanceIndex != INDEX_NONE);
-				InComponent->InstanceReorderTable.Add(DestInstanceIndex);
+
+				if (DestInstanceIndex != INDEX_NONE)
+				{
+					InComponent->InstanceReorderTable.Add(DestInstanceIndex);
+				}
 			}
 			else
 			{
