UE4 - Copy 3733158
Fix static permutation missing while shaders are being compiled

diff --git a/Engine/Source/Runtime/Engine/Private/Materials/MaterialInstance.cpp b/Engine/Source/Runtime/Engine/Private/Materials/MaterialInstance.cpp
--- a/Engine/Source/Runtime/Engine/Private/Materials/MaterialInstance.cpp
+++ b/Engine/Source/Runtime/Engine/Private/Materials/MaterialInstance.cpp
@@ -105,20 +105,23 @@ const FMaterial* FMaterialInstanceResource::GetMaterial(ERHIFeatureLevel::Type I
 			EMaterialQualityLevel::Type ActiveQualityLevel = GetCachedScalabilityCVars().MaterialQualityLevel;
 			FMaterialResource* StaticPermutationResource = Owner->StaticPermutationMaterialResources[ActiveQualityLevel][InFeatureLevel];
 
-			if (StaticPermutationResource->GetRenderingThreadShaderMap())
+			if (StaticPermutationResource)
 			{
-				// Verify that compilation has been finalized, the rendering thread shouldn't be touching it otherwise
-				checkSlow(StaticPermutationResource->GetRenderingThreadShaderMap()->IsCompilationFinalized());
-				// The shader map reference should have been NULL'ed if it did not compile successfully
-				checkSlow(StaticPermutationResource->GetRenderingThreadShaderMap()->CompiledSuccessfully());
-				return StaticPermutationResource;
-			}
-			else
-			{
-				EMaterialDomain Domain = (EMaterialDomain)StaticPermutationResource->GetMaterialDomain();
-				UMaterial* FallbackMaterial = UMaterial::GetDefaultMaterial(Domain);
-				//there was an error, use the default material's resource
-				return FallbackMaterial->GetRenderProxy(IsSelected(), IsHovered())->GetMaterial(InFeatureLevel);
+				if (StaticPermutationResource->GetRenderingThreadShaderMap())
+				{
+					// Verify that compilation has been finalized, the rendering thread shouldn't be touching it otherwise
+					checkSlow(StaticPermutationResource->GetRenderingThreadShaderMap()->IsCompilationFinalized());
+					// The shader map reference should have been NULL'ed if it did not compile successfully
+					checkSlow(StaticPermutationResource->GetRenderingThreadShaderMap()->CompiledSuccessfully());
+					return StaticPermutationResource;
+				}
+				else
+				{
+					EMaterialDomain Domain = (EMaterialDomain)StaticPermutationResource->GetMaterialDomain();
+					UMaterial* FallbackMaterial = UMaterial::GetDefaultMaterial(Domain);
+					//there was an error, use the default material's resource
+					return FallbackMaterial->GetRenderProxy(IsSelected(), IsHovered())->GetMaterial(InFeatureLevel);
+				}
 			}
 		}
 		else
@@ -127,11 +130,10 @@ const FMaterial* FMaterialInstanceResource::GetMaterial(ERHIFeatureLevel::Type I
 			return Parent->GetRenderProxy(IsSelected(), IsHovered())->GetMaterial(InFeatureLevel);
 		}
 	}
-	else 
-	{
-		UMaterial* FallbackMaterial = UMaterial::GetDefaultMaterial(MD_Surface);
-		return FallbackMaterial->GetRenderProxy(IsSelected(), IsHovered())->GetMaterial(InFeatureLevel);
-	}
+
+	// No Parent, or no StaticPermutationResource. This seems to happen if the parent is in the process of using the default material since it's being recompiled or failed to do so.
+	UMaterial* FallbackMaterial = UMaterial::GetDefaultMaterial(MD_Surface);
+	return FallbackMaterial->GetRenderProxy(IsSelected(), IsHovered())->GetMaterial(InFeatureLevel);
 }
 
 FMaterial* FMaterialInstanceResource::GetMaterialNoFallback(ERHIFeatureLevel::Type InFeatureLevel) const
