[NO NON-PII COMMIT MESSAGE]

diff --git a/Engine/Source/Runtime/Engine/Private/HierarchicalInstancedStaticMesh.cpp b/Engine/Source/Runtime/Engine/Private/HierarchicalInstancedStaticMesh.cpp
--- a/Engine/Source/Runtime/Engine/Private/HierarchicalInstancedStaticMesh.cpp
+++ b/Engine/Source/Runtime/Engine/Private/HierarchicalInstancedStaticMesh.cpp
@@ -1919,7 +1919,7 @@ bool UHierarchicalInstancedStaticMeshComponent::RemoveInstances(const TArray<int
 		return true;
 	}
 
-	if (PerInstanceSMData.Num() > 0 && PerInstanceRenderData.IsValid() && PerInstanceRenderData->InstanceBuffer.GetCurrentNumInstances() == 0)
+	if (PerInstanceSMData.Num() > 0 && PerInstanceRenderData.IsValid() && PerInstanceRenderData->InstanceBuffer.GetCurrentNumInstances() == 0 && !KeepInstanceBufferCPUAccess)
 	{
 		UE_LOG(LogStaticMesh, Warning, TEXT("Trying to change instance buffer for component %s, but we have no CPU copy. Set KeepInstanceBufferCPUAccess to true to keep access at the cost of memory."), *GetPathName());
 		return false;
@@ -1957,7 +1957,7 @@ bool UHierarchicalInstancedStaticMeshComponent::RemoveInstance(int32 InstanceInd
 		return false;
 	}
 
-	if (PerInstanceSMData.Num() > 0 && PerInstanceRenderData.IsValid() && PerInstanceRenderData->InstanceBuffer.GetCurrentNumInstances() == 0)
+	if (PerInstanceSMData.Num() > 0 && PerInstanceRenderData.IsValid() && PerInstanceRenderData->InstanceBuffer.GetCurrentNumInstances() == 0 && !KeepInstanceBufferCPUAccess)
 	{
 		UE_LOG(LogStaticMesh, Warning, TEXT("Trying to change instance buffer for component %s, but we have no CPU copy. Set KeepInstanceBufferCPUAccess to true to keep access at the cost of memory."), *GetPathName());
 		return false;
diff --git a/Engine/Source/Runtime/Engine/Private/InstancedStaticMesh.cpp b/Engine/Source/Runtime/Engine/Private/InstancedStaticMesh.cpp
--- a/Engine/Source/Runtime/Engine/Private/InstancedStaticMesh.cpp
+++ b/Engine/Source/Runtime/Engine/Private/InstancedStaticMesh.cpp
@@ -98,7 +98,7 @@ void FStaticMeshInstanceBuffer::UpdateInstanceData(UInstancedStaticMeshComponent
 	int32 NumRenderInstances = InComponent->PerInstanceSMData.Num() - InComponent->RemovedInstances.Num();
 
 	// Allocate the vertex data storage type.
-	if (NumRenderInstances >= InstanceData->GetNumInstances())
+	if (NumRenderInstances > 0 && NumRenderInstances >= InstanceData->GetNumInstances())
 	{
 		InstanceData->AllocateInstances(NumRenderInstances, false);
 	}
@@ -1517,7 +1517,7 @@ int32 UInstancedStaticMeshComponent::AddInstanceInternal(int32 InstanceIndex, FI
 
 int32 UInstancedStaticMeshComponent::AddInstance(const FTransform& InstanceTransform)
 {
-	if (PerInstanceSMData.Num() > 0 && PerInstanceRenderData.IsValid() && PerInstanceRenderData->InstanceBuffer.GetCurrentNumInstances() == 0)
+	if (PerInstanceSMData.Num() > 0 && PerInstanceRenderData.IsValid() && PerInstanceRenderData->InstanceBuffer.GetCurrentNumInstances() == 0 && !KeepInstanceBufferCPUAccess)
 	{
 		UE_LOG(LogStaticMesh, Warning, TEXT("Trying to change instance buffer for component %s, but we have no CPU copy. Set KeepInstanceBufferCPUAccess to true to keep access at the cost of memory."), *GetPathName());
 		return INDEX_NONE;
@@ -1535,7 +1535,7 @@ int32 UInstancedStaticMeshComponent::AddInstanceWorldSpace(const FTransform& Wor
 
 bool UInstancedStaticMeshComponent::RemoveInstanceInternal(int32 InstanceIndex, bool ReorderInstances, bool InstanceAlreadyRemoved)
 {
-	if (PerInstanceSMData.Num() > 0 && PerInstanceRenderData.IsValid() && PerInstanceRenderData->InstanceBuffer.GetCurrentNumInstances() == 0)
+	if (PerInstanceSMData.Num() > 0 && PerInstanceRenderData.IsValid() && PerInstanceRenderData->InstanceBuffer.GetCurrentNumInstances() == 0 && !KeepInstanceBufferCPUAccess)
 	{
 		UE_LOG(LogStaticMesh, Warning, TEXT("Trying to change instance buffer for component %s, but we have no CPU copy. Set KeepInstanceBufferCPUAccess to true to keep access at the cost of memory."), *GetPathName());
 		return false;
@@ -1658,7 +1658,7 @@ bool UInstancedStaticMeshComponent::UpdateInstanceTransform(int32 InstanceIndex,
 		return false;
 	}
 
-	if (PerInstanceSMData.Num() > 0 && PerInstanceRenderData.IsValid() && PerInstanceRenderData->InstanceBuffer.GetCurrentNumInstances() == 0)
+	if (PerInstanceSMData.Num() > 0 && PerInstanceRenderData.IsValid() && PerInstanceRenderData->InstanceBuffer.GetCurrentNumInstances() == 0 && !KeepInstanceBufferCPUAccess)
 	{
 		UE_LOG(LogStaticMesh, Warning, TEXT("Trying to change instance buffer for component %s, but we have no CPU copy. Set KeepInstanceBufferCPUAccess to true to keep access at the cost of memory."), *GetPathName());
 		return false;
