Copying Orion-Staging @ 3748653 (Orion/Dev-General @ 3722124) to //UE4/Main

diff --git a/Engine/Source/Runtime/Engine/Classes/Components/SkeletalMeshComponent.h b/Engine/Source/Runtime/Engine/Classes/Components/SkeletalMeshComponent.h
--- a/Engine/Source/Runtime/Engine/Classes/Components/SkeletalMeshComponent.h
+++ b/Engine/Source/Runtime/Engine/Classes/Components/SkeletalMeshComponent.h
@@ -1063,6 +1063,12 @@ public:
 	UPROPERTY(BlueprintAssignable, Category = Animation)
 	FOnAnimInitialized OnAnimInitialized;
 
+	/**
+		If MeshComponentUpdateFlag == EMeshComponentUpdateFlag::OnlyTickMontagesWhenNotRendered
+		Should we tick Montages only?
+	*/
+	bool ShouldOnlyTickMontages(const float DeltaTime) const;
+
 	/** Tick Animation system */
 	void TickAnimation(float DeltaTime, bool bNeedsValidRootMotion);
 
@@ -1338,7 +1344,10 @@ public:
 	virtual void UnHideBone( int32 BoneIndex ) override;
 	virtual void SetPhysicsAsset(class UPhysicsAsset* NewPhysicsAsset,bool bForceReInit = false) override;
 	virtual void SetSkeletalMesh(class USkeletalMesh* NewMesh, bool bReinitPose = true) override;
-	virtual FVector GetSkinnedVertexPosition(int32 VertexIndex) const override;
+
+	static FVector GetSkinnedVertexPosition(USkeletalMeshComponent* Component, int32 VertexIndex, const FSkeletalMeshLODRenderData& Model, FSkinWeightVertexBuffer& SkinWeightBuffer);
+	static FVector GetSkinnedVertexPosition(USkeletalMeshComponent* Component, int32 VertexIndex, const FSkeletalMeshLODRenderData& Model, FSkinWeightVertexBuffer& SkinWeightBuffer, TArray<FMatrix>& CachedRefToLocals);
+	static void ComputeSkinnedPositions(USkeletalMeshComponent* Component, TArray<FVector> & OutPositions, TArray<FMatrix>& CachedRefToLocals, const FSkeletalMeshLODRenderData& Model, FSkinWeightVertexBuffer& SkinWeightBuffer);
 
 	void SetSkeletalMeshWithoutResettingAnimation(class USkeletalMesh* NewMesh);
 
diff --git a/Engine/Source/Runtime/Engine/Classes/Components/SkinnedMeshComponent.h b/Engine/Source/Runtime/Engine/Classes/Components/SkinnedMeshComponent.h
--- a/Engine/Source/Runtime/Engine/Classes/Components/SkinnedMeshComponent.h
+++ b/Engine/Source/Runtime/Engine/Classes/Components/SkinnedMeshComponent.h
@@ -16,6 +16,9 @@ class FPrimitiveSceneProxy;
 class FColorVertexBuffer;
 class FSkinWeightVertexBuffer;
 class FSkeletalMeshRenderData;
+class FSkeletalMeshLODRenderData;
+struct FSkelMeshRenderSection;
+class FPositionVertexBuffer;
 
 DECLARE_DELEGATE_OneParam(FOnAnimUpdateRateParamsCreated, FAnimUpdateRateParameters*)
 
@@ -369,6 +372,12 @@ public:
 	UPROPERTY(EditAnywhere, AdvancedDisplay, BlueprintReadWrite, Category=SkeletalMesh)
 	uint8 bConsiderAllBodiesForBounds:1;
 
+	/** If true, this component uses its parents LOD when attached if available
+	* ForcedLOD can override this change. By default, it will use parent LOD.
+	*/
+	UPROPERTY(EditAnywhere, BlueprintReadWrite, AdvancedDisplay, Category = Rendering)
+	uint32 bSyncAttachParentLOD : 1;
+
 
 	/** Whether or not we can highlight selected sections - this should really only be done in the editor */
 	UPROPERTY(transient)
@@ -621,20 +630,39 @@ public:
 	 */
 	virtual void PostInitMeshObject(class FSkeletalMeshObject*) {}
 
-	/** 
-	 * Simple, CPU evaluation of a vertex's skinned position (returned in component space) 
+	/**
+	* Simple, CPU evaluation of a vertex's skinned position (returned in component space)
+	*
+	* @param VertexIndex Vertex Index. If compressed, this will be slow.
+	* @param Model The Model to use.
+	* @param SkinWeightBuffer The SkinWeightBuffer to use.
+	* @param CachedRefToLocals Cached RefToLocal matrices.
+	*/
+	static FVector GetSkinnedVertexPosition(USkinnedMeshComponent* Component, int32 VertexIndex, const FSkeletalMeshLODRenderData& LODDatal, FSkinWeightVertexBuffer& SkinWeightBuffer);
+
+	/**
+	 * Simple, CPU evaluation of a vertex's skinned position (returned in component space)
 	 *
-	 * @param VertexIndex Vertex Index. If compressed, this will be slow. 
-	 */
-	virtual FVector GetSkinnedVertexPosition(int32 VertexIndex) const;
+	 * @param VertexIndex Vertex Index. If compressed, this will be slow.
+	 * @param Model The Model to use.
+	 * @param SkinWeightBuffer The SkinWeightBuffer to use.
+	 * @param CachedRefToLocals Cached RefToLocal matrices.
+	*/
+	static FVector GetSkinnedVertexPosition(USkinnedMeshComponent* Component, int32 VertexIndex, const FSkeletalMeshLODRenderData& LODData, FSkinWeightVertexBuffer& SkinWeightBuffer, TArray<FMatrix>& CachedRefToLocals);
 
 	/**
 	* CPU evaluation of the positions of all vertices (returned in component space)
 	*
 	* @param OutPositions buffer to place positions into
+	* @param CachedRefToLocals Cached RefToLocal matrices.
+	* @param Model The Model to use.
+	* @param SkinWeightBuffer The SkinWeightBuffer to use.
 	*/
-	virtual void ComputeSkinnedPositions(TArray<FVector> & OutPositions) const;
+	static void ComputeSkinnedPositions(USkinnedMeshComponent* Component, TArray<FVector> & OutPositions, TArray<FMatrix>& CachedRefToLocals, const FSkeletalMeshLODRenderData& LODData, const FSkinWeightVertexBuffer& SkinWeightBuffer);
 
+	/** Caches the RefToLocal matrices. */
+	void CacheRefToLocalMatrices(TArray<FMatrix>& OutRefToLocal) const;
+	
 	/**
 	* Returns color of the vertex.
 	*
@@ -1085,12 +1113,6 @@ public:
 	/** Animation Update Rate optimization parameters. */
 	struct FAnimUpdateRateParameters* AnimUpdateRateParams;
 
-	/** Updates AnimUpdateRateParams, used by SkinnedMeshComponents.
-	* 
-	* @param bRecentlyRendered : true if at least one SkinnedMeshComponent on this Actor has been rendered in the last second.
-	* @param MaxDistanceFactor : Largest SkinnedMeshComponent of this Actor drawn on screen. */
-	void AnimUpdateRateSetParams(uint8 UpdateRateShift, float DeltaTime, const bool & bInRecentlyRendered, const float& InMaxDistanceFactor, const bool & bPlayingRootMotion);
-
 	virtual bool IsPlayingRootMotion() const { return false; }
 	virtual bool IsPlayingNetworkedRootMotionMontage() const { return false; }
 	virtual bool IsPlayingRootMotionFromEverything() const { return false; }
@@ -1142,3 +1164,15 @@ public:
 		}
 	}
 };
+
+
+/** Simple, CPU evaluation of a vertex's skinned position helper function */
+template <bool bExtraBoneInfluencesT, bool bCachedMatrices>
+FVector GetTypedSkinnedVertexPosition(
+	const USkinnedMeshComponent* SkinnedComp,
+	const FSkelMeshRenderSection& Section,
+	const FPositionVertexBuffer& PositionVertexBuffer,
+	const FSkinWeightVertexBuffer& SkinWeightVertexBuffer,
+	const int32 VertIndex,
+	const TArray<FMatrix> & RefToLocals = TArray<FMatrix>()
+);
diff --git a/Engine/Source/Runtime/Engine/Classes/EdGraph/EdGraph.h b/Engine/Source/Runtime/Engine/Classes/EdGraph/EdGraph.h
--- a/Engine/Source/Runtime/Engine/Classes/EdGraph/EdGraph.h
+++ b/Engine/Source/Runtime/Engine/Classes/EdGraph/EdGraph.h
@@ -170,7 +170,7 @@ public:
 	bool RemoveNode( UEdGraphNode* NodeToRemove );
 
 	/** Signal to listeners that the graph has changed */
-	void NotifyGraphChanged();
+	virtual void NotifyGraphChanged();
 
 	/** 
 	 * Move all nodes from this graph to another graph
@@ -222,7 +222,7 @@ public:
 #endif
 
 protected:
-	void NotifyGraphChanged( const FEdGraphEditAction& Action );
+	virtual void NotifyGraphChanged( const FEdGraphEditAction& Action );
 
 	/** 
 	 * Creates an empty node in this graph. Use FGraphNodeCreator above
diff --git a/Engine/Source/Runtime/Engine/Classes/Engine/CurveTable.h b/Engine/Source/Runtime/Engine/Classes/Engine/CurveTable.h
--- a/Engine/Source/Runtime/Engine/Classes/Engine/CurveTable.h
+++ b/Engine/Source/Runtime/Engine/Classes/Engine/CurveTable.h
@@ -156,7 +156,7 @@ struct ENGINE_API FCurveTableRowHandle
 	/** Returns true if the curve is valid */
 	bool IsValid(const FString& ContextString) const
 	{
-		return (GetCurve(ContextString) != nullptr);
+		return (GetCurve(ContextString, false) != nullptr);
 	}
 
 	/** Returns true if this handle is specifically pointing to nothing */
@@ -166,7 +166,7 @@ struct ENGINE_API FCurveTableRowHandle
 	}
 
 	/** Get the curve straight from the row handle */
-	FRichCurve* GetCurve(const FString& ContextString) const;
+	FRichCurve* GetCurve(const FString& ContextString, bool WarnIfNotFound=true) const;
 
 	/** Evaluate the curve if it is valid
 	 * @param XValue The input X value to the curve
@@ -186,6 +186,12 @@ struct ENGINE_API FCurveTableRowHandle
 	bool operator==(const FCurveTableRowHandle& Other) const;
 	bool operator!=(const FCurveTableRowHandle& Other) const;
 	void PostSerialize(const FArchive& Ar);
+
+	/** Used so we can have a TMap of this struct */
+	FORCEINLINE friend uint32 GetTypeHash(const FCurveTableRowHandle& Handle)
+	{
+		return HashCombine(::GetTypeHash(Handle.RowName), PointerHash(Handle.CurveTable));
+	}
 };
 
 template<>
diff --git a/Engine/Source/Runtime/Engine/Classes/Engine/EngineTypes.h b/Engine/Source/Runtime/Engine/Classes/Engine/EngineTypes.h
--- a/Engine/Source/Runtime/Engine/Classes/Engine/EngineTypes.h
+++ b/Engine/Source/Runtime/Engine/Classes/Engine/EngineTypes.h
@@ -2290,6 +2290,12 @@ public:
 	UPROPERTY()
 	EUpdateRateShiftBucket ShiftBucket;
 
+	UPROPERTY()
+	int32 SkippedUpdateFrames;
+
+	UPROPERTY()
+	int32 SkippedEvalFrames;
+
 public:
 
 	/** Default constructor. */
@@ -2308,6 +2314,8 @@ public:
 		, BaseNonRenderedUpdateRate(4)
 		, MaxEvalRateForInterpolation(4)
 		, ShiftBucket(EUpdateRateShiftBucket::ShiftBucket0)
+		, SkippedUpdateFrames(0)
+		, SkippedEvalFrames(0)
 	{ 
 		BaseVisibleDistanceFactorThesholds.Add(0.24f);
 		BaseVisibleDistanceFactorThesholds.Add(0.12f);
diff --git a/Engine/Source/Runtime/Engine/Classes/Engine/NetConnection.h b/Engine/Source/Runtime/Engine/Classes/Engine/NetConnection.h
--- a/Engine/Source/Runtime/Engine/Classes/Engine/NetConnection.h
+++ b/Engine/Source/Runtime/Engine/Classes/Engine/NetConnection.h
@@ -610,8 +610,10 @@ public:
 
 	/**
 	 * Initializes the PacketHandler
+	 *
+	 * @param InProvider Analytics provider that's passed in to the packet handler
 	 */
-	ENGINE_API virtual void InitHandler();
+	ENGINE_API virtual void InitHandler(TSharedPtr<IAnalyticsProvider> InProvider = nullptr);
 
 	/**
 	 * Initializes the sequence numbers for the connection, usually from shared randomized data
diff --git a/Engine/Source/Runtime/Engine/Classes/Engine/NetDriver.h b/Engine/Source/Runtime/Engine/Classes/Engine/NetDriver.h
--- a/Engine/Source/Runtime/Engine/Classes/Engine/NetDriver.h
+++ b/Engine/Source/Runtime/Engine/Classes/Engine/NetDriver.h
@@ -283,6 +283,8 @@ public:
 	/** Reference to the PacketHandler component, for managing stateless connection handshakes */
 	TWeakPtr<StatelessConnectHandlerComponent> StatelessConnectComponent;
 
+	/** The analytics provider used by the packet handler */
+	TSharedPtr<class IAnalyticsProvider> AnalyticsProvider;
 
 	/** World this net driver is associated with */
 	UPROPERTY()
@@ -472,6 +474,9 @@ public:
 	 */
 	void UpdateStandbyCheatStatus(void);
 
+	/** Sets the analytics provider */
+	void SetAnalyticsProvider(TSharedPtr<class IAnalyticsProvider> InProvider) { AnalyticsProvider = InProvider; }
+
 #if DO_ENABLE_NET_TEST
 	FPacketSimulationSettings	PacketSimulationSettings;
 
diff --git a/Engine/Source/Runtime/Engine/Classes/Engine/NetworkDelegates.h b/Engine/Source/Runtime/Engine/Classes/Engine/NetworkDelegates.h
--- a/Engine/Source/Runtime/Engine/Classes/Engine/NetworkDelegates.h
+++ b/Engine/Source/Runtime/Engine/Classes/Engine/NetworkDelegates.h
@@ -149,6 +149,42 @@ struct FEncryptionKeyResponse
 	}
 };
 
+/** Types of punishment to apply to a cheating client */
+enum class ECheatPunishType : uint8
+{
+	/* Unknown type of cheat punishment */
+	Unknown,
+	/** User should be booted from the game client via logout */
+	KickClient,
+	/** User should be booted from the current game session via disconnect */
+	KickMatch,
+	/** User received info about being punished (eg. ban type etc) */
+	PunishInfo
+};
+
+namespace Lex
+{
+	inline const TCHAR* const ToString(const ECheatPunishType Response)
+	{
+		switch (Response)
+		{
+			case ECheatPunishType::Unknown:
+				return TEXT("Unknown");
+			case ECheatPunishType::KickClient:
+				return TEXT("KickClient");
+			case ECheatPunishType::KickMatch:
+				return TEXT("KickMatch");
+			case ECheatPunishType::PunishInfo:
+				return TEXT("PunishInfo");
+			default:
+				break;
+		}
+
+		checkf(false, TEXT("Missing ECheatPunishType Type: %d"), static_cast<const int32>(Response));
+		return TEXT("");
+	}
+}
+
 /** 
  * Delegate called by the game to provide a response to the encryption key request
  * Provides the encryption key if successful, or a failure reason so the network connection may proceed
@@ -184,11 +220,23 @@ public:
 	static FReceivedNetworkEncryptionAck OnReceivedNetworkEncryptionAck;
 
 	/**
-	 * Delegate fired when a pending net game has created a UNetConnection to the server but hasn't sent the initial join message yet.
 	 *
-	 * @param PendingNetGame pointer to the PendingNetGame that is initializing its connection to a server.
+	 * @param PlayerId net id of the cheating client
+	 * @param PunishType type of punishment to apply to a cheating client
+	 * @param ReasonStr why the client is being punished due to cheating
 	 */
 	DECLARE_MULTICAST_DELEGATE_OneParam(FOnPendingNetGameConnectionCreated, UPendingNetGame* /*PendingNetGame*/);
 	static FOnPendingNetGameConnectionCreated OnPendingNetGameConnectionCreated;
+	
+		/**
+	 * Delegate fires whenever a client cheater is detected in the networking code
+	 *
+	 * @param PlayerId net id of the cheating client
+	 * @param PunishType type of punishment to apply to a cheating client
+	 * @param ReasonStr why the client is being punished due to cheating
+	 * @param InfoStr extra info about the punishment to be applied (eg. ban received, etc)
+	 */
+	DECLARE_DELEGATE_FourParams(FNetworkCheatDetected, const class FUniqueNetId& /*PlayerId*/, ECheatPunishType /*PunishType*/, const FString& /*ReasonStr*/, const FString& /*InfoStr*/);
+	static FNetworkCheatDetected OnNetworkCheatDetected;
 
 };
\ No newline at end of file
diff --git a/Engine/Source/Runtime/Engine/Classes/Engine/RendererSettings.h b/Engine/Source/Runtime/Engine/Classes/Engine/RendererSettings.h
--- a/Engine/Source/Runtime/Engine/Classes/Engine/RendererSettings.h
+++ b/Engine/Source/Runtime/Engine/Classes/Engine/RendererSettings.h
@@ -407,6 +407,9 @@ class ENGINE_API URendererSettings : public UDeveloperSettings
 		ToolTip = "Whether to use original CPU method (loop per morph then by vertex) or use a GPU-based method on Shader Model 5 hardware."))
 	uint32 bUseGPUMorphTargets : 1;
 
+	UPROPERTY(config, EditAnywhere, Category = "Optimizations", meta = (DisplayName = "GPU Particles Support Only Local Vector Field", Tooltip = "Limits Cascade GPU Particle simulations to applying local vector fields.  Global vector fields are not applied."))
+	bool bGPUParticlesLocalVFOnly;
+
 	UPROPERTY(config, EditAnywhere, Category = Debugging, meta = (
 		ConsoleVariable = "r.GPUCrashDebugging", DisplayName = "Enable vendor specific GPU crash analysis tools",
 		ToolTip = "Enables vendor specific GPU crash analysis tools.  Currently only supports NVIDIA Aftermath on DX11.",
diff --git a/Engine/Source/Runtime/Engine/Classes/Engine/Scene.h b/Engine/Source/Runtime/Engine/Classes/Engine/Scene.h
--- a/Engine/Source/Runtime/Engine/Classes/Engine/Scene.h
+++ b/Engine/Source/Runtime/Engine/Classes/Engine/Scene.h
@@ -685,6 +685,8 @@ struct FPostProcessSettings
 
 	UPROPERTY(EditAnywhere, BlueprintReadWrite, Category=Overrides, meta=(PinHiddenByDefault, InlineEditConditionToggle))
 	uint32 bOverride_SceneFringeIntensity:1;
+	UPROPERTY(EditAnywhere, BlueprintReadWrite, Category=Overrides, meta=(PinHiddenByDefault, InlineEditConditionToggle))
+	uint32 bOverride_ChromaticAberrationStartOffset:1;
 
 	UPROPERTY(EditAnywhere, BlueprintReadWrite, Category=Overrides, meta=(PinHiddenByDefault, InlineEditConditionToggle))
 	uint32 bOverride_AmbientCubemapTint:1;
@@ -1118,9 +1120,12 @@ struct FPostProcessSettings
 	FLinearColor SceneColorTint;
 	
 	/** in percent, Scene chromatic aberration / color fringe (camera imperfection) to simulate an artifact that happens in real-world lens, mostly visible in the image corners. */
-	UPROPERTY(interp, BlueprintReadWrite, Category="Lens|Image Effects", meta=(UIMin = "0.0", UIMax = "5.0", editcondition = "bOverride_SceneFringeIntensity", DisplayName = "Chromatic Aberration"))
+	UPROPERTY(interp, BlueprintReadWrite, Category = "Lens|Chromatic Aberration", meta = (UIMin = "0.0", UIMax = "5.0", editcondition = "bOverride_SceneFringeIntensity", DisplayName = "Intensity"))
 	float SceneFringeIntensity;
 
+	/** A normalized distance to the center of the framebuffer where the effect takes place. */
+	UPROPERTY(interp, BlueprintReadWrite, Category = "Lens|Chromatic Aberration", meta = (UIMin = "0.0", UIMax = "1.0", editcondition = "bOverride_ChromaticAberrationStartOffset", DisplayName = "Start Offset"))
+	float ChromaticAberrationStartOffset;
 
 	/** Bloom algorithm */
 	UPROPERTY(EditAnywhere, BlueprintReadWrite, Category = "Lens|Bloom", meta = (editcondition = "bOverride_BloomMethod", DisplayName = "Method"))
diff --git a/Engine/Source/Runtime/Engine/Classes/Engine/SkeletalMesh.h b/Engine/Source/Runtime/Engine/Classes/Engine/SkeletalMesh.h
--- a/Engine/Source/Runtime/Engine/Classes/Engine/SkeletalMesh.h
+++ b/Engine/Source/Runtime/Engine/Classes/Engine/SkeletalMesh.h
@@ -20,6 +20,7 @@
 #include "BoneContainer.h"
 #include "Interfaces/Interface_CollisionDataProvider.h"
 #include "EngineTypes.h"
+#include "SkeletalMeshSampling.h"
 
 #include "SkeletalMesh.generated.h"
 
@@ -39,6 +40,8 @@ class UPhysicsAsset;
 class FSkeletalMeshRenderData;
 class FSkeletalMeshModel;
 class FSkeletalMeshLODModel;
+class FSkeletalMeshLODRenderData;
+class FSkinWeightVertexBuffer;
 
 #if WITH_APEX_CLOTHING
 
@@ -322,12 +325,26 @@ struct FSkeletalMeshLODInfo
 	UPROPERTY()
 	uint32 bHasPerLODVertexColors : 1;
 
+	/** Keeps this LODs data on the CPU so it can be used for things such as sampling in FX. */
+	UPROPERTY(EditAnywhere, Category = SkeletalMeshLODInfo)
+	uint32 bAllowCPUAccess : 1;
+
+	/**
+	Mesh supports uniformly distributed sampling in constant time.
+	Memory cost is 8 bytes per triangle.
+	Example usage is uniform spawning of particles.
+	*/
+	UPROPERTY(EditAnywhere, AdvancedDisplay, Category = SkeletalMeshLODInfo, meta=(EditCondition="bAllowCPUAccess"))
+	uint32 bSupportUniformlyDistributedSampling : 1;
+
 	FSkeletalMeshLODInfo()
 		: ScreenSize(0)
 		, LODHysteresis(0)
 		, bHasBeenSimplified(false)
 		, BakePose(nullptr)
 		, bHasPerLODVertexColors(false)
+		, bAllowCPUAccess(false)
+		, bSupportUniformlyDistributedSampling(false)
 	{
 	}
 
@@ -497,6 +514,10 @@ class ENGINE_API USkeletalMesh : public UObject, public IInterface_CollisionData
 	USkeletalMesh(FVTableHelper& Helper);
 	~USkeletalMesh();
 
+#if WITH_EDITOR
+	/** Notification when anything changed */
+	DECLARE_MULTICAST_DELEGATE(FOnMeshChanged);
+#endif
 private:
 #if WITH_EDITORONLY_DATA
 	/** Imported skeletal mesh geometry information (not used at runtime). */
@@ -614,6 +635,9 @@ public:
 	UPROPERTY()
 	uint32 bHasVertexColors:1;
 
+	//caching optimization to avoid recalculating in non-editor builds
+	uint32 bHasActiveClothingAssets:1;
+
 	/** Uses skinned data for collision data. Per poly collision cannot be used for simulation, in most cases you are better off using the physics asset */
 	UPROPERTY(EditAnywhere, Category = Physics)
 	uint32 bEnablePerPolyCollision : 1;
@@ -798,6 +822,8 @@ public:
 
 	/* Get whether or not any bound clothing assets exist for this mesh **/
 	bool HasActiveClothingAssets() const;
+	/* Compute whether or not any bound clothing assets exist for this mesh **/
+	bool ComputeActiveClothingAssets() const;
 
 	/** Populates OutClothingAssets with all clothing assets that are mapped to sections in the mesh. */
 	void GetClothingAssetsInUse(TArray<UClothingAssetBase*>& OutClothingAssets) const;
@@ -805,12 +831,32 @@ public:
 	/** Adds an asset to this mesh with validation and event broadcast */
 	void AddClothingAsset(UClothingAssetBase* InNewAsset);
 
+	const FSkeletalMeshSamplingInfo& GetSamplingInfo() { return SamplingInfo; }
+
+#if WITH_EDITOR
+	void SetSamplingInfo(const FSkeletalMeshSamplingInfo& InSamplingInfo) { SamplingInfo = InSamplingInfo; }
+	FOnMeshChanged& GetOnMeshChanged() { return OnMeshChanged; }
+#endif
+
+	/** 
+	True if this mesh LOD needs to keep it's data on CPU. 
+	*/
+	bool NeedCPUData(int32 LODIndex)const;
+
 protected:
 
+	/** Defines if and how to generate a set of precomputed data allowing targeted and fast sampling of this mesh on the CPU. */
+	UPROPERTY(EditAnywhere, Category = "Sampling", meta=(ShowOnlyInnerProperties))
+	FSkeletalMeshSamplingInfo SamplingInfo;
+
 	/** Array of user data stored with the asset */
 	UPROPERTY(EditAnywhere, AdvancedDisplay, Instanced, Category=SkeletalMesh)
 	TArray<UAssetUserData*> AssetUserData;
 
+#if WITH_EDITOR
+	FOnMeshChanged OnMeshChanged;
+#endif
+
 private:
 	/** 
 	 *	Array of named socket locations, set up in editor and used as a shortcut instead of specifying 
@@ -1125,3 +1171,5 @@ private:
  * @param	InSkeletalMesh	SkeletalMesh that physics asset has been changed for
  */
 ENGINE_API void RefreshSkelMeshOnPhysicsAssetChange(const USkeletalMesh* InSkeletalMesh);
+
+ENGINE_API FVector GetSkeletalMeshRefVertLocation(const USkeletalMesh* Mesh, const FSkeletalMeshLODRenderData& LODData, const FSkinWeightVertexBuffer& SkinWeightVertexBuffer, const int32 VertIndex);
\ No newline at end of file
diff --git a/Engine/Source/Runtime/Engine/Classes/Engine/StaticMesh.h b/Engine/Source/Runtime/Engine/Classes/Engine/StaticMesh.h
--- a/Engine/Source/Runtime/Engine/Classes/Engine/StaticMesh.h
+++ b/Engine/Source/Runtime/Engine/Classes/Engine/StaticMesh.h
@@ -416,6 +416,9 @@ class UStaticMesh : public UObject, public IInterface_CollisionDataProvider, pub
 #if WITH_EDITOR
 	/** Notification when bounds changed */
 	DECLARE_MULTICAST_DELEGATE_OneParam(FOnExtendedBoundsChanged, const FBoxSphereBounds&);
+
+	/** Notification when anything changed */
+	DECLARE_MULTICAST_DELEGATE(FOnMeshChanged);
 #endif
 
 	/** Pointer to the data used to render this static mesh. */
@@ -613,6 +616,7 @@ class UStaticMesh : public UObject, public IInterface_CollisionDataProvider, pub
 
 #if WITH_EDITOR
 	FOnExtendedBoundsChanged OnExtendedBoundsChanged;
+	FOnMeshChanged OnMeshChanged;
 #endif
 
 protected:
@@ -646,6 +650,7 @@ public:
 	ENGINE_API void BroadcastNavCollisionChange();
 
 	FOnExtendedBoundsChanged& GetOnExtendedBoundsChanged() { return OnExtendedBoundsChanged; }
+	FOnMeshChanged& GetOnMeshChanged() { return OnMeshChanged; }
 #endif // WITH_EDITOR
 
 	ENGINE_API virtual void Serialize(FArchive& Ar) override;
diff --git a/Engine/Source/Runtime/Engine/Classes/GameFramework/Actor.h b/Engine/Source/Runtime/Engine/Classes/GameFramework/Actor.h
--- a/Engine/Source/Runtime/Engine/Classes/GameFramework/Actor.h
+++ b/Engine/Source/Runtime/Engine/Classes/GameFramework/Actor.h
@@ -208,6 +208,7 @@ public:
 	UPROPERTY()
 	uint8 bCollideWhenPlacing:1;
 
+
 	/** If true, this actor should search for an owned camera component to view through when used as a view target. */
 	UPROPERTY(EditAnywhere, BlueprintReadWrite, Category=Actor, AdvancedDisplay)
 	uint8 bFindCameraComponentWhenViewTarget:1;
@@ -287,9 +288,6 @@ private:
 	UPROPERTY()
 	uint8 bActorEnableCollision:1;
 
-	/** Flag indicating we have checked initial simulating physics state to sync networked proxies to the server. */
-	uint8 bNetCheckedInitialPhysicsState:1;
-
 	/** Set when actor is about to be deleted. */
 	uint8 bActorIsBeingDestroyed:1;
 
@@ -313,6 +311,11 @@ private:
 	TEnumAsByte<enum ENetRole> RemoteRole;
 
 public:
+
+	// ORION TODO - was moved to private. Add accessor?
+	/** Flag indicating we have checked initial simulating physics state to sync networked proxies to the server. */
+	uint8 bNetCheckedInitialPhysicsState : 1;
+
 	/**
 	 * Set whether this actor replicates to network clients. When this actor is spawned on the server it will be sent to clients as well.
 	 * Properties flagged for replication will update on clients if they change on the server.
diff --git a/Engine/Source/Runtime/Engine/Classes/GameFramework/GameModeBase.h b/Engine/Source/Runtime/Engine/Classes/GameFramework/GameModeBase.h
--- a/Engine/Source/Runtime/Engine/Classes/GameFramework/GameModeBase.h
+++ b/Engine/Source/Runtime/Engine/Classes/GameFramework/GameModeBase.h
@@ -327,6 +327,7 @@ public:
 	 * @return PlayerController for the player, NULL if there is any reason this player shouldn't exist or due to some error
 	 */
 	virtual APlayerController* SpawnPlayerController(ENetRole InRemoteRole, FVector const& SpawnLocation, FRotator const& SpawnRotation);
+	virtual APlayerController* SpawnReplayPlayerController(ENetRole InRemoteRole, FVector const& SpawnLocation, FRotator const& SpawnRotation);
 
 	/** Signals that a player is ready to enter the game, which may start it up */
 	UFUNCTION(BlueprintNativeEvent, Category=Game)
@@ -533,6 +534,8 @@ protected:
 	UFUNCTION(BlueprintImplementableEvent, Category=Game, meta=(DisplayName="OnSwapPlayerControllers"))
 	void K2_OnSwapPlayerControllers(APlayerController* OldPC, APlayerController* NewPC);
 
+	virtual APlayerController* SpawnPlayerControllerCommon(ENetRole InRemoteRole, FVector const& SpawnLocation, FRotator const& SpawnRotation, TSubclassOf<APlayerController> InPlayerControllerClass);
+
 public:
 	/** Whether the game perform map travels using SeamlessTravel() which loads in the background and doesn't disconnect clients */
 	UPROPERTY(EditAnywhere, BlueprintReadWrite, Category=GameMode)
diff --git a/Engine/Source/Runtime/Engine/Classes/GameFramework/HUD.h b/Engine/Source/Runtime/Engine/Classes/GameFramework/HUD.h
--- a/Engine/Source/Runtime/Engine/Classes/GameFramework/HUD.h
+++ b/Engine/Source/Runtime/Engine/Classes/GameFramework/HUD.h
@@ -62,6 +62,10 @@ class ENGINE_API AHUD : public AActor
 	UPROPERTY(BlueprintReadWrite, Category=HUD)
 	uint8 bShowDebugInfo:1;    
 
+	/** Current target in our considered Targets list for 'showdebug' */
+	UPROPERTY(Transient)
+	int32 CurrentTargetIndex;
+
 	/** If true, show hitbox debugging info. */
 	UPROPERTY(BlueprintReadWrite, Category=HUD)
 	uint8 bShowHitBoxDebugInfo:1;    
@@ -134,7 +138,7 @@ public:
 	UFUNCTION(exec)
 	void ShowDebugForReticleTargetToggle(TSubclassOf<AActor> DesiredClass);
 
-private:
+protected:
 	/** Class filter for selecting 'ShowDebugTargetActor' when 'bShowDebugForReticleTarget' is true. */
 	UPROPERTY()
 	TSubclassOf<AActor> ShowDebugTargetDesiredClass;
@@ -468,6 +472,28 @@ public:
 	 */
 	virtual void ShowDebugInfo(float& YL, float& YPos);
 
+	/** Get Target to view 'showdebug' on */
+	virtual AActor* GetCurrentDebugTargetActor();
+	
+	/** Utility function to add an actor to our consideration list for 'showdebug' 
+	 * Only consider visible, non destroyed Actors in the same world the player is in. */
+	static void AddActorToDebugList(AActor* InActor, TArray<AActor*>& InOutList, UWorld* InWorld);
+
+	/** Utility function to add a component's owner to our consideration list for 'showdebug' */
+	static void AddComponentOwnerToDebugList(UActorComponent* InComponent, TArray<AActor*>& InOutList, UWorld* InWorld);
+
+	/** Get list of considered targets for 'showdebug'
+	 * This list is built contextually based on which 'showdebug' flags have been enabled. */
+	virtual void GetDebugActorList(TArray<AActor*>& InOutList);
+
+	/** Cycle to next target in our considered targets list for 'showdebug' */
+	UFUNCTION(exec)
+	virtual void NextDebugTarget();
+
+	/** Cycle to previous target in our considered targets list for 'showdebug' */
+	UFUNCTION(exec)
+	virtual void PreviousDebugTarget();
+
 	// Callback allowing external systems to register to show debug info
 	static FOnShowDebugInfo OnShowDebugInfo;
 
diff --git a/Engine/Source/Runtime/Engine/Classes/Particles/ParticleEmitter.h b/Engine/Source/Runtime/Engine/Classes/Particles/ParticleEmitter.h
--- a/Engine/Source/Runtime/Engine/Classes/Particles/ParticleEmitter.h
+++ b/Engine/Source/Runtime/Engine/Classes/Particles/ParticleEmitter.h
@@ -86,7 +86,9 @@ struct FParticleBurst
 };
 
 DECLARE_STATS_GROUP(TEXT("Emitters"), STATGROUP_Emitters, STATCAT_Advanced);
+DECLARE_STATS_GROUP(TEXT("Emitters"), STATGROUP_EmittersRT, STATCAT_Advanced);
 DECLARE_DWORD_COUNTER_STAT_EXTERN(TEXT("STAT_EmittersStatGroupTester"), STAT_EmittersStatGroupTester, STATGROUP_Emitters, ENGINE_API);
+DECLARE_DWORD_COUNTER_STAT_EXTERN(TEXT("STAT_EmittersRTStatGroupTester"), STAT_EmittersRTStatGroupTester, STATGROUP_EmittersRT, ENGINE_API);
 
 UCLASS(hidecategories=Object, editinlinenew, abstract, MinimalAPI)
 class UParticleEmitter : public UObject
@@ -379,6 +381,7 @@ class UParticleEmitter : public UObject
 
 	/** Stat id of this object, 0 if nobody asked for it yet */
 	STAT(mutable TStatId StatID;)
+	STAT(mutable TStatId StatIDRT;)
 		
 	/**
 	* Returns the stat ID of the object...
@@ -399,6 +402,22 @@ class UParticleEmitter : public UObject
 #endif
 		return TStatId(); // not doing stats at the moment, or ever
 	}
+
+	FORCEINLINE TStatId GetStatIDRT(bool bForDeferredUse = false) const
+	{
+#if STATS
+		// this is done to avoid even registering stats for a disabled group (unless we plan on using it later)
+		if (bForDeferredUse || FThreadStats::IsCollectingData(GET_STATID(STAT_EmittersRTStatGroupTester)))
+		{
+			if (!StatIDRT.IsValidStat())
+			{
+				CreateStatID();
+			}
+			return StatIDRT;
+		}
+#endif
+		return TStatId(); // not doing stats at the moment, or ever
+	}
 	/**
 	* Creates this stat ID for the emitter...and handle a null this pointer
 	**/
diff --git a/Engine/Source/Runtime/Engine/Classes/Sound/DialogueSoundWaveProxy.h b/Engine/Source/Runtime/Engine/Classes/Sound/DialogueSoundWaveProxy.h
--- a/Engine/Source/Runtime/Engine/Classes/Sound/DialogueSoundWaveProxy.h
+++ b/Engine/Source/Runtime/Engine/Classes/Sound/DialogueSoundWaveProxy.h
@@ -48,4 +48,6 @@ public:
 	virtual void Parse( class FAudioDevice* AudioDevice, const UPTRINT NodeWaveInstanceHash, FActiveSound& AudioComponent, const FSoundParseParameters& ParseParams, TArray<FWaveInstance*>& WaveInstances ) override;
 
 	virtual USoundClass* GetSoundClass() const override;
+
+	virtual class UCurveTable* GetCurveData() const override;
 };
diff --git a/Engine/Source/Runtime/Engine/Classes/Sound/SoundBase.h b/Engine/Source/Runtime/Engine/Classes/Sound/SoundBase.h
--- a/Engine/Source/Runtime/Engine/Classes/Sound/SoundBase.h
+++ b/Engine/Source/Runtime/Engine/Classes/Sound/SoundBase.h
@@ -142,6 +142,11 @@ public:
 	/** Returns whether or not any part of this sound wants interior volumes applied to it */
 	virtual bool ShouldApplyInteriorVolumes();
 
+	/** Returns curves associated with this sound if it has any. By default returns nullptr, but types
+	*	supporting curves can return a corresponding curve table.
+	*/
+	virtual class UCurveTable* GetCurveData() const { return nullptr; }
+
 	/** Returns whether or not this sound is looping. */
 	bool IsLooping();
 
diff --git a/Engine/Source/Runtime/Engine/Classes/Sound/SoundWave.h b/Engine/Source/Runtime/Engine/Classes/Sound/SoundWave.h
--- a/Engine/Source/Runtime/Engine/Classes/Sound/SoundWave.h
+++ b/Engine/Source/Runtime/Engine/Classes/Sound/SoundWave.h
@@ -242,6 +242,8 @@ class ENGINE_API USoundWave : public USoundBase
 	
 #endif // WITH_EDITORONLY_DATA
 
+protected:
+
 	/** Curves associated with this sound wave */
 	UPROPERTY(EditAnywhere, Category = SoundWave, AdvancedDisplay)
 	class UCurveTable* Curves;
@@ -401,6 +403,28 @@ public:
 	 */ 
 	void InvalidateCompressedData();
 
+	/** Returns curves associated with this sound wave */
+	virtual class UCurveTable* GetCurveData() const override { return Curves; }
+
+#if WITH_EDITOR
+	/** These functions are required for support for some custom details/editor functionality.*/
+
+	/** Returns internal curves associated with this sound wave */
+	class UCurveTable* GetInternalCurveData() const { return InternalCurves; }
+
+	/** Returns whether this sound wave has internal curves. */
+	bool HasInternalCurves() const { return InternalCurves != nullptr; }
+
+	/** Sets the curve data for this sound wave. */
+	void SetCurveData(UCurveTable* InCurves) { Curves = InCurves; }
+
+	/** Sets the internal curve data for this sound wave. */
+	void SetInternalCurveData(UCurveTable* InCurves) { InternalCurves = InCurves; }
+
+	/** Gets the member name for the Curves property of the USoundWave object. */
+	static FName GetCurvePropertyName() { return GET_MEMBER_NAME_CHECKED(USoundWave, Curves); }
+#endif
+
 	/**
 	 * Checks whether sound has been categorised as streaming
 	 */
diff --git a/Engine/Source/Runtime/Engine/Engine.Build.cs b/Engine/Source/Runtime/Engine/Engine.Build.cs
--- a/Engine/Source/Runtime/Engine/Engine.Build.cs
+++ b/Engine/Source/Runtime/Engine/Engine.Build.cs
@@ -419,9 +419,24 @@ public class Engine : ModuleRules
 			// module, so it's definitions won't propagate to modules that import Engine.
 			Definitions.Add("WITH_RECAST=0");
 		}
+/*
+		ConfigHierarchy Ini = ConfigCache.ReadHierarchy(ConfigHierarchyType.Engine, System.IO.DirectoryReference.FromFile(Target.ProjectFile), Target.Platform);
+		bool bLocalVectorFieldOnly = false;
+		Ini.GetBool("/Script/Engine.RendererSettings", "bGPUParticlesLocalVFOnly", out bLocalVectorFieldOnly);
+		if (bLocalVectorFieldOnly)
+		{
+			Definitions.Add("GPUPARTICLE_LOCAL_VF_ONLY=1");
+		}
+		else
+		{
+			Definitions.Add("GPUPARTICLE_LOCAL_VF_ONLY=0");
+		}
+*/
+
+		Definitions.Add("GPUPARTICLE_LOCAL_VF_ONLY=1");
 
 		// Add a reference to the stats HTML files referenced by UEngine::DumpFPSChartToHTML. Previously staged by CopyBuildToStagingDirectory.
-		if(Target.bBuildEditor || Target.Configuration != UnrealTargetConfiguration.Shipping)
+    if (Target.bBuildEditor || Target.Configuration != UnrealTargetConfiguration.Shipping)
 		{
 			RuntimeDependencies.Add("$(EngineDir)/Content/Stats/...", StagedFileType.UFS);
 		}
diff --git a/Engine/Source/Runtime/Engine/Private/AI/Navigation/NavLinkRenderingComponent.cpp b/Engine/Source/Runtime/Engine/Private/AI/Navigation/NavLinkRenderingComponent.cpp
--- a/Engine/Source/Runtime/Engine/Private/AI/Navigation/NavLinkRenderingComponent.cpp
+++ b/Engine/Source/Runtime/Engine/Private/AI/Navigation/NavLinkRenderingComponent.cpp
@@ -130,6 +130,12 @@ FNavLinkRenderingProxy::FNavLinkRenderingProxy(const UPrimitiveComponent* InComp
 	}
 }
 
+SIZE_T FNavLinkRenderingProxy::GetTypeHash() const
+{
+	static size_t UniquePointer;
+	return reinterpret_cast<size_t>(&UniquePointer);
+}
+
 void FNavLinkRenderingProxy::StorePointLinks(const FTransform& InLocalToWorld, const TArray<FNavigationLink>& LinksArray)
 {
 	const FNavigationLink* Link = LinksArray.GetData();
diff --git a/Engine/Source/Runtime/Engine/Private/AI/Navigation/NavTestRenderingComponent.cpp b/Engine/Source/Runtime/Engine/Private/AI/Navigation/NavTestRenderingComponent.cpp
--- a/Engine/Source/Runtime/Engine/Private/AI/Navigation/NavTestRenderingComponent.cpp
+++ b/Engine/Source/Runtime/Engine/Private/AI/Navigation/NavTestRenderingComponent.cpp
@@ -14,6 +14,12 @@ static const FColor NavMeshRenderColor_ClosedSet(255,196,0,255);
 static const uint8 NavMeshRenderAlpha_Modifed = 255;
 static const uint8 NavMeshRenderAlpha_NonModified = 64;
 
+SIZE_T FNavTestSceneProxy::GetTypeHash() const
+{
+	static size_t UniquePointer;
+	return reinterpret_cast<size_t>(&UniquePointer);
+}
+
 FNavTestSceneProxy::FNavTestSceneProxy(const UNavTestRenderingComponent* InComponent)
 	: FDebugRenderSceneProxy(InComponent)
 	, NavMeshDrawOffset(0,0,10)
diff --git a/Engine/Source/Runtime/Engine/Private/AI/Navigation/RecastNavMesh.cpp b/Engine/Source/Runtime/Engine/Private/AI/Navigation/RecastNavMesh.cpp
--- a/Engine/Source/Runtime/Engine/Private/AI/Navigation/RecastNavMesh.cpp
+++ b/Engine/Source/Runtime/Engine/Private/AI/Navigation/RecastNavMesh.cpp
@@ -1906,16 +1906,21 @@ void ARecastNavMesh::OnStreamingLevelAdded(ULevel* InLevel, UWorld* InWorld)
 		URecastNavMeshDataChunk* NavDataChunk = GetNavigationDataChunk(InLevel);
 		if (NavDataChunk)
 		{
-			TArray<uint32> AttachedIndices = NavDataChunk->AttachTiles(*RecastNavMeshImpl);
-			if (AttachedIndices.Num() > 0)
-			{
-				InvalidateAffectedPaths(AttachedIndices);
-				RequestDrawingUpdate();
-			}
+			AttachNavMeshDataChunk(*NavDataChunk);
 		}
 	}
 }
 
+void ARecastNavMesh::AttachNavMeshDataChunk(URecastNavMeshDataChunk& NavDataChunk)
+{
+	TArray<uint32> AttachedIndices = NavDataChunk.AttachTiles(*RecastNavMeshImpl);
+	if (AttachedIndices.Num() > 0)
+	{
+		InvalidateAffectedPaths(AttachedIndices);
+		RequestDrawingUpdate();
+	}
+}
+
 void ARecastNavMesh::OnStreamingLevelRemoved(ULevel* InLevel, UWorld* InWorld)
 {
 	QUICK_SCOPE_CYCLE_COUNTER(STAT_RecastNavMesh_OnStreamingLevelRemoved);
@@ -1925,16 +1930,21 @@ void ARecastNavMesh::OnStreamingLevelRemoved(ULevel* InLevel, UWorld* InWorld)
 		URecastNavMeshDataChunk* NavDataChunk = GetNavigationDataChunk(InLevel);
 		if (NavDataChunk)
 		{
-			TArray<uint32> DetachedIndices = NavDataChunk->DetachTiles(*RecastNavMeshImpl);
-			if (DetachedIndices.Num() > 0)
-			{
-				InvalidateAffectedPaths(DetachedIndices);
-				RequestDrawingUpdate();
-			}
+			DetachNavMeshDataChunk(*NavDataChunk);
 		}
 	}
 }
 
+void ARecastNavMesh::DetachNavMeshDataChunk(URecastNavMeshDataChunk& NavDataChunk)
+{
+	TArray<uint32> DetachedIndices = NavDataChunk.DetachTiles(*RecastNavMeshImpl);
+	if (DetachedIndices.Num() > 0)
+	{
+		InvalidateAffectedPaths(DetachedIndices);
+		RequestDrawingUpdate();
+	}
+}
+
 bool ARecastNavMesh::AdjustLocationWithFilter(const FVector& StartLoc, FVector& OutAdjustedLocation, const FNavigationQueryFilter& Filter, const UObject* QueryOwner) const
 {
 	INITIALIZE_NAVQUERY(NavQuery, Filter.GetMaxSearchNodes());
diff --git a/Engine/Source/Runtime/Engine/Private/Actor.cpp b/Engine/Source/Runtime/Engine/Private/Actor.cpp
--- a/Engine/Source/Runtime/Engine/Private/Actor.cpp
+++ b/Engine/Source/Runtime/Engine/Private/Actor.cpp
@@ -2226,18 +2226,6 @@ FString AActor::GetHumanReadableName() const
 
 void AActor::DisplayDebug(UCanvas* Canvas, const FDebugDisplayInfo& DebugDisplay, float& YL, float& YPos)
 {
-	// Draw box around Actor being debugged.
-#if ENABLE_DRAW_DEBUG
-	{
-		FVector BoundsOrigin, BoundsExtent;
-		GetActorBounds(true, BoundsOrigin, BoundsExtent);
-
-		// Expand extent a little bit
-		BoundsExtent *= 1.1f;
-		DrawDebugBox(GetWorld(), BoundsOrigin, BoundsExtent, FColor::Green, false, -1.f, 0, 2.f);
-	}
-#endif
-
 	FDisplayDebugManager& DisplayDebugManager = Canvas->DisplayDebugManager;
 	DisplayDebugManager.SetDrawColor(FColor(255, 0, 0));
 
diff --git a/Engine/Source/Runtime/Engine/Private/Animation/AnimInstance.cpp b/Engine/Source/Runtime/Engine/Private/Animation/AnimInstance.cpp
--- a/Engine/Source/Runtime/Engine/Private/Animation/AnimInstance.cpp
+++ b/Engine/Source/Runtime/Engine/Private/Animation/AnimInstance.cpp
@@ -22,6 +22,8 @@
 #include "Animation/AnimNodeBase.h"
 #include "Animation/AnimInstanceProxy.h"
 #include "Animation/AnimNode_StateMachine.h"
+#include "SkeletalRenderPublic.h"
+#include "Rendering/SkeletalMeshRenderData.h"
 
 /** Anim stats */
 
@@ -346,13 +348,12 @@ void UAnimInstance::UpdateAnimation(float DeltaSeconds, bool bNeedsValidRootMoti
 	{
 		/**
 			If we're set to OnlyTickMontagesWhenNotRendered and we haven't been recently rendered,
-			then only update montages and skip everything else.
+			then only update montages and skip everything else. 
 		*/
-		if ((SkelMeshComp->MeshComponentUpdateFlag == EMeshComponentUpdateFlag::OnlyTickMontagesWhenNotRendered)
-			&& !SkelMeshComp->bRecentlyRendered)
+		if (SkelMeshComp->ShouldOnlyTickMontages(DeltaSeconds))
 		{
 			/**
-				Clear NotifyQueue prior to ticking montages. 
+				Clear NotifyQueue prior to ticking montages.
 				This is typically done in 'PreUpdate', but we're skipping this here since we're not updating the graph.
 				A side effect of this, is that we're stopping all state notifies in the graph, until ticking resumes.
 				This should be fine. But if it is ever a problem, we should keep two versions of them. One for montages and one for the graph.
@@ -368,9 +369,17 @@ void UAnimInstance::UpdateAnimation(float DeltaSeconds, bool bNeedsValidRootMoti
 			UpdateMontage(DeltaSeconds);
 
 			/**
-				We intentionally skip UpdateMontageSyncGroup(), since SyncGroup update is skipped along with AnimGraph update
-				when EMeshComponentUpdateFlag::OnlyTickMontagesWhenNotRendered.
+				We intentionally skip UpdateMontageSyncGroup(), since SyncGroup update is skipped along with AnimGraph update.
+				We do need to reset tick records since the montage will appear to have "jumped" if normal ticking resumes.
+			*/
+			for (FAnimMontageInstance* MontageInstance : MontageInstances)
+			{
+				MontageInstance->bDidUseMarkerSyncThisTick = false;
+				MontageInstance->MarkerTickRecord.Reset();
+				MontageInstance->MarkersPassedThisTick.Reset();
+			};
 
+			/**
 				We also intentionally do not call UpdateMontageEvaluationData after the call to UpdateMontage.
 				As we would have to call 'UpdateAnimation' on the graph as well, so weights could be in sync with this new data.
 				The problem lies in the fact that 'Evaluation' can be called without a call to 'Update' prior.
@@ -696,20 +705,27 @@ void OutputTickRecords(const TArray<FAnimTickRecord>& Records, UCanvas* Canvas,
 
 		DisplayDebugManager.SetLinearDrawColor((PlayerIndex == HighlightIndex) ? HighlightColor : TextColor);
 
-		FString PlayerEntry;
+		FString PlayerEntry = FString::Printf(TEXT("%i) %s (%s) W(%.f%%)"), 
+			PlayerIndex, *Player.SourceAsset->GetName(), *Player.SourceAsset->GetClass()->GetName(), Player.EffectiveBlendWeight*100.f);
 
-		// Part of a sync group
-		if (HighlightIndex != INDEX_NONE)
+		// See if we have access to SequenceLength
+		if (UAnimSequenceBase* AnimSeqBase = Cast<UAnimSequenceBase>(Player.SourceAsset))
 		{
-			PlayerEntry = FString::Printf(TEXT("%i) %s (%s) W:%.1f%% P:%.2f, Prev(i:%d, t:%.3f) Next(i:%d, t:%.3f)"),
-				PlayerIndex, *Player.SourceAsset->GetName(), *Player.SourceAsset->GetClass()->GetName(), Player.EffectiveBlendWeight*100.f, Player.TimeAccumulator != nullptr ? *Player.TimeAccumulator : 0.f
-				, Player.MarkerTickRecord->PreviousMarker.MarkerIndex, Player.MarkerTickRecord->PreviousMarker.TimeToMarker, Player.MarkerTickRecord->NextMarker.MarkerIndex, Player.MarkerTickRecord->NextMarker.TimeToMarker);
+			PlayerEntry += FString::Printf(TEXT(" P(%.2f/%.2f)"), 
+				Player.TimeAccumulator != nullptr ? *Player.TimeAccumulator : 0.f, AnimSeqBase->SequenceLength);
 		}
-		// not part of a sync group
 		else
 		{
-			PlayerEntry = FString::Printf(TEXT("%i) %s (%s) W:%.1f%% P:%.2f"),
-				PlayerIndex, *Player.SourceAsset->GetName(), *Player.SourceAsset->GetClass()->GetName(), Player.EffectiveBlendWeight*100.f, Player.TimeAccumulator != nullptr ? *Player.TimeAccumulator : 0.f);
+			PlayerEntry += FString::Printf(TEXT(" P(%.2f)"),
+				Player.TimeAccumulator != nullptr ? *Player.TimeAccumulator : 0.f);
+		}
+
+		// Part of a sync group
+		if (HighlightIndex != INDEX_NONE)
+		{
+			PlayerEntry += FString::Printf(TEXT(" Prev(i:%d, t:%.3f) Next(i:%d, t:%.3f)"),
+				Player.MarkerTickRecord->PreviousMarker.MarkerIndex, Player.MarkerTickRecord->PreviousMarker.TimeToMarker, 
+				Player.MarkerTickRecord->NextMarker.MarkerIndex, Player.MarkerTickRecord->NextMarker.TimeToMarker);
 		}
 
 		DisplayDebugManager.DrawString(PlayerEntry, Indent);
@@ -760,6 +776,35 @@ void OutputTickRecords(const TArray<FAnimTickRecord>& Records, UCanvas* Canvas,
 	}
 }
 
+void UAnimInstance::DisplayDebugInstance(FDisplayDebugManager& DisplayDebugManager, float& Indent)
+{
+	DisplayDebugManager.SetLinearDrawColor(FLinearColor::Green);
+
+	if (USkeletalMeshComponent* SkelMeshComp = GetSkelMeshComponent())
+	{
+		const int32 MaxLODIndex = SkelMeshComp->MeshObject ? (SkelMeshComp->MeshObject->GetSkeletalMeshRenderData().LODRenderData.Num() - 1) : INDEX_NONE;
+		FAnimInstanceProxy& Proxy = GetProxyOnGameThread<FAnimInstanceProxy>();
+
+		FString DebugText = FString::Printf(TEXT("LOD(%d/%d) UpdateCounter(%d) EvalCounter(%d) CacheBoneCounter(%d) InitCounter(%d) DeltaSeconds(%.3f)"),
+			SkelMeshComp->PredictedLODLevel, MaxLODIndex, Proxy.GetUpdateCounter().Get(), Proxy.GetEvaluationCounter().Get(),
+			Proxy.GetCachedBonesCounter().Get(), Proxy.GetInitializationCounter().Get(), Proxy.GetDeltaSeconds());
+
+		DisplayDebugManager.DrawString(DebugText, Indent);
+
+		if (SkelMeshComp->ShouldUseUpdateRateOptimizations())
+		{
+			if (FAnimUpdateRateParameters* UROParams = SkelMeshComp->AnimUpdateRateParams)
+			{
+				DebugText = FString::Printf(TEXT("URO Rate(%d) SkipUpdate(%d) SkipEval(%d) Interp(%d)"),
+					UROParams->UpdateRate, UROParams->ShouldSkipUpdate(), UROParams->ShouldSkipEvaluation(),
+					UROParams->ShouldInterpolateSkippedFrames());
+
+				DisplayDebugManager.DrawString(DebugText, Indent);
+			}
+		}
+	}
+}
+
 void UAnimInstance::DisplayDebug(class UCanvas* Canvas, const FDebugDisplayInfo& DebugDisplay, float& YL, float& YPos)
 {
 #if ENABLE_DRAW_DEBUG
@@ -796,6 +841,11 @@ void UAnimInstance::DisplayDebug(class UCanvas* Canvas, const FDebugDisplayInfo&
 	FString Heading = FString::Printf(TEXT("Animation: %s"), *GetName());
 	DisplayDebugManager.DrawString(Heading, Indent);
 
+	{
+		FIndenter CustomDebugIndent(Indent);
+		DisplayDebugInstance(DisplayDebugManager, Indent);
+	}
+
 	if (bShowGraph && Proxy.HasRootNode())
 	{
 		DisplayDebugManager.SetLinearDrawColor(TextYellow);
diff --git a/Engine/Source/Runtime/Engine/Private/Animation/AnimNode_SequencePlayer.cpp b/Engine/Source/Runtime/Engine/Private/Animation/AnimNode_SequencePlayer.cpp
--- a/Engine/Source/Runtime/Engine/Private/Animation/AnimNode_SequencePlayer.cpp
+++ b/Engine/Source/Runtime/Engine/Private/Animation/AnimNode_SequencePlayer.cpp
@@ -13,7 +13,8 @@ float FAnimNode_SequencePlayer::GetCurrentAssetTime()
 
 float FAnimNode_SequencePlayer::GetCurrentAssetTimePlayRateAdjusted()
 {
-	float EffectivePlayrate = PlayRate * (Sequence ? Sequence->RateScale : 1.0f);
+	const float SequencePlayRate = (Sequence ? Sequence->RateScale : 1.f);
+	const float EffectivePlayrate = FMath::IsNearlyZero(PlayRateBasis) ? 0.f : (SequencePlayRate * PlayRate / PlayRateBasis);
 	return (EffectivePlayrate < 0.0f) ? GetCurrentAssetLength() - InternalTimeAccumulator : InternalTimeAccumulator;
 }
 
@@ -28,11 +29,11 @@ void FAnimNode_SequencePlayer::Initialize_AnyThread(const FAnimationInitializeCo
 
 	EvaluateGraphExposedInputs.Execute(Context);
 	InternalTimeAccumulator = StartPosition;
-	if (Sequence != NULL)
+	if (Sequence != nullptr)
 	{
 		InternalTimeAccumulator = FMath::Clamp(StartPosition, 0.f, Sequence->SequenceLength);
-
-		if (StartPosition == 0.f && (PlayRate * Sequence->RateScale) < 0.0f)
+		const float EffectivePlayrate = FMath::IsNearlyZero(PlayRateBasis) ? 0.f : (Sequence->RateScale * PlayRate / PlayRateBasis);
+		if ((StartPosition == 0.f) && (EffectivePlayrate < 0.f))
 		{
 			InternalTimeAccumulator = Sequence->SequenceLength;
 		}
@@ -47,10 +48,11 @@ void FAnimNode_SequencePlayer::UpdateAssetPlayer(const FAnimationUpdateContext&
 {
 	EvaluateGraphExposedInputs.Execute(Context);
 
-	if ((Sequence != NULL) && (Context.AnimInstanceProxy->IsSkeletonCompatible(Sequence->GetSkeleton())))
+	if ((Sequence != nullptr) && (Context.AnimInstanceProxy->IsSkeletonCompatible(Sequence->GetSkeleton())))
 	{
 		InternalTimeAccumulator = FMath::Clamp(InternalTimeAccumulator, 0.f, Sequence->SequenceLength);
-		CreateTickRecordForNode(Context, Sequence, bLoopAnimation, PlayRate);
+		const float AdjustedPlayRate = FMath::IsNearlyZero(PlayRateBasis) ? 0.f : (PlayRate / PlayRateBasis);
+		CreateTickRecordForNode(Context, Sequence, bLoopAnimation, AdjustedPlayRate);
 	}
 }
 
@@ -61,7 +63,7 @@ void FAnimNode_SequencePlayer::Evaluate_AnyThread(FPoseContext& Output)
 
 void FAnimNode_SequencePlayer::Evaluate_AnyThread(FPoseContext& Output, bool bExpectsAdditivePose)
 {
-	if ((Sequence != NULL) && (Output.AnimInstanceProxy->IsSkeletonCompatible(Sequence->GetSkeleton())))
+	if ((Sequence != nullptr) && (Output.AnimInstanceProxy->IsSkeletonCompatible(Sequence->GetSkeleton())))
 	{
 		Sequence->GetAnimationPose(Output.Pose, Output.Curve, FAnimExtractContext(InternalTimeAccumulator, Output.AnimInstanceProxy->ShouldExtractRootMotion()));
 	}
diff --git a/Engine/Source/Runtime/Engine/Private/Animation/AnimationRuntime.cpp b/Engine/Source/Runtime/Engine/Private/Animation/AnimationRuntime.cpp
--- a/Engine/Source/Runtime/Engine/Private/Animation/AnimationRuntime.cpp
+++ b/Engine/Source/Runtime/Engine/Private/Animation/AnimationRuntime.cpp
@@ -522,6 +522,51 @@ void FAnimationRuntime::LerpPosesPerBone(FCompactPose& PoseA, const FCompactPose
 	}
 }
 
+void FAnimationRuntime::LerpPosesWithBoneIndexList(FCompactPose& PoseA, const FCompactPose& PoseB, FBlendedCurve& CurveA, const FBlendedCurve& CurveB, float Alpha, const TArray<FCompactPoseBoneIndex>& BoneIndices)
+{
+	// If pose A is full weight, we're set.
+	if (FAnimWeight::IsRelevant(Alpha))
+	{
+		// Make sure poses are compatible with each other.
+		check(&PoseA.GetBoneContainer() == &PoseB.GetBoneContainer());
+
+		// If pose 2 is full weight, just copy, no need to blend.
+		if (FAnimWeight::IsFullWeight(Alpha))
+		{
+			for (int32 Index = 0; Index < BoneIndices.Num(); Index++)
+			{
+				const FCompactPoseBoneIndex& BoneIndex = BoneIndices[Index];
+				PoseA[BoneIndex] = PoseB[BoneIndex];
+			}
+
+			CurveA.CopyFrom(CurveB);
+		}
+		else
+		{
+			const ScalarRegister VWeightOfPose1(1.f - Alpha);
+			const ScalarRegister VWeightOfPose2(Alpha);
+
+			for (int32 Index = 0; Index < BoneIndices.Num(); Index++)
+			{
+				const FCompactPoseBoneIndex& BoneIndex = BoneIndices[Index];
+
+				FTransform& InOutBoneTransform1 = PoseA[BoneIndex];
+				InOutBoneTransform1 *= VWeightOfPose1;
+
+				const FTransform& BoneTransform2 = PoseB[BoneIndex];
+				InOutBoneTransform1.AccumulateWithShortestRotation(BoneTransform2, VWeightOfPose2);
+				InOutBoneTransform1.NormalizeRotation();
+			}
+
+			// @note : This isn't perfect as curve can link to joint, and it would be the best to use that information
+			// but that is very expensive option as we have to have another indirect look up table to search. 
+			// For now, replacing with combine (non-zero will be overridden)
+			// in the future, we might want to do this outside if we want per bone blend to apply curve also UE-39182
+			CurveA.Combine(CurveB);
+		}
+	}
+}
+
 void FAnimationRuntime::LerpBoneTransforms(TArray<FTransform>& A, const TArray<FTransform>& B, float Alpha, const TArray<FBoneIndexType>& RequiredBonesArray)
 {
 	if (FAnimWeight::IsFullWeight(Alpha))
diff --git a/Engine/Source/Runtime/Engine/Private/Components/ArrowComponent.cpp b/Engine/Source/Runtime/Engine/Private/Components/ArrowComponent.cpp
--- a/Engine/Source/Runtime/Engine/Private/Components/ArrowComponent.cpp
+++ b/Engine/Source/Runtime/Engine/Private/Components/ArrowComponent.cpp
@@ -30,9 +30,14 @@ float UArrowComponent::EditorScale = 1.0f;
 #endif
 
 /** Represents a UArrowComponent to the scene manager. */
-class FArrowSceneProxy : public FPrimitiveSceneProxy
+class FArrowSceneProxy final : public FPrimitiveSceneProxy
 {
 public:
+	SIZE_T GetTypeHash() const override
+	{
+		static size_t UniquePointer;
+		return reinterpret_cast<size_t>(&UniquePointer);
+	}
 
 	FArrowSceneProxy(UArrowComponent* Component)
 		: FPrimitiveSceneProxy(Component)
diff --git a/Engine/Source/Runtime/Engine/Private/Components/BillboardComponent.cpp b/Engine/Source/Runtime/Engine/Private/Components/BillboardComponent.cpp
--- a/Engine/Source/Runtime/Engine/Private/Components/BillboardComponent.cpp
+++ b/Engine/Source/Runtime/Engine/Private/Components/BillboardComponent.cpp
@@ -26,9 +26,15 @@ float UBillboardComponent::EditorScale = 1.0f;
 #endif
 
 /** Represents a billboard sprite to the scene manager. */
-class FSpriteSceneProxy : public FPrimitiveSceneProxy
+class FSpriteSceneProxy final : public FPrimitiveSceneProxy
 {
 public:
+	SIZE_T GetTypeHash() const override
+	{
+		static size_t UniquePointer;
+		return reinterpret_cast<size_t>(&UniquePointer);
+	}
+
 	/** Initialization constructor. */
 	FSpriteSceneProxy(const UBillboardComponent* InComponent, float SpriteScale)
 		: FPrimitiveSceneProxy(InComponent)
diff --git a/Engine/Source/Runtime/Engine/Private/Components/BoxComponent.cpp b/Engine/Source/Runtime/Engine/Private/Components/BoxComponent.cpp
--- a/Engine/Source/Runtime/Engine/Private/Components/BoxComponent.cpp
+++ b/Engine/Source/Runtime/Engine/Private/Components/BoxComponent.cpp
@@ -116,9 +116,15 @@ FBoxSphereBounds UBoxComponent::CalcBounds(const FTransform& LocalToWorld) const
 FPrimitiveSceneProxy* UBoxComponent::CreateSceneProxy()
 {
 	/** Represents a UBoxComponent to the scene manager. */
-	class FBoxSceneProxy : public FPrimitiveSceneProxy
+	class FBoxSceneProxy final : public FPrimitiveSceneProxy
 	{
 	public:
+		SIZE_T GetTypeHash() const override
+		{
+			static size_t UniquePointer;
+			return reinterpret_cast<size_t>(&UniquePointer);
+		}
+
 		FBoxSceneProxy(const UBoxComponent* InComponent)
 			:	FPrimitiveSceneProxy(InComponent)
 			,	bDrawOnlyIfSelected( InComponent->bDrawOnlyIfSelected )
diff --git a/Engine/Source/Runtime/Engine/Private/Components/BrushComponent.cpp b/Engine/Source/Runtime/Engine/Private/Components/BrushComponent.cpp
--- a/Engine/Source/Runtime/Engine/Private/Components/BrushComponent.cpp
+++ b/Engine/Source/Runtime/Engine/Private/Components/BrushComponent.cpp
@@ -83,9 +83,15 @@ private:
 };
 #endif // WITH_EDITORONLY_DATA
 
-class FBrushSceneProxy : public FPrimitiveSceneProxy
+class FBrushSceneProxy final : public FPrimitiveSceneProxy
 {
 public:
+	SIZE_T GetTypeHash() const override
+	{
+		static size_t UniquePointer;
+		return reinterpret_cast<size_t>(&UniquePointer);
+	}
+
 	FBrushSceneProxy(UBrushComponent* Component, ABrush* Owner):
 		FPrimitiveSceneProxy(Component),
 #if WITH_EDITORONLY_DATA
diff --git a/Engine/Source/Runtime/Engine/Private/Components/CapsuleComponent.cpp b/Engine/Source/Runtime/Engine/Private/Components/CapsuleComponent.cpp
--- a/Engine/Source/Runtime/Engine/Private/Components/CapsuleComponent.cpp
+++ b/Engine/Source/Runtime/Engine/Private/Components/CapsuleComponent.cpp
@@ -26,9 +26,15 @@ UCapsuleComponent::UCapsuleComponent(const FObjectInitializer& ObjectInitializer
 FPrimitiveSceneProxy* UCapsuleComponent::CreateSceneProxy()
 {
 	/** Represents a UCapsuleComponent to the scene manager. */
-	class FDrawCylinderSceneProxy : public FPrimitiveSceneProxy
+	class FDrawCylinderSceneProxy final : public FPrimitiveSceneProxy
 	{
 	public:
+		SIZE_T GetTypeHash() const override
+		{
+			static size_t UniquePointer;
+			return reinterpret_cast<size_t>(&UniquePointer);
+		}
+
 		FDrawCylinderSceneProxy(const UCapsuleComponent* InComponent)
 			:	FPrimitiveSceneProxy(InComponent)
 			,	bDrawOnlyIfSelected( InComponent->bDrawOnlyIfSelected )
diff --git a/Engine/Source/Runtime/Engine/Private/Components/DrawFrustumComponent.cpp b/Engine/Source/Runtime/Engine/Private/Components/DrawFrustumComponent.cpp
--- a/Engine/Source/Runtime/Engine/Private/Components/DrawFrustumComponent.cpp
+++ b/Engine/Source/Runtime/Engine/Private/Components/DrawFrustumComponent.cpp
@@ -12,9 +12,14 @@
 
 
 /** Represents a draw frustum to the scene manager. */
-class FDrawFrustumSceneProxy : public FPrimitiveSceneProxy
+class FDrawFrustumSceneProxy final : public FPrimitiveSceneProxy
 {
 public:
+	SIZE_T GetTypeHash() const override
+	{
+		static size_t UniquePointer;
+		return reinterpret_cast<size_t>(&UniquePointer);
+	}
 
 	/** 
 	* Initialization constructor. 
diff --git a/Engine/Source/Runtime/Engine/Private/Components/LineBatchComponent.cpp b/Engine/Source/Runtime/Engine/Private/Components/LineBatchComponent.cpp
--- a/Engine/Source/Runtime/Engine/Private/Components/LineBatchComponent.cpp
+++ b/Engine/Source/Runtime/Engine/Private/Components/LineBatchComponent.cpp
@@ -22,6 +22,12 @@ FLineBatcherSceneProxy::FLineBatcherSceneProxy(const ULineBatchComponent* InComp
 	bWillEverBeLit = false;
 }
 
+SIZE_T FLineBatcherSceneProxy::GetTypeHash() const
+{
+	static size_t UniquePointer;
+	return reinterpret_cast<size_t>(&UniquePointer);
+}
+
 void FLineBatcherSceneProxy::GetDynamicMeshElements(const TArray<const FSceneView*>& Views, const FSceneViewFamily& ViewFamily, uint32 VisibilityMap, FMeshElementCollector& Collector) const
 {
 	QUICK_SCOPE_CYCLE_COUNTER( STAT_LineBatcherSceneProxy_GetDynamicMeshElements );
diff --git a/Engine/Source/Runtime/Engine/Private/Components/MaterialBillboardComponent.cpp b/Engine/Source/Runtime/Engine/Private/Components/MaterialBillboardComponent.cpp
--- a/Engine/Source/Runtime/Engine/Private/Components/MaterialBillboardComponent.cpp
+++ b/Engine/Source/Runtime/Engine/Private/Components/MaterialBillboardComponent.cpp
@@ -49,9 +49,14 @@ public:
 };
 
 /** Represents a sprite to the scene manager. */
-class FMaterialSpriteSceneProxy : public FPrimitiveSceneProxy
+class FMaterialSpriteSceneProxy final : public FPrimitiveSceneProxy
 {
 public:
+	SIZE_T GetTypeHash() const override
+	{
+		static size_t UniquePointer;
+		return reinterpret_cast<size_t>(&UniquePointer);
+	}
 
 	/** Initialization constructor. */
 	FMaterialSpriteSceneProxy(const UMaterialBillboardComponent* InComponent)
diff --git a/Engine/Source/Runtime/Engine/Private/Components/ModelComponent.cpp b/Engine/Source/Runtime/Engine/Private/Components/ModelComponent.cpp
--- a/Engine/Source/Runtime/Engine/Private/Components/ModelComponent.cpp
+++ b/Engine/Source/Runtime/Engine/Private/Components/ModelComponent.cpp
@@ -670,44 +670,52 @@ bool UModelComponent::GetPhysicsTriMeshData(struct FTriMeshCollisionData* Collis
 		FModelElement& Element = Elements[ElementIndex];
 		FRawIndexBuffer16or32* IndexBuffer = Element.IndexBuffer;
 		int32 IndexBufferSize = 0;
-		// Check index buffer pointer is valid and has something in it
-		if (IndexBuffer != nullptr && IndexBuffer->Indices.Num() >= 0)
+
+		if (IndexBuffer == nullptr || IndexBuffer->Indices.Num() == 0)
 		{
-			IndexBufferSize = IndexBuffer->Indices.Num();
+			UE_LOG(LogPhysics, Warning, TEXT("Found bad index buffer when cooking UModelComponent physics data! Component: %s, Buffer: %x, Buffer Size: %d, Element: %d"), *GetPathName(), IndexBuffer, IndexBufferSize, ElementIndex);
+			verify(IndexBufferSize >= 0);
+			continue;
+		}
 
-			for (uint32 TriIdx = 0; TriIdx < Element.NumTriangles; TriIdx++)
-			{
-				FTriIndices Triangle;
+		int32 NumVertices = Model->VertexBuffer.Vertices.Num();
+
+		if (NumVertices < (int32)Element.MaxVertexIndex)
+		{
+			UE_LOG(LogPhysics, Warning, TEXT("Found bad vertex buffer when cooking UModelComponent physics data! Component: %s, Element: %d. Verts Exected: %d, Actual: %d"), 
+				*GetPathName(), ElementIndex, Element.MaxVertexIndex, NumVertices);
+			continue;
+		}
+
+		IndexBufferSize = IndexBuffer->Indices.Num();
+
+		for (uint32 TriIdx = 0; TriIdx < Element.NumTriangles; TriIdx++)
+		{
+			FTriIndices Triangle;
 
-				Triangle.v0 = IndexBuffer->Indices[Element.FirstIndex + (TriIdx * 3) + 0];
-				Triangle.v1 = IndexBuffer->Indices[Element.FirstIndex + (TriIdx * 3) + 1];
-				Triangle.v2 = IndexBuffer->Indices[Element.FirstIndex + (TriIdx * 3) + 2];
+			Triangle.v0 = IndexBuffer->Indices[Element.FirstIndex + (TriIdx * 3) + 0];
+			Triangle.v1 = IndexBuffer->Indices[Element.FirstIndex + (TriIdx * 3) + 1];
+			Triangle.v2 = IndexBuffer->Indices[Element.FirstIndex + (TriIdx * 3) + 2];
 
-				if (AreaThreshold >= 0.f)
+			if (AreaThreshold >= 0.f)
+			{
+				const FVector V0 = Model->VertexBuffer.Vertices[Triangle.v0].Position;
+				const FVector V1 = Model->VertexBuffer.Vertices[Triangle.v1].Position;
+				const FVector V2 = Model->VertexBuffer.Vertices[Triangle.v2].Position;
+
+				const FVector V01 = (V1 - V0);
+				const FVector V02 = (V2 - V0);
+				const FVector Cross = FVector::CrossProduct(V01, V02);
+				const float Area = Cross.Size() * 0.5f;
+				if (Area <= AreaThreshold)
 				{
-					const FVector V0 = Model->VertexBuffer.Vertices[Triangle.v0].Position;
-					const FVector V1 = Model->VertexBuffer.Vertices[Triangle.v1].Position;
-					const FVector V2 = Model->VertexBuffer.Vertices[Triangle.v2].Position;
-
-					const FVector V01 = (V1 - V0);
-					const FVector V02 = (V2 - V0);
-					const FVector Cross = FVector::CrossProduct(V01, V02);
-					const float Area = Cross.Size() * 0.5f;
-					if (Area <= AreaThreshold)
-					{
-						nBadArea++;
-						continue;
-					}
+					nBadArea++;
+					continue;
 				}
-
-				CollisionData->Indices.Add(Triangle);
-				CollisionData->MaterialIndices.Add(ElementIndex);
 			}
-		}
-		else
-		{
-			UE_LOG(LogPhysics, Warning, TEXT("Found bad index buffer when cooking UModelComponent physics data! Component: %s, Buffer: %x, Buffer Size: %d, Element: %d"), *GetPathName(this), IndexBuffer, IndexBufferSize, ElementIndex);
-			verify(IndexBufferSize >= 0);
+
+			CollisionData->Indices.Add(Triangle);
+			CollisionData->MaterialIndices.Add(ElementIndex);
 		}
 	}
 
diff --git a/Engine/Source/Runtime/Engine/Private/Components/NavMeshRenderingComponent.cpp b/Engine/Source/Runtime/Engine/Private/Components/NavMeshRenderingComponent.cpp
--- a/Engine/Source/Runtime/Engine/Private/Components/NavMeshRenderingComponent.cpp
+++ b/Engine/Source/Runtime/Engine/Private/Components/NavMeshRenderingComponent.cpp
@@ -790,6 +790,12 @@ void FNavMeshSceneProxyData::GatherData(const ARecastNavMesh* NavMesh, int32 InN
 //////////////////////////////////////////////////////////////////////////
 // FNavMeshSceneProxy
 
+SIZE_T FNavMeshSceneProxy::GetTypeHash() const
+{
+	static size_t UniquePointer;
+	return reinterpret_cast<size_t>(&UniquePointer);
+}
+
 FNavMeshSceneProxy::FNavMeshSceneProxy(const UPrimitiveComponent* InComponent, FNavMeshSceneProxyData* InProxyData, bool ForceToRender)
 	: FDebugRenderSceneProxy(InComponent)
 	, VertexFactory(GetScene().GetFeatureLevel(), "FNavMeshSceneProxy")
diff --git a/Engine/Source/Runtime/Engine/Private/Components/SceneComponent.cpp b/Engine/Source/Runtime/Engine/Private/Components/SceneComponent.cpp
--- a/Engine/Source/Runtime/Engine/Private/Components/SceneComponent.cpp
+++ b/Engine/Source/Runtime/Engine/Private/Components/SceneComponent.cpp
@@ -1596,9 +1596,10 @@ void USceneComponent::SetupAttachment(class USceneComponent* InParent, FName InS
 	{
 		if (ensureMsgf(InParent != this, TEXT("Cannot attach a component to itself.")))
 		{
-			if (ensureMsgf(InParent == nullptr || !InParent->IsAttachedTo(this), TEXT("Setting up attachment would create a cycle.")))
+			// Paragon hack
+			if (/*ensureMsgf*/(InParent == nullptr || !InParent->IsAttachedTo(this))) //, TEXT("Setting up attachment would create a cycle.")))
 			{
-				if (ensureMsgf(AttachParent == nullptr || !AttachParent->AttachChildren.Contains(this), TEXT("SetupAttachment cannot be used once a component has already had AttachTo used to connect it to a parent.")))
+				if (/*ensureMsgf*/(AttachParent == nullptr || !AttachParent->AttachChildren.Contains(this)))// , TEXT("SetupAttachment cannot be used once a component has already had AttachTo used to connect it to a parent.")))
 				{
 					AttachParent = InParent;
 					AttachSocketName = InSocketName;
diff --git a/Engine/Source/Runtime/Engine/Private/Components/SkeletalMeshComponent.cpp b/Engine/Source/Runtime/Engine/Private/Components/SkeletalMeshComponent.cpp
--- a/Engine/Source/Runtime/Engine/Private/Components/SkeletalMeshComponent.cpp
+++ b/Engine/Source/Runtime/Engine/Private/Components/SkeletalMeshComponent.cpp
@@ -13,6 +13,7 @@
 #include "Animation/AnimInstance.h"
 #include "Animation/AnimMontage.h"
 #include "Animation/AnimSingleNodeInstance.h"
+#include "Animation/AnimationSettings.h"
 #include "Engine/SkeletalMeshSocket.h"
 #include "AI/NavigationSystemHelpers.h"
 #include "PhysicsPublic.h"
@@ -37,6 +38,7 @@
 #include "AnimPhysObjectVersion.h"
 #include "ScopeExit.h"
 
+
 #define LOCTEXT_NAMESPACE "SkeletalMeshComponent"
 
 TAutoConsoleVariable<int32> CVarUseParallelAnimationEvaluation(TEXT("a.ParallelAnimEvaluation"), 1, TEXT("If 1, animation evaluation will be run across the task graph system. If 0, evaluation will run purely on the game thread"));
@@ -595,9 +597,9 @@ void USkeletalMeshComponent::InitAnim(bool bForceReinit)
 
 		const USkeleton* AnimSkeleton = (AnimScriptInstance)? AnimScriptInstance->CurrentSkeleton : nullptr;
 
-		bool bClearAnimInstance = AnimScriptInstance && !AnimSkeleton;
-		bool bSkeletonMismatch = AnimSkeleton && (AnimScriptInstance->CurrentSkeleton!=SkeletalMesh->Skeleton);
-		bool bSkeletonNotCompatible = AnimSkeleton && !bSkeletonMismatch && (AnimSkeleton->IsCompatibleMesh(SkeletalMesh) == false);
+		const bool bClearAnimInstance = AnimScriptInstance && !AnimSkeleton;
+		const bool bSkeletonMismatch = AnimSkeleton && (AnimScriptInstance->CurrentSkeleton!=SkeletalMesh->Skeleton);
+		const bool bSkeletonNotCompatible = AnimSkeleton && !bSkeletonMismatch && (AnimSkeleton->IsCompatibleMesh(SkeletalMesh) == false);
 
 		if (bBlueprintMismatch || bSkeletonMismatch || bSkeletonNotCompatible || bClearAnimInstance)
 		{
@@ -605,29 +607,40 @@ void USkeletalMeshComponent::InitAnim(bool bForceReinit)
 		}
 
 		// this has to be called before Initialize Animation because it will required RequiredBones list when InitializeAnimScript
-		RecalcRequiredBones(0);
+		RecalcRequiredBones(PredictedLODLevel);
 
 		const bool bInitializedAnimInstance = InitializeAnimScriptInstance(bForceReinit);
+
 		// Make sure we have a valid pose.
-		if (bInitializedAnimInstance || (AnimScriptInstance == nullptr))
-		{
-			if (bUseRefPoseOnInitAnim)
-			{
-				BoneSpaceTransforms = SkeletalMesh->RefSkeleton.GetRefBonePose();
-				//Mini RefreshBoneTransforms (the bit we actually care about)
-				FillComponentSpaceTransforms(SkeletalMesh, BoneSpaceTransforms, GetEditableComponentSpaceTransforms());
-				bNeedToFlipSpaceBaseBuffers = true; // Have updated space bases so need to flip
-				FlipEditableSpaceBases();
-			}
-			else
-			{
-				TickAnimation(0.f, false);
-				RefreshBoneTransforms();
-			}
+		// We don't allocate transform data when using MasterPoseComponent, so we have nothing to render.
+		if (!MasterPoseComponent.IsValid())
+		{	
+			if (bInitializedAnimInstance || (AnimScriptInstance == nullptr))
+			{ 
+				// In Editor, animations won't get ticked. So Update once to get accurate representation instead of T-Pose.
+				// Also allow this to be an option to support pre-4.19 games that might need it..
+				const bool bTickAnimationNow =
+					((GetWorld()->WorldType == EWorldType::Editor) && !bUseRefPoseOnInitAnim)
+					|| UAnimationSettings::Get()->bTickAnimationOnSkeletalMeshInit;
+
+				if (bTickAnimationNow)
+				{
+					TickAnimation(0.f, false);
+					RefreshBoneTransforms();
+				}
+				else
+				{
+					BoneSpaceTransforms = SkeletalMesh->RefSkeleton.GetRefBonePose();
+					//Mini RefreshBoneTransforms (the bit we actually care about)
+					FillComponentSpaceTransforms(SkeletalMesh, BoneSpaceTransforms, GetEditableComponentSpaceTransforms());
+					bNeedToFlipSpaceBaseBuffers = true; // Have updated space bases so need to flip
+					FlipEditableSpaceBases();
+				}
 
-			if (bInitializedAnimInstance)
-			{
-				OnAnimInitialized.Broadcast();
+				if (bInitializedAnimInstance)
+				{
+					OnAnimInitialized.Broadcast();
+				}
 			}
 		}
 
@@ -915,6 +928,15 @@ void USkeletalMeshComponent::LoadedFromAnotherClass(const FName& OldClassName)
 }
 #endif // WITH_EDITOR
 
+bool USkeletalMeshComponent::ShouldOnlyTickMontages(const float DeltaTime) const
+{
+	// Ignore DeltaSeconds == 0.f, as that is used when we want to force an update followed by RefreshBoneTransforms.
+	// RefreshBoneTransforms will need an updated graph.
+	return (MeshComponentUpdateFlag == EMeshComponentUpdateFlag::OnlyTickMontagesWhenNotRendered)
+		&& !bRecentlyRendered 
+		&& (DeltaTime > 0.f);
+}
+
 void USkeletalMeshComponent::TickAnimation(float DeltaTime, bool bNeedsValidRootMotion)
 {
 	SCOPED_NAMED_EVENT(USkeletalMeshComponent_TickAnimation, FColor::Yellow);
@@ -947,8 +969,7 @@ void USkeletalMeshComponent::TickAnimation(float DeltaTime, bool bNeedsValidRoot
 			If we're called directly for autonomous proxies, TickComponent is not guaranteed to get called.
 			So dispatch all queued events here if we're doing MontageOnly ticking.
 		*/
-		if ((MeshComponentUpdateFlag == EMeshComponentUpdateFlag::OnlyTickMontagesWhenNotRendered)
-			&& !bRecentlyRendered)
+		if (ShouldOnlyTickMontages(DeltaTime))
 		{
 			ConditionallyDispatchQueuedAnimEvents();
 		}
@@ -1757,10 +1778,10 @@ void USkeletalMeshComponent::WritebackClothingSimulationData()
 {
 	if(ClothingSimulation)
 	{
-		USkeletalMeshComponent* OverrideComponent = nullptr;
+		USkinnedMeshComponent* OverrideComponent = nullptr;
 		if(MasterPoseComponent.IsValid())
 		{
-			OverrideComponent = Cast<USkeletalMeshComponent>(MasterPoseComponent.Get());
+			OverrideComponent = Cast<USkinnedMeshComponent>(MasterPoseComponent.Get());
 
 			// Check if our bone map is actually valid, if not there is no clothing data to build
 			if(MasterBoneMap.Num() == 0)
diff --git a/Engine/Source/Runtime/Engine/Private/Components/SkinnedMeshComponent.cpp b/Engine/Source/Runtime/Engine/Private/Components/SkinnedMeshComponent.cpp
--- a/Engine/Source/Runtime/Engine/Private/Components/SkinnedMeshComponent.cpp
+++ b/Engine/Source/Runtime/Engine/Private/Components/SkinnedMeshComponent.cpp
@@ -219,6 +219,28 @@ namespace FAnimUpdateRateManager
 					// Add 1 as an eval rate of 1 is 0 frameskip
 					DesiredEvaluationRate = (*FrameSkip) + 1;
 				}
+				// We haven't found our LOD number into our array. :(
+				// Default to matching settings of previous highest LOD number we've found.
+				// For example if we're missing LOD 3, and we have settings for LOD 2, then match that.
+				// Having no settings means we default to evaluating every frame, which is highest quality setting we have.
+				// This is not what we want to higher LOD numbers.
+				else if (Tracker->UpdateRateParameters.LODToFrameSkipMap.Num() > 0)
+				{
+					TMap<int32, int32>& LODToFrameSkipMap = Tracker->UpdateRateParameters.LODToFrameSkipMap;
+					for (auto Iter = LODToFrameSkipMap.CreateConstIterator(); Iter; ++Iter)
+					{
+						if (Iter.Key() < MinLod)
+						{
+							DesiredEvaluationRate = FMath::Max(Iter.Value(), DesiredEvaluationRate);
+						}
+					}
+
+					// Cache result back into TMap, so we don't have to do this every frame.
+					LODToFrameSkipMap.Add(MinLod, DesiredEvaluationRate);
+
+					// Add 1 as an eval rate of 1 is 0 frameskip
+					DesiredEvaluationRate++;
+				}
 			}
 
 			int32 ForceAnimRate = CVarForceAnimRate.GetValueOnGameThread();
@@ -320,6 +342,7 @@ USkinnedMeshComponent::USkinnedMeshComponent(const FObjectInitializer& ObjectIni
 
 	bCanEverAffectNavigation = false;
 	MasterBoneMapCacheCount = 0;
+	bSyncAttachParentLOD = true;
 }
 
 
@@ -1864,63 +1887,6 @@ void USkinnedMeshComponent::GetUsedMaterials( TArray<UMaterialInterface*>& OutMa
 	}
 }
 
-/** Simple, CPU evaluation of a vertex's skinned position helper function */
-template <bool bExtraBoneInfluencesT, bool bCachedMatrices>
-FORCEINLINE FVector GetTypedSkinnedVertexPosition(
-	const USkinnedMeshComponent* SkinnedComp, 
-	const FSkelMeshRenderSection& Section, 
-	const FPositionVertexBuffer& PositionVertexBuffer,
-	const FSkinWeightVertexBuffer& SkinWeightVertexBuffer,
-	const int32 VertIndex, 
-	const TArray<FMatrix> & RefToLocals = TArray<FMatrix>()
-	)
-{
-	FVector SkinnedPos(0,0,0);
-
-	const USkinnedMeshComponent* const MasterPoseComponentInst = SkinnedComp->MasterPoseComponent.Get();
-	const USkinnedMeshComponent* BaseComponent = MasterPoseComponentInst ? MasterPoseComponentInst : SkinnedComp;
-
-	// Do soft skinning for this vertex.
-	int32 BufferVertIndex = Section.GetVertexBufferIndex() + VertIndex;
-	const TSkinWeightInfo<bExtraBoneInfluencesT>* SrcSkinWeights = SkinWeightVertexBuffer.GetSkinWeightPtr<bExtraBoneInfluencesT>(BufferVertIndex);
-	int32 MaxBoneInfluences = bExtraBoneInfluencesT ? MAX_TOTAL_INFLUENCES : MAX_INFLUENCES_PER_STREAM;
-
-#if !PLATFORM_LITTLE_ENDIAN
-	// uint8[] elements in LOD.VertexBufferGPUSkin have been swapped for VET_UBYTE4 vertex stream use
-	for(int32 InfluenceIndex = MAX_INFLUENCES-1;InfluenceIndex >=  MAX_INFLUENCES- MaxBoneInfluences;InfluenceIndex--)
-#else
-	for(int32 InfluenceIndex = 0;InfluenceIndex < MaxBoneInfluences;InfluenceIndex++)
-#endif
-	{
-		const int32 MeshBoneIndex = Section.BoneMap[SrcSkinWeights->InfluenceBones[InfluenceIndex]];
-		int32 TransformBoneIndex = MeshBoneIndex;
-
-		if(MasterPoseComponentInst)
-		{		
-			const TArray<int32>& MasterBoneMap = SkinnedComp->GetMasterBoneMap();
-			check(MasterBoneMap.Num() == SkinnedComp->SkeletalMesh->RefSkeleton.GetNum());
-			TransformBoneIndex = MasterBoneMap[MeshBoneIndex];
-		}
-
-		const float	Weight = (float)SrcSkinWeights->InfluenceWeights[InfluenceIndex] / 255.0f;
-		{
-			if (bCachedMatrices)
-			{
-				const FMatrix& RefToLocal = RefToLocals[MeshBoneIndex];
-				SkinnedPos += RefToLocal.TransformPosition(PositionVertexBuffer.VertexPosition(BufferVertIndex)) * Weight;
-			}
-			else
-			{
-				const FMatrix BoneTransformMatrix = (TransformBoneIndex != INDEX_NONE) ? BaseComponent->GetComponentSpaceTransforms()[TransformBoneIndex].ToMatrixWithScale() : FMatrix::Identity;
-				const FMatrix RefToLocal = SkinnedComp->SkeletalMesh->RefBasesInvMatrix[MeshBoneIndex] * BoneTransformMatrix;
-				SkinnedPos += RefToLocal.TransformPosition(PositionVertexBuffer.VertexPosition(BufferVertIndex)) * Weight;
-			}
-		}
-	}
-
-	return SkinnedPos;
-}
-
 FSkinWeightVertexBuffer* USkinnedMeshComponent::GetSkinWeightBuffer(int32 LODIndex) const
 {
 	FSkinWeightVertexBuffer* WeightBuffer = nullptr;
@@ -1947,33 +1913,36 @@ FSkinWeightVertexBuffer* USkinnedMeshComponent::GetSkinWeightBuffer(int32 LODInd
 	return WeightBuffer;
 }
 
-
-FVector USkinnedMeshComponent::GetSkinnedVertexPosition(int32 VertexIndex) const
+FVector USkinnedMeshComponent::GetSkinnedVertexPosition(USkinnedMeshComponent* Component, int32 VertexIndex, const FSkeletalMeshLODRenderData& LODData, FSkinWeightVertexBuffer& SkinWeightBuffer) 
 {
-	FVector SkinnedPos(0,0,0);
+	FVector SkinnedPos(0, 0, 0);
 
-	// Fail if no mesh
-	if(!SkeletalMesh || !MeshObject)
-	{
-		return SkinnedPos;
-	}
+	int32 SectionIndex;
+	int32 VertIndex;
+	LODData.GetSectionFromVertexIndex(VertexIndex, SectionIndex, VertIndex);
 
-	FSkeletalMeshLODRenderData& LODData = MeshObject->GetSkeletalMeshRenderData().LODRenderData[0];
+	check(SectionIndex < LODData.RenderSections.Num());
+	const FSkelMeshRenderSection& Section = LODData.RenderSections[SectionIndex];
 
-	//cache RefToLocal matrices
+	return  SkinWeightBuffer.HasExtraBoneInfluences()
+		? GetTypedSkinnedVertexPosition<true, false>(Component, Section, LODData.StaticVertexBuffers.PositionVertexBuffer, SkinWeightBuffer, VertIndex)
+		: GetTypedSkinnedVertexPosition<false, false>(Component, Section, LODData.StaticVertexBuffers.PositionVertexBuffer, SkinWeightBuffer, VertIndex);
+}
+
+FVector USkinnedMeshComponent::GetSkinnedVertexPosition(USkinnedMeshComponent* Component, int32 VertexIndex, const FSkeletalMeshLODRenderData& LODData, FSkinWeightVertexBuffer& SkinWeightBuffer, TArray<FMatrix>& CachedRefToLocals) 
+{
+	FVector SkinnedPos(0, 0, 0);
+	
 	int32 SectionIndex;
 	int32 VertIndex;
 	LODData.GetSectionFromVertexIndex(VertexIndex, SectionIndex, VertIndex);
 
-	//update positions
 	check(SectionIndex < LODData.RenderSections.Num());
 	const FSkelMeshRenderSection& Section = LODData.RenderSections[SectionIndex];
 
-	FSkinWeightVertexBuffer* SkinWeightBuffer = GetSkinWeightBuffer(0);
-
-	return SkinWeightBuffer->HasExtraBoneInfluences()
-		? GetTypedSkinnedVertexPosition<true, false>(this, Section, LODData.StaticVertexBuffers.PositionVertexBuffer, *GetSkinWeightBuffer(0), VertIndex)
-		: GetTypedSkinnedVertexPosition<false, false>(this, Section, LODData.StaticVertexBuffers.PositionVertexBuffer, *GetSkinWeightBuffer(0), VertIndex);
+	return SkinWeightBuffer.HasExtraBoneInfluences()
+		? GetTypedSkinnedVertexPosition<true, false>(Component, Section, LODData.StaticVertexBuffers.PositionVertexBuffer, SkinWeightBuffer, VertIndex,CachedRefToLocals)
+		: GetTypedSkinnedVertexPosition<false, false>(Component, Section, LODData.StaticVertexBuffers.PositionVertexBuffer, SkinWeightBuffer, VertIndex, CachedRefToLocals);
 }
 
 void USkinnedMeshComponent::SetRefPoseOverride(const TArray<FTransform>& NewRefPoseTransforms)
@@ -2057,43 +2026,54 @@ void USkinnedMeshComponent::ClearRefPoseOverride()
 	}
 }
 
-void USkinnedMeshComponent::ComputeSkinnedPositions(TArray<FVector> & OutPositions) const
+void USkinnedMeshComponent::CacheRefToLocalMatrices(TArray<FMatrix>& OutRefToLocal)const
 {
-	OutPositions.Empty();
+	const USkinnedMeshComponent* BaseComponent = MasterPoseComponent.IsValid() ? MasterPoseComponent.Get() : this;
+	OutRefToLocal.SetNumUninitialized(SkeletalMesh->RefBasesInvMatrix.Num());
+	const TArray<FTransform>& CompSpaceTransforms = BaseComponent->GetComponentSpaceTransforms();
+	if(CompSpaceTransforms.Num())
+	{
+		check(CompSpaceTransforms.Num() >= OutRefToLocal.Num());
 
-	// Fail if no mesh
-	if (!SkeletalMesh || !MeshObject)
+		for (int32 MatrixIdx = 0; MatrixIdx < OutRefToLocal.Num(); ++MatrixIdx)
+		{
+			OutRefToLocal[MatrixIdx] = SkeletalMesh->RefBasesInvMatrix[MatrixIdx] * CompSpaceTransforms[MatrixIdx].ToMatrixWithScale();
+		}
+	}
+	else
 	{
-		return;
+		//Possible in some cases to request this before the component space transforms are prepared (undo/redo)
+		for (int32 MatrixIdx = 0; MatrixIdx < OutRefToLocal.Num(); ++MatrixIdx)
+		{
+			OutRefToLocal[MatrixIdx] = SkeletalMesh->RefBasesInvMatrix[MatrixIdx];
+		}
 	}
+}
 
-	FSkeletalMeshLODRenderData& LODData = MeshObject->GetSkeletalMeshRenderData().LODRenderData[0];
-	OutPositions.AddUninitialized(LODData.GetNumVertices());
+void USkinnedMeshComponent::ComputeSkinnedPositions(USkinnedMeshComponent* Component, TArray<FVector> & OutPositions, TArray<FMatrix>& CachedRefToLocals, const FSkeletalMeshLODRenderData& LODData, const FSkinWeightVertexBuffer& SkinWeightBuffer)
+{
+	OutPositions.Empty();
 
-	//cache RefToLocal matrices
-	const USkinnedMeshComponent* BaseComponent = MasterPoseComponent.IsValid() ? MasterPoseComponent.Get() : this;
-	TArray<FMatrix> RefToLocals;
-	RefToLocals.AddUninitialized(SkeletalMesh->RefBasesInvMatrix.Num());
+	// Fail if no mesh
+	if (!Component || !Component->SkeletalMesh)
 	{
-		for (int32 MatrixIdx = 0; MatrixIdx < RefToLocals.Num(); ++MatrixIdx)
-		{
-			RefToLocals[MatrixIdx] = SkeletalMesh->RefBasesInvMatrix[MatrixIdx] * BaseComponent->GetComponentSpaceTransforms()[MatrixIdx].ToMatrixWithScale();
-		}
+		return;
 	}
+	OutPositions.AddUninitialized(LODData.GetNumVertices());
 
 	//update positions
 	for (int32 SectionIdx = 0; SectionIdx < LODData.RenderSections.Num(); ++SectionIdx)
 	{
 		const FSkelMeshRenderSection& Section = LODData.RenderSections[SectionIdx];
-		bool bHasExtraBoneInfluences = GetSkinWeightBuffer(0)->HasExtraBoneInfluences();
+		bool bHasExtraBoneInfluences = SkinWeightBuffer.HasExtraBoneInfluences();
 		{
 			//soft
 			const uint32 SoftOffset = Section.BaseVertexIndex;
 			const uint32 NumSoftVerts = Section.NumVertices;
 			for (uint32 SoftIdx = 0; SoftIdx < NumSoftVerts; ++SoftIdx)
 			{
-				FVector SkinnedPosition = bHasExtraBoneInfluences ? GetTypedSkinnedVertexPosition<true, true>(this, Section, LODData.StaticVertexBuffers.PositionVertexBuffer, *GetSkinWeightBuffer(0), SoftIdx, RefToLocals) :
-																	GetTypedSkinnedVertexPosition<false,true>(this, Section, LODData.StaticVertexBuffers.PositionVertexBuffer, *GetSkinWeightBuffer(0), SoftIdx, RefToLocals);
+				FVector SkinnedPosition = bHasExtraBoneInfluences ? GetTypedSkinnedVertexPosition<true, true>(Component, Section, LODData.StaticVertexBuffers.PositionVertexBuffer, SkinWeightBuffer, SoftIdx, CachedRefToLocals) :
+																	GetTypedSkinnedVertexPosition<false,true>(Component, Section, LODData.StaticVertexBuffers.PositionVertexBuffer, SkinWeightBuffer, SoftIdx, CachedRefToLocals);
 				OutPositions[SoftOffset + SoftIdx] = SkinnedPosition;
 			}
 		}
@@ -2325,6 +2305,10 @@ bool USkinnedMeshComponent::UpdateLODStatus()
 			{
 				PredictedLODLevel = FMath::Clamp(MasterPoseComponent->PredictedLODLevel, 0, MaxLODIndex);
 			}
+			else if (bSyncAttachParentLOD && GetAttachParent() && GetAttachParent()->IsA(USkinnedMeshComponent::StaticClass()))
+			{
+				PredictedLODLevel = FMath::Clamp(CastChecked<USkinnedMeshComponent>(GetAttachParent())->PredictedLODLevel, 0, MaxLODIndex);
+			}
 			else if (MeshObject)
 			{
 #if WITH_EDITOR
@@ -2857,16 +2841,35 @@ void FAnimUpdateRateParameters::SetTrailMode(float DeltaTime, uint8 UpdateRateSh
 	ThisTickDelta = DeltaTime;
 
 	UpdateRate = FMath::Max(NewUpdateRate, 1);
+
 	// Make sure EvaluationRate is a multiple of UpdateRate.
 	EvaluationRate = FMath::Max((NewEvaluationRate / UpdateRate) * UpdateRate, 1);
 	bInterpolateSkippedFrames = (FAnimUpdateRateManager::CVarURODisableInterpolation.GetValueOnAnyThread() == 0) &&
 		((bNewInterpSkippedFrames && (EvaluationRate < MaxEvalRateForInterpolation)) || (FAnimUpdateRateManager::CVarForceInterpolation.GetValueOnAnyThread() == 1));
 
 	// Make sure we don't overflow. we don't need very large numbers.
-	const uint32 Counter = (GFrameCounter + UpdateRateShift)% MAX_uint32;
+	const uint32 Counter = (GFrameCounter + UpdateRateShift) % MAX_uint32;
 
 	bSkipUpdate = ((Counter % UpdateRate) > 0);
 	bSkipEvaluation = ((Counter % EvaluationRate) > 0);
+
+	// As UpdateRate changes, because of LODs for example,
+	// make sure we're not caught in a loop where we don't update longer than our update rate.
+	{
+		SkippedUpdateFrames = bSkipUpdate ? ++SkippedUpdateFrames : 0;
+		SkippedEvalFrames = bSkipEvaluation ? ++SkippedEvalFrames : 0;
+
+		// If we've gone longer that our UpdateRate, force an update to happen.
+		if ((SkippedUpdateFrames >= UpdateRate) || (SkippedEvalFrames >= EvaluationRate))
+		{
+			bSkipUpdate = false;
+			bSkipEvaluation = false;
+			SkippedUpdateFrames = 0;
+			SkippedEvalFrames = 0;
+		}
+	}
+
+	// We should never trigger an Eval without an Update.
 	check((bSkipEvaluation && bSkipUpdate) || (bSkipEvaluation && !bSkipUpdate) || (!bSkipEvaluation && !bSkipUpdate));
 
 	AdditionalTime = 0.f;
@@ -2941,4 +2944,69 @@ float FAnimUpdateRateParameters::GetRootMotionInterp() const
 		return FMath::Clamp(ThisTickDelta / (TickedPoseOffestTime + ThisTickDelta), 0.f, 1.f);
 	}
 	return 1.f;
-}
\ No newline at end of file
+}
+
+/** Simple, CPU evaluation of a vertex's skinned position helper function */
+template <bool bExtraBoneInfluencesT, bool bCachedMatrices>
+FVector GetTypedSkinnedVertexPosition(
+	const USkinnedMeshComponent* SkinnedComp,
+	const FSkelMeshRenderSection& Section,
+	const FPositionVertexBuffer& PositionVertexBuffer,
+	const FSkinWeightVertexBuffer& SkinWeightVertexBuffer,
+	const int32 VertIndex,
+	const TArray<FMatrix> & RefToLocals
+)
+{
+	FVector SkinnedPos(0, 0, 0);
+
+	const USkinnedMeshComponent* const MasterPoseComponentInst = SkinnedComp->MasterPoseComponent.Get();
+	const USkinnedMeshComponent* BaseComponent = MasterPoseComponentInst ? MasterPoseComponentInst : SkinnedComp;
+
+	// Do soft skinning for this vertex.
+	int32 BufferVertIndex = Section.GetVertexBufferIndex() + VertIndex;
+	const TSkinWeightInfo<bExtraBoneInfluencesT>* SrcSkinWeights = SkinWeightVertexBuffer.GetSkinWeightPtr<bExtraBoneInfluencesT>(BufferVertIndex);
+	int32 MaxBoneInfluences = bExtraBoneInfluencesT ? MAX_TOTAL_INFLUENCES : MAX_INFLUENCES_PER_STREAM;
+
+#if !PLATFORM_LITTLE_ENDIAN
+	// uint8[] elements in LOD.VertexBufferGPUSkin have been swapped for VET_UBYTE4 vertex stream use
+	for (int32 InfluenceIndex = MAX_INFLUENCES - 1; InfluenceIndex >= MAX_INFLUENCES - MaxBoneInfluences; InfluenceIndex--)
+#else
+	for (int32 InfluenceIndex = 0; InfluenceIndex < MaxBoneInfluences; InfluenceIndex++)
+#endif
+	{
+		const int32 MeshBoneIndex = Section.BoneMap[SrcSkinWeights->InfluenceBones[InfluenceIndex]];
+		int32 TransformBoneIndex = MeshBoneIndex;
+
+		if (MasterPoseComponentInst)
+		{
+			const TArray<int32>& MasterBoneMap = SkinnedComp->GetMasterBoneMap();
+			check(MasterBoneMap.Num() == SkinnedComp->SkeletalMesh->RefSkeleton.GetNum());
+			TransformBoneIndex = MasterBoneMap[MeshBoneIndex];
+		}
+
+		const float	Weight = (float)SrcSkinWeights->InfluenceWeights[InfluenceIndex] / 255.0f;
+		{
+			if (bCachedMatrices)
+			{
+				const FMatrix& RefToLocal = RefToLocals[MeshBoneIndex];
+				SkinnedPos += RefToLocal.TransformPosition(PositionVertexBuffer.VertexPosition(BufferVertIndex)) * Weight;
+			}
+			else
+			{
+				const FMatrix BoneTransformMatrix = (TransformBoneIndex != INDEX_NONE) ? BaseComponent->GetComponentSpaceTransforms()[TransformBoneIndex].ToMatrixWithScale() : FMatrix::Identity;
+				const FMatrix RefToLocal = SkinnedComp->SkeletalMesh->RefBasesInvMatrix[MeshBoneIndex] * BoneTransformMatrix;
+				SkinnedPos += RefToLocal.TransformPosition(PositionVertexBuffer.VertexPosition(BufferVertIndex)) * Weight;
+			}
+		}
+	}
+
+	return SkinnedPos;
+}
+
+
+
+template FVector GetTypedSkinnedVertexPosition<false, true>(const USkinnedMeshComponent* SkinnedComp, const FSkelMeshRenderSection& Section, const FPositionVertexBuffer& PositionVertexBuffer,
+	const FSkinWeightVertexBuffer& SkinWeightVertexBuffer, const int32 VertIndex, const TArray<FMatrix> & RefToLocals);
+
+template FVector GetTypedSkinnedVertexPosition<true, true>(const USkinnedMeshComponent* SkinnedComp, const FSkelMeshRenderSection& Section, const FPositionVertexBuffer& PositionVertexBuffer,
+	const FSkinWeightVertexBuffer& SkinWeightVertexBuffer, const int32 VertIndex, const TArray<FMatrix> & RefToLocals);
\ No newline at end of file
diff --git a/Engine/Source/Runtime/Engine/Private/Components/SphereComponent.cpp b/Engine/Source/Runtime/Engine/Private/Components/SphereComponent.cpp
--- a/Engine/Source/Runtime/Engine/Private/Components/SphereComponent.cpp
+++ b/Engine/Source/Runtime/Engine/Private/Components/SphereComponent.cpp
@@ -105,9 +105,14 @@ bool USphereComponent::IsZeroExtent() const
 FPrimitiveSceneProxy* USphereComponent::CreateSceneProxy()
 {
 	/** Represents a DrawLightRadiusComponent to the scene manager. */
-	class FSphereSceneProxy : public FPrimitiveSceneProxy
+	class FSphereSceneProxy final : public FPrimitiveSceneProxy
 	{
 	public:
+		SIZE_T GetTypeHash() const override
+		{
+			static size_t UniquePointer;
+			return reinterpret_cast<size_t>(&UniquePointer);
+		}
 
 		/** Initialization constructor. */
 		FSphereSceneProxy(const USphereComponent* InComponent)
diff --git a/Engine/Source/Runtime/Engine/Private/Components/SplineComponent.cpp b/Engine/Source/Runtime/Engine/Private/Components/SplineComponent.cpp
--- a/Engine/Source/Runtime/Engine/Private/Components/SplineComponent.cpp
+++ b/Engine/Source/Runtime/Engine/Private/Components/SplineComponent.cpp
@@ -1364,9 +1364,14 @@ FPrimitiveSceneProxy* USplineComponent::CreateSceneProxy()
 		return Super::CreateSceneProxy();
 	}
 
-	class FSplineSceneProxy : public FPrimitiveSceneProxy
+	class FSplineSceneProxy final : public FPrimitiveSceneProxy
 	{
 	public:
+		SIZE_T GetTypeHash() const override
+		{
+			static size_t UniquePointer;
+			return reinterpret_cast<size_t>(&UniquePointer);
+		}
 
 		FSplineSceneProxy(const USplineComponent* InComponent)
 			: FPrimitiveSceneProxy(InComponent)
diff --git a/Engine/Source/Runtime/Engine/Private/Components/StaticMeshComponent.cpp b/Engine/Source/Runtime/Engine/Private/Components/StaticMeshComponent.cpp
--- a/Engine/Source/Runtime/Engine/Private/Components/StaticMeshComponent.cpp
+++ b/Engine/Source/Runtime/Engine/Private/Components/StaticMeshComponent.cpp
@@ -303,6 +303,11 @@ void UStaticMeshComponent::Serialize(FArchive& Ar)
 		{
 			LODData[ LODIndex ].OwningComponent = this;
 		}
+
+		if (GetName().Contains(TEXT("Dusk_Launcher_EmissTop")))
+		{
+			printf("");
+		}
 	}
 
 #if WITH_EDITORONLY_DATA
@@ -1628,6 +1633,11 @@ void UStaticMeshComponent::UpdateCollisionFromStaticMesh()
 
 void UStaticMeshComponent::PostLoad()
 {
+	if (GetName().Contains(TEXT("SternInhibitorRing1")))
+	{
+		printf("");
+	}
+
 	// need to postload the StaticMesh because super initializes variables based on GetStaticLightingType() which we override and use from the StaticMesh
 	if (GetStaticMesh())
 	{
@@ -2160,6 +2170,11 @@ void UStaticMeshComponent::ApplyComponentInstanceData(FStaticMeshComponentInstan
 {
 	check(StaticMeshInstanceData);
 
+	if (GetName().Contains(TEXT("SternInhibitorRing1")))
+	{
+		printf("");
+	}
+
 	// Note: ApplyComponentInstanceData is called while the component is registered so the rendering thread is already using this component
 	// That means all component state that is modified here must be mirrored on the scene proxy, which will be recreated to receive the changes later due to MarkRenderStateDirty.
 
diff --git a/Engine/Source/Runtime/Engine/Private/Components/TextRenderComponent.cpp b/Engine/Source/Runtime/Engine/Private/Components/TextRenderComponent.cpp
--- a/Engine/Source/Runtime/Engine/Private/Components/TextRenderComponent.cpp
+++ b/Engine/Source/Runtime/Engine/Private/Components/TextRenderComponent.cpp
@@ -559,9 +559,15 @@ FTextRenderComponentMIDCache* FTextRenderComponentMIDCache::InstancePtr = nullpt
 // ------------------------------------------------------
 
 /** Represents a URenderTextComponent to the scene manager. */
-class FTextRenderSceneProxy : public FPrimitiveSceneProxy
+class FTextRenderSceneProxy final : public FPrimitiveSceneProxy
 {
 public:
+	SIZE_T GetTypeHash() const override
+	{
+		static size_t UniquePointer;
+		return reinterpret_cast<size_t>(&UniquePointer);
+	}
+
 	// constructor / destructor
 	FTextRenderSceneProxy(UTextRenderComponent* Component);
 	virtual ~FTextRenderSceneProxy();
diff --git a/Engine/Source/Runtime/Engine/Private/CurveTable.cpp b/Engine/Source/Runtime/Engine/Private/CurveTable.cpp
--- a/Engine/Source/Runtime/Engine/Private/CurveTable.cpp
+++ b/Engine/Source/Runtime/Engine/Private/CurveTable.cpp
@@ -568,18 +568,18 @@ TArray<const UObject*> UCurveTable::GetOwners() const
 //////////////////////////////////////////////////////////////////////////
 
 
-FRichCurve* FCurveTableRowHandle::GetCurve(const FString& ContextString) const
+FRichCurve* FCurveTableRowHandle::GetCurve(const FString& ContextString, bool WarnIfNotFound) const
 {
 	if(CurveTable == NULL)
 	{
 		if (RowName != NAME_None)
 		{
-			UE_LOG(LogCurveTable, Warning, TEXT("FCurveTableRowHandle::FindRow : No CurveTable for row %s (%s)."), *RowName.ToString(), *ContextString);
+			UE_CLOG(WarnIfNotFound, LogCurveTable, Warning, TEXT("FCurveTableRowHandle::FindRow : No CurveTable for row %s (%s)."), *RowName.ToString(), *ContextString);
 		}
 		return NULL;
 	}
 
-	return CurveTable->FindCurve(RowName, ContextString);
+	return CurveTable->FindCurve(RowName, ContextString, WarnIfNotFound);
 }
 
 float FCurveTableRowHandle::Eval(float XValue,const FString& ContextString) const
diff --git a/Engine/Source/Runtime/Engine/Private/DebugRenderSceneProxy.cpp b/Engine/Source/Runtime/Engine/Private/DebugRenderSceneProxy.cpp
--- a/Engine/Source/Runtime/Engine/Private/DebugRenderSceneProxy.cpp
+++ b/Engine/Source/Runtime/Engine/Private/DebugRenderSceneProxy.cpp
@@ -95,6 +95,12 @@ void FDebugDrawDelegateHelper::DrawDebugLabels(UCanvas* Canvas, APlayerControlle
 	Canvas->SetDrawColor(OldDrawColor);
 }
 
+SIZE_T FDebugRenderSceneProxy::GetTypeHash() const
+{
+	static size_t UniquePointer;
+	return reinterpret_cast<size_t>(&UniquePointer);
+}
+
 void FDebugRenderSceneProxy::GetDynamicMeshElements(const TArray<const FSceneView*>& Views, const FSceneViewFamily& ViewFamily, uint32 VisibilityMap, FMeshElementCollector& Collector) const 
 {
 	QUICK_SCOPE_CYCLE_COUNTER( STAT_DebugRenderSceneProxy_GetDynamicMeshElements );
@@ -107,7 +113,7 @@ void FDebugRenderSceneProxy::GetDynamicMeshElements(const TArray<const FSceneVie
 		FMaterialRenderProxy* operator[](FLinearColor Color)
 		{
 			FMaterialRenderProxy* MeshColor = NULL;
-			const uint32 HashKey = GetTypeHash(Color);
+			const uint32 HashKey = ::GetTypeHash(Color);
 			if (MeshColorInstances.Contains(HashKey))
 			{
 				MeshColor = *MeshColorInstances.Find(HashKey);
diff --git a/Engine/Source/Runtime/Engine/Private/DialogueWave.cpp b/Engine/Source/Runtime/Engine/Private/DialogueWave.cpp
--- a/Engine/Source/Runtime/Engine/Private/DialogueWave.cpp
+++ b/Engine/Source/Runtime/Engine/Private/DialogueWave.cpp
@@ -492,6 +492,12 @@ USoundClass* UDialogueSoundWaveProxy::GetSoundClass() const
 	return SoundWave->GetSoundClass();
 }
 
+
+UCurveTable* UDialogueSoundWaveProxy::GetCurveData() const
+{
+	return SoundWave ? SoundWave->GetCurveData() : nullptr;
+}
+
 UDialogueWave::UDialogueWave(const FObjectInitializer& ObjectInitializer)
 	: Super(ObjectInitializer)
 	, LocalizationGUID(FGuid::NewGuid())
diff --git a/Engine/Source/Runtime/Engine/Private/GameMode.cpp b/Engine/Source/Runtime/Engine/Private/GameMode.cpp
--- a/Engine/Source/Runtime/Engine/Private/GameMode.cpp
+++ b/Engine/Source/Runtime/Engine/Private/GameMode.cpp
@@ -414,7 +414,17 @@ void AGameMode::HandleSeamlessTravelPlayer(AController*& C)
 		if (PC->Player != nullptr)
 		{
 			// we need to spawn a new PlayerController to replace the old one
-			APlayerController* NewPC = SpawnPlayerController(PC->IsLocalPlayerController() ? ROLE_SimulatedProxy : ROLE_AutonomousProxy, PC->GetFocalLocation(), PC->GetControlRotation());
+			APlayerController* NewPC = nullptr;
+			if (PC->PlayerState && PC->PlayerState->bOnlySpectator && ReplaySpectatorPlayerControllerClass != nullptr)
+			{
+				NewPC = SpawnReplayPlayerController(PC->IsLocalPlayerController() ? ROLE_SimulatedProxy : ROLE_AutonomousProxy, PC->GetFocalLocation(), PC->GetControlRotation());
+			}
+			else
+			{
+				// We need to spawn a new PlayerController to replace the old one
+				NewPC = SpawnPlayerController(PC->IsLocalPlayerController() ? ROLE_SimulatedProxy : ROLE_AutonomousProxy, PC->GetFocalLocation(), PC->GetControlRotation());
+			}
+
 			if (NewPC == nullptr)
 			{
 				UE_LOG(LogGameMode, Warning, TEXT("Failed to spawn new PlayerController for %s (old class %s)"), *PC->GetHumanReadableName(), *PC->GetClass()->GetName());
diff --git a/Engine/Source/Runtime/Engine/Private/GameModeBase.cpp b/Engine/Source/Runtime/Engine/Private/GameModeBase.cpp
--- a/Engine/Source/Runtime/Engine/Private/GameModeBase.cpp
+++ b/Engine/Source/Runtime/Engine/Private/GameModeBase.cpp
@@ -532,8 +532,17 @@ void AGameModeBase::HandleSeamlessTravelPlayer(AController*& C)
 	APlayerController* PC = Cast<APlayerController>(C);
 	if (PC && PC->Player)
 	{
-		// We need to spawn a new PlayerController to replace the old one
-		APlayerController* NewPC = SpawnPlayerController(PC->IsLocalPlayerController() ? ROLE_SimulatedProxy : ROLE_AutonomousProxy, PC->GetFocalLocation(), PC->GetControlRotation());
+		APlayerController* NewPC = nullptr;
+		if (PC->PlayerState && PC->PlayerState->bOnlySpectator && ReplaySpectatorPlayerControllerClass != nullptr)
+		{
+			NewPC = SpawnReplayPlayerController(PC->IsLocalPlayerController() ? ROLE_SimulatedProxy : ROLE_AutonomousProxy, PC->GetFocalLocation(), PC->GetControlRotation());
+		}
+		else
+		{
+			// We need to spawn a new PlayerController to replace the old one
+			NewPC = SpawnPlayerController(PC->IsLocalPlayerController() ? ROLE_SimulatedProxy : ROLE_AutonomousProxy, PC->GetFocalLocation(), PC->GetControlRotation());
+		}
+
 
 		if (NewPC)
 		{
@@ -644,7 +653,16 @@ PRAGMA_ENABLE_DEPRECATION_WARNINGS
 		return nullptr;
 	}
 
-	APlayerController* NewPlayerController = SpawnPlayerController(InRemoteRole, FVector::ZeroVector, FRotator::ZeroRotator);
+	APlayerController* NewPlayerController = nullptr;
+	if (Options.Contains(FString(TEXT("SpectatorOnly=1"))) && ReplaySpectatorPlayerControllerClass != nullptr)
+	{
+		NewPlayerController = SpawnReplayPlayerController(InRemoteRole, FVector::ZeroVector, FRotator::ZeroRotator);
+	}
+	else
+	{
+		// We need to spawn a new PlayerController to replace the old one
+		NewPlayerController = SpawnPlayerController(InRemoteRole, FVector::ZeroVector, FRotator::ZeroRotator);
+	}
 
 	// Handle spawn failure.
 	if (NewPlayerController == nullptr)
@@ -665,12 +683,22 @@ PRAGMA_ENABLE_DEPRECATION_WARNINGS
 }
 
 APlayerController* AGameModeBase::SpawnPlayerController(ENetRole InRemoteRole, FVector const& SpawnLocation, FRotator const& SpawnRotation)
+{
+	return SpawnPlayerControllerCommon(InRemoteRole, SpawnLocation, SpawnRotation, PlayerControllerClass);
+}
+
+APlayerController* AGameModeBase::SpawnReplayPlayerController(ENetRole InRemoteRole, FVector const& SpawnLocation, FRotator const& SpawnRotation)
+{
+	return SpawnPlayerControllerCommon(InRemoteRole, SpawnLocation, SpawnRotation, ReplaySpectatorPlayerControllerClass);
+}
+
+APlayerController* AGameModeBase::SpawnPlayerControllerCommon(ENetRole InRemoteRole, FVector const& SpawnLocation, FRotator const& SpawnRotation, TSubclassOf<APlayerController> InPlayerControllerClass)
 {
 	FActorSpawnParameters SpawnInfo;
 	SpawnInfo.Instigator = Instigator;
 	SpawnInfo.ObjectFlags |= RF_Transient;	// We never want to save player controllers into a map
 	SpawnInfo.bDeferConstruction = true;
-	APlayerController* NewPC = GetWorld()->SpawnActor<APlayerController>(PlayerControllerClass, SpawnLocation, SpawnRotation, SpawnInfo);
+	APlayerController* NewPC = GetWorld()->SpawnActor<APlayerController>(InPlayerControllerClass, SpawnLocation, SpawnRotation, SpawnInfo);
 	if (NewPC)
 	{
 		if (InRemoteRole == ROLE_SimulatedProxy)
