Copying Orion-Staging @ 3748653 (Orion/Dev-General @ 3722124) to //UE4/Main

diff --git a/Engine/Source/Runtime/Engine/Private/HUD.cpp b/Engine/Source/Runtime/Engine/Private/HUD.cpp
--- a/Engine/Source/Runtime/Engine/Private/HUD.cpp
+++ b/Engine/Source/Runtime/Engine/Private/HUD.cpp
@@ -29,6 +29,9 @@
 #include "Misc/UObjectToken.h"
 #include "DisplayDebugHelpers.h"
 #include "DrawDebugHelpers.h"
+#include "Components/SkeletalMeshComponent.h"
+#include "GameFramework/MovementComponent.h"
+#include "UObject/UObjectIterator.h"
 
 DEFINE_LOG_CATEGORY_STATIC(LogHUD, Log, All);
 
@@ -343,18 +346,27 @@ void AHUD::ShowHUD()
 	bShowHUD = !bShowHUD;
 }
 
-static FName NAME_Reset = FName(TEXT("Reset"));
+namespace ShowDebugNames
+{
+	static const FName Reset(TEXT("Reset"));
+	static const FName HitBox(TEXT("HitBox"));
+	static const FName Animation(TEXT("Animation"));
+	static const FName Physics(TEXT("Physics"));
+}
+
 void AHUD::ShowDebug(FName DebugType)
 {
+	const bool bPreviousShowDebugInfo = bShowDebugInfo;
+
 	if (DebugType == NAME_None)
 	{
 		bShowDebugInfo = !bShowDebugInfo;
 	}
-	else if ( DebugType == FName("HitBox") )
+	else if (DebugType == ShowDebugNames::HitBox)
 	{
 		bShowHitBoxDebugInfo = !bShowHitBoxDebugInfo;
 	}
-	else if( DebugType == NAME_Reset )
+	else if (DebugType == ShowDebugNames::Reset)
 	{
 		DebugDisplay.Reset();
 		bShowDebugInfo = false;
@@ -377,14 +389,19 @@ void AHUD::ShowDebug(FName DebugType)
 		}
 
 		bShowDebugInfo = true;
-
 		SaveConfig();
 	}
+
+	// Reset Target to ourselves when enabled/disabled
+	if (bShowDebugInfo != bPreviousShowDebugInfo)
+	{
+		ShowDebugTargetActor = nullptr;
+	}
 }
 
 void AHUD::ShowDebugToggleSubCategory(FName Category)
 {
-	if( Category == NAME_Reset )
+	if (Category == ShowDebugNames::Reset)
 	{
 		ToggledDebugCategories.Reset();
 		SaveConfig();
@@ -412,36 +429,35 @@ bool AHUD::ShouldDisplayDebug(const FName& DebugType) const
 
 void AHUD::ShowDebugInfo(float& YL, float& YPos)
 {
-	if (DebugCanvas != nullptr )
+	if (DebugCanvas != nullptr)
 	{
+		// Darken background, so we can read text better.
 		FLinearColor BackgroundColor(0.f, 0.f, 0.f, 0.2f);
 		DebugCanvas->Canvas->DrawTile(0, 0, DebugCanvas->ClipX, DebugCanvas->ClipY, 0.f, 0.f, 0.f, 0.f, BackgroundColor);
 
 		FDebugDisplayInfo DisplayInfo(DebugDisplay, ToggledDebugCategories);
+		ShowDebugTargetActor = GetCurrentDebugTargetActor();
 
-		if (bShowDebugForReticleTarget)
-		{
-			FRotator CamRot; FVector CamLoc; PlayerOwner->GetPlayerViewPoint(CamLoc, CamRot);
-
-			FCollisionQueryParams TraceParams(NAME_None, FCollisionQueryParams::GetUnknownStatId(), true, PlayerOwner->PlayerCameraManager->ViewTarget.Target);
-			FHitResult Hit;
-			bool bHit = GetWorld()->LineTraceSingleByChannel(Hit, CamLoc, CamRot.Vector() * 100000.f + CamLoc, ECC_WorldDynamic, TraceParams);
-			if (bHit)
-			{
-				AActor* HitActor = Hit.Actor.Get();
-				if (HitActor && (ShowDebugTargetDesiredClass == NULL || HitActor->IsA(ShowDebugTargetDesiredClass)))
-				{
-					ShowDebugTargetActor = HitActor;
-				}
-			}
-		}
-		else
+		// Draw Header.
 		{
-			ShowDebugTargetActor = PlayerOwner->PlayerCameraManager->ViewTarget.Target;
+			FDisplayDebugManager& DisplayDebugManager = DebugCanvas->DisplayDebugManager;
+			DisplayDebugManager.SetDrawColor(FColor(255, 0, 0));
+			DisplayDebugManager.DrawString(FString::Printf(TEXT("Showing Debug for %s, Press [PageUp] and [PageDown] to cycle between targets."), *GetNameSafe(ShowDebugTargetActor)));
 		}
 
 		if (ShowDebugTargetActor && !ShowDebugTargetActor->IsPendingKill())
 		{
+			// Draw box around Actor being debugged.
+#if ENABLE_DRAW_DEBUG
+			{
+				FVector BoundsOrigin, BoundsExtent;
+				ShowDebugTargetActor->GetActorBounds(true, BoundsOrigin, BoundsExtent);
+
+				// Expand extent a little bit
+				BoundsExtent *= 1.1f;
+				DrawDebugBox(GetWorld(), BoundsOrigin, BoundsExtent, FColor::Green, false, -1.f, 0, 2.f);
+			}
+#endif
 			ShowDebugTargetActor->DisplayDebug(DebugCanvas, DisplayInfo, YL, YPos);
 
 			if (!bShowDebugForReticleTarget && ShowDebugTargetActor->Role == ROLE_SimulatedProxy)
@@ -466,6 +482,160 @@ void AHUD::ShowDebugInfo(float& YL, float& YPos)
 	}
 }
 
+AActor* AHUD::GetCurrentDebugTargetActor()
+{
+	AActor* DebugTargetActor = nullptr;
+
+	// Find targets through the reticle.
+	if (bShowDebugForReticleTarget && PlayerOwner->PlayerCameraManager)
+	{
+		FRotator CamRot; FVector CamLoc; PlayerOwner->GetPlayerViewPoint(CamLoc, CamRot);
+
+		FCollisionQueryParams TraceParams(NAME_None, FCollisionQueryParams::GetUnknownStatId(), true, PlayerOwner->PlayerCameraManager->ViewTarget.Target);
+		FHitResult Hit;
+		bool bHit = GetWorld()->LineTraceSingleByChannel(Hit, CamLoc, CamRot.Vector() * 100000.f + CamLoc, ECC_WorldDynamic, TraceParams);
+		if (bHit)
+		{
+			AActor* HitActor = Hit.Actor.Get();
+			if (HitActor && ((ShowDebugTargetDesiredClass == nullptr) || HitActor->IsA(ShowDebugTargetDesiredClass)))
+			{
+				DebugTargetActor = HitActor;
+			}
+		}
+
+		// If we hit something new, return this.
+		// Otherwise fall back to our last successful hit.
+		return DebugTargetActor ? DebugTargetActor : ShowDebugTargetActor;
+	}
+	else 
+	{
+		// Otherwise we use our Cached DebugTargetActor.
+		DebugTargetActor = ShowDebugTargetActor;
+	}
+
+	// If we have no one to view, default to current view target.
+	if ((DebugTargetActor == nullptr) && PlayerOwner->PlayerCameraManager && PlayerOwner->PlayerCameraManager->ViewTarget.Target)
+	{
+		DebugTargetActor = PlayerOwner->PlayerCameraManager->ViewTarget.Target;
+	}
+
+	return DebugTargetActor;
+}
+
+void AHUD::AddActorToDebugList(AActor* InActor, TArray<AActor*>& InOutList, UWorld* InWorld)
+{
+	// Only consider actors that are visible, not destroyed and in the same world.
+	if (InActor && !InActor->IsPendingKill() && (InActor->GetWorld() == InWorld) && InActor->WasRecentlyRendered())
+	{
+		InOutList.AddUnique(InActor);
+	}
+}
+
+void AHUD::AddComponentOwnerToDebugList(UActorComponent* InComponent, TArray<AActor*>& InOutList, UWorld* InWorld)
+{
+	if (InComponent && InComponent->GetOwner())
+	{
+		AddActorToDebugList(InComponent->GetOwner(), InOutList, InWorld);
+	}
+}
+
+void AHUD::GetDebugActorList(TArray<AActor*>& InOutList)
+{
+	UWorld* World = GetWorld();
+
+	// By default, add all Pawns.
+	for (TActorIterator<APawn> It(World); It; ++It)
+	{
+		AddActorToDebugList(*It, InOutList, World);
+	}
+
+	// If we're viewing animations, add all actors using an AnimInstance.
+	if (ShouldDisplayDebug(ShowDebugNames::Animation))
+	{
+		for (TObjectIterator<USkeletalMeshComponent> It; It; ++It)
+		{
+			if (It->GetAnimInstance())
+			{
+				AddComponentOwnerToDebugList(*It, InOutList, World);
+			}
+		}
+	}
+
+	// If we're viewing physics, add all actors using a movement component.
+	if (ShouldDisplayDebug(ShowDebugNames::Physics))
+	{
+		for (TObjectIterator<UMovementComponent> It; It; ++It)
+		{
+			AddComponentOwnerToDebugList(*It, InOutList, World);
+		}
+	}
+}
+
+void AHUD::NextDebugTarget()
+{
+	TArray<AActor*> Targets;
+	GetDebugActorList(Targets);
+
+	// If we have an existing target, find it in list.
+	// As list can change as actors are spawned and destroyed.
+	if (ShowDebugTargetActor)
+	{
+		if (!Targets.IsValidIndex(CurrentTargetIndex) || (Targets[CurrentTargetIndex] != ShowDebugTargetActor))
+		{
+			int32 FoundIndex;
+			if (Targets.Find(ShowDebugTargetActor, FoundIndex))
+			{
+				CurrentTargetIndex = FoundIndex;
+			}
+		}
+	}
+
+	CurrentTargetIndex = (CurrentTargetIndex + 1) % Targets.Num();
+	if (Targets.IsValidIndex(CurrentTargetIndex))
+	{
+		ShowDebugTargetActor = Targets[CurrentTargetIndex];
+	}
+	else if (PlayerOwner->PlayerCameraManager && PlayerOwner->PlayerCameraManager->ViewTarget.Target)
+	{
+		ShowDebugTargetActor = PlayerOwner->PlayerCameraManager->ViewTarget.Target;
+	}
+}
+
+void AHUD::PreviousDebugTarget()
+{
+	TArray<AActor*> Targets;
+	GetDebugActorList(Targets);
+
+	// If we have an existing target, find it in list.
+	// As list can change as actors are spawned and destroyed.
+	if (ShowDebugTargetActor)
+	{
+		if (!Targets.IsValidIndex(CurrentTargetIndex) || (Targets[CurrentTargetIndex] != ShowDebugTargetActor))
+		{
+			int32 FoundIndex;
+			if (Targets.Find(ShowDebugTargetActor, FoundIndex))
+			{
+				CurrentTargetIndex = FoundIndex;
+			}
+		}
+	}
+
+	CurrentTargetIndex--;
+	if (CurrentTargetIndex < 0)
+	{
+		CurrentTargetIndex = Targets.Num() - 1;
+	}
+
+	if (Targets.IsValidIndex(CurrentTargetIndex))
+	{
+		ShowDebugTargetActor = Targets[CurrentTargetIndex];
+	}
+	else if (PlayerOwner->PlayerCameraManager && PlayerOwner->PlayerCameraManager->ViewTarget.Target)
+	{
+		ShowDebugTargetActor = PlayerOwner->PlayerCameraManager->ViewTarget.Target;
+	}
+}
+
 void AHUD::DrawHUD()
 {
 	HitBoxMap.Reset();
@@ -495,8 +665,6 @@ UFont* AHUD::GetFontFromSizeIndex(int32 FontSizeIndex) const
 	return GEngine->GetLargeFont();
 }
 
-
-
 void AHUD::OnLostFocusPause(bool bEnable)
 {
 	if ( bLostFocusPaused == bEnable )
diff --git a/Engine/Source/Runtime/Engine/Private/HierarchicalInstancedStaticMesh.cpp b/Engine/Source/Runtime/Engine/Private/HierarchicalInstancedStaticMesh.cpp
--- a/Engine/Source/Runtime/Engine/Private/HierarchicalInstancedStaticMesh.cpp
+++ b/Engine/Source/Runtime/Engine/Private/HierarchicalInstancedStaticMesh.cpp
@@ -53,7 +53,8 @@ static TAutoConsoleVariable<int32> CVarDisableCull(
 static TAutoConsoleVariable<int32> CVarCullAll(
 	TEXT("foliage.CullAll"),
 	0,
-	TEXT("If greater than zero, everything is considered culled."));
+	TEXT("If greater than zero, everything is considered culled."),
+	ECVF_Scalability);
 
 static TAutoConsoleVariable<int32> CVarDitheredLOD(
 	TEXT("foliage.DitheredLOD"),
@@ -746,7 +747,7 @@ struct FFoliageElementParams;
 struct FFoliageRenderInstanceParams;
 struct FFoliageCullInstanceParams;
 
-class FHierarchicalStaticMeshSceneProxy : public FInstancedStaticMeshSceneProxy
+class FHierarchicalStaticMeshSceneProxy final : public FInstancedStaticMeshSceneProxy
 {
 	TSharedRef<TArray<FClusterNode>, ESPMode::ThreadSafe> ClusterTreePtr;
 	const TArray<FClusterNode>& ClusterTree;
@@ -772,6 +773,11 @@ class FHierarchicalStaticMeshSceneProxy : public FInstancedStaticMeshSceneProxy
 #endif
 
 public:
+	SIZE_T GetTypeHash() const override
+	{
+		static size_t UniquePointer;
+		return reinterpret_cast<size_t>(&UniquePointer);
+	}
 
 	FHierarchicalStaticMeshSceneProxy(bool bInIsGrass, UHierarchicalInstancedStaticMeshComponent* InComponent, ERHIFeatureLevel::Type InFeatureLevel)
 		: FInstancedStaticMeshSceneProxy(InComponent, InFeatureLevel)
@@ -1110,6 +1116,8 @@ inline bool CanGroup(const FVector& BoundMin, const FVector& BoundMax, const FVe
 	return FarDot < MaxDrawDist;
 }
 
+
+
 template<bool TUseVector>
 void FHierarchicalStaticMeshSceneProxy::Traverse(const FFoliageCullInstanceParams& Params, int32 Index, int32 MinLOD, int32 MaxLOD, bool bFullyContained) const
 {
diff --git a/Engine/Source/Runtime/Engine/Private/InstancedStaticMesh.cpp b/Engine/Source/Runtime/Engine/Private/InstancedStaticMesh.cpp
--- a/Engine/Source/Runtime/Engine/Private/InstancedStaticMesh.cpp
+++ b/Engine/Source/Runtime/Engine/Private/InstancedStaticMesh.cpp
@@ -104,7 +104,6 @@ void FStaticMeshInstanceBuffer::UpdateInstanceData(UInstancedStaticMeshComponent
 	}
 
 	NumInstances = InstanceData->GetNumInstances();
-
 	const FMeshMapBuildData* MeshMapBuildData = NULL;
 
 	if (InComponent->LODData.Num() > 0)
@@ -620,6 +619,11 @@ void FInstancedStaticMeshRenderData::InitStaticMeshVertexFactories(
 	}
 }
 
+SIZE_T FInstancedStaticMeshSceneProxy::GetTypeHash() const
+{
+	static size_t UniquePointer;
+	return reinterpret_cast<size_t>(&UniquePointer);
+}
 
 void FInstancedStaticMeshSceneProxy::GetDynamicMeshElements(const TArray<const FSceneView*>& Views, const FSceneViewFamily& ViewFamily, uint32 VisibilityMap, FMeshElementCollector& Collector) const
 {
diff --git a/Engine/Source/Runtime/Engine/Private/InstancedStaticMesh.h b/Engine/Source/Runtime/Engine/Private/InstancedStaticMesh.h
--- a/Engine/Source/Runtime/Engine/Private/InstancedStaticMesh.h
+++ b/Engine/Source/Runtime/Engine/Private/InstancedStaticMesh.h
@@ -705,6 +705,7 @@ private:
 class FInstancedStaticMeshSceneProxy : public FStaticMeshSceneProxy
 {
 public:
+	SIZE_T GetTypeHash() const override;
 
 	FInstancedStaticMeshSceneProxy(UInstancedStaticMeshComponent* InComponent, ERHIFeatureLevel::Type InFeatureLevel)
 	:	FStaticMeshSceneProxy(InComponent, true)
diff --git a/Engine/Source/Runtime/Engine/Private/Model.cpp b/Engine/Source/Runtime/Engine/Private/Model.cpp
--- a/Engine/Source/Runtime/Engine/Private/Model.cpp
+++ b/Engine/Source/Runtime/Engine/Private/Model.cpp
@@ -743,7 +743,9 @@ void UModel::UpdateVertices()
 #endif
 		VertexBuffer.Buffers.InitModelBuffers(VertexBuffer.Vertices);
 		
-		if (FApp::IsGame())
+		//  Empty this if we have cooked data and thus won't need it later to generate collision
+		// data etc
+		if (FApp::IsGame() && FPlatformProperties::RequiresCookedData())
 		{
 			VertexBuffer.Vertices.Empty();
 		}
diff --git a/Engine/Source/Runtime/Engine/Private/ModelRender.cpp b/Engine/Source/Runtime/Engine/Private/ModelRender.cpp
--- a/Engine/Source/Runtime/Engine/Private/ModelRender.cpp
+++ b/Engine/Source/Runtime/Engine/Private/ModelRender.cpp
@@ -153,6 +153,12 @@ void UModelComponent::BuildRenderData()
 						}
 						Element.MinVertexIndex = FMath::Min(Node.iVertexIndex + Node.NumVertices * BackFace, Element.MinVertexIndex);
 						Element.MaxVertexIndex = FMath::Max(Node.iVertexIndex + Node.NumVertices * BackFace + Node.NumVertices - 1, Element.MaxVertexIndex);
+
+						if (Element.MaxVertexIndex > (uint32)TheModel->VertexBuffer.Vertices.Num())
+						{
+							UE_LOG(LogModelComponent, Warning, TEXT("Model %s has elements that reference missing vertices. MaxVertex=%d, NumVertices=%d"),
+								*GetPathName(), Element.MaxVertexIndex, TheModel->VertexBuffer.Vertices.Num());
+						}
 					}
 				}
 			}
@@ -166,12 +172,17 @@ void UModelComponent::BuildRenderData()
 /**
  * A model component scene proxy.
  */
-class FModelSceneProxy : public FPrimitiveSceneProxy
+class FModelSceneProxy final : public FPrimitiveSceneProxy
 {
 	/** The vertex factory which is used to access VertexBuffer. */
 	FLocalVertexFactory VertexFactory;
 
 public:
+	SIZE_T GetTypeHash() const override
+	{
+		static size_t UniquePointer;
+		return reinterpret_cast<size_t>(&UniquePointer);
+	}
 
 	FModelSceneProxy(UModelComponent* InComponent):
 		FPrimitiveSceneProxy(InComponent),
diff --git a/Engine/Source/Runtime/Engine/Private/NetConnection.cpp b/Engine/Source/Runtime/Engine/Private/NetConnection.cpp
--- a/Engine/Source/Runtime/Engine/Private/NetConnection.cpp
+++ b/Engine/Source/Runtime/Engine/Private/NetConnection.cpp
@@ -132,9 +132,9 @@ void UNetConnection::InitBase(UNetDriver* InDriver,class FSocket* InSocket, cons
 	// Stats
 	StatUpdateTime			= Driver->Time;
 	LastReceiveTime			= Driver->Time;
-	LastReceiveRealtime		= FPlatformTime::Seconds();
-	LastGoodPacketRealtime	= FPlatformTime::Seconds();
-	LastTime				= FPlatformTime::Seconds();
+	LastReceiveRealtime		= 0.0;			// These are set to 0 and initialized on our first tick to delay with scenarios where
+	LastGoodPacketRealtime	= 0.0;			// notable time may elapse between now and then
+	LastTime				= 0.0;
 	LastSendTime			= Driver->Time;
 	LastTickTime			= Driver->Time;
 	LastRecvAckTime			= Driver->Time;
@@ -157,8 +157,7 @@ void UNetConnection::InitBase(UNetDriver* InDriver,class FSocket* InSocket, cons
 	// Reset Handler
 	Handler.Reset(NULL);
 
-	InitHandler();
-
+	InitHandler(InDriver->AnalyticsProvider);
 
 #if DO_ENABLE_NET_TEST
 	// Copy the command line settings from the net driver
@@ -237,7 +236,7 @@ void UNetConnection::InitConnection(UNetDriver* InDriver, EConnectionState InSta
 	PackageMap = PackageMapClient;
 }
 
-void UNetConnection::InitHandler()
+void UNetConnection::InitHandler(TSharedPtr<IAnalyticsProvider> InProvider/*=nullptr*/)
 {
 	check(!Handler.IsValid());
 
@@ -252,7 +251,7 @@ void UNetConnection::InitHandler()
 			Handler::Mode Mode = Driver->ServerConnection != nullptr ? Handler::Mode::Client : Handler::Mode::Server;
 
 			Handler->InitializeDelegates(FPacketHandlerLowLevelSend::CreateUObject(this, &UNetConnection::LowLevelSend));
-			Handler->Initialize(Mode, MaxPacket * 8);
+			Handler->Initialize(Mode, MaxPacket * 8, false, InProvider);
 
 
 			// Add handling for the stateless connect handshake, for connectionless packets, as the outermost layer
@@ -1775,6 +1774,15 @@ void UNetConnection::Tick()
 
 	// Get frame time.
 	const double CurrentRealtimeSeconds = FPlatformTime::Seconds();
+
+	// if this is 0 it's our first tick since init, so start our real-time tracking from here
+	if (LastTime == 0.0)
+	{
+		LastTime = CurrentRealtimeSeconds;
+		LastReceiveRealtime = CurrentRealtimeSeconds;
+		LastGoodPacketRealtime = CurrentRealtimeSeconds;
+	}
+	
 	FrameTime = CurrentRealtimeSeconds - LastTime;
 	LastTime = CurrentRealtimeSeconds;
 	CumulativeTime += FrameTime;
@@ -1847,7 +1855,7 @@ void UNetConnection::Tick()
 	// Handle timeouts.
 	const float Timeout = GetTimeoutValue();
 
-	if ((Driver->Time - LastReceiveTime) > Timeout)
+	if ((CurrentRealtimeSeconds - LastReceiveRealtime) > Timeout)
 	{
 		const TCHAR* const TimeoutString = TEXT("UNetConnection::Tick: Connection TIMED OUT. Closing connection.");
 		const TCHAR* const DestroyString = TEXT("UNetConnection::Tick: Connection closing during pending destroy, not all shutdown traffic may have been negotiated");
diff --git a/Engine/Source/Runtime/Engine/Private/NetworkDelegates.cpp b/Engine/Source/Runtime/Engine/Private/NetworkDelegates.cpp
--- a/Engine/Source/Runtime/Engine/Private/NetworkDelegates.cpp
+++ b/Engine/Source/Runtime/Engine/Private/NetworkDelegates.cpp
@@ -5,3 +5,4 @@
 FNetDelegates::FReceivedNetworkEncryptionToken FNetDelegates::OnReceivedNetworkEncryptionToken;
 FNetDelegates::FReceivedNetworkEncryptionAck FNetDelegates::OnReceivedNetworkEncryptionAck;
 FNetDelegates::FOnPendingNetGameConnectionCreated FNetDelegates::OnPendingNetGameConnectionCreated;
+FNetDelegates::FNetworkCheatDetected FNetDelegates::OnNetworkCheatDetected;
diff --git a/Engine/Source/Runtime/Engine/Private/NetworkDriver.cpp b/Engine/Source/Runtime/Engine/Private/NetworkDriver.cpp
--- a/Engine/Source/Runtime/Engine/Private/NetworkDriver.cpp
+++ b/Engine/Source/Runtime/Engine/Private/NetworkDriver.cpp
@@ -1143,7 +1143,7 @@ void UNetDriver::InitConnectionlessHandler()
 		{
 			ConnectionlessHandler->bConnectionlessHandler = true;
 
-			ConnectionlessHandler->Initialize(Handler::Mode::Server, MAX_PACKET_SIZE, true);
+			ConnectionlessHandler->Initialize(Handler::Mode::Server, MAX_PACKET_SIZE, true, AnalyticsProvider);
 
 			// Add handling for the stateless connect handshake, for connectionless packets, as the outermost layer
 			TSharedPtr<HandlerComponent> NewComponent =
diff --git a/Engine/Source/Runtime/Engine/Private/Particles/ParticleComponents.cpp b/Engine/Source/Runtime/Engine/Private/Particles/ParticleComponents.cpp
--- a/Engine/Source/Runtime/Engine/Private/Particles/ParticleComponents.cpp
+++ b/Engine/Source/Runtime/Engine/Private/Particles/ParticleComponents.cpp
@@ -83,28 +83,28 @@
 #include "GameFramework/GameState.h"
 #include "HAL/LowLevelMemTracker.h"
 
-DECLARE_CYCLE_STAT(TEXT("ParticleComponent InitParticles"), STAT_ParticleSystemComponent_InitParticles, STATGROUP_Particles);
-DECLARE_CYCLE_STAT(TEXT("ParticleComponent SendRenderDynamicData"), STAT_ParticleSystemComponent_SendRenderDynamicData_Concurrent, STATGROUP_Particles);
-DECLARE_CYCLE_STAT(TEXT("ParticleComponent SendRenderTransform Concurrent"), STAT_ParticleSystemComponent_SendRenderTransform_Concurrent, STATGROUP_Particles);
-DECLARE_CYCLE_STAT(TEXT("ParticleComponent DestroyRenderState Concurrent"), STAT_ParticleSystemComponent_DestroyRenderState_Concurrent, STATGROUP_Particles);
-DECLARE_CYCLE_STAT(TEXT("ParticleComponent CreateDynamicData"), STAT_ParticleSystemComponent_CreateDynamicData, STATGROUP_Particles);
-DECLARE_CYCLE_STAT(TEXT("ParticleComponent CreateDynamicData Replay"), STAT_ParticleSystemComponent_CreateDynamicData_Replay, STATGROUP_Particles);
-DECLARE_CYCLE_STAT(TEXT("ParticleComponent CreateDynamicData Capture"), STAT_ParticleSystemComponent_CreateDynamicData_Capture, STATGROUP_Particles);
-DECLARE_CYCLE_STAT(TEXT("ParticleComponent CreateDynamicData Gather"), STAT_ParticleSystemComponent_CreateDynamicData_Gather, STATGROUP_Particles);
-DECLARE_CYCLE_STAT(TEXT("ParticleComponent GetDynamicData"), STAT_ParticleSystemComponent_GetDynamicData, STATGROUP_Particles);
-DECLARE_CYCLE_STAT(TEXT("ParticleComponent GetDynamicData Selected"), STAT_ParticleSystemComponent_GetDynamicData_Selected, STATGROUP_Particles);
-DECLARE_CYCLE_STAT(TEXT("ParticleComponent CreateDynamicData GatherCapture"), STAT_ParticleSystemComponent_CreateDynamicData_GatherCapture, STATGROUP_Particles);
-DECLARE_CYCLE_STAT(TEXT("ParticleComponent UpdateDynamicData"), STAT_ParticleSystemComponent_UpdateDynamicData, STATGROUP_Particles);
-DECLARE_CYCLE_STAT(TEXT("ParticleComponent OrientZAxisTowardCamera"), STAT_UParticleSystemComponent_OrientZAxisTowardCamera, STATGROUP_Particles);
-DECLARE_CYCLE_STAT(TEXT("ParticleComponent QueueFinalize"), STAT_UParticleSystemComponent_QueueFinalize, STATGROUP_Particles);
-DECLARE_CYCLE_STAT(TEXT("ParticleComponent CheckForReset"), STAT_UParticleSystemComponent_CheckForReset, STATGROUP_Particles);
-DECLARE_CYCLE_STAT(TEXT("ParticleComponent LOD_Inactive"), STAT_UParticleSystemComponent_LOD_Inactive, STATGROUP_Particles);
-DECLARE_CYCLE_STAT(TEXT("ParticleComponent LOD"), STAT_UParticleSystemComponent_LOD, STATGROUP_Particles);
-DECLARE_CYCLE_STAT(TEXT("ParticleComponent ResetAndCheckParallel"), STAT_UParticleSystemComponent_ResetAndCheckParallel, STATGROUP_Particles);
-DECLARE_CYCLE_STAT(TEXT("ParticleComponent QueueTasks"), STAT_UParticleSystemComponent_QueueTasks, STATGROUP_Particles);
-DECLARE_CYCLE_STAT(TEXT("ParticleComponent QueueAsync"), STAT_UParticleSystemComponent_QueueAsync, STATGROUP_Particles);
-DECLARE_CYCLE_STAT(TEXT("ParticleComponent WaitForAsyncAndFinalize"), STAT_UParticleSystemComponent_WaitForAsyncAndFinalize, STATGROUP_Particles);
-DECLARE_CYCLE_STAT(TEXT("ParticleComponent CreateRenderState Concurrent"), STAT_ParticleSystemComponent_CreateRenderState_Concurrent, STATGROUP_Particles);
+DECLARE_CYCLE_STAT(TEXT("ParticleComponent InitParticles GT"), STAT_ParticleSystemComponent_InitParticles, STATGROUP_Particles);
+DECLARE_CYCLE_STAT(TEXT("ParticleComponent SendRenderDynamicData GT"), STAT_ParticleSystemComponent_SendRenderDynamicData_Concurrent, STATGROUP_Particles);
+DECLARE_CYCLE_STAT(TEXT("ParticleComponent SendRenderTransform Concurrent GT"), STAT_ParticleSystemComponent_SendRenderTransform_Concurrent, STATGROUP_Particles);
+DECLARE_CYCLE_STAT(TEXT("ParticleComponent DestroyRenderState Concurrent GT"), STAT_ParticleSystemComponent_DestroyRenderState_Concurrent, STATGROUP_Particles);
+DECLARE_CYCLE_STAT(TEXT("ParticleComponent CreateDynamicData GT"), STAT_ParticleSystemComponent_CreateDynamicData, STATGROUP_Particles);
+DECLARE_CYCLE_STAT(TEXT("ParticleComponent CreateDynamicData Replay GT"), STAT_ParticleSystemComponent_CreateDynamicData_Replay, STATGROUP_Particles);
+DECLARE_CYCLE_STAT(TEXT("ParticleComponent CreateDynamicData Capture GT"), STAT_ParticleSystemComponent_CreateDynamicData_Capture, STATGROUP_Particles);
+DECLARE_CYCLE_STAT(TEXT("ParticleComponent CreateDynamicData Gather GT"), STAT_ParticleSystemComponent_CreateDynamicData_Gather, STATGROUP_Particles);
+DECLARE_CYCLE_STAT(TEXT("ParticleComponent GetDynamicData GT"), STAT_ParticleSystemComponent_GetDynamicData, STATGROUP_Particles);
+DECLARE_CYCLE_STAT(TEXT("ParticleComponent GetDynamicData Selected GT"), STAT_ParticleSystemComponent_GetDynamicData_Selected, STATGROUP_Particles);
+DECLARE_CYCLE_STAT(TEXT("ParticleComponent CreateDynamicData GatherCapture GT"), STAT_ParticleSystemComponent_CreateDynamicData_GatherCapture, STATGROUP_Particles);
+DECLARE_CYCLE_STAT(TEXT("ParticleComponent UpdateDynamicData GT"), STAT_ParticleSystemComponent_UpdateDynamicData, STATGROUP_Particles);
+DECLARE_CYCLE_STAT(TEXT("ParticleComponent OrientZAxisTowardCamera GT"), STAT_UParticleSystemComponent_OrientZAxisTowardCamera, STATGROUP_Particles);
+DECLARE_CYCLE_STAT(TEXT("ParticleComponent QueueFinalize GT"), STAT_UParticleSystemComponent_QueueFinalize, STATGROUP_Particles);
+DECLARE_CYCLE_STAT(TEXT("ParticleComponent CheckForReset GT"), STAT_UParticleSystemComponent_CheckForReset, STATGROUP_Particles);
+DECLARE_CYCLE_STAT(TEXT("ParticleComponent LOD_Inactive GT"), STAT_UParticleSystemComponent_LOD_Inactive, STATGROUP_Particles);
+DECLARE_CYCLE_STAT(TEXT("ParticleComponent LOD GT"), STAT_UParticleSystemComponent_LOD, STATGROUP_Particles);
+DECLARE_CYCLE_STAT(TEXT("ParticleComponent ResetAndCheckParallel GT"), STAT_UParticleSystemComponent_ResetAndCheckParallel, STATGROUP_Particles);
+DECLARE_CYCLE_STAT(TEXT("ParticleComponent QueueTasksGT"), STAT_UParticleSystemComponent_QueueTasks, STATGROUP_Particles);
+DECLARE_CYCLE_STAT(TEXT("ParticleComponent QueueAsyncGT"), STAT_UParticleSystemComponent_QueueAsync, STATGROUP_Particles);
+DECLARE_CYCLE_STAT(TEXT("ParticleComponent WaitForAsyncAndFinalize GT"), STAT_UParticleSystemComponent_WaitForAsyncAndFinalize, STATGROUP_Particles);
+DECLARE_CYCLE_STAT(TEXT("ParticleComponent CreateRenderState Concurrent GT"), STAT_ParticleSystemComponent_CreateRenderState_Concurrent, STATGROUP_Particles);
 
 #include "InGamePerformanceTracker.h"
 
@@ -1773,6 +1773,7 @@ bool UParticleEmitter::HasAnyEnabledLODs()const
 #if STATS
 
 DEFINE_STAT(STAT_EmittersStatGroupTester);
+DEFINE_STAT(STAT_EmittersRTStatGroupTester);
 
 void UParticleEmitter::CreateStatID() const
 {
@@ -1783,6 +1784,7 @@ void UParticleEmitter::CreateStatID() const
 	FString LongName;
 	LongName = FString(TEXT("Emitter")) / OuterName.ToString() / EmitterName.ToString();
 	StatID = FDynamicStats::CreateStatId<FStatGroup_STATGROUP_Emitters>(LongName);
+	StatIDRT = FDynamicStats::CreateStatId<FStatGroup_STATGROUP_EmittersRT>(LongName / TEXT("RT"));
 }
 #endif
 
@@ -3807,7 +3809,7 @@ FDynamicEmitterDataBase* UParticleSystemComponent::CreateDynamicDataFromReplay(
 #if STATS
 	if (EmitterData)
 	{
-		EmitterData->StatID = EmitterInstance->SpriteTemplate->GetStatID();
+		EmitterData->StatID = EmitterInstance->SpriteTemplate->GetStatIDRT();
 	}
 #endif
 
@@ -3978,7 +3980,7 @@ FParticleDynamicData* UParticleSystemComponent::CreateDynamicData(ERHIFeatureLev
 					if( NewDynamicEmitterData != NULL )
 					{
 #if STATS
-						NewDynamicEmitterData->StatID = EmitterInst->SpriteTemplate->GetStatID();
+						NewDynamicEmitterData->StatID = EmitterInst->SpriteTemplate->GetStatIDRT();
 #endif
 						NewDynamicEmitterData->bValid = true;
 						ParticleDynamicData->DynamicEmitterDataArray.Add( NewDynamicEmitterData );
@@ -4376,7 +4378,7 @@ public:
 
 TAutoConsoleVariable<int32> CVarFXEarlySchedule(TEXT("FX.EarlyScheduleAsync"), 0, TEXT("If 1, particle system components that can run async will be scheduled earlier in the frame"));
 
-DECLARE_CYCLE_STAT(TEXT("PSys Comp Marshall Time"),STAT_UParticleSystemComponent_Marshall,STATGROUP_Particles);
+DECLARE_CYCLE_STAT(TEXT("PSys Comp Marshall Time GT"),STAT_UParticleSystemComponent_Marshall,STATGROUP_Particles);
 
 bool UParticleSystemComponent::IsReadyForOwnerToAutoDestroy() const
 {
@@ -4841,14 +4843,14 @@ void UParticleSystemComponent::FinalizeTickComponent()
 	else
 	{
 		LastSignificantTime = CurrTime;
-	// If component has just totally finished, call script event.
-	const bool bIsCompleted = HasCompleted(); 
-	if (bIsCompleted && !bWasCompleted)
-	{
+		// If component has just totally finished, call script event.
+		const bool bIsCompleted = HasCompleted(); 
+		if (bIsCompleted && !bWasCompleted)
+		{
 			Complete();
 		}
 		bWasCompleted = bIsCompleted;
-		}
+	}
 
 	// Update bounding box.
 	if (!bWarmingUp && !bWasCompleted && !Template->bUseFixedRelativeBoundingBox && !bIsTransformDirty)
diff --git a/Engine/Source/Runtime/Engine/Private/Particles/ParticleEmitterInstances.cpp b/Engine/Source/Runtime/Engine/Private/Particles/ParticleEmitterInstances.cpp
--- a/Engine/Source/Runtime/Engine/Private/Particles/ParticleEmitterInstances.cpp
+++ b/Engine/Source/Runtime/Engine/Private/Particles/ParticleEmitterInstances.cpp
@@ -14,6 +14,8 @@
 #include "StaticMeshResources.h"
 #include "FXSystem.h"
 
+#include "HAL/PlatformStackWalk.h"
+
 #include "Particles/SubUV/ParticleModuleSubUV.h"
 #include "Particles/Collision/ParticleModuleCollisionGPU.h"
 #include "Particles/Event/ParticleModuleEventGenerator.h"
@@ -85,6 +87,8 @@ DEFINE_STAT(STAT_GPUSpriteSpawnTime);
 DEFINE_STAT(STAT_GPUSpriteTickTime);
 DEFINE_STAT(STAT_GPUSingleIterationEmitters);
 DEFINE_STAT(STAT_GPUMultiIterationsEmitters);
+DEFINE_STAT(STAT_GPUParticlesInjectionTime);
+DEFINE_STAT(STAT_GPUParticlesSimulationCommands);
 
 /** Particle memory stats */
 
@@ -135,15 +139,152 @@ DEFINE_STAT(STAT_DynamicAnimTrailCount_MAX);
 DEFINE_STAT(STAT_DynamicAnimTrailGTMem_MAX);
 DEFINE_STAT(STAT_DynamicUntrackedGTMem_MAX);
 
-DECLARE_CYCLE_STAT(TEXT("EmitterInstance Init"), STAT_ParticleEmitterInstance_Init, STATGROUP_Particles);
-DECLARE_CYCLE_STAT(TEXT("MeshEmitterInstance Init"), STAT_MeshEmitterInstance_Init, STATGROUP_Particles);
-DECLARE_CYCLE_STAT(TEXT("EmitterInstance InitParams"), STAT_ParticleEmitterInstance_InitParameters, STATGROUP_Particles);
-DECLARE_CYCLE_STAT(TEXT("MeshEmitterInstance InitParams"), STAT_MeshEmitterInstance_InitParameters, STATGROUP_Particles);
-DECLARE_CYCLE_STAT(TEXT("EmitterInstance Init Sizes"), STAT_ParticleEmitterInstance_InitSize, STATGROUP_Particles);
-DECLARE_CYCLE_STAT(TEXT("EmitterInstance PrepPerInstanceBlock"), STAT_PrepPerInstanceBlock, STATGROUP_Particles);
-DECLARE_CYCLE_STAT(TEXT("EmitterInstance Resize"), STAT_ParticleEmitterInstance_Resize, STATGROUP_Particles);
+DECLARE_CYCLE_STAT(TEXT("EmitterInstance Init GT"), STAT_ParticleEmitterInstance_Init, STATGROUP_Particles);
+DECLARE_CYCLE_STAT(TEXT("MeshEmitterInstance Init GT"), STAT_MeshEmitterInstance_Init, STATGROUP_Particles);
+DECLARE_CYCLE_STAT(TEXT("EmitterInstance InitParams GT"), STAT_ParticleEmitterInstance_InitParameters, STATGROUP_Particles);
+DECLARE_CYCLE_STAT(TEXT("MeshEmitterInstance InitParams GT"), STAT_MeshEmitterInstance_InitParameters, STATGROUP_Particles);
+DECLARE_CYCLE_STAT(TEXT("EmitterInstance Init Sizes GT"), STAT_ParticleEmitterInstance_InitSize, STATGROUP_Particles);
+DECLARE_CYCLE_STAT(TEXT("EmitterInstance PrepPerInstanceBlock GT"), STAT_PrepPerInstanceBlock, STATGROUP_Particles);
+DECLARE_CYCLE_STAT(TEXT("EmitterInstance Resize GT"), STAT_ParticleEmitterInstance_Resize, STATGROUP_Particles);
+
+
+#define USE_FAST_PARTICLE_POOL 1
+#if USE_FAST_PARTICLE_POOL
+
+static int32 GEnableFastPools = 0;
+static FAutoConsoleVariableRef CVarEnableFastPools(
+	TEXT("r.Emitter.FastPoolEnable"),
+	GEnableFastPools,
+	TEXT("Should we use fast pools for emitters.\n")
+	TEXT(" 0: Don't pool anything\n")
+	TEXT(" 1: Pool the emitters bro (default)\n"),
+	ECVF_Default
+);
+
+static int32 GMaxFreePoolSizeBytes = 2 * 1024 * 1024;
+static FAutoConsoleVariableRef CVarFastPoolMaxFreeSize(
+	TEXT("r.Emitter.FastPoolMaxFreeSize"),
+	GMaxFreePoolSizeBytes,
+	TEXT("Max free pool size to keep around without cleaning up."),
+	ECVF_Default
+);
+
+FCriticalSection GFastPoolsCriticalSection;
+
+struct FFastPoolFreePool
+{
+	FFastPoolFreePool() : LastUsedTime(0.0) { }
+
+	TArray<void*> FreeAllocations;
+	double LastUsedTime;
+};
+
+TMap<int32, FFastPoolFreePool> GFastPoolFreedAllocations;
+int32 GFreePoolSizeBytes = 0;
+
+#define FASTPARTICLEALLOC_CHECKSIZE 0
+#if FASTPARTICLEALLOC_CHECKSIZE 
+FCriticalSection GFastPoolSizeMapCriticalSection;
+TMap<void*, int32> GFastPoolSizeMap;
+#endif
+
+FORCEINLINE static void* FastParticleSmallBlockAlloc(size_t AllocSize)
+{
+	QUICK_SCOPE_CYCLE_COUNTER(STAT_PARTALLOC);
+
+	if (GEnableFastPools)
+	{
+		FScopeLock S(&GFastPoolsCriticalSection);
+
+		FFastPoolFreePool* Allocations = GFastPoolFreedAllocations.Find(AllocSize);
+		if ( Allocations  && Allocations->FreeAllocations.Num() )
+		{
+			void* Result = Allocations->FreeAllocations[0];
+			Allocations->FreeAllocations.RemoveAtSwap(0,1, false);
+			Allocations->LastUsedTime = FPlatformTime::Seconds();
+			GFreePoolSizeBytes -= AllocSize;
+#if FASTPARTICLEALLOC_CHECKSIZE
+			FScopeLock Q(&GFastPoolSizeMapCriticalSection);
+			GFastPoolSizeMap.Add(Result, AllocSize);
+#endif
+			return Result;
+		}
+
+	}
 
+#if FASTPARTICLEALLOC_CHECKSIZE
+	FScopeLock S(&GFastPoolSizeMapCriticalSection);
+	void *Result = FMemory::Malloc(AllocSize);
+	GFastPoolSizeMap.Add(Result, AllocSize);
+	return Result;
+#else
+	return FMemory::Malloc(AllocSize);
+#endif
+}
 
+FORCEINLINE static void FastParticleSmallBlockFree(void *RawMemory, size_t AllocSize)
+{
+	QUICK_SCOPE_CYCLE_COUNTER(STAT_PARTALLOC);
+
+
+#if FASTPARTICLEALLOC_CHECKSIZE
+	{
+		FScopeLock S(&GFastPoolSizeMapCriticalSection);
+		int32* AllocatedSize = GFastPoolSizeMap.Find(RawMemory);
+		GFastPoolSizeMap.Remove(RawMemory);
+		check(AllocatedSize);
+		if (AllocatedSize)
+		{
+			check(*AllocatedSize == AllocSize);
+		}
+	}
+#endif
+	if (GEnableFastPools)
+	{
+		FScopeLock S(&GFastPoolsCriticalSection);
+		FFastPoolFreePool& Allocations = GFastPoolFreedAllocations.FindOrAdd(AllocSize);
+		Allocations.FreeAllocations.Add(RawMemory);
+		GFreePoolSizeBytes += AllocSize;
+
+		if ( GFreePoolSizeBytes > GMaxFreePoolSizeBytes )
+		{
+			// free the oldest allocation
+			FFastPoolFreePool* OldestPool = nullptr;
+			int32 OldestPoolAllocSize  = 0;
+			for ( auto& PoolItr : GFastPoolFreedAllocations )
+			{
+				FFastPoolFreePool& Pool = PoolItr.Value;
+				if ( Pool.FreeAllocations.Num() )
+				{
+					if ( OldestPool == nullptr )
+					{
+						OldestPoolAllocSize = PoolItr.Key;
+						OldestPool = &Pool;
+					}
+					else
+					{
+						if (OldestPool->LastUsedTime > Pool.LastUsedTime)
+						{
+							OldestPoolAllocSize = PoolItr.Key;
+							OldestPool = &Pool;
+						}
+					}
+				}
+			}
+			check( OldestPool );
+			check( OldestPoolAllocSize  != 0 );
+			void* OldAllocation = OldestPool->FreeAllocations[0];
+			OldestPool->FreeAllocations.RemoveAtSwap(0, 1, false);
+			GFreePoolSizeBytes -= OldestPoolAllocSize;
+			FMemory::Free(OldAllocation);
+		}
+		return;
+	}
+	FMemory::Free(RawMemory);
+}
+
+
+#else
 FORCEINLINE static void* FastParticleSmallBlockAlloc(size_t AllocSize)
 {
 	QUICK_SCOPE_CYCLE_COUNTER(STAT_PARTALLOC);
@@ -156,6 +297,9 @@ FORCEINLINE static void FastParticleSmallBlockFree(void *RawMemory, size_t Alloc
 	FMemory::Free(RawMemory);
 }
 
+#endif
+
+
 
 void FParticleDataContainer::Alloc(int32 InParticleDataNumBytes, int32 InParticleIndicesNumShorts)
 {
diff --git a/Engine/Source/Runtime/Engine/Private/Particles/ParticleGpuSimulation.cpp b/Engine/Source/Runtime/Engine/Private/Particles/ParticleGpuSimulation.cpp
--- a/Engine/Source/Runtime/Engine/Private/Particles/ParticleGpuSimulation.cpp
+++ b/Engine/Source/Runtime/Engine/Private/Particles/ParticleGpuSimulation.cpp
@@ -46,8 +46,8 @@
 #include "CoreDelegates.h"
 #include "PipelineStateCache.h"
 
-DECLARE_CYCLE_STAT(TEXT("GPUSpriteEmitterInstance Init"), STAT_GPUSpriteEmitterInstance_Init, STATGROUP_Particles);
-DECLARE_FLOAT_COUNTER_STAT(TEXT("Particle Simulation"), Stat_GPU_ParticleSimulation, STATGROUP_GPU);
+DECLARE_CYCLE_STAT(TEXT("GPUSpriteEmitterInstance Init GT"), STAT_GPUSpriteEmitterInstance_Init, STATGROUP_Particles);
+DECLARE_FLOAT_COUNTER_STAT(TEXT("Particle Simulation RT"), Stat_GPU_ParticleSimulation, STATGROUP_GPU);
 
 /*------------------------------------------------------------------------------
 	Constants to tune memory and performance for GPU particle simulation.
@@ -91,7 +91,11 @@ static_assert(TILES_PER_INSTANCE <= MAX_PARTICLES_PER_INSTANCE, "Tiles per insta
 static_assert((TILES_PER_INSTANCE & (TILES_PER_INSTANCE - 1)) == 0, "Tiles per instance is not a power of two.");
 
 /** Maximum number of vector fields that can be evaluated at once. */
+#if GPUPARTICLE_LOCAL_VF_ONLY
+enum { MAX_VECTOR_FIELDS = 1 };
+#else
 enum { MAX_VECTOR_FIELDS = 4 };
+#endif
 
 // Using a fix step 1/30, allows game targetting 30 fps and 60 fps to have single iteration updates.
 static TAutoConsoleVariable<float> CVarGPUParticleFixDeltaSeconds(TEXT("r.GPUParticle.FixDeltaSeconds"), 1.f/30.f,TEXT("GPU particle fix delta seconds."));
@@ -1148,7 +1152,7 @@ public:
 	void SetVectorFieldParameters(FRHICommandList& RHICmdList, const FVectorFieldUniformBufferRef& UniformBuffer, const FTexture3DRHIParamRef VolumeTexturesRHI[])
 	{
 		FPixelShaderRHIParamRef PixelShaderRHI = GetPixelShader();
-		SetUniformBufferParameter(RHICmdList, PixelShaderRHI, GetUniformBufferParameter<FVectorFieldUniformParameters>(), UniformBuffer);
+			SetUniformBufferParameter(RHICmdList, PixelShaderRHI, GetUniformBufferParameter<FVectorFieldUniformParameters>(), UniformBuffer);
 		
 		FSamplerStateRHIParamRef SamplerStateLinear = TStaticSamplerState<SF_Bilinear,AM_Clamp,AM_Clamp,AM_Clamp>::GetRHI();
 
@@ -1380,6 +1384,7 @@ struct FSimulationCommandGPU
 	FParticlePerFrameSimulationParameters PerFrameParameters;
 	/** Parameters to sample the local vector field for this simulation. */
 	FVectorFieldUniformBufferRef VectorFieldsUniformBuffer;
+	FLocalUniformBuffer VectorFieldsUniformBufferLocal;
 	/** Vector field volume textures for this simulation. */
 	FTexture3DRHIParamRef VectorFieldTexturesRHI[MAX_VECTOR_FIELDS];
 	/** The number of tiles to simulate. */
@@ -1428,6 +1433,7 @@ void ExecuteSimulationCommands(
 		return;
 	}
 
+	SCOPE_CYCLE_COUNTER(STAT_GPUParticlesSimulationCommands);
 	SCOPED_DRAW_EVENT(RHICmdList, ParticleSimulation);
 	SCOPED_GPU_STAT(RHICmdList, Stat_GPU_ParticleSimulation);
 
@@ -1461,7 +1467,7 @@ void ExecuteSimulationCommands(
 			RHICmdList, 
 			Command.VectorFieldsUniformBuffer,
 			Command.VectorFieldTexturesRHI
-			);
+		);
 		DrawAlignedParticleTiles(RHICmdList, Command.TileCount);
 	}
 
@@ -2481,6 +2487,12 @@ public:
 	/** The vertex factory for visualizing the local vector field. */
 	FVectorFieldVisualizationVertexFactory* VectorFieldVisualizationVertexFactory;
 
+#if GPUPARTICLE_LOCAL_VF_ONLY
+	/** If we only support local VF's then we can cache this uniform buffer */
+	FVectorFieldUniformBufferRef LocalVectorFieldUniformBuffer;
+	float LocalIntensity;
+#endif
+
 	/** The simulation index within the associated FX system. */
 	int32 SimulationIndex;
 
@@ -2506,6 +2518,9 @@ public:
 	FParticleSimulationGPU()
 		: EmitterSimulationResources(NULL)
 		, VectorFieldVisualizationVertexFactory(NULL)
+#if GPUPARTICLE_LOCAL_VF_ONLY
+		, LocalIntensity(0.0f)
+#endif
 		, SimulationIndex(INDEX_NONE)
 		, SimulationPhase(EParticleSimulatePhase::Main)
 		, bWantsCollision(false)
@@ -4609,8 +4624,8 @@ void FFXSystem::SimulateGPUParticles(
 			VectorFieldParameters.IntensityAndTightness[Index] = FVector4(0.0f);
 		}
 		VectorFieldParameters.Count = 0;
-		EmptyVectorFieldUniformBuffer = FVectorFieldUniformBufferRef::CreateUniformBufferImmediate(VectorFieldParameters, UniformBuffer_SingleFrame);
-	}
+			EmptyVectorFieldUniformBuffer = FVectorFieldUniformBufferRef::CreateUniformBufferImmediate(VectorFieldParameters, UniformBuffer_SingleFrame);
+		}
 
 	// Gather simulation commands from all active simulations.
 	static TArray<FSimulationCommandGPU> SimulationCommands;
@@ -4650,10 +4665,13 @@ void FFXSystem::SimulateGPUParticles(
 
 				// Add the local vector field.
 				VectorFieldParameters.Count = 0;
+
+				
+				float LocalIntensity = 0.0f;
 				if (Simulation->LocalVectorField.Resource)
 				{
-					const float Intensity = Simulation->LocalVectorField.Intensity * Simulation->LocalVectorField.Resource->Intensity;
-					if (FMath::Abs(Intensity) > 0.0f)
+					LocalIntensity = Simulation->LocalVectorField.Intensity * Simulation->LocalVectorField.Resource->Intensity;
+					if (FMath::Abs(LocalIntensity) > 0.0f)
 					{
 						Simulation->LocalVectorField.Resource->Update(RHICmdList, Simulation->PerFrameSimulationParameters.DeltaSeconds);
 						SimulationCommand->VectorFieldTexturesRHI[0] = Simulation->LocalVectorField.Resource->VolumeTextureRHI;
@@ -4661,6 +4679,7 @@ void FFXSystem::SimulateGPUParticles(
 					}
 				}
 
+#if !GPUPARTICLE_LOCAL_VF_ONLY
 				// Add any world vector fields that intersect the simulation.
 				const float GlobalVectorFieldScale = Simulation->EmitterSimulationResources->GlobalVectorFieldScale;
 				const float GlobalVectorFieldTightness = Simulation->EmitterSimulationResources->GlobalVectorFieldTightness;
@@ -4679,10 +4698,12 @@ void FFXSystem::SimulateGPUParticles(
 						}
 					}
 				}
+#endif
 
 				// Fill out any remaining vector field entries.
 				if (VectorFieldParameters.Count > 0)
 				{
+#if !GPUPARTICLE_LOCAL_VF_ONLY
 					int32 PadCount = VectorFieldParameters.Count;
 					while (PadCount < MAX_VECTOR_FIELDS)
 					{
@@ -4692,7 +4713,20 @@ void FFXSystem::SimulateGPUParticles(
 						VectorFieldParameters.VolumeSize[Index] = FVector4(1.0f);
 						VectorFieldParameters.IntensityAndTightness[Index] = FVector4(0.0f);
 					}
+#endif
+		
+						
+#if GPUPARTICLE_LOCAL_VF_ONLY
+					const bool bRecreateBuffer = !Simulation->LocalVectorFieldUniformBuffer || LocalIntensity != Simulation->LocalIntensity;
+					if (bRecreateBuffer)
+					{
+						Simulation->LocalVectorFieldUniformBuffer = FVectorFieldUniformBufferRef::CreateUniformBufferImmediate(VectorFieldParameters, UniformBuffer_MultiFrame);
+						Simulation->LocalIntensity = LocalIntensity;
+					}
+					SimulationCommand->VectorFieldsUniformBuffer = Simulation->LocalVectorFieldUniformBuffer;
+#else
 					SimulationCommand->VectorFieldsUniformBuffer = FVectorFieldUniformBufferRef::CreateUniformBufferImmediate(VectorFieldParameters, UniformBuffer_SingleFrame);
+#endif
 				}
 			}
 		
@@ -4755,6 +4789,7 @@ void FFXSystem::SimulateGPUParticles(
 	// Inject any new particles that have spawned into the simulation.
 	if (NewParticles.Num())
 	{
+		SCOPE_CYCLE_COUNTER(STAT_GPUParticlesInjectionTime);
 		SCOPED_DRAW_EVENT(RHICmdList, ParticleInjection);
 		SCOPED_GPU_STAT(RHICmdList, Stat_GPU_ParticleSimulation);		
 
diff --git a/Engine/Source/Runtime/Engine/Private/Particles/ParticleModules_Location.cpp b/Engine/Source/Runtime/Engine/Private/Particles/ParticleModules_Location.cpp
--- a/Engine/Source/Runtime/Engine/Private/Particles/ParticleModules_Location.cpp
+++ b/Engine/Source/Runtime/Engine/Private/Particles/ParticleModules_Location.cpp
@@ -2038,8 +2038,15 @@ bool UParticleModuleLocationBoneSocket::GetSocketInfoForSourceIndex(FModuleLocat
 	{
 		case EBoneSocketSourceIndexMode::SourceLocations:
 		{
-			OutSocket = SourceComponent->SkeletalMesh->FindSocket(SourceLocations[SourceIndex].BoneSocketName);
-			OutOffset = SourceLocations[SourceIndex].Offset + UniversalOffset;
+			if (ensureMsgf(SourceIndex < SourceLocations.Num(), TEXT("Invalid index of %s for %s"), SourceIndex, *GetPathName()))
+			{
+				OutSocket = SourceComponent->SkeletalMesh->FindSocket(SourceLocations[SourceIndex].BoneSocketName);
+				OutOffset = SourceLocations[SourceIndex].Offset + UniversalOffset;
+			}
+			else
+			{
+				return false;
+			}
 		}
 		break;
 		case EBoneSocketSourceIndexMode::PreSelectedIndices:
@@ -2818,9 +2825,11 @@ bool UParticleModuleLocationSkelVertSurface::GetParticleLocation(FParticleEmitte
 
 	if (SkelMeshResource)
 	{
+		FSkeletalMeshLODRenderData& LODData = SkelMeshResource->LODRenderData[0];
+		FSkinWeightVertexBuffer& SkinWeightBuffer = *InSkelMeshComponent->GetSkinWeightBuffer(0);
 		if (SourceType == VERTSURFACESOURCE_Vert)
 		{
-			FVector VertPos = InSkelMeshComponent->GetSkinnedVertexPosition(InPrimaryVertexIndex);
+			FVector VertPos = USkeletalMeshComponent::GetSkinnedVertexPosition(InSkelMeshComponent, InPrimaryVertexIndex, LODData, SkinWeightBuffer);
 			OutPosition = InSkelMeshComponent->GetComponentTransform().TransformPosition(VertPos);
 			OutRotation = FQuat::Identity;
 		}
@@ -2828,14 +2837,13 @@ bool UParticleModuleLocationSkelVertSurface::GetParticleLocation(FParticleEmitte
 		{
 			FVector Verts[3];
 			int32 VertIndex[3];
-			FSkeletalMeshLODRenderData& LODData = SkelMeshResource->LODRenderData[0];
 
 			VertIndex[0] = LODData.MultiSizeIndexContainer.GetIndexBuffer()->Get( InPrimaryVertexIndex );
 			VertIndex[1] = LODData.MultiSizeIndexContainer.GetIndexBuffer()->Get( InPrimaryVertexIndex+1 );
 			VertIndex[2] = LODData.MultiSizeIndexContainer.GetIndexBuffer()->Get( InPrimaryVertexIndex+2 );
-			Verts[0] = InSkelMeshComponent->GetComponentTransform().TransformPosition(InSkelMeshComponent->GetSkinnedVertexPosition(VertIndex[0]));
-			Verts[1] = InSkelMeshComponent->GetComponentTransform().TransformPosition(InSkelMeshComponent->GetSkinnedVertexPosition(VertIndex[1]));
-			Verts[2] = InSkelMeshComponent->GetComponentTransform().TransformPosition(InSkelMeshComponent->GetSkinnedVertexPosition(VertIndex[2]));
+			Verts[0] = InSkelMeshComponent->GetComponentTransform().TransformPosition(USkeletalMeshComponent::GetSkinnedVertexPosition(InSkelMeshComponent, VertIndex[0], LODData, SkinWeightBuffer));
+			Verts[1] = InSkelMeshComponent->GetComponentTransform().TransformPosition(USkeletalMeshComponent::GetSkinnedVertexPosition(InSkelMeshComponent, VertIndex[1], LODData, SkinWeightBuffer));
+			Verts[2] = InSkelMeshComponent->GetComponentTransform().TransformPosition(USkeletalMeshComponent::GetSkinnedVertexPosition(InSkelMeshComponent, VertIndex[2], LODData, SkinWeightBuffer));
 
 			FVector V0ToV2 = (Verts[2] - Verts[0]);
 			V0ToV2.Normalize();
diff --git a/Engine/Source/Runtime/Engine/Private/Particles/ParticleSystemRender.cpp b/Engine/Source/Runtime/Engine/Private/Particles/ParticleSystemRender.cpp
--- a/Engine/Source/Runtime/Engine/Private/Particles/ParticleSystemRender.cpp
+++ b/Engine/Source/Runtime/Engine/Private/Particles/ParticleSystemRender.cpp
@@ -38,10 +38,10 @@
 #include "Particles/ParticleLODLevel.h"
 #include "Engine/StaticMesh.h"
 
-DECLARE_CYCLE_STAT(TEXT("ParticleSystemSceneProxy GetMeshElements"), STAT_FParticleSystemSceneProxy_GetMeshElements, STATGROUP_Particles);
-DECLARE_CYCLE_STAT(TEXT("DynamicSpriteEmitterData GetDynamicMeshElementsEmitter GetParticleOrderData"), STAT_FDynamicSpriteEmitterData_GetDynamicMeshElementsEmitter_GetParticleOrderData, STATGROUP_Particles);
-DECLARE_CYCLE_STAT(TEXT("DynamicSpriteEmitterData PerParticleWorkOrTasks"), STAT_FDynamicSpriteEmitterData_PerParticleWorkOrTasks, STATGROUP_Particles);
-DECLARE_CYCLE_STAT(TEXT("DynamicSpriteEmitterData GetDynamicMeshElementsEmitter Task"), STAT_FDynamicSpriteEmitterData_GetDynamicMeshElementsEmitter_Task, STATGROUP_Particles);
+DECLARE_CYCLE_STAT(TEXT("ParticleSystemSceneProxy GetMeshElements RT"), STAT_FParticleSystemSceneProxy_GetMeshElements, STATGROUP_Particles);
+DECLARE_CYCLE_STAT(TEXT("DynamicSpriteEmitterData GetDynamicMeshElementsEmitter GetParticleOrderData RT"), STAT_FDynamicSpriteEmitterData_GetDynamicMeshElementsEmitter_GetParticleOrderData, STATGROUP_Particles);
+DECLARE_CYCLE_STAT(TEXT("DynamicSpriteEmitterData PerParticleWorkOrTasks RT"), STAT_FDynamicSpriteEmitterData_PerParticleWorkOrTasks, STATGROUP_Particles);
+DECLARE_CYCLE_STAT(TEXT("DynamicSpriteEmitterData GetDynamicMeshElementsEmitter Task RT"), STAT_FDynamicSpriteEmitterData_GetDynamicMeshElementsEmitter_Task, STATGROUP_Particles);
 
 
 #include "InGamePerformanceTracker.h"
@@ -6748,7 +6748,7 @@ int32 FDynamicAnimTrailEmitterData::FillVertexData(struct FAsyncBufferFillData&
 //	ParticleSystemSceneProxy
 ///////////////////////////////////////////////////////////////////////////////
 /** Initialization constructor. */
-FParticleSystemSceneProxy::FParticleSystemSceneProxy(const UParticleSystemComponent* Component, FParticleDynamicData* InDynamicData)
+FParticleSystemSceneProxy::FParticleSystemSceneProxy(const UParticleSystemComponent* Component, FParticleDynamicData* InDynamicData, bool InbCanBeOccluded)
 	: FPrimitiveSceneProxy(Component, Component->Template ? Component->Template->GetFName() : NAME_None)
 	, Owner(Component->GetOwner())
 	, bCastShadow(Component->CastShadow)
@@ -6772,6 +6772,8 @@ FParticleSystemSceneProxy::FParticleSystemSceneProxy(const UParticleSystemCompon
 	, FirstFreeMeshBatch(0)
 	, bVertexFactoriesDirty(false)
 	, FeatureLevel(GetScene().GetFeatureLevel())
+	, bCanBeOccluded(InbCanBeOccluded)
+	, bHasCustomOcclusionBounds(false)
 {
 	WireframeColor = FLinearColor(3.0f, 0.0f, 0.0f);
 	LevelColor = FLinearColor(1.0f, 1.0f, 0.0f);
@@ -6781,6 +6783,18 @@ FParticleSystemSceneProxy::FParticleSystemSceneProxy(const UParticleSystemCompon
 
 	// Particle systems intrinsically always have motion, but is this motion relevant to systems external to particle systems?
 	bAlwaysHasVelocity = Component->Template && Component->Template->DoesAnyEmitterHaveMotionBlur(Component->GetCurrentLODIndex());
+
+	if (bCanBeOccluded && Component->Template && (Component->Template->OcclusionBoundsMethod == EPSOBM_CustomBounds))
+	{
+		OcclusionBounds = FBoxSphereBounds(Component->Template->CustomOcclusionBounds);
+		bHasCustomOcclusionBounds = true;
+	}
+}
+
+SIZE_T FParticleSystemSceneProxy::GetTypeHash() const
+{
+	static size_t UniquePointer;
+	return reinterpret_cast<size_t>(&UniquePointer);
 }
 
 FParticleSystemSceneProxy::~FParticleSystemSceneProxy()
@@ -7159,25 +7173,6 @@ void FParticleSystemSceneProxy::GatherSimpleLights(const FSceneViewFamily& ViewF
 	}
 }
 
-/**
- *	Occluding particle system scene proxy...
- */
-/** Initialization constructor. */
-FParticleSystemOcclusionSceneProxy::FParticleSystemOcclusionSceneProxy(const UParticleSystemComponent* Component, FParticleDynamicData* InDynamicData) :
-	  FParticleSystemSceneProxy(Component,InDynamicData)
-	, bHasCustomOcclusionBounds(false)
-{
-	if (Component->Template && (Component->Template->OcclusionBoundsMethod == EPSOBM_CustomBounds))
-	{
-		OcclusionBounds = FBoxSphereBounds(Component->Template->CustomOcclusionBounds);
-		bHasCustomOcclusionBounds = true;
-	}
-}
-
-FParticleSystemOcclusionSceneProxy::~FParticleSystemOcclusionSceneProxy()
-{
-}
-
 FPrimitiveSceneProxy* UParticleSystemComponent::CreateSceneProxy()
 {
 	FParticleSystemSceneProxy* NewProxy = NULL;
@@ -7203,11 +7198,11 @@ FPrimitiveSceneProxy* UParticleSystemComponent::CreateSceneProxy()
 		if (CanBeOccluded())
 		{
 			Template->CustomOcclusionBounds.IsValid = true;
-			NewProxy = ::new FParticleSystemOcclusionSceneProxy(this,ParticleDynamicData);
+			NewProxy = ::new FParticleSystemSceneProxy(this,ParticleDynamicData, true);
 		}
 		else
 		{
-			NewProxy = ::new FParticleSystemSceneProxy(this,ParticleDynamicData);
+			NewProxy = ::new FParticleSystemSceneProxy(this,ParticleDynamicData, false);
 		}
 		check (NewProxy);
 		if (ParticleDynamicData)
diff --git a/Engine/Source/Runtime/Engine/Private/Particles/ParticleTrail2EmitterInstance.cpp b/Engine/Source/Runtime/Engine/Private/Particles/ParticleTrail2EmitterInstance.cpp
--- a/Engine/Source/Runtime/Engine/Private/Particles/ParticleTrail2EmitterInstance.cpp
+++ b/Engine/Source/Runtime/Engine/Private/Particles/ParticleTrail2EmitterInstance.cpp
@@ -43,7 +43,7 @@ DEFINE_STAT(STAT_TrailTickTime);
 
 DEFINE_STAT(STAT_AnimTrailNotifyTime);
 
-DECLARE_CYCLE_STAT(TEXT("TrailEmitterInstance Init"), STAT_TrailEmitterInstance_Init, STATGROUP_Particles);
+DECLARE_CYCLE_STAT(TEXT("TrailEmitterInstance Init GT"), STAT_TrailEmitterInstance_Init, STATGROUP_Particles);
 
 
 #define MAX_TRAIL_INDICES	65535
diff --git a/Engine/Source/Runtime/Engine/Private/PhysicsEngine/PhysAnim.cpp b/Engine/Source/Runtime/Engine/Private/PhysicsEngine/PhysAnim.cpp
--- a/Engine/Source/Runtime/Engine/Private/PhysicsEngine/PhysAnim.cpp
+++ b/Engine/Source/Runtime/Engine/Private/PhysicsEngine/PhysAnim.cpp
@@ -640,21 +640,24 @@ void USkeletalMeshComponent::UpdateKinematicBonesToAnim(const TArray<FTransform>
 		{
 			if (bNeedsSkinning)
 			{
+					const FSkeletalMeshLODRenderData& LODData = MeshObject->GetSkeletalMeshRenderData().LODRenderData[0];
+				FSkinWeightVertexBuffer& SkinWeightBuffer = *GetSkinWeightBuffer(0);
+				TArray<FMatrix> RefToLocals;
 				TArray<FVector> NewPositions;
 				if (true)
 				{
 					SCOPE_CYCLE_COUNTER(STAT_SkinPerPolyVertices);
-					ComputeSkinnedPositions(NewPositions);
+					CacheRefToLocalMatrices(RefToLocals);
+					ComputeSkinnedPositions(this, NewPositions, RefToLocals, LODData, SkinWeightBuffer);
 				}
 				else	//keep old way around for now - useful for comparing performance
 				{
-					const FSkeletalMeshLODRenderData& LODData = MeshObject->GetSkeletalMeshRenderData().LODRenderData[0];
 					NewPositions.AddUninitialized(LODData.GetNumVertices());
 					{
 						SCOPE_CYCLE_COUNTER(STAT_SkinPerPolyVertices);
 						for (uint32 VertIndex = 0; VertIndex < LODData.GetNumVertices(); ++VertIndex)
 						{
-							NewPositions[VertIndex] = GetSkinnedVertexPosition(VertIndex);
+							NewPositions[VertIndex] = GetSkinnedVertexPosition(this, VertIndex, LODData, SkinWeightBuffer);
 						}
 					}
 				}
diff --git a/Engine/Source/Runtime/Engine/Private/PhysicsEngine/PhysDerivedData.h b/Engine/Source/Runtime/Engine/Private/PhysicsEngine/PhysDerivedData.h
--- a/Engine/Source/Runtime/Engine/Private/PhysicsEngine/PhysDerivedData.h
+++ b/Engine/Source/Runtime/Engine/Private/PhysicsEngine/PhysDerivedData.h
@@ -50,7 +50,7 @@ public:
 		// This is a version string that mimics the old versioning scheme. If you
 		// want to bump this version, generate a new guid using VS->Tools->Create GUID and
 		// return it here. Ex.
-		return TEXT("3B0AB674B19F4A698B28F7025C031EA2");	
+		return TEXT("8E2B126ED2B644F1BBF553FEF6480DB4");	
 	}
 
 	virtual FString GetPluginSpecificCacheKeySuffix() const override
diff --git a/Engine/Source/Runtime/Engine/Private/PrecomputedVolumetricLightmap.cpp b/Engine/Source/Runtime/Engine/Private/PrecomputedVolumetricLightmap.cpp
--- a/Engine/Source/Runtime/Engine/Private/PrecomputedVolumetricLightmap.cpp
+++ b/Engine/Source/Runtime/Engine/Private/PrecomputedVolumetricLightmap.cpp
@@ -26,7 +26,7 @@ void FVolumetricLightmapDataLayer::CreateTexture(FIntVector Dimensions)
 		Dimensions.Z, 
 		Format,
 		1,
-		TexCreate_ShaderResource,
+		TexCreate_ShaderResource | TexCreate_DisableAutoDefrag,
 		CreateInfo);
 }
 
diff --git a/Engine/Source/Runtime/Engine/Private/SceneView.cpp b/Engine/Source/Runtime/Engine/Private/SceneView.cpp
--- a/Engine/Source/Runtime/Engine/Private/SceneView.cpp
+++ b/Engine/Source/Runtime/Engine/Private/SceneView.cpp
@@ -1313,6 +1313,7 @@ void FSceneView::OverridePostProcessSettings(const FPostProcessSettings& Src, fl
 
 		LERP_PP(SceneColorTint);
 		LERP_PP(SceneFringeIntensity);
+		LERP_PP(ChromaticAberrationStartOffset);
 		LERP_PP(BloomIntensity);
 		LERP_PP(BloomThreshold);
 		LERP_PP(Bloom1Tint);
diff --git a/Engine/Source/Runtime/Engine/Private/SkeletalMesh.cpp b/Engine/Source/Runtime/Engine/Private/SkeletalMesh.cpp
--- a/Engine/Source/Runtime/Engine/Private/SkeletalMesh.cpp
+++ b/Engine/Source/Runtime/Engine/Private/SkeletalMesh.cpp
@@ -528,6 +528,15 @@ int32 USkeletalMesh::GetClothingAssetIndex(const FGuid& InAssetGuid) const
 }
 
 bool USkeletalMesh::HasActiveClothingAssets() const
+{
+#if WITH_EDITOR
+	return ComputeActiveClothingAssets();
+#else
+	return bHasActiveClothingAssets;
+#endif
+}
+
+bool USkeletalMesh::ComputeActiveClothingAssets() const
 {
 	if(FSkeletalMeshRenderData* Resource = GetResourceForRendering())
 	{
@@ -576,6 +585,11 @@ void USkeletalMesh::GetClothingAssetsInUse(TArray<UClothingAssetBase*>& OutCloth
 	}
 }
 
+bool USkeletalMesh::NeedCPUData(int32 LODIndex)const
+{
+	return SamplingInfo.IsSamplingEnabled(this, LODIndex);
+}
+
 void USkeletalMesh::InitResources()
 {
 	UpdateUVChannelData(false);
@@ -920,8 +934,30 @@ void USkeletalMesh::PostEditChangeProperty(FPropertyChangedEvent& PropertyChange
 		FMultiComponentReregisterContext ReregisterContext(ComponentsToReregister);
 	}
 
+	if (PropertyThatChanged && PropertyChangedEvent.MemberProperty)
+	{
+		if (PropertyChangedEvent.MemberProperty->GetFName() == FName(TEXT("SamplingInfo")))
+		{
+			SamplingInfo.BuildRegions(this);
+		}
+		else if (PropertyChangedEvent.MemberProperty->GetFName() == FName(TEXT("LODInfo")))
+		{
+			SamplingInfo.BuildWholeMesh(this);
+		}
+	}
+	else
+	{
+		//Rebuild the lot. No property could mean a reimport.
+		SamplingInfo.BuildRegions(this);
+		SamplingInfo.BuildWholeMesh(this);
+	}
+
 	UpdateUVChannelData(true);
 
+#if WITH_EDITOR
+	OnMeshChanged.Broadcast();
+#endif
+
 	Super::PostEditChangeProperty(PropertyChangedEvent);
 }
 
@@ -1594,6 +1630,8 @@ void USkeletalMesh::PostLoad()
 			ClothingAsset->InvalidateCachedData();
 		}
 	}
+
+	bHasActiveClothingAssets = ComputeActiveClothingAssets();
 }
 
 #if WITH_EDITORONLY_DATA
@@ -2330,9 +2368,12 @@ void USkeletalMesh::ReleaseCPUResources()
 	FSkeletalMeshRenderData* SkelMeshRenderData = GetResourceForRendering();
 	if (SkelMeshRenderData)
 	{
-		for (int32 Index = 0; Index < SkelMeshRenderData->LODRenderData.Num(); ++Index)
-	{
-			SkelMeshRenderData->LODRenderData[Index].ReleaseCPUResources();
+		for(int32 Index = 0; Index < SkelMeshRenderData->LODRenderData.Num(); ++Index)
+		{
+			if (!NeedCPUData(Index))
+			{
+				SkelMeshRenderData->LODRenderData[Index].ReleaseCPUResources();
+			}
 		}
 	}
 }
@@ -3546,6 +3587,12 @@ void FSkeletalMeshSceneProxy::GetDynamicElementsSection(const TArray<const FScen
 	}
 }
 
+SIZE_T FSkeletalMeshSceneProxy::GetTypeHash() const
+{
+	static size_t UniquePointer;
+	return reinterpret_cast<size_t>(&UniquePointer);
+}
+
 bool FSkeletalMeshSceneProxy::HasDynamicIndirectShadowCasterRepresentation() const
 {
 	return CastsDynamicShadow() && CastsDynamicIndirectShadow();
@@ -3888,4 +3935,58 @@ FSkeletalMeshComponentRecreateRenderStateContext::~FSkeletalMeshComponentRecreat
 	}
 }
 
+//////////////////////////////////////////////////////////////////////////
+
+template <bool bExtraBoneInfluencesT>
+FVector GetRefVertexLocationTyped(
+	const USkeletalMesh* Mesh,
+	const FSkelMeshRenderSection& Section,
+	const FPositionVertexBuffer& PositionBuffer,
+	const FSkinWeightVertexBuffer& SkinWeightVertexBuffer,
+	const int32 VertIndex
+)
+{
+	FVector SkinnedPos(0, 0, 0);
+
+	// Do soft skinning for this vertex.
+	int32 BufferVertIndex = Section.GetVertexBufferIndex() + VertIndex;
+	const TSkinWeightInfo<bExtraBoneInfluencesT>* SrcSkinWeights = SkinWeightVertexBuffer.GetSkinWeightPtr<bExtraBoneInfluencesT>(BufferVertIndex);
+	int32 MaxBoneInfluences = bExtraBoneInfluencesT ? MAX_TOTAL_INFLUENCES : MAX_INFLUENCES_PER_STREAM;
+
+#if !PLATFORM_LITTLE_ENDIAN
+	// uint8[] elements in LOD.VertexBufferGPUSkin have been swapped for VET_UBYTE4 vertex stream use
+	for (int32 InfluenceIndex = MAX_INFLUENCES - 1; InfluenceIndex >= MAX_INFLUENCES - MaxBoneInfluences; InfluenceIndex--)
+#else
+	for (int32 InfluenceIndex = 0; InfluenceIndex < MaxBoneInfluences; InfluenceIndex++)
+#endif
+	{
+		const int32 MeshBoneIndex = Section.BoneMap[SrcSkinWeights->InfluenceBones[InfluenceIndex]];
+		const float	Weight = (float)SrcSkinWeights->InfluenceWeights[InfluenceIndex] / 255.0f;
+		{
+			const FMatrix BoneTransformMatrix = FMatrix::Identity;//Mesh->GetComposedRefPoseMatrix(MeshBoneIndex);
+			const FMatrix RefToLocal = Mesh->RefBasesInvMatrix[MeshBoneIndex] * BoneTransformMatrix;
+
+			SkinnedPos += BoneTransformMatrix.TransformPosition(PositionBuffer.VertexPosition(BufferVertIndex)) * Weight;
+		}
+	}
+
+	return SkinnedPos;
+}
+
+FVector GetSkeletalMeshRefVertLocation(const USkeletalMesh* Mesh, const FSkeletalMeshLODRenderData& LODData, const FSkinWeightVertexBuffer& SkinWeightVertexBuffer, const int32 VertIndex)
+{
+	int32 SectionIndex;
+	int32 VertIndexInChunk;
+	LODData.GetSectionFromVertexIndex(VertIndex, SectionIndex, VertIndexInChunk);
+	const FSkelMeshRenderSection& Section = LODData.RenderSections[SectionIndex];
+	if (SkinWeightVertexBuffer.HasExtraBoneInfluences())
+	{
+		return GetRefVertexLocationTyped<true>(Mesh, Section, LODData.StaticVertexBuffers.PositionVertexBuffer, SkinWeightVertexBuffer, VertIndexInChunk);
+	}
+	else
+	{
+		return GetRefVertexLocationTyped<false>(Mesh, Section, LODData.StaticVertexBuffers.PositionVertexBuffer, SkinWeightVertexBuffer, VertIndexInChunk);
+	}
+}
+
 #undef LOCTEXT_NAMESPACE
diff --git a/Engine/Source/Runtime/Engine/Private/SkeletalMeshComponentPhysics.cpp b/Engine/Source/Runtime/Engine/Private/SkeletalMeshComponentPhysics.cpp
--- a/Engine/Source/Runtime/Engine/Private/SkeletalMeshComponentPhysics.cpp
+++ b/Engine/Source/Runtime/Engine/Private/SkeletalMeshComponentPhysics.cpp
@@ -1809,29 +1809,26 @@ void USkeletalMeshComponent::UpdatePhysicsToRBChannels()
 	}
 }
 
-FVector USkeletalMeshComponent::GetSkinnedVertexPosition(int32 VertexIndex) const
+template<bool bCachedMatrices>
+FVector GetTypedSkinnedVertexPositionWithCloth(USkeletalMeshComponent* Component, int32 VertexIndex, const FSkeletalMeshLODRenderData& LODData, FSkinWeightVertexBuffer& SkinWeightBuffer, TArray<FMatrix>& CachedRefToLocals)
 {
+	// Find the chunk and vertex within that chunk, and skinning type, for this vertex.
+	int32 SectionIndex;
+	int32 VertIndexInChunk;
+	LODData.GetSectionFromVertexIndex(VertexIndex, SectionIndex, VertIndexInChunk);
+	const FSkelMeshRenderSection& Section = LODData.RenderSections[SectionIndex];
+
 	// only if this component has clothing and is showing simulated results	
-	if (SkeletalMesh &&
-		SkeletalMesh->MeshClothingAssets.Num() > 0 &&
-		MeshObject &&
-		!bDisableClothSimulation &&
-		ClothBlendWeight > 0.0f // if cloth blend weight is 0.0, only showing skinned vertices regardless of simulation positions
+	if (Component->SkeletalMesh &&
+		Component->SkeletalMesh->MeshClothingAssets.Num() > 0 &&
+		!Component->bDisableClothSimulation &&
+		Component->ClothBlendWeight > 0.0f // if cloth blend weight is 0.0, only showing skinned vertices regardless of simulation positions
 		)
 	{
-		FSkeletalMeshLODRenderData& LODData = MeshObject->GetSkeletalMeshRenderData().LODRenderData[0];
-
-		// Find the chunk and vertex within that chunk, and skinning type, for this vertex.
-		int32 SectionIndex;
-		int32 VertIndexInChunk;
-		LODData.GetSectionFromVertexIndex(VertexIndex, SectionIndex, VertIndexInChunk);
-
 		bool bClothVertex = false;
 		FGuid ClothAssetGuid;
 
 		// if this section corresponds to a cloth section, returns corresponding cloth section's info instead
-		const FSkelMeshRenderSection& Section = LODData.RenderSections[SectionIndex];
-
 		// if this chunk has cloth data
 		if (Section.HasClothingData())
 		{
@@ -1842,25 +1839,146 @@ FVector USkeletalMeshComponent::GetSkinnedVertexPosition(int32 VertexIndex) cons
 		if (bClothVertex)
 		{
 			FVector SimulatedPos;
-			if (GetClothSimulatedPosition_GameThread(ClothAssetGuid, VertIndexInChunk, SimulatedPos))
+			if (Component->GetClothSimulatedPosition_GameThread(ClothAssetGuid, VertIndexInChunk, SimulatedPos))
 			{
 				// a simulated position is in world space and convert this to local space
 				// because SkinnedMeshComponent::GetSkinnedVertexPosition() returns the position in local space
-				SimulatedPos = GetComponentTransform().InverseTransformPosition(SimulatedPos);
+				SimulatedPos = Component->GetComponentTransform().InverseTransformPosition(SimulatedPos);
 
 				// if blend weight is 1.0, doesn't need to blend with a skinned position
-				if (ClothBlendWeight < 1.0f) 
+				if (Component->ClothBlendWeight < 1.0f)
 				{
 					// blend with a skinned position
-					FVector SkinnedPos = Super::GetSkinnedVertexPosition(VertexIndex);
-					SimulatedPos = SimulatedPos*ClothBlendWeight + SkinnedPos*(1.0f - ClothBlendWeight);
+					FVector SkinnedPos = SkinWeightBuffer.HasExtraBoneInfluences() ?
+						GetTypedSkinnedVertexPosition<true, bCachedMatrices>(Component, Section, LODData.StaticVertexBuffers.PositionVertexBuffer, SkinWeightBuffer, VertIndexInChunk, CachedRefToLocals) :
+						GetTypedSkinnedVertexPosition<false, bCachedMatrices>(Component, Section, LODData.StaticVertexBuffers.PositionVertexBuffer, SkinWeightBuffer, VertIndexInChunk, CachedRefToLocals);
+
+					SimulatedPos = SimulatedPos*Component->ClothBlendWeight + SkinnedPos*(1.0f - Component->ClothBlendWeight);
 				}
 				return SimulatedPos;
 			}
 		}
 	}
 
-	return Super::GetSkinnedVertexPosition(VertexIndex);
+	return SkinWeightBuffer.HasExtraBoneInfluences() ?
+		GetTypedSkinnedVertexPosition<true, bCachedMatrices>(Component, Section, LODData.StaticVertexBuffers.PositionVertexBuffer, SkinWeightBuffer, VertIndexInChunk, CachedRefToLocals) :
+		GetTypedSkinnedVertexPosition<false, bCachedMatrices>(Component, Section, LODData.StaticVertexBuffers.PositionVertexBuffer, SkinWeightBuffer, VertIndexInChunk, CachedRefToLocals);
+}
+
+FVector USkeletalMeshComponent::GetSkinnedVertexPosition(USkeletalMeshComponent* Component, int32 VertexIndex, const FSkeletalMeshLODRenderData& LODData, FSkinWeightVertexBuffer& SkinWeightBuffer)
+{
+	TArray<FMatrix> Dummy;
+	return GetTypedSkinnedVertexPositionWithCloth<false>(Component, VertexIndex, LODData, SkinWeightBuffer, Dummy);
+}
+
+FVector USkeletalMeshComponent::GetSkinnedVertexPosition(USkeletalMeshComponent* Component, int32 VertexIndex, const FSkeletalMeshLODRenderData& LODData, FSkinWeightVertexBuffer& SkinWeightBuffer, TArray<FMatrix>& CachedRefToLocals)
+{
+	return GetTypedSkinnedVertexPositionWithCloth<true>(Component, VertexIndex, LODData, SkinWeightBuffer, CachedRefToLocals);
+}
+
+void USkeletalMeshComponent::ComputeSkinnedPositions(USkeletalMeshComponent* Component, TArray<FVector> & OutPositions, TArray<FMatrix>& CachedRefToLocals, const FSkeletalMeshLODRenderData& LODData, FSkinWeightVertexBuffer& SkinWeightBuffer)
+{
+	// Fail if no mesh
+	if (!Component->SkeletalMesh)
+	{
+		return;
+	}
+
+	OutPositions.Empty();
+	OutPositions.AddUninitialized(LODData.GetNumVertices());
+	bool bHasExtraBoneInfluences = SkinWeightBuffer.HasExtraBoneInfluences();
+	
+	if (Component->SkeletalMesh->MeshClothingAssets.Num() > 0 &&
+		!Component->bDisableClothSimulation &&
+		Component->ClothBlendWeight > 0.0f // if cloth blend weight is 0.0, only showing skinned vertices regardless of simulation positions
+		)
+	{
+		//update positions
+		for (int32 SectionIdx = 0; SectionIdx < LODData.RenderSections.Num(); ++SectionIdx)
+		{
+			const FSkelMeshRenderSection& Section = LODData.RenderSections[SectionIdx];
+
+			bool bClothVertex = false;
+			int32 ClothAssetIndex = -1;
+			FGuid ClothAssetGuid;
+
+			// if this section corresponds to a cloth section, returns corresponding cloth section's info instead
+			// if this chunk has cloth data
+			if (Section.HasClothingData())
+			{
+				bClothVertex = true;
+				ClothAssetIndex = Section.CorrespondClothAssetIndex;
+				ClothAssetGuid = Section.ClothingData.AssetGuid;
+			}
+
+			if (bClothVertex)
+			{
+				int32 AssetIndex = Component->SkeletalMesh->GetClothingAssetIndex(ClothAssetGuid);
+				if (AssetIndex != INDEX_NONE)
+				{
+					const FClothSimulData* ActorData = Component->CurrentSimulationData_GameThread.Find(AssetIndex);
+
+					if (ActorData)
+					{
+						const uint32 SoftOffset = Section.GetVertexBufferIndex();
+						const uint32 NumSoftVerts = Section.GetNumVertices();
+
+						// if blend weight is 1.0, doesn't need to blend with a skinned position
+						if (Component->ClothBlendWeight < 1.0f)
+						{
+							for (uint32 SoftIdx = 0; SoftIdx < NumSoftVerts; ++SoftIdx)
+							{
+								FVector SimulatedPos = ActorData->Positions[SoftIdx];
+
+								FVector SkinnedPosition = bHasExtraBoneInfluences ? 
+									GetTypedSkinnedVertexPosition<true, true>(Component, Section, LODData.StaticVertexBuffers.PositionVertexBuffer, SkinWeightBuffer, SoftIdx, CachedRefToLocals) :
+									GetTypedSkinnedVertexPosition<false, true>(Component, Section, LODData.StaticVertexBuffers.PositionVertexBuffer, SkinWeightBuffer, SoftIdx, CachedRefToLocals);
+
+								OutPositions[SoftOffset + SoftIdx] = SimulatedPos*Component->ClothBlendWeight + SkinnedPosition*(1.0f - Component->ClothBlendWeight);
+							}
+						}
+						else
+						{
+							for (uint32 SoftIdx = 0; SoftIdx < NumSoftVerts; ++SoftIdx)
+							{
+								OutPositions[SoftOffset + SoftIdx] = ActorData->Positions[SoftIdx];
+							}
+						}
+
+						return;
+					}
+				}
+			}
+
+			//fall back to just regular skinning.
+			const uint32 SoftOffset = Section.GetVertexBufferIndex();
+			const uint32 NumSoftVerts = Section.GetNumVertices();
+			for (uint32 SoftIdx = 0; SoftIdx < NumSoftVerts; ++SoftIdx)
+			{
+				FVector SkinnedPosition = bHasExtraBoneInfluences ?
+					GetTypedSkinnedVertexPosition<true, true>(Component, Section, LODData.StaticVertexBuffers.PositionVertexBuffer, SkinWeightBuffer, SoftIdx, CachedRefToLocals) :
+					GetTypedSkinnedVertexPosition<false, true>(Component, Section, LODData.StaticVertexBuffers.PositionVertexBuffer, SkinWeightBuffer, SoftIdx, CachedRefToLocals);
+				OutPositions[SoftOffset + SoftIdx] = SkinnedPosition;
+			}
+		}
+	}
+	else
+	{
+		for (int32 SectionIdx = 0; SectionIdx < LODData.RenderSections.Num(); ++SectionIdx)
+		{
+			const FSkelMeshRenderSection& Section = LODData.RenderSections[SectionIdx];
+
+			const uint32 SoftOffset = Section.GetVertexBufferIndex();
+			const uint32 NumSoftVerts = Section.GetNumVertices();
+			for (uint32 SoftIdx = 0; SoftIdx < NumSoftVerts; ++SoftIdx)
+			{
+				FVector SkinnedPosition = bHasExtraBoneInfluences ? 
+					GetTypedSkinnedVertexPosition<true, true>(Component, Section, LODData.StaticVertexBuffers.PositionVertexBuffer, SkinWeightBuffer, SoftIdx, CachedRefToLocals) :
+					GetTypedSkinnedVertexPosition<false, true>(Component, Section, LODData.StaticVertexBuffers.PositionVertexBuffer, SkinWeightBuffer, SoftIdx, CachedRefToLocals);
+				OutPositions[SoftOffset + SoftIdx] = SkinnedPosition;
+			}
+		}
+	}
 }
 
 void USkeletalMeshComponent::SetEnableBodyGravity(bool bEnableGravity, FName BoneName)
diff --git a/Engine/Source/Runtime/Engine/Private/SkeletalRenderGPUSkin.cpp b/Engine/Source/Runtime/Engine/Private/SkeletalRenderGPUSkin.cpp
--- a/Engine/Source/Runtime/Engine/Private/SkeletalRenderGPUSkin.cpp
+++ b/Engine/Source/Runtime/Engine/Private/SkeletalRenderGPUSkin.cpp
@@ -53,6 +53,7 @@ static FAutoConsoleVariableRef CVarMorphTargetWeightThreshold(
 	ECVF_Default
 );
 
+
 /*-----------------------------------------------------------------------------
 FMorphVertexBuffer
 -----------------------------------------------------------------------------*/
@@ -1279,15 +1280,87 @@ void FDynamicSkelMeshObjectDataGPUSkin::Clear()
 	ClothBlendWeight = 0.0f;
 }
 
+#define SKELETON_POOL_GPUSKINS 1
+#if SKELETON_POOL_GPUSKINS
+TArray<FDynamicSkelMeshObjectDataGPUSkin*> FreeGpuSkins;
+FCriticalSection FreeGpuSkinsCriticalSection;
+
+static int32 GPoolGpuSkins = 1;
+static int32 GMinPoolCount = 0;
+static int32 GAllocationCounter = 0;
+static const int32 GAllocationsBeforeCleanup = 1000; // number of allocations we make before we clean up the pool, this number is increased when we have to allocate not from the pool
+static FAutoConsoleVariableRef CVarPoolGpuSkins(
+	TEXT("r.GpuSkin.Pool"),
+	GPoolGpuSkins,
+	TEXT("Should we pool gpu skins.\n")
+	TEXT(" 0: Don't pool anything\n")
+	TEXT(" 1: Pool gpu skins bro (default)\n"),
+	ECVF_Default
+);
+#endif
+
 
 FDynamicSkelMeshObjectDataGPUSkin* FDynamicSkelMeshObjectDataGPUSkin::AllocDynamicSkelMeshObjectDataGPUSkin()
 {
+#if SKELETON_POOL_GPUSKINS
+	if (!GPoolGpuSkins)
+	{
+		return new FDynamicSkelMeshObjectDataGPUSkin;
+	}
+
+	FScopeLock S(&FreeGpuSkinsCriticalSection);
+	++GAllocationCounter;
+	GMinPoolCount = FMath::Min(FreeGpuSkins.Num(), GMinPoolCount);
+	if ( FreeGpuSkins.Num() > 0 )
+	{
+		FDynamicSkelMeshObjectDataGPUSkin *Result = FreeGpuSkins[0];
+		FreeGpuSkins.RemoveAtSwap(0);
+		return Result;
+	}
+	else
+	{
+		return new FDynamicSkelMeshObjectDataGPUSkin;
+	}
+#else
 	return new FDynamicSkelMeshObjectDataGPUSkin;
+#endif
 }
 
 void FDynamicSkelMeshObjectDataGPUSkin::FreeDynamicSkelMeshObjectDataGPUSkin(FDynamicSkelMeshObjectDataGPUSkin* Who)
 {
+#if SKELETON_POOL_GPUSKINS
+	if (!GPoolGpuSkins)
+	{
+		delete Who;
+
+		if ( FreeGpuSkins.Num() > 0 )
+		{
+			FScopeLock S(&FreeGpuSkinsCriticalSection);
+			for ( FDynamicSkelMeshObjectDataGPUSkin* I : FreeGpuSkins )
+			{
+				delete I;
+			}
+			FreeGpuSkins.Empty();
+		}
+		return;
+	}
+
+	Who->Clear();
+	FScopeLock S(&FreeGpuSkinsCriticalSection);
+	FreeGpuSkins.Add(Who);
+	if ( GAllocationCounter > GAllocationsBeforeCleanup )
+	{
+		GAllocationCounter = 0;
+		for ( int32 I = 0; I < GMinPoolCount; ++I )
+		{
+			delete FreeGpuSkins[0];
+			FreeGpuSkins.RemoveAtSwap(0);
+		}
+		GMinPoolCount = FreeGpuSkins.Num();
+	}
+#else
 	delete Who;
+#endif
 }
 
 void FDynamicSkelMeshObjectDataGPUSkin::InitDynamicSkelMeshObjectDataGPUSkin(
diff --git a/Engine/Source/Runtime/Engine/Private/SkeletalRenderGPUSkin.h b/Engine/Source/Runtime/Engine/Private/SkeletalRenderGPUSkin.h
--- a/Engine/Source/Runtime/Engine/Private/SkeletalRenderGPUSkin.h
+++ b/Engine/Source/Runtime/Engine/Private/SkeletalRenderGPUSkin.h
@@ -12,6 +12,7 @@
 #include "ShaderParameters.h"
 #include "Components/SkinnedMeshComponent.h"
 #include "GlobalShader.h"
+#include "GPUSkinCache.h"
 #include "GPUSkinVertexFactory.h"
 #include "SkeletalRenderPublic.h"
 #include "ClothingSystemRuntimeTypes.h"
@@ -645,4 +646,4 @@ protected:
 	FShaderParameter MorphTargetWeightParameter;
 	FShaderParameter MorphWorkItemsParameter;
 	FShaderParameter PositionScaleParameter;
-};
\ No newline at end of file
+};
diff --git a/Engine/Source/Runtime/Engine/Private/SplineMeshSceneProxy.cpp b/Engine/Source/Runtime/Engine/Private/SplineMeshSceneProxy.cpp
--- a/Engine/Source/Runtime/Engine/Private/SplineMeshSceneProxy.cpp
+++ b/Engine/Source/Runtime/Engine/Private/SplineMeshSceneProxy.cpp
@@ -39,6 +39,12 @@ FSplineMeshSceneProxy::FSplineMeshSceneProxy(USplineMeshComponent* InComponent)
 	}
 }
 
+SIZE_T FSplineMeshSceneProxy::GetTypeHash() const
+{
+	static size_t UniquePointer;
+	return reinterpret_cast<size_t>(&UniquePointer);
+}
+
 bool FSplineMeshSceneProxy::GetShadowMeshElement(int32 LODIndex, int32 BatchIndex, uint8 InDepthPriorityGroup, FMeshBatch& OutMeshBatch, bool bDitheredLODTransition) const
 {
 	//checkf(LODIndex == 0, TEXT("Getting spline static mesh element with invalid LOD [%d]"), LODIndex);
diff --git a/Engine/Source/Runtime/Engine/Private/SplineMeshSceneProxy.h b/Engine/Source/Runtime/Engine/Private/SplineMeshSceneProxy.h
--- a/Engine/Source/Runtime/Engine/Private/SplineMeshSceneProxy.h
+++ b/Engine/Source/Runtime/Engine/Private/SplineMeshSceneProxy.h
@@ -131,7 +131,7 @@ private:
 // SplineMeshSceneProxy
 
 /** Scene proxy for SplineMesh instance */
-class FSplineMeshSceneProxy : public FStaticMeshSceneProxy
+class FSplineMeshSceneProxy final : public FStaticMeshSceneProxy
 {
 protected:
 	struct FLODResources
@@ -145,6 +145,7 @@ protected:
 		}
 	};
 public:
+	SIZE_T GetTypeHash() const override;
 
 	FSplineMeshSceneProxy(USplineMeshComponent* InComponent);
 
diff --git a/Engine/Source/Runtime/Engine/Private/StaticMeshRender.cpp b/Engine/Source/Runtime/Engine/Private/StaticMeshRender.cpp
--- a/Engine/Source/Runtime/Engine/Private/StaticMeshRender.cpp
+++ b/Engine/Source/Runtime/Engine/Private/StaticMeshRender.cpp
@@ -295,6 +295,12 @@ void UStaticMeshComponent::SetLODDataCount( const uint32 MinSize, const uint32 M
 	}
 }
 
+SIZE_T FStaticMeshSceneProxy::GetTypeHash() const
+{
+	static size_t UniquePointer;
+	return reinterpret_cast<size_t>(&UniquePointer);
+}
+
 bool FStaticMeshSceneProxy::GetShadowMeshElement(int32 LODIndex, int32 BatchIndex, uint8 InDepthPriorityGroup, FMeshBatch& OutMeshBatch, bool bDitheredLODTransition) const
 {
 	const FStaticMeshLODResources& LOD = RenderData->LODResources[LODIndex];
diff --git a/Engine/Source/Runtime/Engine/Private/TimerManager.cpp b/Engine/Source/Runtime/Engine/Private/TimerManager.cpp
--- a/Engine/Source/Runtime/Engine/Private/TimerManager.cpp
+++ b/Engine/Source/Runtime/Engine/Private/TimerManager.cpp
@@ -9,6 +9,7 @@
 #include "Misc/CoreDelegates.h"
 #include "Engine/World.h"
 #include "UnrealEngine.h"
+#include "TimeGuard.h"
 
 DECLARE_CYCLE_STAT(TEXT("SetTimer"), STAT_SetTimer, STATGROUP_Engine);
 DECLARE_CYCLE_STAT(TEXT("ClearTimer"), STAT_ClearTimer, STATGROUP_Engine);
@@ -525,6 +526,22 @@ DECLARE_DWORD_COUNTER_STAT(TEXT("TimerManager Heap Size"),STAT_NumHeapEntries,ST
 
 void FTimerManager::Tick(float DeltaTime)
 {
+#if DO_TIMEGUARD && 0
+	TArray<FTimerUnifiedDelegate> RunTimerDelegates;
+	FTimerNameDelegate NameFunction = FTimerNameDelegate::CreateLambda([&] {
+			FString ActiveDelegates;
+			for ( const FTimerUnifiedDelegate& Descriptor : RunTimerDelegates )
+			{
+				ActiveDelegates += FString::Printf(TEXT("Delegate %s, "), *Descriptor.ToString() );
+			}
+			return FString::Printf(TEXT("UWorld::Tick - TimerManager, %s"), *ActiveDelegates);
+		});
+
+
+	// no delegate should take longer then 5ms to run 
+	SCOPE_TIME_GUARD_DELEGATE_MS(NameFunction, 5);
+#endif	
+
 	// @todo, might need to handle long-running case
 	// (e.g. every X seconds, renormalize to InternalTime = 0)
 
@@ -568,6 +585,18 @@ void FTimerManager::Tick(float DeltaTime)
 			// Now call the function
 			for (int32 CallIdx=0; CallIdx<CallCount; ++CallIdx)
 			{ 
+#if DO_TIMEGUARD && 0
+				FTimerNameDelegate NameFunction = FTimerNameDelegate::CreateLambda([&] { 
+						return FString::Printf(TEXT("FTimerManager slowtick from delegate %s "), *CurrentlyExecutingTimer.TimerDelegate.ToString());
+					});
+				// no delegate should take longer then 2ms to run 
+				SCOPE_TIME_GUARD_DELEGATE_MS(NameFunction, 2);
+#endif
+#if DO_TIMEGUARD && 0
+				RunTimerDelegates.Add(CurrentlyExecutingTimer.TimerDelegate);
+#endif
+
+
 				CurrentlyExecutingTimer.TimerDelegate.Execute();
 
 				// If timer was cleared in the delegate execution, don't execute further 
diff --git a/Engine/Source/Runtime/Engine/Private/UnrealEngine.cpp b/Engine/Source/Runtime/Engine/Private/UnrealEngine.cpp
--- a/Engine/Source/Runtime/Engine/Private/UnrealEngine.cpp
+++ b/Engine/Source/Runtime/Engine/Private/UnrealEngine.cpp
@@ -9017,7 +9017,7 @@ UWorld* UEngine::GetWorldFromContextObject(const UObject* Object, EGetWorldError
 			check(Object);
 			break;
 		case EGetWorldErrorMode::LogAndReturnNull:
-			FFrame::KismetExecutionMessage(TEXT("A null object was passed as a world context object to UEngine::GetWorldFromContextObject()."), ELogVerbosity::Error);
+			FFrame::KismetExecutionMessage(TEXT("A null object was passed as a world context object to UEngine::GetWorldFromContextObject()."), ELogVerbosity::Warning);
 			//UE_LOG(LogEngine, Warning, TEXT("UEngine::GetWorldFromContextObject() passed a nullptr"));
 			break;
 		case EGetWorldErrorMode::ReturnNull:
@@ -9030,7 +9030,7 @@ UWorld* UEngine::GetWorldFromContextObject(const UObject* Object, EGetWorldError
 	UWorld* World = (ErrorMode == EGetWorldErrorMode::Assert) ? Object->GetWorldChecked(/*out*/ bSupported) : Object->GetWorld();
 	if (bSupported && (World == nullptr) && (ErrorMode == EGetWorldErrorMode::LogAndReturnNull))
 	{
-		FFrame::KismetExecutionMessage(*FString::Printf(TEXT("No world was found for object (%s) passed in to UEngine::GetWorldFromContextObject()."), *GetPathNameSafe(Object)), ELogVerbosity::Error);
+		FFrame::KismetExecutionMessage(*FString::Printf(TEXT("No world was found for object (%s) passed in to UEngine::GetWorldFromContextObject()."), *GetPathNameSafe(Object)), ELogVerbosity::Warning);
 	}
 	return (bSupported ? World : GWorld);
 }
diff --git a/Engine/Source/Runtime/Engine/Private/VectorField.cpp b/Engine/Source/Runtime/Engine/Private/VectorField.cpp
--- a/Engine/Source/Runtime/Engine/Private/VectorField.cpp
+++ b/Engine/Source/Runtime/Engine/Private/VectorField.cpp
@@ -90,7 +90,7 @@ void FVectorFieldInstance::UpdateTransforms(const FMatrix& LocalToWorld)
 	VolumeToWorldNoScale = LocalToWorld.GetMatrixWithoutScale().RemoveTranslation();
 	VolumeToWorld = FScaleMatrix(VolumeScale) * FTranslationMatrix(VolumeOffset)
 		* LocalToWorld;
-	WorldToVolume = VolumeToWorld.InverseFast();
+	WorldToVolume = VolumeToWorld.Inverse();
 }
 
 /*------------------------------------------------------------------------------
@@ -392,9 +392,14 @@ public:
 	Scene proxy for visualizing vector fields.
 ------------------------------------------------------------------------------*/
 
-class FVectorFieldSceneProxy : public FPrimitiveSceneProxy
+class FVectorFieldSceneProxy final : public FPrimitiveSceneProxy
 {
 public:
+	SIZE_T GetTypeHash() const override
+	{
+		static size_t UniquePointer;
+		return reinterpret_cast<size_t>(&UniquePointer);
+	}
 
 	/** Initialization constructor. */
 	explicit FVectorFieldSceneProxy( UVectorFieldComponent* VectorFieldComponent )
diff --git a/Engine/Source/Runtime/Engine/Public/Animation/AnimInstanceProxy.h b/Engine/Source/Runtime/Engine/Public/Animation/AnimInstanceProxy.h
--- a/Engine/Source/Runtime/Engine/Public/Animation/AnimInstanceProxy.h
+++ b/Engine/Source/Runtime/Engine/Public/Animation/AnimInstanceProxy.h
@@ -99,6 +99,20 @@ struct ENGINE_API FAnimInstanceProxy
 
 public:
 	FAnimInstanceProxy()
+		: AnimInstanceObject(nullptr)
+		, AnimClassInterface(nullptr)
+		, Skeleton(nullptr)
+		, SkeletalMeshComponent(nullptr)
+		, CurrentDeltaSeconds(0.0f)
+#if WITH_EDITORONLY_DATA
+		, bIsBeingDebugged(false)
+#endif
+		, RootNode(nullptr)
+		, SubInstanceInputNode(nullptr)
+		, SyncGroupWriteIndex(0)
+		, RootMotionMode(ERootMotionMode::NoRootMotionExtraction)
+		, bShouldExtractRootMotion(false)
+		, bBoneCachesInvalidated(false)
 	{
 	}
 
diff --git a/Engine/Source/Runtime/Engine/Public/Animation/AnimTypes.h b/Engine/Source/Runtime/Engine/Public/Animation/AnimTypes.h
--- a/Engine/Source/Runtime/Engine/Public/Animation/AnimTypes.h
+++ b/Engine/Source/Runtime/Engine/Public/Animation/AnimTypes.h
@@ -587,3 +587,35 @@ namespace EAxisOption
 		Custom
 	};
 }
+
+struct FAxisOption
+{
+	static FVector GetAxisVector(const TEnumAsByte<EAxisOption::Type> InAxis, const FVector& CustomAxis)
+	{
+		switch (InAxis)
+		{
+		case EAxisOption::X:
+			return FVector::ForwardVector;
+		case EAxisOption::X_Neg:
+			return -FVector::ForwardVector;
+		case EAxisOption::Y:
+			return FVector::RightVector;
+		case EAxisOption::Y_Neg:
+			return -FVector::RightVector;
+		case EAxisOption::Z:
+			return FVector::UpVector;
+		case EAxisOption::Z_Neg:
+			return -FVector::UpVector;
+		case EAxisOption::Custom:
+			return CustomAxis;
+		}
+
+		return FVector::ForwardVector;
+	}
+
+	static FVector GetAxisVector(const TEnumAsByte<EAxisOption::Type> InAxis)
+	{
+		return GetAxisVector(InAxis, FVector::ForwardVector);
+	}
+};
+
diff --git a/Engine/Source/Runtime/Engine/Public/AnimationRuntime.h b/Engine/Source/Runtime/Engine/Public/AnimationRuntime.h
--- a/Engine/Source/Runtime/Engine/Public/AnimationRuntime.h
+++ b/Engine/Source/Runtime/Engine/Public/AnimationRuntime.h
@@ -310,20 +310,19 @@ public:
 	static void AccumulateMeshSpaceRotationAdditiveToLocalPose(FCompactPose& BasePose, const FCompactPose& MeshSpaceRotationAdditive, FBlendedCurve& BaseCurve, const FBlendedCurve& AdditiveCurve, float Weight) { AccumulateAdditivePose(BasePose, MeshSpaceRotationAdditive, BaseCurve, AdditiveCurve, Weight, EAdditiveAnimationType::AAT_RotationOffsetMeshSpace); }
 
 
-	/** Lerp for BoneTransforms. Stores results in A. Performs A = Lerp(A, B, Alpha);
-	* @param A : In/Out FCompactPose.
-	* @param B : B In FCompactPose.
-	* @param Alpha : Alpha.
-	*/
+	/** Lerp for FCompactPose. Stores results in PoseA. Performs PoseA = Lerp(PoseA, PoseB, Alpha); */
 	static void LerpPoses(FCompactPose& PoseA, const FCompactPose& PoseB, FBlendedCurve& CurveA, const FBlendedCurve& CurveB, float Alpha);
 
-	/** Lerp for BoneTransforms. Stores results in A. Performs A = Lerp(A, B, Alpha);
-	* @param A : In/Out FCompactPose.
-	* @param B : B In FCompactPose.
-	* @param Alpha : Alpha.
-	*/
+	/** Lerp for FCompactPose. Stores results in PoseA. 
+	 * For each bone performs BoneA[i] = Lerp(BoneA[i], BoneB[i], Alpha * PerBoneWeights[i]);
+	 */
 	static void LerpPosesPerBone(FCompactPose& PoseA, const FCompactPose& PoseB, FBlendedCurve& CurveA, const FBlendedCurve& CurveB, float Alpha, const TArray<float>& PerBoneWeights);
 
+	/** Lerp for FCompactPose. Stores results in PoseA. Performs PoseA = Lerp(PoseA, PoseB, Alpha);
+	 * on reduced set of bones defined in BoneIndices list.
+	 */
+	static void LerpPosesWithBoneIndexList(FCompactPose& PoseA, const FCompactPose& PoseB, FBlendedCurve& CurveA, const FBlendedCurve& CurveB, float Alpha, const TArray<FCompactPoseBoneIndex>& BoneIndices);
+
 	/** Lerp for BoneTransforms. Stores results in A. Performs A = Lerp(A, B, Alpha);
 	 * @param A : In/Out transform array.
 	 * @param B : B In transform array.
diff --git a/Engine/Source/Runtime/Engine/Public/AudioStreaming.h b/Engine/Source/Runtime/Engine/Public/AudioStreaming.h
--- a/Engine/Source/Runtime/Engine/Public/AudioStreaming.h
+++ b/Engine/Source/Runtime/Engine/Public/AudioStreaming.h
@@ -294,7 +294,7 @@ protected:
 	TMap<USoundWave*, FWaveRequest> WaveRequests;
 
 	/** Results of async loading audio chunks. */
-	TQueue<FASyncAudioChunkLoadResult*> AsyncAudioStreamChunkResults;
+	TQueue<FASyncAudioChunkLoadResult*, EQueueMode::Mpsc> AsyncAudioStreamChunkResults;
 
 	/** Critical section to protect usage of shared gamethread/audiothread members */
 	mutable FCriticalSection CriticalSection;
diff --git a/Engine/Source/Runtime/Engine/Public/BonePose.h b/Engine/Source/Runtime/Engine/Public/BonePose.h
--- a/Engine/Source/Runtime/Engine/Public/BonePose.h
+++ b/Engine/Source/Runtime/Engine/Public/BonePose.h
@@ -161,6 +161,12 @@ public:
 		this->InitBones(BoneContainer->GetBoneIndicesArray().Num());
 	}
 
+	void CopyAndAssignBoneContainer(FBoneContainer& NewBoneContainer)
+	{
+		NewBoneContainer = *BoneContainer;
+		BoneContainer = &NewBoneContainer;
+	}
+
 	void InitFrom(const FBaseCompactPose& SrcPose)
 	{
 		SetBoneContainer(SrcPose.BoneContainer);
@@ -205,6 +211,12 @@ public:
 		DestPoseBones = this->Bones;
 	}
 
+	void Empty()
+	{
+		BoneContainer = nullptr;
+		this->Bones.Empty();
+	}
+
 	// Sets this pose to its ref pose
 	void ResetToRefPose()
 	{
@@ -450,6 +462,19 @@ struct FCSPose
 		ComponentSpaceFlags = SrcPose.GetComponentSpaceFlags();
 	}
 
+	void CopyAndAssignBoneContainer(FBoneContainer& NewBoneContainer)
+	{
+		Pose.CopyAndAssignBoneContainer(NewBoneContainer);
+	}
+
+	void Empty()
+	{
+		Pose.Empty();
+		ComponentSpaceFlags.Empty();
+		BonesToConvert.Empty();
+		BoneMask.Empty();
+	}
+
 	const PoseType& GetPose() const { return Pose; }
 	const TCustomBoneIndexArray<uint8, BoneIndexType>& GetComponentSpaceFlags() const { return ComponentSpaceFlags; }
 
diff --git a/Engine/Source/Runtime/Engine/Public/DebugRenderSceneProxy.h b/Engine/Source/Runtime/Engine/Public/DebugRenderSceneProxy.h
--- a/Engine/Source/Runtime/Engine/Public/DebugRenderSceneProxy.h
+++ b/Engine/Source/Runtime/Engine/Public/DebugRenderSceneProxy.h
@@ -22,6 +22,8 @@ DECLARE_DELEGATE_TwoParams(FDebugDrawDelegate, class UCanvas*, class APlayerCont
 class FDebugRenderSceneProxy : public FPrimitiveSceneProxy
 {
 public:
+	ENGINE_API SIZE_T GetTypeHash() const override;
+
 	enum EDrawType
 	{
 		SolidMesh = 0,
diff --git a/Engine/Source/Runtime/Engine/Public/MeshBatch.h b/Engine/Source/Runtime/Engine/Public/MeshBatch.h
--- a/Engine/Source/Runtime/Engine/Public/MeshBatch.h
+++ b/Engine/Source/Runtime/Engine/Public/MeshBatch.h
@@ -61,6 +61,7 @@ struct FMeshBatchElement
 	/** Conceptual element index used for debug viewmodes. */
 	int8 VisualizeElementIndex;
 #endif
+	FVertexBufferRHIParamRef IndirectArgsBuffer;
 
 	FMeshBatchElement()
 	:	PrimitiveUniformBufferResource(nullptr)
@@ -81,7 +82,7 @@ struct FMeshBatchElement
 #if !(UE_BUILD_SHIPPING || UE_BUILD_TEST)
 	,	VisualizeElementIndex(INDEX_NONE)
 #endif
-
+	,	IndirectArgsBuffer(nullptr)
 	{
 	}
 };
