Fix for -

- Check CSProj Target Versions warning
- VS2017 compile issue in Niagara
- Non-unity Win64 / Mac UE4Game errors
- Build DDC Mac asset error
- "elements that reference missing vertices" warning during tests

diff --git a/Engine/Plugins/FX/Niagara/Source/NiagaraEditor/Public/NiagaraNodeCustomHlsl.h b/Engine/Plugins/FX/Niagara/Source/NiagaraEditor/Public/NiagaraNodeCustomHlsl.h
--- a/Engine/Plugins/FX/Niagara/Source/NiagaraEditor/Public/NiagaraNodeCustomHlsl.h
+++ b/Engine/Plugins/FX/Niagara/Source/NiagaraEditor/Public/NiagaraNodeCustomHlsl.h
@@ -25,7 +25,7 @@ public:
 
 	virtual TSharedPtr<SGraphNode> CreateVisualWidget() override;
 	virtual void OnRenameNode(const FString& NewName) override;
-	virtual FLinearColor UNiagaraNodeCustomHlsl::GetNodeTitleColor() const override;
+	virtual FLinearColor GetNodeTitleColor() const override;
 
 	FText GetHlslText() const;
 	void OnCustomHlslTextCommitted(const FText& InText, ETextCommit::Type InType);
diff --git a/Engine/Source/Runtime/Engine/Private/ModelRender.cpp b/Engine/Source/Runtime/Engine/Private/ModelRender.cpp
--- a/Engine/Source/Runtime/Engine/Private/ModelRender.cpp
+++ b/Engine/Source/Runtime/Engine/Private/ModelRender.cpp
@@ -156,7 +156,7 @@ void UModelComponent::BuildRenderData()
 
 						if (Element.MaxVertexIndex > (uint32)TheModel->VertexBuffer.Vertices.Num())
 						{
-							UE_LOG(LogModelComponent, Warning, TEXT("Model %s has elements that reference missing vertices. MaxVertex=%d, NumVertices=%d"),
+							UE_LOG(LogModelComponent, Log, TEXT("Model %s has elements that reference missing vertices. MaxVertex=%d, NumVertices=%d"),
 								*GetPathName(), Element.MaxVertexIndex, TheModel->VertexBuffer.Vertices.Num());
 						}
 					}
diff --git a/Engine/Source/Runtime/GameplayTags/Private/GameplayTagsManager.cpp b/Engine/Source/Runtime/GameplayTags/Private/GameplayTagsManager.cpp
--- a/Engine/Source/Runtime/GameplayTags/Private/GameplayTagsManager.cpp
+++ b/Engine/Source/Runtime/GameplayTags/Private/GameplayTagsManager.cpp
@@ -1,6 +1,8 @@
 // Copyright 1998-2017 Epic Games, Inc. All Rights Reserved.
 
+
 #include "GameplayTagsManager.h"
+#include "Engine/Engine.h"
 #include "HAL/PlatformFilemanager.h"
 #include "HAL/FileManager.h"
 #include "Misc/Paths.h"
