Copying WEX-Staging @ (WEX/Main @ 3740665) to //UE4/Main

diff --git a/Engine/Source/Editor/LevelEditor/LevelEditor.Build.cs b/Engine/Source/Editor/LevelEditor/LevelEditor.Build.cs
--- a/Engine/Source/Editor/LevelEditor/LevelEditor.Build.cs
+++ b/Engine/Source/Editor/LevelEditor/LevelEditor.Build.cs
@@ -12,8 +12,6 @@ public class LevelEditor : ModuleRules
 				"ClassViewer",
 				"MainFrame",
                 "PlacementMode",
-				"ReferenceViewer",
-				"SizeMap",
 				"SlateReflector",
                 "IntroTutorials",
                 "AppFramework",
@@ -87,8 +85,6 @@ public class LevelEditor : ModuleRules
 				"NewLevelDialog",
 				"DeviceProfileEditor",
                 "PlacementMode",
-				"ReferenceViewer",
-				"SizeMap",
                 "IntroTutorials",
 				"HeadMountedDisplay",
 				"VREditor"
diff --git a/Engine/Source/Editor/LevelEditor/Private/LevelEditor.cpp b/Engine/Source/Editor/LevelEditor/Private/LevelEditor.cpp
--- a/Engine/Source/Editor/LevelEditor/Private/LevelEditor.cpp
+++ b/Engine/Source/Editor/LevelEditor/Private/LevelEditor.cpp
@@ -648,18 +648,6 @@ void FLevelEditorModule::BindGlobalLevelEditorCommands()
 		FExecuteAction::CreateStatic( &FLevelEditorActionCallbacks::FindInContentBrowser_Clicked )
 		);
 
-	ActionList.MapAction( 
-		FGlobalEditorCommonCommands::Get().ViewReferences, 
-		FExecuteAction::CreateStatic( &FLevelEditorActionCallbacks::ViewReferences_Execute ),
-		FCanExecuteAction::CreateStatic( &FLevelEditorActionCallbacks::CanViewReferences )
-		);
-
-	ActionList.MapAction( 
-		FGlobalEditorCommonCommands::Get().ViewSizeMap, 
-		FExecuteAction::CreateStatic( &FLevelEditorActionCallbacks::ViewSizeMap_Execute ),
-		FCanExecuteAction::CreateStatic( &FLevelEditorActionCallbacks::CanViewSizeMap )
-		);
-
 	const FVector* NullVector = nullptr;
 	ActionList.MapAction(
 		Commands.GoHere,
diff --git a/Engine/Source/Editor/LevelEditor/Private/LevelEditorActions.cpp b/Engine/Source/Editor/LevelEditor/Private/LevelEditorActions.cpp
--- a/Engine/Source/Editor/LevelEditor/Private/LevelEditorActions.cpp
+++ b/Engine/Source/Editor/LevelEditor/Private/LevelEditorActions.cpp
@@ -62,8 +62,6 @@
 #include "SourceCodeNavigation.h"
 #include "EngineAnalytics.h"
 #include "Interfaces/IAnalyticsProvider.h"
-#include "ReferenceViewer.h"
-#include "ISizeMapModule.h"
 #include "EditorClassUtils.h"
 
 #include "EditorActorFolders.h"
@@ -1107,70 +1105,6 @@ void FLevelEditorActionCallbacks::FindInContentBrowser_Clicked()
 	GEditor->SyncToContentBrowser();
 }
 
-void FLevelEditorActionCallbacks::ViewReferences_Execute()
-{
-	if( GEditor->GetSelectedActorCount() > 0 )
-	{
-		TArray< UObject* > ReferencedAssets;
-		GEditor->GetReferencedAssetsForEditorSelection( ReferencedAssets );
-
-		if (ReferencedAssets.Num() > 0)
-		{
-			TArray< FName > ViewableObjects;
-			for( auto ObjectIter = ReferencedAssets.CreateConstIterator(); ObjectIter; ++ObjectIter )
-			{
-				// Don't allow user to perform certain actions on objects that aren't actually assets (e.g. Level Script blueprint objects)
-				const auto EditingObject = *ObjectIter;
-				if( EditingObject != NULL && EditingObject->IsAsset() )
-				{
-					ViewableObjects.Add( EditingObject->GetOuter()->GetFName());
-				}
-			}
-
-			IReferenceViewerModule::Get().InvokeReferenceViewerTab(ViewableObjects);
-		}
-	}
-}
-
-bool FLevelEditorActionCallbacks::CanViewReferences()
-{
-	TArray< UObject* > ReferencedAssets;
-	GEditor->GetReferencedAssetsForEditorSelection(ReferencedAssets);
-	return ReferencedAssets.Num() > 0;
-}
-
-void FLevelEditorActionCallbacks::ViewSizeMap_Execute()
-{
-	if( GEditor->GetSelectedActorCount() > 0 )
-	{
-		TArray< UObject* > ReferencedAssets;
-		GEditor->GetReferencedAssetsForEditorSelection( ReferencedAssets );
-
-		if (ReferencedAssets.Num() > 0)
-		{
-			TArray< FName > ViewableObjects;
-			for( auto ObjectIter = ReferencedAssets.CreateConstIterator(); ObjectIter; ++ObjectIter )
-			{
-				// Don't allow user to perform certain actions on objects that aren't actually assets (e.g. Level Script blueprint objects)
-				const auto EditingObject = *ObjectIter;
-				if( EditingObject != NULL && EditingObject->IsAsset() )
-				{
-					ViewableObjects.Add( EditingObject->GetOuter()->GetFName());
-				}
-			}
-
-			ISizeMapModule::Get().InvokeSizeMapTab(ViewableObjects);
-		}
-	}
-}
-
-bool FLevelEditorActionCallbacks::CanViewSizeMap()
-{
-	TArray< UObject* > ReferencedAssets;
-	GEditor->GetReferencedAssetsForEditorSelection(ReferencedAssets);
-	return ReferencedAssets.Num() > 0;
-}
-
 void FLevelEditorActionCallbacks::EditAsset_Clicked( const EToolkitMode::Type ToolkitMode, TWeakPtr< SLevelEditor > LevelEditor, bool bConfirmMultiple )
 {
 	if( GEditor->GetSelectedActorCount() > 0 )
diff --git a/Engine/Source/Editor/LevelEditor/Private/LevelEditorContextMenu.cpp b/Engine/Source/Editor/LevelEditor/Private/LevelEditorContextMenu.cpp
--- a/Engine/Source/Editor/LevelEditor/Private/LevelEditorContextMenu.cpp
+++ b/Engine/Source/Editor/LevelEditor/Private/LevelEditorContextMenu.cpp
@@ -264,10 +264,6 @@ void FLevelEditorContextMenu::FillMenu( FMenuBuilder& MenuBuilder, TWeakPtr<SLev
 						);
 
 				}
-
-				MenuBuilder.AddMenuEntry(FGlobalEditorCommonCommands::Get().ViewReferences);
-				MenuBuilder.AddMenuEntry(FGlobalEditorCommonCommands::Get().ViewSizeMap);
-
 			}
 			MenuBuilder.EndSection();
 		}
diff --git a/Engine/Source/Editor/LevelEditor/Public/LevelEditorActions.h b/Engine/Source/Editor/LevelEditor/Public/LevelEditorActions.h
--- a/Engine/Source/Editor/LevelEditor/Public/LevelEditorActions.h
+++ b/Engine/Source/Editor/LevelEditor/Public/LevelEditorActions.h
@@ -834,18 +834,6 @@ public:
 	 */
 	static void FindInContentBrowser_Clicked();
 
-	/** Called when the ViewReferences command is executed */
-	static void ViewReferences_Execute();
-
-	/** If true ViewReferences_Execute can be called */
-	static bool CanViewReferences();
-
-	/** Called when the ViewSizeMap command is executed */
-	static void ViewSizeMap_Execute();
-
-	/** If true ViewSizeMap_Execute can be called */
-	static bool CanViewSizeMap();
-
 	/** Called to when "Edit Asset" is clicked */
 	static void EditAsset_Clicked( const EToolkitMode::Type ToolkitMode, TWeakPtr< class SLevelEditor > LevelEditor, bool bAskMultiple );
 
diff --git a/Engine/Source/Editor/UMGEditor/Private/WidgetBlueprintCompiler.cpp b/Engine/Source/Editor/UMGEditor/Private/WidgetBlueprintCompiler.cpp
--- a/Engine/Source/Editor/UMGEditor/Private/WidgetBlueprintCompiler.cpp
+++ b/Engine/Source/Editor/UMGEditor/Private/WidgetBlueprintCompiler.cpp
@@ -152,15 +152,15 @@ void FWidgetBlueprintCompiler::CleanAndSanitizeClass(UBlueprintGeneratedClass* C
 {
 	UWidgetBlueprint* WidgetBP = WidgetBlueprint();
 
+	const bool bRecompilingOnLoad = Blueprint->bIsRegeneratingOnLoad;
+	const ERenameFlags RenFlags = REN_DontCreateRedirectors | (bRecompilingOnLoad ? REN_ForceNoResetLoaders : 0) | REN_NonTransactional | REN_DoNotDirty;
+
 	if ( !Blueprint->bIsRegeneratingOnLoad && bIsFullCompile )
 	{
 		UPackage* WidgetTemplatePackage = WidgetBP->GetWidgetTemplatePackage();
 		UUserWidget* OldArchetype = FindObjectFast<UUserWidget>(WidgetTemplatePackage, TEXT("WidgetArchetype"));
-		if ( OldArchetype )
+		if (OldArchetype)
 		{
-			const bool bRecompilingOnLoad = Blueprint->bIsRegeneratingOnLoad;
-			const ERenameFlags RenFlags = REN_DontCreateRedirectors | ( bRecompilingOnLoad ? REN_ForceNoResetLoaders : 0 ) | REN_NonTransactional | REN_DoNotDirty;
-
 			FString TransientArchetypeString = FString::Printf(TEXT("OLD_TEMPLATE_%s"), *OldArchetype->GetName());
 			FName TransientArchetypeName = MakeUniqueObjectName(GetTransientPackage(), OldArchetype->GetClass(), FName(*TransientArchetypeString));
 			OldArchetype->Rename(*TransientArchetypeName.ToString(), GetTransientPackage(), RenFlags);
@@ -187,9 +187,9 @@ void FWidgetBlueprintCompiler::CleanAndSanitizeClass(UBlueprintGeneratedClass* C
 	// Make sure our typed pointer is set
 	check(ClassToClean == NewClass && NewWidgetBlueprintClass == NewClass);
 
-	for ( UWidgetAnimation* Animation : NewWidgetBlueprintClass->Animations )
+	for (UWidgetAnimation* Animation : NewWidgetBlueprintClass->Animations)
 	{
-		Animation->Rename(nullptr, GetTransientPackage(), REN_ForceNoResetLoaders | REN_DontCreateRedirectors);
+		Animation->Rename(nullptr, GetTransientPackage(), RenFlags);
 	}
 	NewWidgetBlueprintClass->Animations.Empty();
 
@@ -198,6 +198,8 @@ void FWidgetBlueprintCompiler::CleanAndSanitizeClass(UBlueprintGeneratedClass* C
 
 void FWidgetBlueprintCompiler::SaveSubObjectsFromCleanAndSanitizeClass(FSubobjectCollection& SubObjectsToSave, UBlueprintGeneratedClass* ClassToClean)
 {
+	Super::SaveSubObjectsFromCleanAndSanitizeClass(SubObjectsToSave, ClassToClean);
+
 	// Make sure our typed pointer is set
 	check(ClassToClean == NewClass);
 	NewWidgetBlueprintClass = CastChecked<UWidgetBlueprintGeneratedClass>((UObject*)NewClass);
@@ -207,14 +209,6 @@ void FWidgetBlueprintCompiler::SaveSubObjectsFromCleanAndSanitizeClass(FSubobjec
 	// We need to save the widget tree to survive the initial sub-object clean blitz, 
 	// otherwise they all get renamed, and it causes early loading errors.
 	SubObjectsToSave.AddObject(WidgetBP->WidgetTree);
-
-	// We need to save all the animations to survive the initial sub-object clean blitz, 
-	// otherwise they all get renamed, and it causes early loading errors.
-	SubObjectsToSave.AddObject(NewWidgetBlueprintClass->WidgetTree);
-	for ( UWidgetAnimation* Animation : NewWidgetBlueprintClass->Animations )
-	{
-		SubObjectsToSave.AddObject(Animation);
-	}
 }
 
 void FWidgetBlueprintCompiler::CreateClassVariablesFromBlueprint()
diff --git a/Engine/Source/Editor/UnrealEd/Private/Commandlets/AssetRegistryGenerator.cpp b/Engine/Source/Editor/UnrealEd/Private/Commandlets/AssetRegistryGenerator.cpp
--- a/Engine/Source/Editor/UnrealEd/Private/Commandlets/AssetRegistryGenerator.cpp
+++ b/Engine/Source/Editor/UnrealEd/Private/Commandlets/AssetRegistryGenerator.cpp
@@ -414,7 +414,8 @@ bool FAssetRegistryGenerator::SaveManifests(FSandboxPlatformFile* InSandboxFile)
 
 		if (!bUseAssetManager)
 		{
-			GenerateAssetChunkInformationCSV(FPaths::Combine(*FPaths::ProjectLogDir(), TEXT("ChunkLists")));
+			// In new flow, this is written later
+			GenerateAssetChunkInformationCSV(FPaths::Combine(*FPaths::ProjectLogDir(), TEXT("ChunkLists")), true);
 		}
 	}
 
@@ -701,7 +702,7 @@ void FAssetRegistryGenerator::AddAssetToFileOrderRecursive(FAssetData* InAsset,
 	}
 }
 
-bool FAssetRegistryGenerator::SaveAssetRegistry(const FString& SandboxPath, bool bSerializeDevelopmentAssetRegistry )
+bool FAssetRegistryGenerator::SaveAssetRegistry(const FString& SandboxPath, bool bSerializeDevelopmentAssetRegistry)
 {
 	UE_LOG(LogAssetRegistryGenerator, Display, TEXT("Saving asset registry."));
 	const TMap<FName, const FAssetData*>& ObjectToDataMap = State.GetObjectPathToAssetDataMap();
@@ -731,6 +732,14 @@ bool FAssetRegistryGenerator::SaveAssetRegistry(const FString& SandboxPath, bool
 		FString PlatformSandboxPath = SandboxPath.Replace(TEXT("[Platform]"), *TargetPlatform->PlatformName());
 		PlatformSandboxPath.ReplaceInline(TEXT("AssetRegistry.bin"), TEXT("DevelopmentAssetRegistry.bin"));
 		FFileHelper::SaveArrayToFile(SerializedAssetRegistry, *PlatformSandboxPath);
+
+		if (bGenerateChunks && bUseAssetManager)
+		{
+			FString ChunkListsPath = PlatformSandboxPath.Replace(TEXT("/DevelopmentAssetRegistry.bin"), TEXT(""));
+
+			// Write out CSV file with chunking information
+			GenerateAssetChunkInformationCSV(ChunkListsPath, false);
+		}
 	}
 
 	if (SaveOptions.bSerializeAssetRegistry)
@@ -959,7 +968,7 @@ bool FAssetRegistryGenerator::GatherAllPackageDependencies(FName PackageName, TA
 	return true;
 }
 
-bool FAssetRegistryGenerator::GenerateAssetChunkInformationCSV(const FString& OutputPath)
+bool FAssetRegistryGenerator::GenerateAssetChunkInformationCSV(const FString& OutputPath, bool bWriteIndividualFiles)
 {
 	FString TmpString;
 	FString CSVString;
@@ -969,29 +978,31 @@ bool FAssetRegistryGenerator::GenerateAssetChunkInformationCSV(const FString& Ou
 	CSVString = HeaderText;
 
 	const TMap<FName, const FAssetData*>& ObjectToDataMap = State.GetObjectPathToAssetDataMap();
+	TArray<const FAssetData*> AssetDataList;
+	for (const TPair<FName, const FAssetData*> Pair : ObjectToDataMap)
+	{
+		AssetDataList.Add(Pair.Value);
+	}
+
+	// Sort list so it's consistent over time
+	AssetDataList.Sort([](const FAssetData& A, const FAssetData& B)
+	{
+		return A.ObjectPath < B.ObjectPath;
+	});
+
 	for (int32 ChunkID = 0, ChunkNum = FinalChunkManifests.Num(); ChunkID < ChunkNum; ++ChunkID)
 	{
 		FString PerChunkManifestCSV = HeaderText;
-		for (const TPair<FName, const FAssetData*> Pair : ObjectToDataMap)
+		for (const FAssetData* AssetDataPtr : AssetDataList)
 		{
-			const FAssetData& AssetData = *Pair.Value;
+			const FAssetData& AssetData = *AssetDataPtr;
 			// Add only assets that have actually been cooked and belong to any chunk
 			if (AssetData.ChunkIDs.Num() > 0)
 			{
-				FString Fullname;
-				if (AssetData.ChunkIDs.Contains(ChunkID) && FPackageName::DoesPackageExist(*AssetData.PackageName.ToString(), nullptr, &Fullname))
+				const FAssetPackageData* PackageData = State.GetAssetPackageData(AssetData.PackageName);
+				if (AssetData.ChunkIDs.Contains(ChunkID) && PackageData && PackageData->DiskSize >= 0)
 				{
-					auto FileSize = IFileManager::Get().FileSize(*FPackageName::LongPackageNameToFilename(*AssetData.PackageName.ToString(), FPackageName::GetAssetPackageExtension()));
-					if (FileSize == INDEX_NONE)
-					{
-						FileSize = IFileManager::Get().FileSize(*FPackageName::LongPackageNameToFilename(*AssetData.PackageName.ToString(), FPackageName::GetMapPackageExtension()));
-					}
-
-					if (FileSize == INDEX_NONE)
-					{
-						FileSize = 0;
-					}
-
+					int64 FileSize = PackageData->DiskSize;
 					FString SoftChain;
 					bool bHardChunk = false;
 					if (ChunkID < ChunkManifests.Num())
@@ -1035,7 +1046,10 @@ bool FAssetRegistryGenerator::GenerateAssetChunkInformationCSV(const FString& Ou
 			}
 		}
 
-		FFileHelper::SaveStringToFile(PerChunkManifestCSV, *FPaths::Combine(*OutputPath, *FString::Printf(TEXT("Chunks%dInfo.csv"), ChunkID)));
+		if (bWriteIndividualFiles)
+		{
+			FFileHelper::SaveStringToFile(PerChunkManifestCSV, *FPaths::Combine(*OutputPath, *FString::Printf(TEXT("Chunks%dInfo.csv"), ChunkID)));
+		}
 	}
 
 	return FFileHelper::SaveStringToFile(CSVString, *FPaths::Combine(*OutputPath, TEXT("AllChunksInfo.csv")));
diff --git a/Engine/Source/Editor/UnrealEd/Private/EditorEngine.cpp b/Engine/Source/Editor/UnrealEd/Private/EditorEngine.cpp
--- a/Engine/Source/Editor/UnrealEd/Private/EditorEngine.cpp
+++ b/Engine/Source/Editor/UnrealEd/Private/EditorEngine.cpp
@@ -928,9 +928,6 @@ void UEditorEngine::Init(IEngineLoop* InEngineLoop)
 			TEXT("BehaviorTreeEditor"),
 			TEXT("HardwareTargeting"),
 			TEXT("LocalizationDashboard"),
-			TEXT("ReferenceViewer"),
-			TEXT("TreeMap"),
-			TEXT("SizeMap"),
 			TEXT("MergeActors"),
 			TEXT("InputBindingEditor"),
 			TEXT("AudioEditor")
diff --git a/Engine/Source/Editor/UnrealEd/Private/Toolkits/AssetEditorToolkit.cpp b/Engine/Source/Editor/UnrealEd/Private/Toolkits/AssetEditorToolkit.cpp
--- a/Engine/Source/Editor/UnrealEd/Private/Toolkits/AssetEditorToolkit.cpp
+++ b/Engine/Source/Editor/UnrealEd/Private/Toolkits/AssetEditorToolkit.cpp
@@ -22,8 +22,6 @@
 #include "CollectionManagerModule.h"
 #include "Widgets/SToolTip.h"
 #include "IDocumentation.h"
-#include "ReferenceViewer.h"
-#include "ISizeMapModule.h"
 #include "IIntroTutorials.h"
 #include "Widgets/Docking/SDockTab.h"
 #include "IAssetTools.h"
@@ -33,7 +31,8 @@
 #define LOCTEXT_NAMESPACE "AssetEditorToolkit"
 
 TWeakPtr< IToolkitHost > FAssetEditorToolkit::PreviousWorldCentricToolkitHostForNewAssetEditor;
-
+TSharedPtr<FExtensibilityManager> FAssetEditorToolkit::SharedMenuExtensibilityManager;
+TSharedPtr<FExtensibilityManager> FAssetEditorToolkit::SharedToolBarExtensibilityManager;
 
 const FName FAssetEditorToolkit::ToolbarTabId( TEXT( "AssetEditorToolkit_Toolbar" ) );
 
@@ -46,7 +45,6 @@ FAssetEditorToolkit::FAssetEditorToolkit()
 	WorkspaceMenuCategory = FWorkspaceItem::NewGroup(LOCTEXT("WorkspaceMenu_BaseAssetEditor", "Asset Editor"));
 }
 
-
 void FAssetEditorToolkit::InitAssetEditor( const EToolkitMode::Type Mode, const TSharedPtr< class IToolkitHost >& InitToolkitHost, const FName AppIdentifier, const TSharedRef<FTabManager::FLayout>& StandaloneDefaultLayout, const bool bCreateDefaultStandaloneMenu, const bool bCreateDefaultToolbar, UObject* ObjectToEdit, const bool bInIsToolbarFocusable )
 {
 	TArray< UObject* > ObjectsToEdit;
@@ -196,7 +194,6 @@ void FAssetEditorToolkit::InitAssetEditor( const EToolkitMode::Type Mode, const
 
 		NewMajorTab->SetContent
 		( 
-
 			SAssignNew( NewStandaloneHost, SStandaloneAssetEditorToolkitHost, NewTabManager, AppIdentifier )
 			.OnRequestClose(this, &FAssetEditorToolkit::OnRequestClose)
 		);
@@ -204,31 +201,13 @@ void FAssetEditorToolkit::InitAssetEditor( const EToolkitMode::Type Mode, const
 		// Assign our toolkit host before we setup initial content.  (Important: We must cache this pointer here as SetupInitialContent
 		// will callback into the toolkit host.)
 		ToolkitHost = NewStandaloneHost;
-	}
 
+		StandaloneHost = NewStandaloneHost;
+	}
 
 	check( ToolkitHost.IsValid() );
 	ToolkitManager.RegisterNewToolkit( SharedThis( this ) );
 	
-	if (ToolkitMode == EToolkitMode::Standalone)
-	{
-		TSharedRef<FTabManager::FLayout> LayoutToUse = FLayoutSaveRestore::LoadFromConfig(GEditorLayoutIni, StandaloneDefaultLayout);
-
-		// Actually create the widget content
-		NewStandaloneHost->SetupInitialContent( LayoutToUse, NewMajorTab, bCreateDefaultStandaloneMenu );
-	}
-	StandaloneHost = NewStandaloneHost;
-	
-
-	if (bCreateDefaultToolbar)
-	{
-		GenerateToolbar();
-	}
-	else
-	{
-		Toolbar = SNullWidget::NullWidget;
-	}
-
 	ToolkitCommands->MapAction(
 		FAssetEditorCommonCommands::Get().SaveAsset,
 		FExecuteAction::CreateSP( this, &FAssetEditorToolkit::SaveAsset_Execute ),
@@ -243,17 +222,7 @@ void FAssetEditorToolkit::InitAssetEditor( const EToolkitMode::Type Mode, const
 		FGlobalEditorCommonCommands::Get().FindInContentBrowser,
 		FExecuteAction::CreateSP( this, &FAssetEditorToolkit::FindInContentBrowser_Execute ),
 		FCanExecuteAction::CreateSP( this, &FAssetEditorToolkit::CanFindInContentBrowser ));
-	
-	ToolkitCommands->MapAction(
-		FGlobalEditorCommonCommands::Get().ViewReferences,
-		FExecuteAction::CreateSP( this, &FAssetEditorToolkit::ViewReferences_Execute ),
-		FCanExecuteAction::CreateSP( this, &FAssetEditorToolkit::CanViewReferences ));
-	
-	ToolkitCommands->MapAction(
-		FGlobalEditorCommonCommands::Get().ViewSizeMap,
-		FExecuteAction::CreateSP( this, &FAssetEditorToolkit::ViewSizeMap_Execute ),
-		FCanExecuteAction::CreateSP( this, &FAssetEditorToolkit::CanViewSizeMap ));
-	
+		
 	ToolkitCommands->MapAction(
 		FGlobalEditorCommonCommands::Get().OpenDocumentation,
 		FExecuteAction::CreateSP( this, &FAssetEditorToolkit::BrowseDocumentation_Execute ) );
@@ -280,6 +249,29 @@ void FAssetEditorToolkit::InitAssetEditor( const EToolkitMode::Type Mode, const
 		}
 	}
 
+	// Create menus
+	if (ToolkitMode == EToolkitMode::Standalone)
+	{
+		AddMenuExtender(GetSharedMenuExtensibilityManager()->GetAllExtenders(ToolkitCommands, EditingObjects));
+
+		TSharedRef<FTabManager::FLayout> LayoutToUse = FLayoutSaveRestore::LoadFromConfig(GEditorLayoutIni, StandaloneDefaultLayout);
+
+		// Actually create the widget content
+		NewStandaloneHost->SetupInitialContent(LayoutToUse, NewMajorTab, bCreateDefaultStandaloneMenu);
+	}
+	
+	// Create toolbars
+	AddToolbarExtender(GetSharedToolBarExtensibilityManager()->GetAllExtenders(ToolkitCommands, EditingObjects));
+
+	if (bCreateDefaultToolbar)
+	{
+		GenerateToolbar();
+	}
+	else
+	{
+		Toolbar = SNullWidget::NullWidget;
+	}
+
 	// NOTE: Currently, the AssetEditorManager will keep a hard reference to our object as we're editing it
 	FAssetEditorManager::Get().NotifyAssetsOpened( EditingObjects, this );
 }
@@ -788,52 +780,6 @@ void FAssetEditorToolkit::BrowseDocumentation_Execute() const
 	IDocumentation::Get()->Open(GetDocumentationLink(), FDocumentationSourceInfo(TEXT("help_menu_asset")));
 }
 
-void FAssetEditorToolkit::ViewReferences_Execute()
-{
-	if (ensure( ViewableObjects.Num() > 0))
-	{
-		IReferenceViewerModule::Get().InvokeReferenceViewerTab(ViewableObjects);
-	}
-}
-
-bool FAssetEditorToolkit::CanViewReferences()
-{
-	ViewableObjects.Empty();
-	for (const auto EditingObject : EditingObjects)
-	{
-		// Don't allow user to perform certain actions on objects that aren't actually assets (e.g. Level Script blueprint objects)
-		if (EditingObject != NULL && EditingObject->IsAsset())
-		{
-			ViewableObjects.Add(EditingObject->GetOuter()->GetFName());
-		}
-	}
-
-	return ViewableObjects.Num() > 0;
-}
-
-void FAssetEditorToolkit::ViewSizeMap_Execute()
-{
-	if (ensure( ViewableObjects.Num() > 0))
-	{
-		ISizeMapModule::Get().InvokeSizeMapTab(ViewableObjects);
-	}
-}
-
-bool FAssetEditorToolkit::CanViewSizeMap()
-{
-	ViewableObjects.Empty();
-	for (const auto EditingObject : EditingObjects)
-	{
-		// Don't allow user to perform certain actions on objects that aren't actually assets (e.g. Level Script blueprint objects)
-		if (EditingObject != NULL && EditingObject->IsAsset())
-		{
-			ViewableObjects.Add(EditingObject->GetOuter()->GetFName());
-		}
-	}
-
-	return ViewableObjects.Num() > 0;
-}
-
 FString FAssetEditorToolkit::GetDocumentationLink() const
 {
 	return FString(TEXT("%ROOT%"));
@@ -955,9 +901,6 @@ void FAssetEditorToolkit::FillDefaultAssetMenuCommands( FMenuBuilder& MenuBuilde
 	// Commands we only want to be accessible when editing an asset should go here 
 	if( IsActuallyAnAsset() )
 	{
-		MenuBuilder.AddMenuEntry( FGlobalEditorCommonCommands::Get().ViewReferences);
-		MenuBuilder.AddMenuEntry( FGlobalEditorCommonCommands::Get().ViewSizeMap);
-
 		// Add a reimport menu entry for each supported editable object
 		for( auto ObjectIter = EditingObjects.CreateConstIterator(); ObjectIter; ++ObjectIter )
 		{
@@ -1125,6 +1068,24 @@ void FAssetEditorToolkit::RemoveToolbarExtender(TSharedPtr<FExtender> Extender)
 	ToolbarExtenders.Remove(Extender);
 }
 
+TSharedPtr<FExtensibilityManager> FAssetEditorToolkit::GetSharedMenuExtensibilityManager()
+{
+	if (!SharedMenuExtensibilityManager.IsValid())
+	{
+		SharedMenuExtensibilityManager = MakeShareable(new FExtensibilityManager);
+	}
+	return SharedMenuExtensibilityManager;
+}
+
+TSharedPtr<FExtensibilityManager> FAssetEditorToolkit::GetSharedToolBarExtensibilityManager()
+{
+	if (!SharedToolBarExtensibilityManager.IsValid())
+	{
+		SharedToolBarExtensibilityManager = MakeShareable(new FExtensibilityManager);
+	}
+	return SharedToolBarExtensibilityManager;
+}
+
 void FAssetEditorToolkit::SetMenuOverlay( TSharedRef<SWidget> Widget )
 {
 	StandaloneHost.Pin()->SetMenuOverlay( Widget );
diff --git a/Engine/Source/Editor/UnrealEd/Private/Toolkits/GlobalEditorCommonCommands.cpp b/Engine/Source/Editor/UnrealEd/Private/Toolkits/GlobalEditorCommonCommands.cpp
--- a/Engine/Source/Editor/UnrealEd/Private/Toolkits/GlobalEditorCommonCommands.cpp
+++ b/Engine/Source/Editor/UnrealEd/Private/Toolkits/GlobalEditorCommonCommands.cpp
@@ -33,11 +33,8 @@ void FGlobalEditorCommonCommands::RegisterCommands()
 	UI_COMMAND( SummonOpenAssetDialog, "Open Asset...", "Summons an asset picker", EUserInterfaceActionType::Button, FInputChord(EModifierKey::Control, EKeys::P) );
 	UI_COMMAND( SummonOpenAssetDialogAlternate, "Open Asset...", "Summons an asset picker", EUserInterfaceActionType::Button, FInputChord(EModifierKey::Alt | EModifierKey::Shift, EKeys::O));
 	UI_COMMAND( FindInContentBrowser, "Browse to Asset", "Browses to the associated asset and selects it in the most recently used Content Browser (summoning one if necessary)", EUserInterfaceActionType::Button, FInputChord(EModifierKey::Control, EKeys::B));
-	UI_COMMAND( ViewReferences, "Reference Viewer...", "Launches the reference viewer showing the selected assets' references", EUserInterfaceActionType::Button, FInputChord(EModifierKey::Shift | EModifierKey::Alt, EKeys::R));
-	UI_COMMAND( ViewSizeMap, "Size Map...", "Displays an interactive map showing the approximate size of this asset and everything it references", EUserInterfaceActionType::Button, FInputChord(EModifierKey::Shift | EModifierKey::Alt, EKeys::M));	// @todo sizemap: Make sure key is not used already
 	
 	UI_COMMAND( OpenConsoleCommandBox, "Open Console Command Box", "Opens an edit box where you can type in a console command", EUserInterfaceActionType::Button, FInputChord(EKeys::Tilde));
-
 	UI_COMMAND( OpenDocumentation, "Open Documentation...", "Opens documentation for this tool", EUserInterfaceActionType::Button, FInputChord(EKeys::F1) );
 }
 
diff --git a/Engine/Source/Editor/UnrealEd/Public/Commandlets/AssetRegistryGenerator.h b/Engine/Source/Editor/UnrealEd/Public/Commandlets/AssetRegistryGenerator.h
--- a/Engine/Source/Editor/UnrealEd/Public/Commandlets/AssetRegistryGenerator.h
+++ b/Engine/Source/Editor/UnrealEd/Public/Commandlets/AssetRegistryGenerator.h
@@ -295,49 +295,31 @@ private:
 	/** Gather a list of dependencies required by to completely load this package */
 	bool GatherAllPackageDependencies(FName PackageName, TArray<FName>& DependentPackageNames);
 
-	/**
-	* Gather the list of dependencies that link the source to the target.  Output array includes the target.
-	*/
+	/** Gather the list of dependencies that link the source to the target.  Output array includes the target */
 	bool GetPackageDependencyChain(FName SourcePackage, FName TargetPackage, TSet<FName>& VisitedPackages, TArray<FName>& OutDependencyChain);
 
-	/**
-	* Get an array of Packages this package will import.
-	*/
+	/** Get an array of Packages this package will import */
 	bool GetPackageDependencies(FName PackageName, TArray<FName>& DependentPackageNames, EAssetRegistryDependencyType::Type InDependencyType);
 
-	/**
-	 * Save a CSV dump of chunk asset information.
-	 */
-	bool GenerateAssetChunkInformationCSV(const FString& OutputPath);
+	/** Save a CSV dump of chunk asset information, if bWriteIndividualFiles is true it writes a CSV per chunk in addition to AllChunksInfo */
+	bool GenerateAssetChunkInformationCSV(const FString& OutputPath, bool bWriteIndividualFiles = false);
 
-	/**
-	* Finds the asset belonging to ChunkID with the smallest number of links to Packages In PackageNames.
-	*/
-	void			FindShortestReferenceChain(TArray<FReferencePair> PackageNames, int32 ChunkID, uint32& OutParentIndex, FString& OutChainPath);
+	/** Finds the asset belonging to ChunkID with the smallest number of links to Packages In PackageNames */
+	void FindShortestReferenceChain(TArray<FReferencePair> PackageNames, int32 ChunkID, uint32& OutParentIndex, FString& OutChainPath);
 
-	/**
-	* Helper function for FindShortestReferenceChain
-	*/
-	FString			GetShortestReferenceChain(FName PackageName, int32 ChunkID);
+	/** Helper function for FindShortestReferenceChain */
+	FString	GetShortestReferenceChain(FName PackageName, int32 ChunkID);
 
-	/**
-	* 
-	*/
-	void			ResolveChunkDependencyGraph(const FChunkDependencyTreeNode& Node, FChunkPackageSet BaseAssetSet, TArray<TArray<FName>>& OutPackagesMovedBetweenChunks);
+	/** Deprecated method to remove redundant chunks */
+	void ResolveChunkDependencyGraph(const FChunkDependencyTreeNode& Node, FChunkPackageSet BaseAssetSet, TArray<TArray<FName>>& OutPackagesMovedBetweenChunks);
 
-	/**
-	* Helper function to verify Chunk asset assigment is valid.
-	*/
-	bool			CheckChunkAssetsAreNotInChild(const FChunkDependencyTreeNode& Node);
+	/** Helper function to verify Chunk asset assigment is valid */
+	bool CheckChunkAssetsAreNotInChild(const FChunkDependencyTreeNode& Node);
 
-	/**
-	* Helper function to create a given collection.
-	*/
-	bool			CreateOrEmptyCollection(FName CollectionName);
+	/** Helper function to create a given collection. */
+	bool CreateOrEmptyCollection(FName CollectionName);
 
-	/**
-	* Helper function to fill a given collection with a set of packages.
-	*/
-	void			WriteCollection(FName CollectionName, const TArray<FName>& PackageNames);
+	/** Helper function to fill a given collection with a set of packages */
+	void WriteCollection(FName CollectionName, const TArray<FName>& PackageNames);
 
 };
diff --git a/Engine/Source/Editor/UnrealEd/Public/Toolkits/AssetEditorToolkit.h b/Engine/Source/Editor/UnrealEd/Public/Toolkits/AssetEditorToolkit.h
--- a/Engine/Source/Editor/UnrealEd/Public/Toolkits/AssetEditorToolkit.h
+++ b/Engine/Source/Editor/UnrealEd/Public/Toolkits/AssetEditorToolkit.h
@@ -34,6 +34,49 @@ enum class EAssetEditorToolkitTabLocation : int32
 	Standalone,
 };
 
+DECLARE_DELEGATE_RetVal_TwoParams(TSharedRef<FExtender>, FAssetEditorExtender, const TSharedRef<FUICommandList>, const TArray<UObject*>);
+
+/**
+* Extensibility managers simply keep a series of FExtenders for a single menu/toolbar/anything
+* It is here to keep a standardized approach to editor extensibility among modules
+*/
+class UNREALED_API FExtensibilityManager
+{
+public:
+	virtual ~FExtensibilityManager() {}
+
+	/** Functions for outsiders to add or remove their extenders */
+	virtual void AddExtender(TSharedPtr<FExtender> Extender) { Extenders.AddUnique(Extender); }
+	virtual void RemoveExtender(TSharedPtr<FExtender> Extender) { Extenders.Remove(Extender); }
+
+	/** Gets all extender delegates for this manager */
+	virtual TArray<FAssetEditorExtender>& GetExtenderDelegates() { return ExtenderDelegates; }
+
+	/** Gets all extenders, consolidated, for use by the editor to be extended */
+	virtual TSharedPtr<FExtender> GetAllExtenders();
+	/** Gets all extenders and asset editor extenders from delegates consolidated */
+	virtual TSharedPtr<FExtender> GetAllExtenders(const TSharedRef<FUICommandList>& CommandList, const TArray<UObject*>& ContextSensitiveObjects);
+
+private:
+	/** A list of extenders the editor will use */
+	TArray<TSharedPtr<FExtender>> Extenders;
+	/** A list of extender delegates the editor will use */
+	TArray<FAssetEditorExtender> ExtenderDelegates;
+};
+
+/** Indicates that a class has a default menu that is extensible */
+class IHasMenuExtensibility
+{
+public:
+	virtual TSharedPtr<FExtensibilityManager> GetMenuExtensibilityManager() = 0;
+};
+
+/** Indicates that a class has a default toolbar that is extensible */
+class IHasToolBarExtensibility
+{
+public:
+	virtual TSharedPtr<FExtensibilityManager> GetToolBarExtensibilityManager() = 0;
+};
 
 /**
  * Base class for toolkits that are used for asset editing (abstract)
@@ -151,6 +194,10 @@ public:
 	void AddToolbarExtender(TSharedPtr<FExtender> Extender);
 	void RemoveToolbarExtender(TSharedPtr<FExtender> Extender);
 
+	/** Returns the default extensibility managers, these are applied for all asset types */
+	static TSharedPtr<FExtensibilityManager> GetSharedMenuExtensibilityManager();
+	static TSharedPtr<FExtensibilityManager> GetSharedToolBarExtensibilityManager();
+
 	/**
 	 * Allows the caller to set a menu overlay, displayed to the far right of the editor's menu bar
 	 *
@@ -273,18 +320,6 @@ private:
 		FLayoutSaveRestore::SaveToConfig(GEditorLayoutIni, LayoutToSave);
 	}
 
-	/** Called when "View References" is called for this asset */
-	void ViewReferences_Execute();
-
-	/** If true ViewReferences_Execute can be called, also caches ViewableObjects */
-	bool CanViewReferences();
-
-	/** Called when "View Size Map" is called for this asset */
-	void ViewSizeMap_Execute();
-
-	/** If true ViewSizeMap_Execute can be called, also caches ViewableObjects */
-	bool CanViewSizeMap();
-
 private:
 	/**
 	 * Report the references of the EditingObjects to the GC.  The level of indirection necessary 
@@ -329,6 +364,10 @@ private:
 	/** Static: World centric toolkit host to use for the next created asset editing toolkit */
 	static TWeakPtr< IToolkitHost > PreviousWorldCentricToolkitHostForNewAssetEditor;
 
+	/** The extensibility managers shared by all asset types */
+	static TSharedPtr<FExtensibilityManager> SharedMenuExtensibilityManager;
+	static TSharedPtr<FExtensibilityManager> SharedToolBarExtensibilityManager;
+
 	/** The object we're currently editing */
 	// @todo toolkit minor: Currently we don't need to serialize this object reference because the FAssetEditorManager is kept in sync (and will always serialize it.)
 	TArray<UObject*> EditingObjects;
@@ -350,55 +389,4 @@ private:
 
 	/**	The tab ids for all the tabs used */
 	static const FName ToolbarTabId;
-
-	/** A cached list of selected objects that can be viewed in the reference viewer */
-	TArray< FName > ViewableObjects;
-};
-
-
-
-DECLARE_DELEGATE_RetVal_TwoParams( TSharedRef<FExtender>, FAssetEditorExtender, const TSharedRef<FUICommandList>, const TArray<UObject*>);
-
-/**
- * Extensibility managers simply keep a series of FExtenders for a single menu/toolbar/anything
- * It is here to keep a standardized approach to editor extensibility among modules
- */
-class UNREALED_API FExtensibilityManager
-{
-public:
-	virtual ~FExtensibilityManager() {}
-
-	/** Functions for outsiders to add or remove their extenders */
-	virtual void AddExtender(TSharedPtr<FExtender> Extender) {Extenders.AddUnique(Extender);}
-	virtual void RemoveExtender(TSharedPtr<FExtender> Extender) {Extenders.Remove(Extender);}
-	
-	/** Gets all extender delegates for this manager */
-	virtual TArray<FAssetEditorExtender>& GetExtenderDelegates() {return ExtenderDelegates;}
-
-	/** Gets all extenders, consolidated, for use by the editor to be extended */
-	virtual TSharedPtr<FExtender> GetAllExtenders();
-	/** Gets all extenders and asset editor extenders from delegates consolidated */
-	virtual TSharedPtr<FExtender> GetAllExtenders(const TSharedRef<FUICommandList>& CommandList, const TArray<UObject*>& ContextSensitiveObjects);
-	
-private:
-	/** A list of extenders the editor will use */
-	TArray<TSharedPtr<FExtender>> Extenders;
-	/** A list of extender delegates the editor will use */
-	TArray<FAssetEditorExtender> ExtenderDelegates;
-};
-
-
-/** Indicates that a class has a default menu that is extensible */
-class IHasMenuExtensibility
-{
-public:
-	virtual TSharedPtr<FExtensibilityManager> GetMenuExtensibilityManager() = 0;
-};
-
-
-/** Indicates that a class has a default toolbar that is extensible */
-class IHasToolBarExtensibility
-{
-public:
-	virtual TSharedPtr<FExtensibilityManager> GetToolBarExtensibilityManager() = 0;
 };
diff --git a/Engine/Source/Editor/UnrealEd/Public/Toolkits/GlobalEditorCommonCommands.h b/Engine/Source/Editor/UnrealEd/Public/Toolkits/GlobalEditorCommonCommands.h
--- a/Engine/Source/Editor/UnrealEd/Public/Toolkits/GlobalEditorCommonCommands.h
+++ b/Engine/Source/Editor/UnrealEd/Public/Toolkits/GlobalEditorCommonCommands.h
@@ -34,8 +34,6 @@ public:
 	TSharedPtr<FUICommandInfo> SummonOpenAssetDialog;
 	TSharedPtr<FUICommandInfo> SummonOpenAssetDialogAlternate;
 	TSharedPtr<FUICommandInfo> OpenDocumentation;
-	TSharedPtr<FUICommandInfo> ViewReferences;
-	TSharedPtr<FUICommandInfo> ViewSizeMap;
 	TSharedPtr<FUICommandInfo> OpenConsoleCommandBox;
 };
 
diff --git a/Engine/Source/Editor/UnrealEd/UnrealEd.Build.cs b/Engine/Source/Editor/UnrealEd/UnrealEd.Build.cs
--- a/Engine/Source/Editor/UnrealEd/UnrealEd.Build.cs
+++ b/Engine/Source/Editor/UnrealEd/UnrealEd.Build.cs
@@ -47,11 +47,9 @@ public class UnrealEd : ModuleRules
 				"AudioEditor",
 				"ViewportSnapping",
 				"SourceCodeAccess",
-				"ReferenceViewer",
 				"IntroTutorials",
 				"OutputLog",
 				"Landscape",
-				"SizeMap",
 				"LocalizationService",
 				"HierarchicalLODUtilities",
 				"MessagingRpc",
@@ -140,7 +138,6 @@ public class UnrealEd : ModuleRules
 				"MeshPaintMode",
 				"Foliage",
 				"VectorVM",
-				"TreeMap",
 				"MaterialUtilities",
 				"Localization",
 				"LocalizationService",
@@ -212,10 +209,8 @@ public class UnrealEd : ModuleRules
 				"GameplayTasksEditor",
 				"UndoHistory",
 				"SourceCodeAccess",
-				"ReferenceViewer",
 				"HotReload",
 				"HTML5PlatformEditor",
-				"SizeMap",
 				"PortalProxies",
 				"PortalServices",
 				"GeometryCacheEd",
diff --git a/Engine/Source/Programs/AutomationTool/IOS/IOSPlatform.Automation.cs b/Engine/Source/Programs/AutomationTool/IOS/IOSPlatform.Automation.cs
--- a/Engine/Source/Programs/AutomationTool/IOS/IOSPlatform.Automation.cs
+++ b/Engine/Source/Programs/AutomationTool/IOS/IOSPlatform.Automation.cs
@@ -1036,12 +1036,24 @@ public class IOSPlatform : Platform
 
 			// copy in the application
 			string AppName = Path.GetFileNameWithoutExtension(ProjectIPA) + ".app";
+			if (!File.Exists(ProjectIPA))
+			{
+				Console.WriteLine("Couldn't find IPA: " + ProjectIPA);
+			}
 			using (ZipFile Zip = new ZipFile(ProjectIPA))
 			{
 				Zip.ExtractAll(ArchivePath, ExtractExistingFileAction.OverwriteSilently);
 
 				List<string> Dirs = new List<string>(Directory.EnumerateDirectories(Path.Combine(ArchivePath, "Payload"), "*.app"));
 				AppName = Dirs[0].Substring(Dirs[0].LastIndexOf(UnrealBuildTool.BuildHostPlatform.Current.Platform != UnrealTargetPlatform.Mac ? "\\" : "/") + 1);
+				foreach (string Dir in Dirs)
+				{
+					if (Dir.Contains(Params.ShortProjectName + ".app"))
+					{
+						Console.WriteLine("Using Directory: " + Dir);
+						AppName = Dir.Substring(Dir.LastIndexOf(UnrealBuildTool.BuildHostPlatform.Current.Platform != UnrealTargetPlatform.Mac ? "\\" : "/") + 1);
+					}
+				}
 				CopyDirectory_NoExceptions(Path.Combine(ArchivePath, "Payload", AppName), Path.Combine(ArchiveName, "Products", "Applications", AppName));
 			}
 
@@ -1094,23 +1106,33 @@ public class IOSPlatform : Platform
 			string AppPlist = Path.Combine(ArchiveName, "Products", "Applications", AppName, "Info.plist");
 			string OldPListData = File.Exists(AppPlist) ? File.ReadAllText(AppPlist) : "";
 
-			// bundle identifier
-			int index = OldPListData.IndexOf("CFBundleIdentifier");
-			index = OldPListData.IndexOf("<string>", index) + 8;
-			int length = OldPListData.IndexOf("</string>", index) - index;
-			string BundleIdentifier = OldPListData.Substring(index, length);
-
-			// short version
-			index = OldPListData.IndexOf("CFBundleShortVersionString");
-			index = OldPListData.IndexOf("<string>", index) + 8;
-			length = OldPListData.IndexOf("</string>", index) - index;
-			string BundleShortVersion = OldPListData.Substring(index, length);
-
-			// bundle version
-			index = OldPListData.IndexOf("CFBundleVersion");
-			index = OldPListData.IndexOf("<string>", index) + 8;
-			length = OldPListData.IndexOf("</string>", index) - index;
-			string BundleVersion = OldPListData.Substring(index, length);
+			string BundleIdentifier = "";
+			string BundleShortVersion = "";
+			string BundleVersion = "";
+			if (!string.IsNullOrEmpty(OldPListData))
+			{
+				// bundle identifier
+				int index = OldPListData.IndexOf("CFBundleIdentifier");
+				index = OldPListData.IndexOf("<string>", index) + 8;
+				int length = OldPListData.IndexOf("</string>", index) - index;
+				BundleIdentifier = OldPListData.Substring(index, length);
+
+				// short version
+				index = OldPListData.IndexOf("CFBundleShortVersionString");
+				index = OldPListData.IndexOf("<string>", index) + 8;
+				length = OldPListData.IndexOf("</string>", index) - index;
+				BundleShortVersion = OldPListData.Substring(index, length);
+
+				// bundle version
+				index = OldPListData.IndexOf("CFBundleVersion");
+				index = OldPListData.IndexOf("<string>", index) + 8;
+				length = OldPListData.IndexOf("</string>", index) - index;
+				BundleVersion = OldPListData.Substring(index, length);
+			}
+			else
+			{
+				Console.WriteLine("Could not load Info.plist");
+			}
 
 			// date we made this
 			const string Iso8601DateTimeFormat = "yyyy-MM-ddTHH:mm:ssZ";
diff --git a/Engine/Source/Programs/UnrealBuildTool/Platform/Android/UEDeployAndroid.cs b/Engine/Source/Programs/UnrealBuildTool/Platform/Android/UEDeployAndroid.cs
--- a/Engine/Source/Programs/UnrealBuildTool/Platform/Android/UEDeployAndroid.cs
+++ b/Engine/Source/Programs/UnrealBuildTool/Platform/Android/UEDeployAndroid.cs
@@ -2128,6 +2128,12 @@ namespace UnrealBuildTool
 
 			Text.AppendLine("\t\t<receiver android:name=\"com.epicgames.ue4.LocalNotificationReceiver\" />");
 
+			Text.AppendLine("\t\t<receiver android:name=\"com.epicgames.ue4.MulticastBroadcastReceiver\" android:exported=\"true\">");
+			Text.AppendLine("\t\t\t<intent-filter>");
+			Text.AppendLine("\t\t\t\t<action android:name=\"com.android.vending.INSTALL_REFERRER\" />");
+			Text.AppendLine("\t\t\t</intent-filter>");
+			Text.AppendLine("\t\t</receiver>");
+
 			// Max supported aspect ratio
 			string MaxAspectRatioString = MaxAspectRatioValue.ToString("f", System.Globalization.CultureInfo.InvariantCulture);
             Text.AppendLine(string.Format("\t\t<meta-data android:name=\"android.max_aspect\" android:value=\"{0}\" />", MaxAspectRatioString));
diff --git a/Engine/Source/Programs/UnrealBuildTool/Platform/IOS/IOSToolChain.cs b/Engine/Source/Programs/UnrealBuildTool/Platform/IOS/IOSToolChain.cs
--- a/Engine/Source/Programs/UnrealBuildTool/Platform/IOS/IOSToolChain.cs
+++ b/Engine/Source/Programs/UnrealBuildTool/Platform/IOS/IOSToolChain.cs
@@ -143,6 +143,43 @@ namespace UnrealBuildTool
                     FileReference AssetFile = FileReference.Combine(Binary.Config.OutputFilePath.Directory, "AssetCatalog", "Assets.car");
                     BuildProducts.Add(AssetFile, BuildProductType.RequiredResource);
                 }
+				else if (CppPlatform == CppPlatform.IOS && Settings.Value.IOSSDKVersionFloat >= 11.0f && BuildHostPlatform.Current.Platform == UnrealTargetPlatform.Mac)
+				{
+					FileReference AssetFile = FileReference.Combine(Binary.Config.OutputFilePath.Directory, "Payload", Binary.Config.OutputFilePath.GetFileNameWithoutExtension() + ".app", "Assets.car");
+					BuildProducts.Add(AssetFile, BuildProductType.RequiredResource);
+					AssetFile = FileReference.Combine(Binary.Config.OutputFilePath.Directory, "Payload", Binary.Config.OutputFilePath.GetFileNameWithoutExtension() + ".app", "[REDACTED_EMAIL]");
+					BuildProducts.Add(AssetFile, BuildProductType.RequiredResource);
+					AssetFile = FileReference.Combine(Binary.Config.OutputFilePath.Directory, "Payload", Binary.Config.OutputFilePath.GetFileNameWithoutExtension() + ".app", "AppIcon20x20@2x~ipad.png");
+					BuildProducts.Add(AssetFile, BuildProductType.RequiredResource);
+					AssetFile = FileReference.Combine(Binary.Config.OutputFilePath.Directory, "Payload", Binary.Config.OutputFilePath.GetFileNameWithoutExtension() + ".app", "[REDACTED_EMAIL]");
+					BuildProducts.Add(AssetFile, BuildProductType.RequiredResource);
+					AssetFile = FileReference.Combine(Binary.Config.OutputFilePath.Directory, "Payload", Binary.Config.OutputFilePath.GetFileNameWithoutExtension() + ".app", "AppIcon20x20~ipad.png");
+					BuildProducts.Add(AssetFile, BuildProductType.RequiredResource);
+					AssetFile = FileReference.Combine(Binary.Config.OutputFilePath.Directory, "Payload", Binary.Config.OutputFilePath.GetFileNameWithoutExtension() + ".app", "[REDACTED_EMAIL]");
+					BuildProducts.Add(AssetFile, BuildProductType.RequiredResource);
+					AssetFile = FileReference.Combine(Binary.Config.OutputFilePath.Directory, "Payload", Binary.Config.OutputFilePath.GetFileNameWithoutExtension() + ".app", "AppIcon29x29@2x~ipad.png");
+					BuildProducts.Add(AssetFile, BuildProductType.RequiredResource);
+					AssetFile = FileReference.Combine(Binary.Config.OutputFilePath.Directory, "Payload", Binary.Config.OutputFilePath.GetFileNameWithoutExtension() + ".app", "[REDACTED_EMAIL]");
+					BuildProducts.Add(AssetFile, BuildProductType.RequiredResource);
+					AssetFile = FileReference.Combine(Binary.Config.OutputFilePath.Directory, "Payload", Binary.Config.OutputFilePath.GetFileNameWithoutExtension() + ".app", "AppIcon29x29~ipad.png");
+					BuildProducts.Add(AssetFile, BuildProductType.RequiredResource);
+					AssetFile = FileReference.Combine(Binary.Config.OutputFilePath.Directory, "Payload", Binary.Config.OutputFilePath.GetFileNameWithoutExtension() + ".app", "[REDACTED_EMAIL]");
+					BuildProducts.Add(AssetFile, BuildProductType.RequiredResource);
+					AssetFile = FileReference.Combine(Binary.Config.OutputFilePath.Directory, "Payload", Binary.Config.OutputFilePath.GetFileNameWithoutExtension() + ".app", "AppIcon40x40@2x~ipad.png");
+					BuildProducts.Add(AssetFile, BuildProductType.RequiredResource);
+					AssetFile = FileReference.Combine(Binary.Config.OutputFilePath.Directory, "Payload", Binary.Config.OutputFilePath.GetFileNameWithoutExtension() + ".app", "[REDACTED_EMAIL]");
+					BuildProducts.Add(AssetFile, BuildProductType.RequiredResource);
+					AssetFile = FileReference.Combine(Binary.Config.OutputFilePath.Directory, "Payload", Binary.Config.OutputFilePath.GetFileNameWithoutExtension() + ".app", "AppIcon40x40~ipad.png");
+					BuildProducts.Add(AssetFile, BuildProductType.RequiredResource);
+					AssetFile = FileReference.Combine(Binary.Config.OutputFilePath.Directory, "Payload", Binary.Config.OutputFilePath.GetFileNameWithoutExtension() + ".app", "[REDACTED_EMAIL]");
+					BuildProducts.Add(AssetFile, BuildProductType.RequiredResource);
+					AssetFile = FileReference.Combine(Binary.Config.OutputFilePath.Directory, "Payload", Binary.Config.OutputFilePath.GetFileNameWithoutExtension() + ".app", "AppIcon76x76@2x~ipad.png");
+					BuildProducts.Add(AssetFile, BuildProductType.RequiredResource);
+					AssetFile = FileReference.Combine(Binary.Config.OutputFilePath.Directory, "Payload", Binary.Config.OutputFilePath.GetFileNameWithoutExtension() + ".app", "AppIcon76x76~ipad.png");
+					BuildProducts.Add(AssetFile, BuildProductType.RequiredResource);
+					AssetFile = FileReference.Combine(Binary.Config.OutputFilePath.Directory, "Payload", Binary.Config.OutputFilePath.GetFileNameWithoutExtension() + ".app", "AppIcon83.5x83.5@2x~ipad.png");
+					BuildProducts.Add(AssetFile, BuildProductType.RequiredResource);
+				}
 			}
             if ((ProjectSettings.bGeneratedSYMFile == true || ProjectSettings.bGeneratedSYMBundle == true) && ProjectSettings.bGenerateCrashReportSymbols && Binary.Config.Type == UEBuildBinaryType.Executable)
             {
diff --git a/Engine/Source/Programs/UnrealBuildTool/System/UProjectInfo.cs b/Engine/Source/Programs/UnrealBuildTool/System/UProjectInfo.cs
--- a/Engine/Source/Programs/UnrealBuildTool/System/UProjectInfo.cs
+++ b/Engine/Source/Programs/UnrealBuildTool/System/UProjectInfo.cs
@@ -140,7 +140,7 @@ namespace UnrealBuildTool
 		{
 			DateTime StartTime = DateTime.Now;
 
-			List<DirectoryReference> DirectoriesToSearch = new List<DirectoryReference>();
+			List<DirectoryInfo> DirectoriesToSearch = new List<DirectoryInfo>();
 
 			// Find all the .uprojectdirs files contained in the root folder and add their entries to the search array
 			string EngineSourceDirectory = Path.GetFullPath(Path.Combine(RootDirectory, "Engine", "Source"));
@@ -156,7 +156,7 @@ namespace UnrealBuildTool
 						DirectoryReference BaseProjectDir = DirectoryReference.Combine(UnrealBuildTool.RootDirectory, TrimLine);
 						if(BaseProjectDir.IsUnderDirectory(UnrealBuildTool.RootDirectory))
 						{
-							DirectoriesToSearch.Add(BaseProjectDir);
+							DirectoriesToSearch.Add(new DirectoryInfo(BaseProjectDir.FullName));
 						}
 						else
 						{
@@ -168,18 +168,18 @@ namespace UnrealBuildTool
 
 			Log.TraceVerbose("\tFound {0} directories to search", DirectoriesToSearch.Count);
 
-			foreach (DirectoryReference DirToSearch in DirectoriesToSearch)
+			foreach (DirectoryInfo DirToSearch in DirectoriesToSearch)
 			{
 				Log.TraceVerbose("\t\tSearching {0}", DirToSearch.FullName);
-				if (DirectoryReference.Exists(DirToSearch))
+				if (DirToSearch.Exists)
 				{
-					foreach (DirectoryReference SubDir in DirectoryReference.EnumerateDirectories(DirToSearch, "*", SearchOption.TopDirectoryOnly))
+					foreach (DirectoryInfo SubDir in DirToSearch.EnumerateDirectories())
 					{
-						Log.TraceVerbose("\t\t\tFound subdir {0}", SubDir.FullName);
-						foreach(FileReference UProjFile in DirectoryReference.EnumerateFiles(SubDir, "*.uproject", SearchOption.TopDirectoryOnly))
+						Log.TraceVerbose("\t\t\tFound subdir {0} ({1})", SubDir.FullName, SubDir.Name);
+						foreach(FileInfo UProjFile in SubDir.EnumerateFiles("*.uproject", SearchOption.TopDirectoryOnly))
 						{
 							Log.TraceVerbose("\t\t\t\t{0}", UProjFile.FullName);
-							AddProject(UProjFile);
+							AddProject(new FileReference(UProjFile));
 						}
 					}
 				}
diff --git a/Engine/Source/Runtime/AssetRegistry/Private/AssetRegistry.cpp b/Engine/Source/Runtime/AssetRegistry/Private/AssetRegistry.cpp
--- a/Engine/Source/Runtime/AssetRegistry/Private/AssetRegistry.cpp
+++ b/Engine/Source/Runtime/AssetRegistry/Private/AssetRegistry.cpp
@@ -1508,6 +1508,11 @@ void UAssetRegistryImpl::InitializeTemporaryAssetRegistryState(FAssetRegistrySta
 	OutState.InitializeFromExisting(DataToUse, State.CachedDependsNodes, State.CachedPackageData, Options, bRefreshExisting);
 }
 
+const FAssetRegistryState* UAssetRegistryImpl::GetAssetRegistryState() const
+{
+	return &State;
+}
+
 void UAssetRegistryImpl::ScanPathsAndFilesSynchronous(const TArray<FString>& InPaths, const TArray<FString>& InSpecificFiles, bool bForceRescan, EAssetDataCacheMode AssetDataCacheMode)
 {
 	ScanPathsAndFilesSynchronous(InPaths, InSpecificFiles, bForceRescan, AssetDataCacheMode, nullptr, nullptr);
@@ -2460,14 +2465,15 @@ void UAssetRegistryImpl::SetManageReferences(const TMultiMap<FAssetIdentifier, F
 	}
 
 	TSet<FDependsNode*> Visited;
-	TArray<FDependsNode*> NodesToManage;
+	TMap<FDependsNode*, EAssetRegistryDependencyType::Type> NodesToManage;
+	TArray<FDependsNode*> NodesToHardReference;
 	TArray<FDependsNode*> NodesToRecurse;
 
 	// For each explicitly set asset
-	for (const TPair<FDependsNode*, TArray<FDependsNode *>>& Pair : ExplicitMap)
+	for (const TPair<FDependsNode*, TArray<FDependsNode *>>& ExplicitPair : ExplicitMap)
 	{
-		FDependsNode* BaseManagedNode = Pair.Key;
-		const TArray<FDependsNode*>& ManagerNodes = Pair.Value;
+		FDependsNode* BaseManagedNode = ExplicitPair.Key;
+		const TArray<FDependsNode*>& ManagerNodes = ExplicitPair.Value;
 
 		for (FDependsNode* ManagerNode : ManagerNodes)
 		{	
@@ -2477,7 +2483,7 @@ void UAssetRegistryImpl::SetManageReferences(const TMultiMap<FAssetIdentifier, F
 			
 			FDependsNode* SourceNode = ManagerNode;
 
-			auto IterateFunction = [&ManagerNode, &SourceNode, &ShouldSetManager, &NodesToManage, &NodesToRecurse, &Visited, &ExplicitMap, &ExistingManagedNodes](FDependsNode* TargetNode, EAssetRegistryDependencyType::Type DependencyType)
+			auto IterateFunction = [&ManagerNode, &SourceNode, &ShouldSetManager, &NodesToManage, &NodesToHardReference, &NodesToRecurse, &Visited, &ExplicitMap, &ExistingManagedNodes](FDependsNode* TargetNode, EAssetRegistryDependencyType::Type DependencyType)
 			{
 				// Only recurse if we haven't already visited, and this node passes recursion test
 				if (!Visited.Contains(TargetNode))
@@ -2493,7 +2499,8 @@ void UAssetRegistryImpl::SetManageReferences(const TMultiMap<FAssetIdentifier, F
 						return;
 					}
 
-					NodesToManage.Push(TargetNode);
+					EAssetRegistryDependencyType::Type ManageType = (Flags & EAssetSetManagerFlags::IsDirectSet) ? EAssetRegistryDependencyType::HardManage : EAssetRegistryDependencyType::SoftManage;
+					NodesToManage.Add(TargetNode, ManageType);
 
 					if (Result == EAssetSetManagerResult::SetAndRecurse)
 					{
@@ -2519,14 +2526,10 @@ void UAssetRegistryImpl::SetManageReferences(const TMultiMap<FAssetIdentifier, F
 				}
 			}
 
-			// Now set all the dependencies
-			while (NodesToManage.Num())
+			for (TPair<FDependsNode*, EAssetRegistryDependencyType::Type>& ManagePair : NodesToManage)
 			{
-				// Pull off end of array, order doesn't matter
-				FDependsNode* ManagedNode = NodesToManage.Pop();
-
-				ManagerNode->AddDependency(ManagedNode, EAssetRegistryDependencyType::Manage);
-				ManagedNode->AddReferencer(ManagerNode);
+				ManagePair.Key->AddReferencer(ManagerNode);
+				ManagerNode->AddDependency(ManagePair.Key, ManagePair.Value);
 			}
 		}
 	}
diff --git a/Engine/Source/Runtime/AssetRegistry/Private/AssetRegistry.h b/Engine/Source/Runtime/AssetRegistry/Private/AssetRegistry.h
--- a/Engine/Source/Runtime/AssetRegistry/Private/AssetRegistry.h
+++ b/Engine/Source/Runtime/AssetRegistry/Private/AssetRegistry.h
@@ -68,6 +68,7 @@ public:
 	virtual uint32 GetAllocatedSize(bool bLogDetailed = false) const override;
 	virtual void LoadPackageRegistryData(FArchive& Ar, TArray<FAssetData*>& Data) const override;
 	virtual void InitializeTemporaryAssetRegistryState(FAssetRegistryState& OutState, const FAssetRegistrySerializationOptions& Options, bool bRefreshExisting = false, const TMap<FName, FAssetData*>& OverrideData = TMap<FName, FAssetData*>()) const override;
+	virtual const FAssetRegistryState* GetAssetRegistryState() const override;
 	virtual void InitializeSerializationOptions(FAssetRegistrySerializationOptions& Options, const FString& PlatformIniName = FString()) const override;
 
 	virtual void SaveRegistryData(FArchive& Ar, TMap<FName, FAssetData*>& Data, TArray<FName>* InMaps = nullptr) override;
diff --git a/Engine/Source/Runtime/AssetRegistry/Private/AssetRegistryState.cpp b/Engine/Source/Runtime/AssetRegistry/Private/AssetRegistryState.cpp
--- a/Engine/Source/Runtime/AssetRegistry/Private/AssetRegistryState.cpp
+++ b/Engine/Source/Runtime/AssetRegistry/Private/AssetRegistryState.cpp
@@ -508,7 +508,6 @@ bool FAssetRegistryState::GetReferencers(const FAssetIdentifier& AssetIdentifier
 {
 	const FDependsNode* const* NodePtr = CachedDependsNodes.Find(AssetIdentifier);
 	const FDependsNode* Node = nullptr;
-	bool bShowingAllReferences = InReferenceType == EAssetRegistryDependencyType::All;
 
 	if (NodePtr != nullptr)
 	{
@@ -518,28 +517,11 @@ bool FAssetRegistryState::GetReferencers(const FAssetIdentifier& AssetIdentifier
 	if (Node != nullptr)
 	{
 		TArray<FDependsNode*> DependencyNodes;
-		Node->GetReferencers(DependencyNodes);
+		Node->GetReferencers(DependencyNodes, InReferenceType);
 
 		for (FDependsNode* DependencyNode : DependencyNodes)
 		{
-			if (!bShowingAllReferences)
-			{
-				TArray<FDependsNode*> DependenciesFromReferencer;
-				DependencyNode->GetDependencies(DependenciesFromReferencer, InReferenceType);
-
-				for (FDependsNode* Dependency : DependenciesFromReferencer)
-				{
-					if (Dependency == Node)
-					{
-						OutReferencers.Add(DependencyNode->GetIdentifier());
-						break;
-					}
-				}
-			}
-			else
-			{
-				OutReferencers.Add(DependencyNode->GetIdentifier());
-			}
+			OutReferencers.Add(DependencyNode->GetIdentifier());
 		}
 
 		return true;
@@ -611,7 +593,8 @@ bool FAssetRegistryState::Serialize(FArchive& OriginalAr, FAssetRegistrySerializ
 			DepCounts.Add(EAssetRegistryDependencyType::Hard, 0);
 			DepCounts.Add(EAssetRegistryDependencyType::Soft, 0);
 			DepCounts.Add(EAssetRegistryDependencyType::SearchableName, 0);
-			DepCounts.Add(EAssetRegistryDependencyType::Manage, 0);
+			DepCounts.Add(EAssetRegistryDependencyType::SoftManage, 0);
+			DepCounts.Add(EAssetRegistryDependencyType::HardManage, 0);
 			DepCounts.Add(EAssetRegistryDependencyType::None, 0); // Referencers
 
 			auto DependencyProcessor = [&](FDependsNode* InDependency, EAssetRegistryDependencyType::Type InDependencyType)
@@ -634,7 +617,8 @@ bool FAssetRegistryState::Serialize(FArchive& OriginalAr, FAssetRegistrySerializ
 			}
 			if (Options.bSerializeManageDependencies)
 			{
-				DependentNode->IterateOverDependencies(DependencyProcessor, EAssetRegistryDependencyType::Manage);
+				DependentNode->IterateOverDependencies(DependencyProcessor, EAssetRegistryDependencyType::SoftManage);
+				DependentNode->IterateOverDependencies(DependencyProcessor, EAssetRegistryDependencyType::HardManage);
 			}
 
 			DependentNode->IterateOverReferencers([&](FDependsNode* InReferencer) 
@@ -650,7 +634,8 @@ bool FAssetRegistryState::Serialize(FArchive& OriginalAr, FAssetRegistrySerializ
 			Ar << DepCounts[EAssetRegistryDependencyType::Hard];
 			Ar << DepCounts[EAssetRegistryDependencyType::Soft];
 			Ar << DepCounts[EAssetRegistryDependencyType::SearchableName];
-			Ar << DepCounts[EAssetRegistryDependencyType::Manage];
+			Ar << DepCounts[EAssetRegistryDependencyType::SoftManage];
+			Ar << DepCounts[EAssetRegistryDependencyType::HardManage];
 			Ar << DepCounts[EAssetRegistryDependencyType::None];
 
 			for (FDependsNode* Dependency : ProcessedDependencies)
@@ -734,7 +719,15 @@ bool FAssetRegistryState::Serialize(FArchive& OriginalAr, FAssetRegistrySerializ
 			Ar << DepCounts.FindOrAdd(EAssetRegistryDependencyType::Hard);
 			Ar << DepCounts.FindOrAdd(EAssetRegistryDependencyType::Soft);
 			Ar << DepCounts.FindOrAdd(EAssetRegistryDependencyType::SearchableName);
-			Ar << DepCounts.FindOrAdd(EAssetRegistryDependencyType::Manage);
+			Ar << DepCounts.FindOrAdd(EAssetRegistryDependencyType::SoftManage);
+			if (Version < FAssetRegistryVersion::AddedHardManage)
+			{
+				DepCounts.Add(EAssetRegistryDependencyType::HardManage, 0);
+			}
+			else
+			{
+				Ar << DepCounts.FindOrAdd(EAssetRegistryDependencyType::HardManage);
+			}
 			Ar << DepCounts.FindOrAdd(EAssetRegistryDependencyType::None); // Referencers
 						
 			// Create the node if we're actually saving dependencies, otherwise just fake serialize
@@ -743,7 +736,7 @@ bool FAssetRegistryState::Serialize(FArchive& OriginalAr, FAssetRegistrySerializ
 			if (NewDependsNodeData)
 			{
 				NewDependsNodeData->SetIdentifier(AssetIdentifier);
-				NewDependsNodeData->Reserve(DepCounts[EAssetRegistryDependencyType::Hard], DepCounts[EAssetRegistryDependencyType::Soft], DepCounts[EAssetRegistryDependencyType::SearchableName], DepCounts[EAssetRegistryDependencyType::Manage], DepCounts[EAssetRegistryDependencyType::None]);
+				NewDependsNodeData->Reserve(DepCounts[EAssetRegistryDependencyType::Hard], DepCounts[EAssetRegistryDependencyType::Soft], DepCounts[EAssetRegistryDependencyType::SearchableName], DepCounts[EAssetRegistryDependencyType::SoftManage], DepCounts[EAssetRegistryDependencyType::HardManage], DepCounts[EAssetRegistryDependencyType::None]);
 				CachedDependsNodes.Add(NewDependsNodeData->GetIdentifier(), NewDependsNodeData);
 			}
 
@@ -777,7 +770,8 @@ bool FAssetRegistryState::Serialize(FArchive& OriginalAr, FAssetRegistrySerializ
 			SerializeDependencyType(EAssetRegistryDependencyType::Hard, Options.bSerializeDependencies);
 			SerializeDependencyType(EAssetRegistryDependencyType::Soft, Options.bSerializeDependencies);
 			SerializeDependencyType(EAssetRegistryDependencyType::SearchableName, Options.bSerializeDependencies && Options.bSerializeSearchableNameDependencies);
-			SerializeDependencyType(EAssetRegistryDependencyType::Manage, Options.bSerializeDependencies && Options.bSerializeManageDependencies);
+			SerializeDependencyType(EAssetRegistryDependencyType::SoftManage, Options.bSerializeDependencies && Options.bSerializeManageDependencies);
+			SerializeDependencyType(EAssetRegistryDependencyType::HardManage, Options.bSerializeDependencies && Options.bSerializeManageDependencies);
 			SerializeDependencyType(EAssetRegistryDependencyType::None, Options.bSerializeDependencies);
 		}
 
diff --git a/Engine/Source/Runtime/AssetRegistry/Private/DependsNode.cpp b/Engine/Source/Runtime/AssetRegistry/Private/DependsNode.cpp
--- a/Engine/Source/Runtime/AssetRegistry/Private/DependsNode.cpp
+++ b/Engine/Source/Runtime/AssetRegistry/Private/DependsNode.cpp
@@ -52,13 +52,31 @@ void FDependsNode::GetReferencers(TArray<FDependsNode*>& OutReferencers, EAssetR
 		// If type specified, filter
 		if (InDependencyType != EAssetRegistryDependencyType::All)
 		{
-			IterateOverDependencies([&bShouldAdd,this](const FDependsNode* InDependency, EAssetRegistryDependencyType::Type /*InDependencyType*/)
+			// If type is specified, filter. We don't use the iterate wrapper for performance
+			if (InDependencyType & EAssetRegistryDependencyType::Hard)
 			{
-				if (InDependency == this)
-				{
-					bShouldAdd = true;
-				}
-			}, InDependencyType);
+				bShouldAdd = Referencer->HardDependencies.Contains(this);
+			}
+
+			if (InDependencyType & EAssetRegistryDependencyType::Soft && !bShouldAdd)
+			{
+				bShouldAdd = Referencer->SoftDependencies.Contains(this);
+			}
+
+			if (InDependencyType & EAssetRegistryDependencyType::HardManage && !bShouldAdd)
+			{
+				bShouldAdd = Referencer->HardManageDependencies.Contains(this);
+			}
+
+			if (InDependencyType & EAssetRegistryDependencyType::SoftManage && !bShouldAdd)
+			{
+				bShouldAdd = Referencer->SoftManageDependencies.Contains(this);
+			}
+
+			if (InDependencyType & EAssetRegistryDependencyType::SearchableName && !bShouldAdd)
+			{
+				bShouldAdd = Referencer->NameDependencies.Contains(this);
+			}
 		}
 		else
 		{
@@ -183,5 +201,5 @@ void FDependsNode::PrintReferencersRecursive(const FString& Indent, TSet<const F
 
 int32 FDependsNode::GetConnectionCount() const
 {
-	return HardDependencies.Num() + SoftDependencies.Num() + NameDependencies.Num() + ManageDependencies.Num() + Referencers.Num();
+	return HardDependencies.Num() + SoftDependencies.Num() + NameDependencies.Num() + SoftManageDependencies.Num() + HardManageDependencies.Num() + Referencers.Num();
 }
\ No newline at end of file
diff --git a/Engine/Source/Runtime/AssetRegistry/Private/DependsNode.h b/Engine/Source/Runtime/AssetRegistry/Private/DependsNode.h
--- a/Engine/Source/Runtime/AssetRegistry/Private/DependsNode.h
+++ b/Engine/Source/Runtime/AssetRegistry/Private/DependsNode.h
@@ -50,7 +50,7 @@ public:
 	/** Returns amount of memory used by the arrays */
 	uint32 GetAllocatedSize(void) const
 	{
-		return HardDependencies.GetAllocatedSize() + SoftDependencies.GetAllocatedSize() + NameDependencies.GetAllocatedSize() + ManageDependencies.GetAllocatedSize() + Referencers.GetAllocatedSize();
+		return HardDependencies.GetAllocatedSize() + SoftDependencies.GetAllocatedSize() + NameDependencies.GetAllocatedSize() + SoftManageDependencies.GetAllocatedSize() + HardManageDependencies.GetAllocatedSize() + Referencers.GetAllocatedSize();
 	}
 
 	/** Iterate over all the dependencies of this node, filtered by the supplied type parameter, and call the supplied lambda parameter on the record */
@@ -76,12 +76,13 @@ public:
 		}
 	}
 
-	void Reserve(int32 InNumHardDependencies, int32 InNumSoftDependencies, int32 InNumNameDependencies, int32 InNumManageDependencies, int32 InNumReferencers)
+	void Reserve(int32 InNumHardDependencies, int32 InNumSoftDependencies, int32 InNumNameDependencies, int32 InNumSoftManageDependencies, int32 InNumHardManageDependencies, int32 InNumReferencers)
 	{
 		HardDependencies.Reserve(InNumHardDependencies);
 		SoftDependencies.Reserve(InNumSoftDependencies);
 		NameDependencies.Reserve(InNumNameDependencies);
-		ManageDependencies.Reserve(InNumManageDependencies);
+		SoftManageDependencies.Reserve(InNumSoftManageDependencies);
+		HardManageDependencies.Reserve(InNumHardManageDependencies);
 		Referencers.Reserve(InNumReferencers);
 	}
 
@@ -91,6 +92,7 @@ private:
 	template <class T>
 	FORCEINLINE void IterateOverDependencyArrays(T InCallback, EAssetRegistryDependencyType::Type InDependencyType = EAssetRegistryDependencyType::All) const
 	{
+		// This iteration is specific so it gets the "most important" references first in the array
 		if (InDependencyType & EAssetRegistryDependencyType::Hard)
 		{
 			InCallback(const_cast<TArray<FDependsNode*>&>(HardDependencies), EAssetRegistryDependencyType::Hard);
@@ -101,14 +103,19 @@ private:
 			InCallback(const_cast<TArray<FDependsNode*>&>(SoftDependencies), EAssetRegistryDependencyType::Soft);
 		}
 
-		if (InDependencyType & EAssetRegistryDependencyType::SearchableName)
+		if (InDependencyType & EAssetRegistryDependencyType::HardManage)
 		{
-			InCallback(const_cast<TArray<FDependsNode*>&>(NameDependencies), EAssetRegistryDependencyType::SearchableName);
+			InCallback(const_cast<TArray<FDependsNode*>&>(HardManageDependencies), EAssetRegistryDependencyType::HardManage);
 		}
 
-		if (InDependencyType & EAssetRegistryDependencyType::Manage)
+		if (InDependencyType & EAssetRegistryDependencyType::SoftManage)
 		{
-			InCallback(const_cast<TArray<FDependsNode*>&>(ManageDependencies), EAssetRegistryDependencyType::Manage);
+			InCallback(const_cast<TArray<FDependsNode*>&>(SoftManageDependencies), EAssetRegistryDependencyType::SoftManage);
+		}
+
+		if (InDependencyType & EAssetRegistryDependencyType::SearchableName)
+		{
+			InCallback(const_cast<TArray<FDependsNode*>&>(NameDependencies), EAssetRegistryDependencyType::SearchableName);
 		}
 	}
 
@@ -125,8 +132,10 @@ private:
 	TArray<FDependsNode*> SoftDependencies;
 	/** The list of searchable name dependencies for this node */
 	TArray<FDependsNode*> NameDependencies;
-	/** The list of manage dependencies for this node */
-	TArray<FDependsNode*> ManageDependencies;
+	/** The list of hard manage dependencies for this node */
+	TArray<FDependsNode*> SoftManageDependencies;
+	/** The list of soft manage dependencies for this node */
+	TArray<FDependsNode*> HardManageDependencies;
 	/** The list of referencers to this node */
 	TArray<FDependsNode*> Referencers;
 };
diff --git a/Engine/Source/Runtime/AssetRegistry/Public/AssetData.h b/Engine/Source/Runtime/AssetRegistry/Public/AssetData.h
--- a/Engine/Source/Runtime/AssetRegistry/Public/AssetData.h
+++ b/Engine/Source/Runtime/AssetRegistry/Public/AssetData.h
@@ -28,6 +28,7 @@ struct ASSETREGISTRY_API FAssetRegistryVersion
 		AddAssetRegistryState,	// Added FAssetRegistryState and support for piecemeal serialization
 		ChangedAssetData,		// AssetData serialization format changed, versions before this are not readable
 		RemovedMD5Hash,			// Removed MD5 hash from package data
+		AddedHardManage,		// Added hard/soft manage references
 
 		// -----<new versions can be added above this line>-------------------------------------------------
 		VersionPlusOne,
diff --git a/Engine/Source/Runtime/AssetRegistry/Public/IAssetRegistry.h b/Engine/Source/Runtime/AssetRegistry/Public/IAssetRegistry.h
--- a/Engine/Source/Runtime/AssetRegistry/Public/IAssetRegistry.h
+++ b/Engine/Source/Runtime/AssetRegistry/Public/IAssetRegistry.h
@@ -348,6 +348,9 @@ public:
 	 */
 	virtual void InitializeTemporaryAssetRegistryState(FAssetRegistryState& OutState, const FAssetRegistrySerializationOptions& Options, bool bRefreshExisting = false, const TMap<FName, FAssetData*>& OverrideData = TMap<FName, FAssetData*>()) const = 0;
 
+	/** Returns read only reference to the current asset registry state. The contents of this may change at any time so do not save internal pointers */
+	virtual const FAssetRegistryState* GetAssetRegistryState() const = 0;
+
 	/** Fills in FAssetRegistrySerializationOptions from ini, optionally using a target platform ini name */
 	virtual void InitializeSerializationOptions(FAssetRegistrySerializationOptions& Options, const FString& PlatformIniName = FString()) const = 0;
 
diff --git a/Engine/Source/Runtime/Core/Private/IOS/IOSPlatformMisc.cpp b/Engine/Source/Runtime/Core/Private/IOS/IOSPlatformMisc.cpp
--- a/Engine/Source/Runtime/Core/Private/IOS/IOSPlatformMisc.cpp
+++ b/Engine/Source/Runtime/Core/Private/IOS/IOSPlatformMisc.cpp
@@ -97,6 +97,43 @@ void FIOSPlatformMisc::PlatformHandleSplashScreen(bool ShowSplashScreen)
 //    GShowSplashScreen = ShowSplashScreen;
 }
 
+const TCHAR* FIOSPlatformMisc::GamePersistentDownloadDir()
+{
+    static FString GamePersistentDownloadDir = TEXT("");
+    
+    if (GamePersistentDownloadDir.Len() == 0)
+    {
+        FString BaseProjectDir = ProjectDir();
+        
+        if (BaseProjectDir.Len() > 0)
+        {
+            GamePersistentDownloadDir = BaseProjectDir / TEXT("PersistentDownloadDir");
+        }
+        
+        // create the directory so we can exclude it from iCloud backup
+        FString Result = GamePersistentDownloadDir;
+        Result.ReplaceInline(TEXT("../"), TEXT(""));
+        Result.ReplaceInline(TEXT(".."), TEXT(""));
+        Result.ReplaceInline(FPlatformProcess::BaseDir(), TEXT(""));
+        FString DownloadPath = FString([NSSearchPathForDirectoriesInDomains(NSDocumentDirectory, NSUserDomainMask, YES) objectAtIndex:0]) + TEXT("/");
+        Result = DownloadPath + Result;
+        NSURL* URL = [NSURL fileURLWithPath : Result.GetNSString()];
+        if (![[NSFileManager defaultManager] fileExistsAtPath:[URL path]])
+        {
+            [[NSFileManager defaultManager] createDirectoryAtURL:URL withIntermediateDirectories : YES attributes : nil error : nil];
+        }
+        
+        // mark it to not be uploaded
+        NSError *error = nil;
+        BOOL success = [URL setResourceValue : [NSNumber numberWithBool : YES] forKey : NSURLIsExcludedFromBackupKey error : &error];
+        if (!success)
+        {
+            NSLog(@"Error excluding %@ from backup %@",[URL lastPathComponent], error);
+        }
+    }
+    return *GamePersistentDownloadDir;
+}
+
 EAppReturnType::Type FIOSPlatformMisc::MessageBoxExt( EAppMsgType::Type MsgType, const TCHAR* Text, const TCHAR* Caption )
 {
 	extern EAppReturnType::Type MessageBoxExtImpl( EAppMsgType::Type MsgType, const TCHAR* Text, const TCHAR* Caption );
diff --git a/Engine/Source/Runtime/Core/Public/IOS/IOSPlatformMisc.h b/Engine/Source/Runtime/Core/Public/IOS/IOSPlatformMisc.h
--- a/Engine/Source/Runtime/Core/Public/IOS/IOSPlatformMisc.h
+++ b/Engine/Source/Runtime/Core/Public/IOS/IOSPlatformMisc.h
@@ -31,6 +31,7 @@ struct CORE_API FIOSPlatformMisc : public FApplePlatformMisc
 	static bool DeleteStoredValue(const FString& InStoreId, const FString& InSectionName, const FString& InKeyName);
 	static void GetValidTargetPlatforms(class TArray<class FString>& TargetPlatformNames);
 	static bool HasActiveWiFiConnection();
+    static const TCHAR* GamePersistentDownloadDir();
 
 	static int GetAudioVolume();
 	static bool AreHeadphonesPluggedIn();
diff --git a/Engine/Source/Runtime/Core/Public/Mac/MacPlatformCompilerPreSetup.h b/Engine/Source/Runtime/Core/Public/Mac/MacPlatformCompilerPreSetup.h
--- a/Engine/Source/Runtime/Core/Public/Mac/MacPlatformCompilerPreSetup.h
+++ b/Engine/Source/Runtime/Core/Public/Mac/MacPlatformCompilerPreSetup.h
@@ -146,8 +146,14 @@
 #pragma clang diagnostic ignored "-Wconstant-logical-operand"
 #pragma clang diagnostic ignored "-Wreserved-user-defined-literal"
 
-// Apple LLVM 8.1.0 (Xcode 8.3) introduced more warnings that we want to ignore for now
+// Apple LLVM 9.0 (Xcode 9.0)
+#if (__clang_major__ >= 9)
+#pragma clang diagnostic ignored "-Wunguarded-availability-new"
+#endif
+
+// Apple LLVM 8.1.0 (Xcode 8.3) introduced -Wundefined-var-template
 #if (__clang_major__ > 8) || (__clang_major__ == 8 && __clang_minor__ >= 1)
+#pragma clang diagnostic ignored "-Wundefined-var-template"
 #pragma clang diagnostic ignored "-Wnullability-inferred-on-nested-type"
 #pragma clang diagnostic ignored "-Wobjc-protocol-property-synthesis"
 #pragma clang diagnostic ignored "-Wnullability-completeness-on-arrays"
diff --git a/Engine/Source/Runtime/CoreUObject/Private/Serialization/AsyncLoading.cpp b/Engine/Source/Runtime/CoreUObject/Private/Serialization/AsyncLoading.cpp
--- a/Engine/Source/Runtime/CoreUObject/Private/Serialization/AsyncLoading.cpp
+++ b/Engine/Source/Runtime/CoreUObject/Private/Serialization/AsyncLoading.cpp
@@ -7042,6 +7042,7 @@ void FArchiveAsync2::ReadCallback(bool bWasCancelled, IAsyncReadRequest* Request
 			if (GEventDrivenLoaderEnabled)
 			{
 				// in this case we don't need to serialize the summary because we know the header is the whole file
+				FScopeLock Lock(&SummaryRacePreventer);
 				HeaderSize = FileSize;
 				LogItem(TEXT("Starting Split Header"), 0, FileSize);
 				PrecacheInternal(0, HeaderSize);
diff --git a/Engine/Source/Runtime/CoreUObject/Public/Misc/AssetRegistryInterface.h b/Engine/Source/Runtime/CoreUObject/Public/Misc/AssetRegistryInterface.h
--- a/Engine/Source/Runtime/CoreUObject/Public/Misc/AssetRegistryInterface.h
+++ b/Engine/Source/Runtime/CoreUObject/Public/Misc/AssetRegistryInterface.h
@@ -11,22 +11,26 @@ namespace EAssetRegistryDependencyType
 	enum Type
 	{
 		// Dependencies which don't need to be loaded for the object to be used (i.e. string asset references)
-		Soft = 1,
+		Soft = 0x01,
 
 		// Dependencies which are required for correct usage of the source asset, and must be loaded at the same time
-		Hard = 2,
+		Hard = 0x02,
 
 		// References to specific SearchableNames inside a package
-		SearchableName = 4,
+		SearchableName = 0x04,
 
-		// Reference that says one object manages another object, set when Primary Assets manage individual assets
-		Manage = 8,
+		// Indirect management references, these are set through recursion for Primary Assets that manage packages or other primary assets
+		SoftManage = 0x08,
+
+		// Reference that says one object directly manages another object, set when Primary Assets manage things explicitly
+		HardManage = 0x10,
 	};
 
 	static const Type None = (Type)(0);
-	static const Type All = (Type)(Soft | Hard | SearchableName | Manage);
+	static const Type All = (Type)(Soft | Hard | SearchableName | SoftManage | HardManage);
 	static const Type Packages = (Type)(Soft | Hard);
-};
+	static const Type Manage = (Type)(SoftManage | HardManage);
+}
 
 /**
 * HotReload module interface
diff --git a/Engine/Source/Runtime/Engine/Classes/Engine/AssetManager.h b/Engine/Source/Runtime/Engine/Classes/Engine/AssetManager.h
--- a/Engine/Source/Runtime/Engine/Classes/Engine/AssetManager.h
+++ b/Engine/Source/Runtime/Engine/Classes/Engine/AssetManager.h
@@ -49,6 +49,15 @@ public:
 	/** Asset Type of Label used to tag other assets */
 	static const FPrimaryAssetType PrimaryAssetLabelType;
 
+	/** Type representing a packaging chunk, this is a virtual type that is never loaded off disk */
+	static const FPrimaryAssetType PackageChunkType;
+
+	/** Creates a PrimaryAssetId from a chunk id */
+	static FPrimaryAssetId CreatePrimaryAssetIdFromChunkId(int32 ChunkId);
+
+	/** Extracts a chunk id from a primary asset id, returns INDEX_NONE if it is not PackageChunkType */
+	static int32 ExtractChunkIdFromPrimaryAssetId(const FPrimaryAssetId& PrimaryAssetId);
+
 	// BUILDING ASSET DIRECTORY
 
 	/** 
@@ -83,7 +92,7 @@ public:
 	virtual bool AddDynamicAsset(const FPrimaryAssetId& PrimaryAssetId, const FSoftObjectPath& AssetPath, const FAssetBundleData& BundleData);
 
 	/** This will expand out references in the passed in AssetBundleData that are pointing to other primary assets with bundles. This is useful to preload entire webs of assets */
-	virtual void RecursivelyExpandBundleData(FAssetBundleData& BundleData);
+	virtual void RecursivelyExpandBundleData(FAssetBundleData& BundleData) const;
 
 	// ACCESSING ASSET DIRECTORY
 
@@ -233,6 +242,17 @@ public:
 	/** Fills in a TMap with the pending/active loading state of every asset */
 	void GetPrimaryAssetBundleStateMap(TMap<FPrimaryAssetId, TArray<FName>>& BundleStateMap, bool bForceCurrent = false) const;
 
+	/**
+	 * Fills in a set of object paths with the assets that need to be loaded, for a given Primary Asset and bundle list
+	 *
+	 * @param OutAssetLoadSet	Set that will have asset paths added to it
+	 * @param PrimaryAssetId	Asset that would be loaded
+	 * @param LoadBundles		List of bundles to load for those assets
+	 * @param bLoadRecursive	If true, this will call RecursivelyExpandBundleData and recurse into sub bundles of other primary assets loaded by a bundle reference
+	 * @return					True if primary asset id was found
+	 */
+	bool GetPrimaryAssetLoadSet(TSet<FSoftObjectPath>& OutAssetLoadSet, const FPrimaryAssetId& PrimaryAssetId, const TArray<FName>& LoadBundles, bool bLoadRecursive) const;
+
 	/**
 	 * Preloads data for a set of assets in a specific bundle state, and returns a handle you must keep active.
 	 * These assets are not officially Loaded, so Unload/ChangeBundleState will not affect them and if you release the handle without otherwise loading the assets they will be freed.
@@ -366,8 +386,18 @@ public:
 	/** Returns true if the specified asset package can be cooked, will error and return false if it is disallowed */
 	virtual bool VerifyCanCookPackage(FName PackageName, bool bLogError = true) const;
 
-	/** For a given package and platform, return what Chunks it should be assigned to, games can override this as needed. Returns false if no information found */
-	virtual bool GetPackageChunkIds(FName PackageName, const class ITargetPlatform* TargetPlatform, const TArray<int32>& ExistingChunkList, TArray<int32>& OutChunkList) const;
+	/** 
+	 * For a given package and platform, return what Chunks it should be assigned to, games can override this as needed. Returns false if no information found 
+	 * @param PackageName Package to check chunks for
+	 * @param TargetPlatform Can be used to do platform-specific chunking, this is null when previewing in the editor
+	 * @param ExistingChunkList List of chunks that the asset was previously assigned to, may be empty
+	 * @param OutChunkList List of chunks to actually assign this to
+	 * @param OutOverrideChunkList List of chunks that were added due to override rules, not just normal primary asset rules. Tools use this for dependency checking
+	 */
+	virtual bool GetPackageChunkIds(FName PackageName, const class ITargetPlatform* TargetPlatform, const TArray<int32>& ExistingChunkList, TArray<int32>& OutChunkList, TArray<int32>* OutOverrideChunkList = nullptr) const;
+
+	/** Returns the list of chunks assigned to the list of primary assets, which is usually a manager list. This is called by GetPackageChunkIds */
+	virtual bool GetPrimaryAssetSetChunkIds(const TSet<FPrimaryAssetId>& PrimaryAssetSet, const class ITargetPlatform* TargetPlatform, const TArray<int32>& ExistingChunkList, TArray<int32>& OutChunkList) const;
 
 	/** Refresh the entire set of asset data, can call from editor when things have changed dramatically */
 	virtual void RefreshPrimaryAssetDirectory();
@@ -378,6 +408,9 @@ public:
 	/** Updates the asset management database if needed */
 	virtual void UpdateManagementDatabase(bool bForceRefresh = false);
 
+	/** Returns the chunk information map computed during UpdateManagementDatabase */
+	const TMap<int32, FAssetManagerChunkInfo>& GetChunkManagementMap() const;
+
 	/** Handles applying Asset Labels and should be overridden to do any game-specific labelling */
 	virtual void ApplyPrimaryAssetLabels();
 
@@ -482,8 +515,10 @@ protected:
 
 	/** Copy of the asset state before PIE was entered, return to that when PIE completes */
 	TMap<FPrimaryAssetId, TArray<FName>> PrimaryAssetStateBeforePIE;
-#endif // WITH_EDITOR
 
+	/** Cached map of chunk ids to lists of assets in that chunk */
+	TMap<int32, FAssetManagerChunkInfo> CachedChunkMap;
+#endif // WITH_EDITOR
 
 	/** Map from object path to Primary Asset Id */
 	TMap<FName, FPrimaryAssetId> AssetPathMap;
diff --git a/Engine/Source/Runtime/Engine/Classes/Engine/AssetManagerTypes.h b/Engine/Source/Runtime/Engine/Classes/Engine/AssetManagerTypes.h
--- a/Engine/Source/Runtime/Engine/Classes/Engine/AssetManagerTypes.h
+++ b/Engine/Source/Runtime/Engine/Classes/Engine/AssetManagerTypes.h
@@ -138,3 +138,13 @@ public:
 	ENGINE_API bool FillRuntimeData();
 
 };
+
+/** Information about a package chunk, computed by the asset manager or read out of the cooked asset registry */
+struct FAssetManagerChunkInfo
+{
+	/** Packages/PrimaryAssets that were explicitly added to a chunk */
+	TSet<FAssetIdentifier> ExplicitAssets;
+
+	/** All packages/Primary Assets in a chunk, includes everything in Explicit plus recursively added ones */
+	TSet<FAssetIdentifier> AllAssets;
+};
\ No newline at end of file
diff --git a/Engine/Source/Runtime/Engine/Private/AssetManager.cpp b/Engine/Source/Runtime/Engine/Private/AssetManager.cpp
--- a/Engine/Source/Runtime/Engine/Private/AssetManager.cpp
+++ b/Engine/Source/Runtime/Engine/Private/AssetManager.cpp
@@ -97,6 +97,7 @@ struct FPrimaryAssetTypeData
 
 const FPrimaryAssetType UAssetManager::MapType = FName(TEXT("Map"));
 const FPrimaryAssetType UAssetManager::PrimaryAssetLabelType = FName(TEXT("PrimaryAssetLabel"));
+const FPrimaryAssetType UAssetManager::PackageChunkType = FName(TEXT("PackageChunk"));
 
 UAssetManager::UAssetManager()
 {
@@ -193,6 +194,27 @@ UAssetManager* UAssetManager::GetIfValid()
 	return nullptr;
 }
 
+FPrimaryAssetId UAssetManager::CreatePrimaryAssetIdFromChunkId(int32 ChunkId)
+{
+	if (ChunkId == INDEX_NONE)
+	{
+		return FPrimaryAssetId();
+	}
+
+	// Name_0 is actually stored as 1 inside FName, so offset
+	static FName ChunkName = "Chunk";
+	return FPrimaryAssetId(PackageChunkType, FName(ChunkName, ChunkId + 1));
+}
+
+int32 UAssetManager::ExtractChunkIdFromPrimaryAssetId(const FPrimaryAssetId& PrimaryAssetId)
+{
+	if (PrimaryAssetId.PrimaryAssetType == PackageChunkType)
+	{
+		return PrimaryAssetId.PrimaryAssetName.GetNumber() - 1;
+	}
+	return INDEX_NONE;
+}
+
 IAssetRegistry& UAssetManager::GetAssetRegistry() const
 {
 	if (!CachedAssetRegistry)
@@ -603,7 +625,7 @@ bool UAssetManager::AddDynamicAsset(const FPrimaryAssetId& PrimaryAssetId, const
 	return true;
 }
 
-void UAssetManager::RecursivelyExpandBundleData(FAssetBundleData& BundleData)
+void UAssetManager::RecursivelyExpandBundleData(FAssetBundleData& BundleData) const
 {
 	TArray<FSoftObjectPath> ReferencesToExpand;
 	TSet<FName> FoundBundleNames;
@@ -1130,49 +1152,54 @@ TSharedPtr<FStreamableHandle> UAssetManager::ChangeBundleStateForMatchingPrimary
 	return nullptr;
 }
 
-TSharedPtr<FStreamableHandle> UAssetManager::PreloadPrimaryAssets(const TArray<FPrimaryAssetId>& AssetsToLoad, const TArray<FName>& LoadBundles, bool bLoadRecursive, FStreamableDelegate DelegateToCall, TAsyncLoadPriority Priority)
+bool UAssetManager::GetPrimaryAssetLoadSet(TSet<FSoftObjectPath>& OutAssetLoadSet, const FPrimaryAssetId& PrimaryAssetId, const TArray<FName>& LoadBundles, bool bLoadRecursive) const
 {
-	TSet<FSoftObjectPath> PathsToLoad;
-	FString DebugName;
-	TSharedPtr<FStreamableHandle> ReturnHandle;
-
-	for (const FPrimaryAssetId& PrimaryAssetId : AssetsToLoad)
+	const FPrimaryAssetData* NameData = GetNameData(PrimaryAssetId);
+	if (NameData)
 	{
-		FPrimaryAssetData* NameData = GetNameData(PrimaryAssetId);
+		// Gather asset refs
+		const FSoftObjectPath& AssetPath = NameData->AssetPtr.ToSoftObjectPath();
+		if (!AssetPath.IsNull())
+		{
+			// Dynamic types can have no base asset path
+			OutAssetLoadSet.Add(AssetPath);
+		}
 
-		if (NameData)
+		// Construct a temporary bundle data with the bundles specified
+		FAssetBundleData TempBundleData;
+		for (const FName& BundleName : LoadBundles)
 		{
-			// Gather asset refs
-			const FSoftObjectPath& AssetPath = NameData->AssetPtr.ToSoftObjectPath();
+			FAssetBundleEntry Entry = GetAssetBundleEntry(PrimaryAssetId, BundleName);
 
-			if (!AssetPath.IsNull())
+			if (Entry.IsValid())
 			{
-				// Dynamic types can have no base asset path
-				PathsToLoad.Add(AssetPath);
+				TempBundleData.Bundles.Add(Entry);
 			}
+		}
 
-			// Construct a temporary bundle data with the bundles specified
-			FAssetBundleData TempBundleData;
-			for (const FName& BundleName : LoadBundles)
-			{
-				FAssetBundleEntry Entry = GetAssetBundleEntry(PrimaryAssetId, BundleName);
-
-				if (Entry.IsValid())
-				{
-					TempBundleData.Bundles.Add(Entry);
-				}
-			}
+		if (bLoadRecursive)
+		{
+			RecursivelyExpandBundleData(TempBundleData);
+		}
 
-			if (bLoadRecursive)
-			{
-				RecursivelyExpandBundleData(TempBundleData);
-			}
+		for (const FAssetBundleEntry& Entry : TempBundleData.Bundles)
+		{
+			OutAssetLoadSet.Append(Entry.BundleAssets);
+		}
+	}
+	return NameData != nullptr;
+}
 
-			for (const FAssetBundleEntry& Entry : TempBundleData.Bundles)
-			{
-				PathsToLoad.Append(Entry.BundleAssets);
-			}
+TSharedPtr<FStreamableHandle> UAssetManager::PreloadPrimaryAssets(const TArray<FPrimaryAssetId>& AssetsToLoad, const TArray<FName>& LoadBundles, bool bLoadRecursive, FStreamableDelegate DelegateToCall, TAsyncLoadPriority Priority)
+{
+	TSet<FSoftObjectPath> PathsToLoad;
+	FString DebugName;
+	TSharedPtr<FStreamableHandle> ReturnHandle;
 
+	for (const FPrimaryAssetId& PrimaryAssetId : AssetsToLoad)
+	{
+		if (GetPrimaryAssetLoadSet(PathsToLoad, PrimaryAssetId, LoadBundles, bLoadRecursive))
+		{
 			if (DebugName.IsEmpty())
 			{
 				DebugName += TEXT("Preloading ");
@@ -2194,6 +2221,12 @@ void UAssetManager::DumpReferencersForPackage(const TArray< FString >& PackageNa
 
 bool UAssetManager::ShouldScanPrimaryAssetType(FPrimaryAssetTypeInfo& TypeInfo) const
 {
+	if (!ensureMsgf(TypeInfo.PrimaryAssetType != PackageChunkType, TEXT("Cannot use %s as an asset manager type, this is reserved for internal use"), *TypeInfo.PrimaryAssetType.ToString()))
+	{
+		// Cannot use this as a proper type
+		return false;
+	}
+
 	if (TypeInfo.bIsEditorOnly && !GIsEditor)
 	{
 		return false;
@@ -2298,8 +2331,11 @@ bool UAssetManager::GetPackageManagers(FName PackageName, bool bRecurseToParents
 				{
 					for (const FPrimaryAssetId& Manager : *ManagementParents)
 					{
-						// Add to end of list, this will recurse again if needed
-						ReferencingPrimaryAssets.AddUnique(Manager);
+						if (!ManagerSet.Contains(Manager))
+						{
+							// Add to end of list, this will recurse again if needed
+							ReferencingPrimaryAssets.Add(Manager);
+						}
 					}
 				}
 			}
@@ -2526,6 +2562,9 @@ void UAssetManager::UpdateManagementDatabase(bool bForceRefresh)
 	// List of references to not recurse on, priority doesn't matter
 	TMultiMap<FAssetIdentifier, FAssetIdentifier> NoReferenceManagementMap;
 
+	// List of packages that need to have their chunks updated
+	TSet<FName> PackagesToUpdateChunksFor;
+
 	for (const TPair<FName, TSharedRef<FPrimaryAssetTypeData>>& TypePair : AssetTypeMap)
 	{
 		const FPrimaryAssetTypeData& TypeData = TypePair.Value.Get();
@@ -2547,6 +2586,7 @@ void UAssetManager::UpdateManagementDatabase(bool bForceRefresh)
 				FName PackageName = FName(*AssetRef.GetLongPackageName());
 
 				AssetPackagesReferenced.AddUnique(PackageName);
+				PackagesToUpdateChunksFor.Add(PackageName);
 			}
 
 			TMap<FName, FAssetBundleEntry>* BundleMap = CachedAssetBundles.Find(PrimaryAssetId);
@@ -2558,9 +2598,10 @@ void UAssetManager::UpdateManagementDatabase(bool bForceRefresh)
 				{
 					for (const FSoftObjectPath& BundleAssetRef : BundlePair.Value.BundleAssets)
 					{
-						FString PackageName = BundleAssetRef.GetLongPackageName();
+						FName PackageName = FName(*BundleAssetRef.GetLongPackageName());
 
-						AssetPackagesReferenced.AddUnique(FName(*PackageName));
+						AssetPackagesReferenced.AddUnique(PackageName);
+						PackagesToUpdateChunksFor.Add(PackageName);
 					}
 				}
 			}
@@ -2585,9 +2626,14 @@ void UAssetManager::UpdateManagementDatabase(bool bForceRefresh)
 	const bool bAllowInPIE = true;
 	SlowTask.MakeDialog(bShowCancelButton, bAllowInPIE);
 
-	auto SetManagerPredicate = [this](const FAssetIdentifier& Manager, const FAssetIdentifier& Source, const FAssetIdentifier& Target, EAssetRegistryDependencyType::Type DependencyType, EAssetSetManagerFlags::Type Flags)
+	auto SetManagerPredicate = [this, &PackagesToUpdateChunksFor](const FAssetIdentifier& Manager, const FAssetIdentifier& Source, const FAssetIdentifier& Target, EAssetRegistryDependencyType::Type DependencyType, EAssetSetManagerFlags::Type Flags)
 	{
-		return this->ShouldSetManager(Manager, Source, Target, DependencyType, Flags);
+		EAssetSetManagerResult::Type Result = this->ShouldSetManager(Manager, Source, Target, DependencyType, Flags);
+		if (Result != EAssetSetManagerResult::DoNotSet && Target.IsPackage())
+		{
+			PackagesToUpdateChunksFor.Add(Target.PackageName);
+		}
+		return Result;
 	};
 
 	for (int32 PriorityIndex = 0; PriorityIndex < PriorityArray.Num(); PriorityIndex++)
@@ -2607,6 +2653,7 @@ void UAssetManager::UpdateManagementDatabase(bool bForceRefresh)
 
 
 	TMultiMap<FAssetIdentifier, FAssetIdentifier> PrimaryAssetIdManagementMap;
+	TArray<int32> ChunkList, ExistingChunkList, OverrideChunkList;
 
 	// Update management parent list, which is PrimaryAssetId -> PrimaryAssetId
 	for (const TPair<FName, TSharedRef<FPrimaryAssetTypeData>>& TypePair : AssetTypeMap)
@@ -2619,12 +2666,12 @@ void UAssetManager::UpdateManagementDatabase(bool bForceRefresh)
 			FPrimaryAssetId PrimaryAssetId(TypePair.Key, NamePair.Key);
 			const FSoftObjectPath& AssetRef = NameData.AssetPtr.ToSoftObjectPath();
 
+			TSet<FPrimaryAssetId> Managers;
+
 			if (AssetRef.IsValid())
 			{
 				FName PackageName = FName(*AssetRef.GetLongPackageName());
 
-				TSet<FPrimaryAssetId> Managers;
-
 				if (GetPackageManagers(PackageName, false, Managers) && Managers.Num() > 1)
 				{
 					// Find all managers that aren't this specific asset
@@ -2640,6 +2687,22 @@ void UAssetManager::UpdateManagementDatabase(bool bForceRefresh)
 					}
 				}
 			}
+			else
+			{
+				Managers.Add(PrimaryAssetId);
+			}
+
+			// Compute chunk assignment and store those as manager references
+			ChunkList.Reset();
+			GetPrimaryAssetSetChunkIds(Managers, nullptr, ExistingChunkList, ChunkList);
+
+			for (int32 ChunkId : ChunkList)
+			{
+				FPrimaryAssetId ChunkPrimaryAsset = CreatePrimaryAssetIdFromChunkId(ChunkId);
+
+				CachedChunkMap.FindOrAdd(ChunkId).ExplicitAssets.Add(PrimaryAssetId);
+				PrimaryAssetIdManagementMap.Add(ChunkPrimaryAsset, PrimaryAssetId);
+			}
 		}
 	}
 
@@ -2648,10 +2711,36 @@ void UAssetManager::UpdateManagementDatabase(bool bForceRefresh)
 		AssetRegistry.SetManageReferences(PrimaryAssetIdManagementMap, false, EAssetRegistryDependencyType::None);
 	}
 
+	// Update chunk package list for all chunks
+	for (FName PackageName : PackagesToUpdateChunksFor)
+	{
+		ChunkList.Reset();
+		OverrideChunkList.Reset();
+		GetPackageChunkIds(PackageName, nullptr, ExistingChunkList, ChunkList, &OverrideChunkList);
+
+		if (ChunkList.Num() > 0)
+		{
+			for (int32 ChunkId : ChunkList)
+			{
+				CachedChunkMap.FindOrAdd(ChunkId).AllAssets.Add(PackageName);
+
+				if (OverrideChunkList.Contains(ChunkId))
+				{
+					// This was in the override list, so add an explicit dependency
+					CachedChunkMap.FindOrAdd(ChunkId).ExplicitAssets.Add(PackageName);
+				}
+			}
+		}
+	}
 
 	bIsManagementDatabaseCurrent = true;
 }
 
+const TMap<int32, FAssetManagerChunkInfo>& UAssetManager::GetChunkManagementMap() const
+{
+	return CachedChunkMap;
+}
+
 void UAssetManager::ApplyPrimaryAssetLabels()
 {
 	// Load all of them off disk. Turn off string asset reference tracking to avoid them getting cooked
@@ -2759,24 +2848,38 @@ bool UAssetManager::VerifyCanCookPackage(FName PackageName, bool bLogError) cons
 	return true;
 }
 
-bool UAssetManager::GetPackageChunkIds(FName PackageName, const ITargetPlatform* TargetPlatform, const TArray<int32>& ExistingChunkList, TArray<int32>& OutChunkList) const
+bool UAssetManager::GetPackageChunkIds(FName PackageName, const ITargetPlatform* TargetPlatform, const TArray<int32>& ExistingChunkList, TArray<int32>& OutChunkList, TArray<int32>* OutOverrideChunkList) const
 {
 	// Include preset chunks
 	OutChunkList.Append(ExistingChunkList);
+	if (OutOverrideChunkList)
+	{
+		OutOverrideChunkList->Append(ExistingChunkList);
+	}
 
 	if (PackageName.ToString().StartsWith(TEXT("/Engine/"), ESearchCase::CaseSensitive))
 	{
 		// Some engine content is only referenced by string, make sure it's all in chunk 0 to avoid issues
 		OutChunkList.AddUnique(0);
+
+		if (OutOverrideChunkList)
+		{
+			OutOverrideChunkList->AddUnique(0);
+		}
 	}
 
 	// Add all chunk ids from the asset rules of managers. By default priority will not override other chunks
-	bool bFoundAny = false;
 	TSet<FPrimaryAssetId> Managers;
 
 	GetPackageManagers(PackageName, true, Managers);
+	return GetPrimaryAssetSetChunkIds(Managers, TargetPlatform, ExistingChunkList, OutChunkList);
+}
 
-	for (const FPrimaryAssetId& PrimaryAssetId : Managers)
+bool UAssetManager::GetPrimaryAssetSetChunkIds(const TSet<FPrimaryAssetId>& PrimaryAssetSet, const class ITargetPlatform* TargetPlatform, const TArray<int32>& ExistingChunkList, TArray<int32>& OutChunkList) const
+{
+	bool bFoundAny = false;
+	int32 HighestChunk = 0;
+	for (const FPrimaryAssetId& PrimaryAssetId : PrimaryAssetSet)
 	{
 		FPrimaryAssetRules Rules = GetPrimaryAssetRules(PrimaryAssetId);
 
@@ -2784,20 +2887,16 @@ bool UAssetManager::GetPackageChunkIds(FName PackageName, const ITargetPlatform*
 		{
 			bFoundAny = true;
 			OutChunkList.AddUnique(Rules.ChunkId);
-		}
-	}
 
-	int32 HighestChunk = 0;
-	for (int32 ChunkId : OutChunkList)
-	{
-		if (ChunkId > HighestChunk)
-		{
-			HighestChunk = ChunkId;
+			if (Rules.ChunkId > HighestChunk)
+			{
+				HighestChunk = Rules.ChunkId;
+			}
 		}
 	}
 
 	// Use chunk dependency info to remove redundant chunks
-	UChunkDependencyInfo *DependencyInfo = GetMutableDefault<UChunkDependencyInfo>();
+	UChunkDependencyInfo* DependencyInfo = GetMutableDefault<UChunkDependencyInfo>();
 	DependencyInfo->GetOrBuildChunkDependencyGraph(HighestChunk);
 	DependencyInfo->RemoveRedundantChunks(OutChunkList);
 
diff --git a/Engine/Source/Runtime/Engine/Private/Components/ChildActorComponent.cpp b/Engine/Source/Runtime/Engine/Private/Components/ChildActorComponent.cpp
--- a/Engine/Source/Runtime/Engine/Private/Components/ChildActorComponent.cpp
+++ b/Engine/Source/Runtime/Engine/Private/Components/ChildActorComponent.cpp
@@ -532,8 +532,9 @@ void UChildActorComponent::CreateChildActor()
 				{
 					Params.ObjectFlags &= ~RF_Transactional;
 				}
-				if (HasAllFlags(RF_Transient))
+				if (HasAllFlags(RF_Transient) || IsEditorOnly())
 				{
+					// If we are either transient or editor only, set our created actor to transient. We can't programatically set editor only on an actor so this is the best option
 					Params.ObjectFlags |= RF_Transient;
 				}
 
diff --git a/Engine/Source/Runtime/Engine/Private/GameplayStatics.cpp b/Engine/Source/Runtime/Engine/Private/GameplayStatics.cpp
--- a/Engine/Source/Runtime/Engine/Private/GameplayStatics.cpp
+++ b/Engine/Source/Runtime/Engine/Private/GameplayStatics.cpp
@@ -61,6 +61,7 @@ struct FSaveGameFileVersion
 
 DECLARE_CYCLE_STAT(TEXT("BreakHitResult"), STAT_BreakHitResult, STATGROUP_Game);
 DECLARE_CYCLE_STAT(TEXT("MakeHitResult"), STAT_MakeHitResult, STATGROUP_Game);
+DECLARE_CYCLE_STAT(TEXT("SpawnTime"), STAT_SpawnTime, STATGROUP_Game);
 
 //////////////////////////////////////////////////////////////////////////
 // UGameplayStatics
@@ -411,6 +412,7 @@ class AActor* UGameplayStatics::BeginSpawningActorFromClass(const UObject* World
 
 class AActor* UGameplayStatics::BeginDeferredActorSpawnFromClass(const UObject* WorldContextObject, TSubclassOf<AActor> ActorClass, const FTransform& SpawnTransform, ESpawnActorCollisionHandlingMethod CollisionHandlingOverride /*= ESpawnActorCollisionHandlingMethod::Undefined*/, AActor* Owner /*= nullptr*/)
 {
+	SCOPE_CYCLE_COUNTER(STAT_SpawnTime);
 	if (UClass* Class = *ActorClass)
 	{
 		// If the WorldContextObject is a Pawn we will use that as the instigator.
@@ -445,6 +447,7 @@ class AActor* UGameplayStatics::BeginDeferredActorSpawnFromClass(const UObject*
 
 AActor* UGameplayStatics::FinishSpawningActor(AActor* Actor, const FTransform& SpawnTransform)
 {
+	SCOPE_CYCLE_COUNTER(STAT_SpawnTime);
 	if (Actor)
 	{
 		Actor->FinishSpawning(SpawnTransform);
diff --git a/Engine/Source/Runtime/JsonUtilities/Private/JsonObjectConverter.cpp b/Engine/Source/Runtime/JsonUtilities/Private/JsonObjectConverter.cpp
--- a/Engine/Source/Runtime/JsonUtilities/Private/JsonObjectConverter.cpp
+++ b/Engine/Source/Runtime/JsonUtilities/Private/JsonObjectConverter.cpp
@@ -622,12 +622,12 @@ bool FJsonObjectConverter::JsonValueToUProperty(TSharedPtr<FJsonValue> JsonValue
 		return false;
 	}
 
-	bool bArrayProperty = Property->IsA<UArrayProperty>();
+	bool bArrayLikeProperty = Property->IsA<UArrayProperty>() || Property->IsA<USetProperty>();
 	bool bJsonArray = JsonValue->Type == EJson::Array;
 
 	if (!bJsonArray)
 	{
-		if (bArrayProperty)
+		if (bArrayLikeProperty)
 		{
 			UE_LOG(LogJson, Error, TEXT("JsonValueToUProperty - Attempted to import TArray from non-array JSON key"));
 			return false;			
@@ -642,7 +642,7 @@ bool FJsonObjectConverter::JsonValueToUProperty(TSharedPtr<FJsonValue> JsonValue
 	}
 
 	// In practice, the ArrayDim == 1 check ought to be redundant, since nested arrays of UPropertys are not supported
-	if (bArrayProperty && Property->ArrayDim == 1)
+	if (bArrayLikeProperty && Property->ArrayDim == 1)
 	{
 		// Read into TArray
 		return ConvertScalarJsonValueToUProperty(JsonValue, Property, OutValue, CheckFlags, SkipFlags);
diff --git a/Engine/Source/Runtime/MovieScene/Public/MovieSceneFolder.h b/Engine/Source/Runtime/MovieScene/Public/MovieSceneFolder.h
--- a/Engine/Source/Runtime/MovieScene/Public/MovieSceneFolder.h
+++ b/Engine/Source/Runtime/MovieScene/Public/MovieSceneFolder.h
@@ -11,8 +11,8 @@
 
 class UMovieSceneTrack;
 
-/** Reprents a folder used for organizing objects in tracks in a movie scene. */
-UCLASS()
+/** Represents a folder used for organizing objects in tracks in a movie scene. */
+UCLASS(DefaultToInstanced)
 class MOVIESCENE_API UMovieSceneFolder : public UObject
 {
 	GENERATED_UCLASS_BODY()
