Copying //UE4/Dev-Framework to //UE4/Dev-Main (Source: //UE4/Dev-Framework @ 3779049)

diff --git a/Engine/Source/Runtime/Engine/Private/StaticMesh.cpp b/Engine/Source/Runtime/Engine/Private/StaticMesh.cpp
--- a/Engine/Source/Runtime/Engine/Private/StaticMesh.cpp
+++ b/Engine/Source/Runtime/Engine/Private/StaticMesh.cpp
@@ -1568,11 +1568,6 @@ void UStaticMesh::InitResources()
 #endif // STATS
 }
 
-/**
- * Returns the size of the object/ resource for display to artists/ LDs in the Editor.
- *
- * @return size of resource as to be displayed to artists/ LDs in the Editor.
- */
 void UStaticMesh::GetResourceSizeEx(FResourceSizeEx& CumulativeResourceSize)
 {
 	Super::GetResourceSizeEx(CumulativeResourceSize);
@@ -1581,30 +1576,6 @@ void UStaticMesh::GetResourceSizeEx(FResourceSizeEx& CumulativeResourceSize)
 	{
 		RenderData->GetResourceSizeEx(CumulativeResourceSize);
 	}
-	if (CumulativeResourceSize.GetResourceSizeMode() == EResourceSizeMode::Inclusive)
-	{
-		TSet<UMaterialInterface*> UniqueMaterials;
-		for (int32 MaterialIndex = 0; MaterialIndex < StaticMaterials.Num(); ++MaterialIndex)
-		{
-			const FStaticMaterial& StaticMaterial = StaticMaterials[MaterialIndex];
-			bool bAlreadyCounted = false;
-			UniqueMaterials.Add(StaticMaterial.MaterialInterface,&bAlreadyCounted);
-			if (!bAlreadyCounted && StaticMaterial.MaterialInterface)
-			{
-				StaticMaterial.MaterialInterface->GetResourceSizeEx(CumulativeResourceSize);
-			}
-		}
-
-		if(BodySetup)
-		{
-			BodySetup->GetResourceSizeEx(CumulativeResourceSize);
-		}
-	}
-}
-
-SIZE_T FStaticMeshRenderData::GetResourceSize() const
-{
-	return GetResourceSizeBytes();
 }
 
 void FStaticMeshRenderData::GetResourceSizeEx(FResourceSizeEx& CumulativeResourceSize) const
@@ -1613,10 +1584,6 @@ void FStaticMeshRenderData::GetResourceSizeEx(FResourceSizeEx& CumulativeResourc
 
 	// Count dynamic arrays.
 	CumulativeResourceSize.AddUnknownMemoryBytes(LODResources.GetAllocatedSize());
-#if WITH_EDITORONLY_DATA
-	CumulativeResourceSize.AddDedicatedSystemMemoryBytes(DerivedDataKey.GetAllocatedSize());
-	CumulativeResourceSize.AddDedicatedSystemMemoryBytes(WedgeMap.GetAllocatedSize());
-#endif // #if WITH_EDITORONLY_DATA
 
 	for(int32 LODIndex = 0;LODIndex < LODResources.Num();LODIndex++)
 	{
@@ -1647,13 +1614,6 @@ void FStaticMeshRenderData::GetResourceSizeEx(FResourceSizeEx& CumulativeResourc
 #endif // #if WITH_EDITORONLY_DATA
 }
 
-SIZE_T FStaticMeshRenderData::GetResourceSizeBytes() const
-{
-	FResourceSizeEx ResSize;
-	GetResourceSizeEx(ResSize);
-	return ResSize.GetTotalMemoryBytes();
-}
-
 int32 UStaticMesh::GetNumVertices(int32 LODIndex) const
 {
 	int32 NumVertices = 0;
@@ -2094,11 +2054,10 @@ void UStaticMesh::GetAssetRegistryTags(TArray<FAssetRegistryTag>& OutTags) const
 		DefaultCollisionName = BodySetup->DefaultInstance.GetCollisionProfileName();
 	}
 
-	static const UEnum *ComplexityEnum = FindObject<UEnum>(ANY_PACKAGE, TEXT("ECollisionTraceFlag"), true);
 	FString ComplexityString;
-	if (BodySetup && ComplexityEnum)
+	if (BodySetup != nullptr)
 	{
-		ComplexityString = ComplexityEnum->GetNameStringByValue((int64)BodySetup->GetCollisionTraceFlag());
+		ComplexityString = Lex::ToString((ECollisionTraceFlag)BodySetup->GetCollisionTraceFlag());
 	}
 
 	OutTags.Add( FAssetRegistryTag("Triangles", FString::FromInt(NumTriangles), FAssetRegistryTag::TT_Numerical) );
diff --git a/Engine/Source/Runtime/Engine/Private/StaticMeshActor.cpp b/Engine/Source/Runtime/Engine/Private/StaticMeshActor.cpp
--- a/Engine/Source/Runtime/Engine/Private/StaticMeshActor.cpp
+++ b/Engine/Source/Runtime/Engine/Private/StaticMeshActor.cpp
@@ -167,16 +167,20 @@ void AStaticMeshActor::CheckForErrors()
 
 	if( !StaticMeshComponent )
 	{
+		FFormatNamedArguments Arguments;
+		Arguments.Add(TEXT("ActorName"), FText::FromString(GetName()));
 		MapCheck.Warning()
 			->AddToken(FUObjectToken::Create(this))
-			->AddToken(FTextToken::Create(LOCTEXT( "MapCheck_Message_StaticMeshComponent", "Static mesh actor has NULL StaticMeshComponent property - please delete" ) ))
+			->AddToken(FTextToken::Create(FText::Format(LOCTEXT( "MapCheck_Message_StaticMeshComponent", "Static mesh actor {ActorName} has NULL StaticMeshComponent property - please delete" ), Arguments)))
 			->AddToken(FMapErrorToken::Create(FMapErrors::StaticMeshComponent));
 	}
 	else if( StaticMeshComponent->GetStaticMesh() == nullptr )
 	{
+		FFormatNamedArguments Arguments;
+		Arguments.Add(TEXT("ActorName"), FText::FromString(GetName()));
 		MapCheck.Warning()
 			->AddToken(FUObjectToken::Create(this))
-			->AddToken(FTextToken::Create(LOCTEXT( "MapCheck_Message_StaticMeshNull", "Static mesh actor has NULL StaticMesh property" ) ))
+			->AddToken(FTextToken::Create(FText::Format(LOCTEXT( "MapCheck_Message_StaticMeshNull", "Static mesh actor {ActorName} has NULL StaticMesh property" ), Arguments)))
 			->AddToken(FMapErrorToken::Create(FMapErrors::StaticMeshNull));
 	}
 	else
diff --git a/Engine/Source/Runtime/Engine/Private/Texture2D.cpp b/Engine/Source/Runtime/Engine/Private/Texture2D.cpp
--- a/Engine/Source/Runtime/Engine/Private/Texture2D.cpp
+++ b/Engine/Source/Runtime/Engine/Private/Texture2D.cpp
@@ -919,14 +919,13 @@ void UTexture2D::GetResourceSizeEx(FResourceSizeEx& CumulativeResourceSize)
 
 	if (CumulativeResourceSize.GetResourceSizeMode() == EResourceSizeMode::Exclusive)
 	{
-		CumulativeResourceSize.AddUnknownMemoryBytes(CalcTextureMemorySize(GetNumResidentMips()));
+		// Use only loaded mips
+		CumulativeResourceSize.AddDedicatedVideoMemoryBytes(CalcTextureMemorySize(GetNumResidentMips()));
 	}
 	else
 	{
-		if (PlatformData)
-		{
-			CumulativeResourceSize.AddUnknownMemoryBytes(CalcTextureSize(GetSizeX(), GetSizeY(), GetPixelFormat(), GetNumMips()));
-		}
+		// Use all possible mips
+		CumulativeResourceSize.AddDedicatedVideoMemoryBytes(CalcTextureMemorySize(GetNumMipsAllowed(true)));
 	}
 }
 
diff --git a/Engine/Source/Runtime/Engine/Private/TextureCube.cpp b/Engine/Source/Runtime/Engine/Private/TextureCube.cpp
--- a/Engine/Source/Runtime/Engine/Private/TextureCube.cpp
+++ b/Engine/Source/Runtime/Engine/Private/TextureCube.cpp
@@ -351,7 +351,17 @@ FTextureResource* UTextureCube::CreateResource()
 void UTextureCube::GetResourceSizeEx(FResourceSizeEx& CumulativeResourceSize)
 {
 	Super::GetResourceSizeEx(CumulativeResourceSize);
-	CumulativeResourceSize.AddUnknownMemoryBytes(CalcTextureMemorySizeEnum(TMC_ResidentMips));
+
+	if (CumulativeResourceSize.GetResourceSizeMode() == EResourceSizeMode::Exclusive)
+	{
+		// Use only loaded mips
+		CumulativeResourceSize.AddDedicatedVideoMemoryBytes(CalcTextureMemorySizeEnum(TMC_ResidentMips));
+	}
+	else
+	{
+		// Use all possible mips
+		CumulativeResourceSize.AddDedicatedVideoMemoryBytes(CalcTextureMemorySizeEnum(TMC_AllMipsBiased));
+	}
 }
 
 #if WITH_EDITOR
diff --git a/Engine/Source/Runtime/Engine/Private/TickTaskManager.cpp b/Engine/Source/Runtime/Engine/Private/TickTaskManager.cpp
--- a/Engine/Source/Runtime/Engine/Private/TickTaskManager.cpp
+++ b/Engine/Source/Runtime/Engine/Private/TickTaskManager.cpp
@@ -807,17 +807,12 @@ public:
 		bool bDeferredRemove;
 	};
 
-	void RemoveAndRescheduleForInterval(FTickFunction* TickFunction)
-	{
-		verify(AllEnabledTickFunctions.Remove(TickFunction) == 1);
-		TickFunctionsToReschedule.Add(FTickScheduleDetails(TickFunction, TickFunction->TickInterval));
-	}
 	void RescheduleForIntervalParallel(FTickFunction* TickFunction)
 	{
 		// note we do the remove later!
 		TickFunctionsToReschedule.AddThreadsafe(FTickScheduleDetails(TickFunction, TickFunction->TickInterval, true));
 	}
-	/* Puts a TickFunction in to the cooldown state*/
+	/* Helper to presize reschedule array */
 	void ReserveTickFunctionCooldowns(int32 NumToReserve)
 	{
 		TickFunctionsToReschedule.Reserve(NumToReserve);
@@ -972,7 +967,7 @@ public:
 	/**
 	 * If there is infinite recursive spawning, log that and discard them
 	 */
-	void LogAndDisardRunawayNewlySpawned(ETickingGroup CurrentTickGroup)
+	void LogAndDiscardRunawayNewlySpawned(ETickingGroup CurrentTickGroup)
 	{
 		Context.TickGroup = CurrentTickGroup;
 		FTickTaskSequencer& TTS = FTickTaskSequencer::Get();
@@ -1480,7 +1475,7 @@ public:
 				// this is runaway recursive spawning.
 				for( int32 LevelIndex = 0; LevelIndex < LevelList.Num(); LevelIndex++ )
 				{
-					LevelList[LevelIndex]->LogAndDisardRunawayNewlySpawned(Context.TickGroup);
+					LevelList[LevelIndex]->LogAndDiscardRunawayNewlySpawned(Context.TickGroup);
 				}
 			}
 		}
diff --git a/Engine/Source/Runtime/Engine/Private/UnrealEngine.cpp b/Engine/Source/Runtime/Engine/Private/UnrealEngine.cpp
--- a/Engine/Source/Runtime/Engine/Private/UnrealEngine.cpp
+++ b/Engine/Source/Runtime/Engine/Private/UnrealEngine.cpp
@@ -2858,7 +2858,7 @@ struct FSortedParticleSet
 		, ModuleSize(0)
 		, ComponentSize(0)
 		, ComponentCount(0)
-		, ComponentResourceSize(EResourceSizeMode::Inclusive)
+		, ComponentResourceSize(EResourceSizeMode::EstimatedTotal)
 		, ComponentTrueResourceSize(EResourceSizeMode::Exclusive)
 	{
 	}
@@ -4448,7 +4448,7 @@ bool UEngine::HandleListParticleSystemsCommand( const TCHAR* Cmd, FOutputDevice&
 		FArchiveCountMem Count( Tree );
 		int32 RootSize = Count.GetMax();
 
-		SortedSets.Add(FSortedParticleSet(Description, RootSize, RootSize, 0, 0, 0, FResourceSizeEx(EResourceSizeMode::Inclusive), FResourceSizeEx(EResourceSizeMode::Exclusive)));
+		SortedSets.Add(FSortedParticleSet(Description, RootSize, RootSize, 0, 0, 0, FResourceSizeEx(EResourceSizeMode::EstimatedTotal), FResourceSizeEx(EResourceSizeMode::Exclusive)));
 		SortMap.Add(Tree,SortedSets.Num() - 1);
 	}
 
@@ -4478,7 +4478,7 @@ bool UEngine::HandleListParticleSystemsCommand( const TCHAR* Cmd, FOutputDevice&
 			Set.ComponentSize += ComponentCount.GetMax();
 
 			// Save this for adding to the total
-			FResourceSizeEx CompResSize = FResourceSizeEx(EResourceSizeMode::Inclusive);
+			FResourceSizeEx CompResSize = FResourceSizeEx(EResourceSizeMode::EstimatedTotal);
 			Comp->GetResourceSizeEx(CompResSize);
 
 			Set.ComponentResourceSize += CompResSize;
@@ -4740,8 +4740,8 @@ bool UEngine::HandleParticleMeshUsageCommand( const TCHAR* Cmd, FOutputDevice& A
 	{
 		FORCEINLINE bool operator()( UStaticMesh& A, UStaticMesh& B ) const
 		{
-			const SIZE_T ResourceSizeA = A.GetResourceSizeBytes(EResourceSizeMode::Inclusive);
-			const SIZE_T ResourceSizeB = B.GetResourceSizeBytes(EResourceSizeMode::Inclusive);
+			const SIZE_T ResourceSizeA = A.GetResourceSizeBytes(EResourceSizeMode::EstimatedTotal);
+			const SIZE_T ResourceSizeB = B.GetResourceSizeBytes(EResourceSizeMode::EstimatedTotal);
 			return ResourceSizeB < ResourceSizeA;
 		}
 	};
@@ -4754,7 +4754,7 @@ bool UEngine::HandleParticleMeshUsageCommand( const TCHAR* Cmd, FOutputDevice& A
 	for( int32 StaticMeshIndex=0; StaticMeshIndex<UniqueReferencedMeshes.Num(); StaticMeshIndex++ )
 	{
 		UStaticMesh* StaticMesh	= UniqueReferencedMeshes[StaticMeshIndex];
-		const SIZE_T StaticMeshResourceSize = StaticMesh->GetResourceSizeBytes(EResourceSizeMode::Inclusive);
+		const SIZE_T StaticMeshResourceSize = StaticMesh->GetResourceSizeBytes(EResourceSizeMode::EstimatedTotal);
 		TotalSize += StaticMeshResourceSize;
 	}
 
@@ -4768,7 +4768,7 @@ bool UEngine::HandleParticleMeshUsageCommand( const TCHAR* Cmd, FOutputDevice& A
 		TArray<UParticleSystem*> ParticleSystems;
 		StaticMeshToParticleSystemMap.MultiFind( StaticMesh, ParticleSystems );
 
-		const SIZE_T StaticMeshResourceSize = StaticMesh->GetResourceSizeBytes(EResourceSizeMode::Inclusive);
+		const SIZE_T StaticMeshResourceSize = StaticMesh->GetResourceSizeBytes(EResourceSizeMode::EstimatedTotal);
 
 		// Log meshes including resource size and referencing particle systems.
 		Ar.Logf(TEXT("%5i KByte  %s"), StaticMeshResourceSize / 1024, *StaticMesh->GetFullName());
@@ -10357,9 +10357,6 @@ bool UEngine::LoadMap( FWorldContext& WorldContext, FURL URL, class UPendingNetG
 		{
 			if (!bCalled)
 			{
-				PRAGMA_DISABLE_DEPRECATION_WARNINGS
-				FCoreUObjectDelegates::PostLoadMap.Broadcast();
-				PRAGMA_ENABLE_DEPRECATION_WARNINGS
 				FCoreUObjectDelegates::PostLoadMapWithWorld.Broadcast(nullptr);
 			}
 		}
@@ -10813,9 +10810,6 @@ bool UEngine::LoadMap( FWorldContext& WorldContext, FURL URL, class UPendingNetG
 
 	// send a callback message
 	PostLoadMapCaller.bCalled = true;
-	PRAGMA_DISABLE_DEPRECATION_WARNINGS
-	FCoreUObjectDelegates::PostLoadMap.Broadcast();
-	PRAGMA_ENABLE_DEPRECATION_WARNINGS
 	FCoreUObjectDelegates::PostLoadMapWithWorld.Broadcast(WorldContext.World());
 	
 	WorldContext.World()->bWorldWasLoadedThisTick = true;
diff --git a/Engine/Source/Runtime/Engine/Private/World.cpp b/Engine/Source/Runtime/Engine/Private/World.cpp
--- a/Engine/Source/Runtime/Engine/Private/World.cpp
+++ b/Engine/Source/Runtime/Engine/Private/World.cpp
@@ -5624,9 +5624,6 @@ UWorld* FSeamlessTravelHandler::Tick()
 				LoadedWorld->BeginPlay();
 
 				FCoreUObjectDelegates::PostLoadMapWithWorld.Broadcast(LoadedWorld);
-				PRAGMA_DISABLE_DEPRECATION_WARNINGS
-				FCoreUObjectDelegates::PostLoadMap.Broadcast();
-				PRAGMA_ENABLE_DEPRECATION_WARNINGS
 			}
 			else
 			{
@@ -6650,10 +6647,9 @@ void UWorld::GetAssetRegistryTags(TArray<FAssetRegistryTag>& OutTags) const
 	// Get the full file path with extension
 	const FString FullFilePath = FPackageName::LongPackageNameToFilename(GetOutermost()->GetName(), FPackageName::GetMapPackageExtension());
 
-	// Save/Display the file size and modify date
+	// Save/Display the modify date. File size is handled generically for all packages
 	FDateTime AssetDateModified = IFileManager::Get().GetTimeStamp(*FullFilePath);
 	OutTags.Add(FAssetRegistryTag("DateModified", AssetDateModified.ToString(), FAssetRegistryTag::TT_Chronological, FAssetRegistryTag::TD_Date));
-	OutTags.Add(FAssetRegistryTag("MapFileSize", Lex::ToString(IFileManager::Get().FileSize(*FullFilePath)), FAssetRegistryTag::TT_Numerical, FAssetRegistryTag::TD_Memory));
 
 	FWorldDelegates::GetAssetTags.Broadcast(this, OutTags);
 }
diff --git a/Engine/Source/Runtime/Engine/Public/DistanceFieldAtlas.h b/Engine/Source/Runtime/Engine/Public/DistanceFieldAtlas.h
--- a/Engine/Source/Runtime/Engine/Public/DistanceFieldAtlas.h
+++ b/Engine/Source/Runtime/Engine/Public/DistanceFieldAtlas.h
@@ -168,16 +168,10 @@ public:
 		delete this;
 	}
 
-	DEPRECATED(4.14, "GetResourceSize is deprecated. Please use GetResourceSizeEx or GetResourceSizeBytes instead.")
-	SIZE_T GetResourceSize() const
-	{
-		return GetResourceSizeBytes();
-	}
-
 	void GetResourceSizeEx(FResourceSizeEx& CumulativeResourceSize) const
 	{
 		CumulativeResourceSize.AddDedicatedSystemMemoryBytes(sizeof(*this));
-		CumulativeResourceSize.AddUnknownMemoryBytes(CompressedDistanceFieldVolume.GetAllocatedSize());
+		CumulativeResourceSize.AddDedicatedSystemMemoryBytes(CompressedDistanceFieldVolume.GetAllocatedSize());
 	}
 
 	SIZE_T GetResourceSizeBytes() const
diff --git a/Engine/Source/Runtime/Engine/Public/MaterialShared.h b/Engine/Source/Runtime/Engine/Public/MaterialShared.h
--- a/Engine/Source/Runtime/Engine/Public/MaterialShared.h
+++ b/Engine/Source/Runtime/Engine/Public/MaterialShared.h
@@ -1841,9 +1841,6 @@ public:
 
 	ENGINE_API void GetRepresentativeShaderTypesAndDescriptions(TMap<FName, FString>& OutShaderTypeNameAndDescriptions) const;
 
-	DEPRECATED(4.14, "GetResourceSizeInclusive is deprecated. Please use GetResourceSizeEx instead.")
-	ENGINE_API SIZE_T GetResourceSizeInclusive();
-
 	ENGINE_API void GetResourceSizeEx(FResourceSizeEx& CumulativeResourceSize);
 
 	ENGINE_API virtual void LegacySerialize(FArchive& Ar) override;
diff --git a/Engine/Source/Runtime/Engine/Public/ParticleEmitterInstances.h b/Engine/Source/Runtime/Engine/Public/ParticleEmitterInstances.h
--- a/Engine/Source/Runtime/Engine/Public/ParticleEmitterInstances.h
+++ b/Engine/Source/Runtime/Engine/Public/ParticleEmitterInstances.h
@@ -616,29 +616,11 @@ public:
 		OutMax = 0;
 	}
 	
-	/**
-	 * Returns the size of the object/ resource for display to artists/ LDs in the Editor.
-	 *
-	 * @param	Mode	Specifies which resource size should be displayed. ( see EResourceSizeMode::Type )
-	 * @return  Size of resource as to be displayed to artists/ LDs in the Editor.
-	 */
-	DEPRECATED(4.14, "GetResourceSize is deprecated. Please use GetResourceSizeEx or GetResourceSizeBytes instead.")
-	virtual SIZE_T GetResourceSize(EResourceSizeMode::Type Mode)
-	{
-		return GetResourceSizeBytes(Mode);
-	}
-
+	/** Returns resource size, similar to UObject function */
 	virtual void GetResourceSizeEx(FResourceSizeEx& CumulativeResourceSize)
 	{
 	}
 
-	SIZE_T GetResourceSizeBytes(EResourceSizeMode::Type Mode)
-	{
-		FResourceSizeEx ResSize = FResourceSizeEx(Mode);
-		GetResourceSizeEx(ResSize);
-		return ResSize.GetTotalMemoryBytes();
-	}
-
 	/**
 	 *	Process received events.
 	 */
diff --git a/Engine/Source/Runtime/Engine/Public/SkeletalRenderPublic.h b/Engine/Source/Runtime/Engine/Public/SkeletalRenderPublic.h
--- a/Engine/Source/Runtime/Engine/Public/SkeletalRenderPublic.h
+++ b/Engine/Source/Runtime/Engine/Public/SkeletalRenderPublic.h
@@ -152,24 +152,9 @@ public:
 		delete this;
 	}
 	
-	/**
-	* Returns the size of memory allocated by render data
-	*/
-	DEPRECATED(4.14, "GetResourceSize is deprecated. Please use GetResourceSizeEx or GetResourceSizeBytes instead.")
-	SIZE_T GetResourceSize()
-	{
-		return GetResourceSizeBytes();
-	}
-
+	/** Returns the size of memory allocated by render data */
 	virtual void GetResourceSizeEx(FResourceSizeEx& CumulativeResourceSize) = 0;
 
-	SIZE_T GetResourceSizeBytes()
-	{
-		FResourceSizeEx ResSize;
-		GetResourceSizeEx(ResSize);
-		return ResSize.GetTotalMemoryBytes();
-	}
-
 	/**
 	 * List of sections to be rendered based on instance weight usage. Full swap of weights will render with its own sections.
 	 * @return Sections to iterate over for rendering
diff --git a/Engine/Source/Runtime/GeometryCache/Classes/GeometryCacheMeshData.h b/Engine/Source/Runtime/GeometryCache/Classes/GeometryCacheMeshData.h
--- a/Engine/Source/Runtime/GeometryCache/Classes/GeometryCacheMeshData.h
+++ b/Engine/Source/Runtime/GeometryCache/Classes/GeometryCacheMeshData.h
@@ -127,12 +127,6 @@ struct FGeometryCacheMeshData
 		return Ar;
 	}
 
-	DEPRECATED(4.14, "GetResourceSize is deprecated. Please use GetResourceSizeEx or GetResourceSizeBytes instead.")
-	SIZE_T GetResourceSize() const
-	{
-		return GetResourceSizeBytes();
-	}
-
 	void GetResourceSizeEx(FResourceSizeEx& CumulativeResourceSize) const
 	{
 		// Calculate resource size according to what is actually serialized
@@ -144,11 +138,4 @@ struct FGeometryCacheMeshData
 		CumulativeResourceSize.AddUnknownMemoryBytes(Indices.Num() * sizeof(uint32));
 		CumulativeResourceSize.AddUnknownMemoryBytes(sizeof(Indices));
 	}
-
-	SIZE_T GetResourceSizeBytes() const
-	{
-		FResourceSizeEx ResSize;
-		GetResourceSizeEx(ResSize);
-		return ResSize.GetTotalMemoryBytes();
-	}
 };
diff --git a/Engine/Source/Runtime/Landscape/Private/Landscape.cpp b/Engine/Source/Runtime/Landscape/Private/Landscape.cpp
--- a/Engine/Source/Runtime/Landscape/Private/Landscape.cpp
+++ b/Engine/Source/Runtime/Landscape/Private/Landscape.cpp
@@ -367,7 +367,7 @@ void ULandscapeComponent::GetResourceSizeEx(FResourceSizeEx& CumulativeResourceS
 {
 	Super::GetResourceSizeEx(CumulativeResourceSize);
 
-	CumulativeResourceSize.AddUnknownMemoryBytes(GrassData->GetAllocatedSize());
+	CumulativeResourceSize.AddDedicatedSystemMemoryBytes(GrassData->GetAllocatedSize());
 }
 
 #if WITH_EDITOR
diff --git a/Engine/Source/Runtime/Landscape/Private/LandscapeGrass.cpp b/Engine/Source/Runtime/Landscape/Private/LandscapeGrass.cpp
--- a/Engine/Source/Runtime/Landscape/Private/LandscapeGrass.cpp
+++ b/Engine/Source/Runtime/Landscape/Private/LandscapeGrass.cpp
@@ -1124,10 +1124,6 @@ SIZE_T FLandscapeComponentGrassData::GetAllocatedSize() const
 	}
 	return sizeof(*this)
 		+ HeightData.GetAllocatedSize()
-#if WITH_EDITORONLY_DATA
-		+ HeightMipData.GetAllocatedSize()
-		+ Algo::TransformAccumulate(HeightMipData, [](const TPair<int32, TArray<uint16>>& HeightMipDataPair) { return HeightMipDataPair.Value.GetAllocatedSize(); }, 0)
-#endif
 		+ WeightData.GetAllocatedSize() + WeightSize;
 }
 
diff --git a/Engine/Source/Runtime/RuntimeAssetCache/Private/RuntimeAssetCacheBuilders.cpp b/Engine/Source/Runtime/RuntimeAssetCache/Private/RuntimeAssetCacheBuilders.cpp
--- a/Engine/Source/Runtime/RuntimeAssetCache/Private/RuntimeAssetCacheBuilders.cpp
+++ b/Engine/Source/Runtime/RuntimeAssetCache/Private/RuntimeAssetCacheBuilders.cpp
@@ -235,8 +235,8 @@ void UExampleTextureCacheBuilder::OnAssetPostLoad()
 int64 UExampleTextureCacheBuilder::GetSerializedDataSizeEstimate()
 {
 	int64 DataSize = sizeof(FTexturePlatformData);
-	DataSize += sizeof(FString) + (sizeof(TCHAR) * 12);							// Guess the size of the pixel format string (most are less than 12 characters, but we don't need to be exact)
-	DataSize += Texture->GetResourceSizeBytes(EResourceSizeMode::Exclusive);	// Size of all the mips
-	DataSize += (sizeof(int32) * 3) * Texture->GetNumMips();					// Each mip stores its X and Y size, and its BulkDataSize
+	DataSize += sizeof(FString) + (sizeof(TCHAR) * 12);								// Guess the size of the pixel format string (most are less than 12 characters, but we don't need to be exact)
+	DataSize += Texture->GetResourceSizeBytes(EResourceSizeMode::EstimatedTotal);	// Size of all the mips
+	DataSize += (sizeof(int32) * 3) * Texture->GetNumMips();						// Each mip stores its X and Y size, and its BulkDataSize
 	return DataSize;
 }
diff --git a/Engine/Source/Runtime/SlateRHIRenderer/Private/SlateRHIResourceManager.cpp b/Engine/Source/Runtime/SlateRHIRenderer/Private/SlateRHIResourceManager.cpp
--- a/Engine/Source/Runtime/SlateRHIRenderer/Private/SlateRHIResourceManager.cpp
+++ b/Engine/Source/Runtime/SlateRHIRenderer/Private/SlateRHIResourceManager.cpp
@@ -80,7 +80,7 @@ void FDynamicResourceMap::AddUTextureResource( UTexture* TextureObject, TSharedR
 		check(TextureObject == InResource->TextureObject);
 		TextureMap.Add(TextureObject, InResource);
 
-		TextureMemorySincePurge += TextureObject->GetResourceSizeBytes(EResourceSizeMode::Inclusive);
+		TextureMemorySincePurge += TextureObject->GetResourceSizeBytes(EResourceSizeMode::Exclusive);
 	}
 }
 
@@ -100,7 +100,7 @@ void FDynamicResourceMap::RemoveUTextureResource( UTexture* TextureObject )
 	if(TextureObject)
 	{
 		TextureMap.Remove(TextureObject);
-		TextureMemorySincePurge -= TextureObject->GetResourceSizeBytes(EResourceSizeMode::Inclusive);
+		TextureMemorySincePurge -= TextureObject->GetResourceSizeBytes(EResourceSizeMode::Exclusive);
 	}
 }
 
