Fix a crash caused by BSP releasing vertex buffers that are still in use

diff --git a/Engine/Source/Runtime/Engine/Classes/Components/ModelComponent.h b/Engine/Source/Runtime/Engine/Classes/Components/ModelComponent.h
--- a/Engine/Source/Runtime/Engine/Classes/Components/ModelComponent.h
+++ b/Engine/Source/Runtime/Engine/Classes/Components/ModelComponent.h
@@ -87,6 +87,8 @@ public:
 	ENGINE_API void CopyElementsFrom(UModelComponent* OtherModelComponent);
 
 	//~ Begin UPrimitiveComponent Interface.
+	virtual void CreateRenderState_Concurrent() override;
+	virtual void DestroyRenderState_Concurrent() override;
 	virtual bool GetLightMapResolution( int32& Width, int32& Height ) const override;
 	virtual int32 GetStaticLightMapResolution() const override;
 	virtual void GetLightAndShadowMapMemoryUsage( int32& LightMapMemoryUsage, int32& ShadowMapMemoryUsage ) const override;
diff --git a/Engine/Source/Runtime/Engine/Private/ModelRender.cpp b/Engine/Source/Runtime/Engine/Private/ModelRender.cpp
--- a/Engine/Source/Runtime/Engine/Private/ModelRender.cpp
+++ b/Engine/Source/Runtime/Engine/Private/ModelRender.cpp
@@ -196,8 +196,29 @@ public:
 			FColor(157,149,223,255))
 #endif
 	{
-		FModelVertexBuffer& VertexBuffer = InComponent->GetModel()->VertexBuffer;
-		VertexBuffer.Buffers.InitModelVF(&VertexFactory);
+		ENQUEUE_RENDER_COMMAND(InitOrUpdateVertexBufferCmd)([VertexBuffer = &InComponent->GetModel()->VertexBuffer](FRHICommandList&)
+		{
+			if (!VertexBuffer->Buffers.PositionVertexBuffer.IsInitialized())
+			{
+				VertexBuffer->Buffers.PositionVertexBuffer.InitResource();
+			}
+			else if (VertexBuffer->RefCount == 0)
+			{
+				// Only allow update when no other FModelSceneProxy's is using the vertex buffer
+				// otherwise their resources will become invalid
+				VertexBuffer->Buffers.PositionVertexBuffer.UpdateRHI();
+			}
+
+			if (!VertexBuffer->Buffers.StaticMeshVertexBuffer.IsInitialized())
+			{
+				VertexBuffer->Buffers.StaticMeshVertexBuffer.InitResource();
+			}
+			else if (VertexBuffer->RefCount == 0)
+			{
+				VertexBuffer->Buffers.StaticMeshVertexBuffer.UpdateRHI();
+			}
+		});
+		InComponent->GetModel()->VertexBuffer.Buffers.InitModelVF(&VertexFactory);
 
 		OverrideOwnerName(NAME_BSP);
 		const TIndirectArray<FModelElement>& SourceElements = Component->GetElements();
@@ -749,6 +770,26 @@ public:
 	friend class UModelComponent;
 };
 
+void UModelComponent::CreateRenderState_Concurrent()
+{
+	if (GetModel())
+	{
+		++GetModel()->VertexBuffer.RefCount;
+	}
+
+	Super::CreateRenderState_Concurrent();
+}
+
+void UModelComponent::DestroyRenderState_Concurrent()
+{
+	if (GetModel())
+	{
+		--GetModel()->VertexBuffer.RefCount;
+	}
+
+	Super::DestroyRenderState_Concurrent();
+}
+
 FPrimitiveSceneProxy* UModelComponent::CreateSceneProxy()
 {
 	FPrimitiveSceneProxy* Proxy = ::new FModelSceneProxy(this);
diff --git a/Engine/Source/Runtime/Engine/Private/StaticMesh.cpp b/Engine/Source/Runtime/Engine/Private/StaticMesh.cpp
--- a/Engine/Source/Runtime/Engine/Private/StaticMesh.cpp
+++ b/Engine/Source/Runtime/Engine/Private/StaticMesh.cpp
@@ -410,8 +410,8 @@ void FStaticMeshVertexBuffers::InitModelVF(FLocalVertexFactory* VertexFactory)
 	ENQUEUE_RENDER_COMMAND(StaticMeshVertexBuffersLegacyBspInit)(
 		[VertexFactory, Self](FRHICommandListImmediate& RHICmdList)
 	{
-		InitOrUpdateResource(&Self->PositionVertexBuffer);
-		InitOrUpdateResource(&Self->StaticMeshVertexBuffer);
+		check(Self->PositionVertexBuffer.IsInitialized());
+		check(Self->StaticMeshVertexBuffer.IsInitialized());
 
 		FLocalVertexFactory::FDataType Data;
 		Self->PositionVertexBuffer.BindPositionVertexBuffer(VertexFactory, Data);
diff --git a/Engine/Source/Runtime/Engine/Public/Model.h b/Engine/Source/Runtime/Engine/Public/Model.h
--- a/Engine/Source/Runtime/Engine/Public/Model.h
+++ b/Engine/Source/Runtime/Engine/Public/Model.h
@@ -301,6 +301,10 @@ struct FModelVertex
 class FModelVertexBuffer
 {
 public:
+	// Number of FModelSceneProxy's using this
+	// Only access in CreateRenderState_Concurrent and DestroyRenderState_Concurrent
+	uint32 RefCount = 0;
+
 	/** model vertex data */
 	TArray<FModelVertex> Vertices;
 
