HTML5 - disable SupportsScreenPercentage() - render the full screen

otherwise, this is causing HTML5 screen to only render a portion of the screen and in black...
- there may be another function is that not getting fired off "to render the screen" properly
- this may be due to IsMobileHDR checks that are not fully accounted for the HTML5 platform

#rn

diff --git a/Engine/Source/Runtime/Engine/Private/SceneView.cpp b/Engine/Source/Runtime/Engine/Private/SceneView.cpp
--- a/Engine/Source/Runtime/Engine/Private/SceneView.cpp
+++ b/Engine/Source/Runtime/Engine/Private/SceneView.cpp
@@ -2406,7 +2406,7 @@ bool FSceneViewFamily::SupportsScreenPercentage() const
 	}
 
 	// Mobile renderer does not support screen percentage with LDR.
-	if (GetFeatureLevel() <= ERHIFeatureLevel::ES3_1 && !IsMobileHDR())
+	if ((GetFeatureLevel() <= ERHIFeatureLevel::ES3_1 && !IsMobileHDR()) || PLATFORM_HTML5)
 	{
 		return false;
 	}
