Make sure that resources are released before clearing index buffer, fixes race condition in CL 3709272 (original change by [REDACTED_PERSON])

diff --git a/Engine/Source/Runtime/Engine/Private/Level.cpp b/Engine/Source/Runtime/Engine/Private/Level.cpp
--- a/Engine/Source/Runtime/Engine/Private/Level.cpp
+++ b/Engine/Source/Runtime/Engine/Private/Level.cpp
@@ -1034,10 +1034,17 @@ void ULevel::CreateModelComponents()
 
 	SlowTask.EnterProgressFrame(4);
 
-	Model->InvalidSurfaces = 0;
+	Model->InvalidSurfaces = false;
+	
+	// It is possible that the BSP model has existing buffers from an undo/redo operation
+	if (Model->MaterialIndexBuffers.Num())
+	{
+		// Make sure model resources are released which only happens on the rendering thread
+		FlushRenderingCommands();
 
-	// Clear the model index buffers.
-	Model->MaterialIndexBuffers.Empty();
+		// Clear the model index buffers.
+		Model->MaterialIndexBuffers.Empty();
+	}
 
 	struct FNodeIndices
 	{
