[NO NON-PII COMMIT MESSAGE]

diff --git a/Engine/Source/Runtime/Engine/Private/SoundWave.cpp b/Engine/Source/Runtime/Engine/Private/SoundWave.cpp
--- a/Engine/Source/Runtime/Engine/Private/SoundWave.cpp
+++ b/Engine/Source/Runtime/Engine/Private/SoundWave.cpp
@@ -947,19 +947,23 @@ bool USoundWave::GetChunkData(int32 ChunkIndex, uint8** OutChunkData)
 {
 	if (RunningPlatformData->TryLoadChunk(ChunkIndex, OutChunkData) == false)
 	{
-		// Unable to load chunks from the cache. Rebuild the sound and try again.
-		UE_LOG(LogAudio, Warning, TEXT("GetChunkData failed for %s"), *GetPathName());
 #if WITH_EDITORONLY_DATA
+		// Unable to load chunks from the cache. Rebuild the sound and attempt to recache it.
+		UE_LOG(LogAudio, Display, TEXT("GetChunkData failed, rebuilding %s"), *GetPathName());
+
 		ForceRebuildPlatformData();
 		if (RunningPlatformData->TryLoadChunk(ChunkIndex, OutChunkData) == false)
 		{
-			UE_LOG(LogAudio, Error, TEXT("Failed to build sound %s."), *GetPathName());
+			UE_LOG(LogAudio, Display, TEXT("Failed to build sound %s."), *GetPathName());
 		}
 		else
 		{
 			// Succeeded after rebuilding platform data
 			return true;
 		}
+#else
+		// Failed to find the SoundWave chunk in the cooked package.
+		UE_LOG(LogAudio, Warning, TEXT("GetChunkData failed while streaming. Ensure the following file is cooked: %s"), *GetPathName());
 #endif // #if WITH_EDITORONLY_DATA
 		return false;
 	}
