Fix to runtime error with textures being loaded from bulkdata when disabling texture streaming.

diff --git a/Engine/Source/Runtime/CoreUObject/Private/Serialization/BulkData.cpp b/Engine/Source/Runtime/CoreUObject/Private/Serialization/BulkData.cpp
--- a/Engine/Source/Runtime/CoreUObject/Private/Serialization/BulkData.cpp
+++ b/Engine/Source/Runtime/CoreUObject/Private/Serialization/BulkData.cpp
@@ -1433,8 +1433,17 @@ void FUntypedBulkData::LoadDataIntoMemory( void* Dest )
 	{
 		// load from the specied filename when the linker has been cleared
 		checkf( Filename != TEXT(""), TEXT( "Attempted to load bulk data without a proper filename." ) );
-	
-		UE_CLOG(GEventDrivenLoaderEnabled && (IsInGameThread() || IsInAsyncLoadingThread()), LogSerialization, Error, TEXT("Attempt to sync load bulk data with EDL enabled (LoadDataIntoMemory). This is not desireable. File %s"), *Filename);
+
+#if PLATFORM_SUPPORTS_TEXTURE_STREAMING
+		static auto CVarTextureStreamingEnabled = IConsoleManager::Get().FindTConsoleVariableDataInt(TEXT("r.TextureStreaming"));
+		check(CVarTextureStreamingEnabled);
+		// Because "r.TextureStreaming" is driven by the project setting as well as the command line option "-NoTextureStreaming", 
+		// is it possible for streaming mips to be loaded in non streaming ways.
+		if (CVarTextureStreamingEnabled->GetValueOnAnyThread() != 0)
+		{
+			UE_CLOG(GEventDrivenLoaderEnabled && (IsInGameThread() || IsInAsyncLoadingThread()), LogSerialization, Error, TEXT("Attempt to sync load bulk data with EDL enabled (LoadDataIntoMemory). This is not desireable. File %s"), *Filename);
+		}
+#endif
 
 		if (GEventDrivenLoaderEnabled && (Filename.EndsWith(TEXT(".uasset")) || Filename.EndsWith(TEXT(".umap"))))
 		{
diff --git a/Engine/Source/Runtime/Engine/Private/TextureDerivedData.cpp b/Engine/Source/Runtime/Engine/Private/TextureDerivedData.cpp
--- a/Engine/Source/Runtime/Engine/Private/TextureDerivedData.cpp
+++ b/Engine/Source/Runtime/Engine/Private/TextureDerivedData.cpp
@@ -21,6 +21,7 @@
 #include "DeviceProfiles/DeviceProfile.h"
 #include "DeviceProfiles/DeviceProfileManager.h"
 #include "TextureDerivedDataTask.h"
+#include "Streaming/TextureStreamingHelpers.h"
 
 #if WITH_EDITOR
 
@@ -815,13 +816,16 @@ bool FTexturePlatformData::TryLoadMips(int32 FirstMipToLoad, void** OutMipData)
 			if (OutMipData)
 			{
 				OutMipData[MipIndex - FirstMipToLoad] = FMemory::Malloc(Mip.BulkData.GetBulkDataSize());
-#if 0
-				checkSlow(!Mip.BulkData.GetFilename().EndsWith(TEXT(".ubulk"))); // We want to make sure that any non-streamed mips are coming from the texture asset file, and not from an external bulk file
-#elif !(WITH_DEV_AUTOMATION_TESTS || WITH_PERF_AUTOMATION_TESTS)
-				UE_CLOG(Mip.BulkData.GetFilename().EndsWith(TEXT(".ubulk")), LogTexture, Error, TEXT("Loading non-streamed mips from an external bulk file.  This is not desireable.  File %s"), *(Mip.BulkData.GetFilename() ) );
+
+#if PLATFORM_SUPPORTS_TEXTURE_STREAMING
+				// We want to make sure that any non-streamed mips are coming from the texture asset file, and not from an external bulk file.
+				// But because "r.TextureStreaming" is driven by the project setting as well as the command line option "-NoTextureStreaming", 
+				// is it possible for streaming mips to be loaded in non streaming ways.
+				if (CVarSetTextureStreaming.GetValueOnAnyThread() != 0)
+				{
+					UE_CLOG(Mip.BulkData.GetFilename().EndsWith(TEXT(".ubulk")), LogTexture, Error, TEXT("Loading non-streamed mips from an external bulk file.  This is not desireable.  File %s"), *(Mip.BulkData.GetFilename() ) );
+				}
 #endif
-				
-				
 				Mip.BulkData.GetCopy(&OutMipData[MipIndex - FirstMipToLoad]);
 			}
 			NumMipsCached++;
@@ -995,8 +999,16 @@ static void SerializePlatformData(
 	if (bCooked && Ar.IsSaving())
 	{
 		int32 MinMipToInline = 0;
-			
+		
+#if WITH_EDITORONLY_DATA
+		check(Ar.CookingTarget());
+		// This also needs to check whether the project enables texture streaming.
+		// Currently, there is no reliable way to implement this because there is no difference
+		// between the project settings (CVar) and the command line setting (from -NoTextureStreaming)
+		if (bStreamable && Ar.CookingTarget()->SupportsFeature(ETargetPlatformFeatures::TextureStreaming))
+#else
 		if (bStreamable)
+#endif
 		{
 			MinMipToInline = FMath::Max(0, NumMips - PlatformData->GetNumNonStreamingMips());
 		}
