Last batch of benign PVS-Studio warnings (rest appear valid bugs; will follow up with authors separately)

diff --git a/Engine/Source/Editor/IntroTutorials/Private/STutorialOverlay.cpp b/Engine/Source/Editor/IntroTutorials/Private/STutorialOverlay.cpp
--- a/Engine/Source/Editor/IntroTutorials/Private/STutorialOverlay.cpp
+++ b/Engine/Source/Editor/IntroTutorials/Private/STutorialOverlay.cpp
@@ -166,11 +166,8 @@ int32 STutorialOverlay::TraverseWidgets(TSharedRef<SWidget> InWidget, const FGeo
 		// if we are picking, we need to draw an outline here	
 		if(WidgetNameToHighlight != NAME_None )
 		{
-			if(bIsPicking == true)
-			{
-				const FLinearColor Color = bIsPicking && bShouldHighlight ? FLinearColor::Green : FLinearColor::White;
-				FSlateDrawElement::MakeBox(OutDrawElements, LayerId++, InGeometry.ToPaintGeometry(), FCoreStyle::Get().GetBrush(TEXT("Debug.Border")), ESlateDrawEffect::None, Color);
-			}
+			const FLinearColor Color = bIsPicking && bShouldHighlight ? FLinearColor::Green : FLinearColor::White;
+			FSlateDrawElement::MakeBox(OutDrawElements, LayerId++, InGeometry.ToPaintGeometry(), FCoreStyle::Get().GetBrush(TEXT("Debug.Border")), ESlateDrawEffect::None, Color);
 		}	
 	}
 
diff --git a/Engine/Source/Editor/MaterialEditor/Private/MaterialEditor.cpp b/Engine/Source/Editor/MaterialEditor/Private/MaterialEditor.cpp
--- a/Engine/Source/Editor/MaterialEditor/Private/MaterialEditor.cpp
+++ b/Engine/Source/Editor/MaterialEditor/Private/MaterialEditor.cpp
@@ -1788,7 +1788,7 @@ void FMaterialEditor::UpdateMaterialInfoList(bool bForceDisplay)
 		FeatureLevelsToDisplay[NumFeatureLevels++] = ERHIFeatureLevel::ES3_1;
 	}
 
-	if (NumFeatureLevels > 0)
+	if (NumFeatureLevels > 0) //-V547
 	{
 		UMaterial* MaterialForStats = bStatsFromPreviewMaterial ? Material : OriginalMaterial;
 
diff --git a/Engine/Source/Editor/MaterialEditor/Private/SMaterialLayersFunctionsTree.cpp b/Engine/Source/Editor/MaterialEditor/Private/SMaterialLayersFunctionsTree.cpp
--- a/Engine/Source/Editor/MaterialEditor/Private/SMaterialLayersFunctionsTree.cpp
+++ b/Engine/Source/Editor/MaterialEditor/Private/SMaterialLayersFunctionsTree.cpp
@@ -530,16 +530,13 @@ public:
 			switch (StackParameterData->ParameterInfo.Association)
 			{
 			case EMaterialParameterAssociation::LayerParameter:
-			{
 				StackBrush = const_cast<FSlateBrush*>(FEditorStyle::GetBrush("MaterialInstanceEditor.StackBody"));
 				break;
-			}
 			case EMaterialParameterAssociation::BlendParameter:
-			{
 				StackBrush = const_cast<FSlateBrush*>(FEditorStyle::GetBrush("MaterialInstanceEditor.StackBodyBlend"));
 				break;
-			}
 			default:
+				StackBrush = nullptr;
 				break;
 			}
 			WrapperWidget->AddSlot()
diff --git a/Engine/Source/Editor/Matinee/Private/MatineeDraw.cpp b/Engine/Source/Editor/Matinee/Private/MatineeDraw.cpp
--- a/Engine/Source/Editor/Matinee/Private/MatineeDraw.cpp
+++ b/Engine/Source/Editor/Matinee/Private/MatineeDraw.cpp
@@ -290,7 +290,7 @@ void FMatinee::GetGroupLabelPosition( UInterpGroup* InGroup, int32& LabelTopPosi
 						{
 							// If we're parented then we're only visible if our parent group is not collapsed
 							check( CurParentGroup != NULL );
-							if( CurParentGroup->bCollapsed )
+							if( CurParentGroup->bCollapsed ) //-V522
 							{
 								// Parent group is collapsed, so we should not be rendered
 								bIsGroupVisible = false;
diff --git a/Engine/Source/Runtime/Core/Private/Misc/Base64.cpp b/Engine/Source/Runtime/Core/Private/Misc/Base64.cpp
--- a/Engine/Source/Runtime/Core/Private/Misc/Base64.cpp
+++ b/Engine/Source/Runtime/Core/Private/Misc/Base64.cpp
@@ -217,10 +217,6 @@ template<typename CharType> bool FBase64::Decode(const CharType* Source, uint32
 		}
 
 		// Now we can tear the uint32 into bytes
-		if(Length >= 4)
-		{
-			Dest[2] = OriginalTriplet & 0xFF;
-		}
 		OriginalTriplet >>= 8;
 		if(Length >= 3)
 		{
diff --git a/Engine/Source/Runtime/Engine/Private/Components/CharacterMovementComponent.cpp b/Engine/Source/Runtime/Engine/Private/Components/CharacterMovementComponent.cpp
--- a/Engine/Source/Runtime/Engine/Private/Components/CharacterMovementComponent.cpp
+++ b/Engine/Source/Runtime/Engine/Private/Components/CharacterMovementComponent.cpp
@@ -1377,7 +1377,7 @@ void UCharacterMovementComponent::SimulatedTick(float DeltaSeconds)
 				const FRotator OldRotation = OldRotationQuat.Rotator();
 				const FRotator NewRotation = UpdatedComponent->GetComponentRotation();
 				const FVector NewLocation = UpdatedComponent->GetComponentLocation();
-				DrawDebugCoordinateSystem(GetWorld(), CharacterOwner->GetMesh()->GetComponentLocation() + FVector(0,0,1), NewRotation, 50.f, false);
+				DrawDebugCoordinateSystem(GetWorld(), CharacterOwner->GetMesh()->GetComponentLocation() + FVector(0,0,1), NewRotation, 50.f, false); //-V595
 				DrawDebugLine(GetWorld(), OldLocation, NewLocation, FColor::Red, true, 10.f);
 
 				UE_LOG(LogRootMotion, Log,  TEXT("UCharacterMovementComponent::SimulatedTick DeltaMovement Translation: %s, Rotation: %s, MovementBase: %s"), //-V595
diff --git a/Engine/Source/Runtime/Engine/Private/HierarchicalInstancedStaticMesh.cpp b/Engine/Source/Runtime/Engine/Private/HierarchicalInstancedStaticMesh.cpp
--- a/Engine/Source/Runtime/Engine/Private/HierarchicalInstancedStaticMesh.cpp
+++ b/Engine/Source/Runtime/Engine/Private/HierarchicalInstancedStaticMesh.cpp
@@ -1583,7 +1583,7 @@ void FHierarchicalStaticMeshSceneProxy::GetDynamicMeshElements(const TArray<cons
 					SCOPE_CYCLE_COUNTER(STAT_FoliageTraversalTime);
 
 					// validate that the bounding box is layed out correctly in memory
-					check(&((const FVector4*)&ClusterTree[0].BoundMin)[1] == (const FVector4*)&ClusterTree[0].BoundMax);
+					check((const FVector4*)&ClusterTree[0].BoundMin + 1 == (const FVector4*)&ClusterTree[0].BoundMax);
 					//check(UPTRINT(&ClusterTree[0].BoundMin) % 16 == 0);
 					//check(UPTRINT(&ClusterTree[0].BoundMax) % 16 == 0);
 
