Fix non-unity compile error introduced from Release-Staging-4.19.

diff --git a/Engine/Source/Runtime/Engine/Private/Rendering/SkinWeightVertexBuffer.cpp b/Engine/Source/Runtime/Engine/Private/Rendering/SkinWeightVertexBuffer.cpp
--- a/Engine/Source/Runtime/Engine/Private/Rendering/SkinWeightVertexBuffer.cpp
+++ b/Engine/Source/Runtime/Engine/Private/Rendering/SkinWeightVertexBuffer.cpp
@@ -3,6 +3,7 @@
 #include "Rendering/SkinWeightVertexBuffer.h"
 #include "EngineUtils.h"
 #include "Rendering/SkeletalMeshLODModel.h"
+#include "RenderUtils.h"
 
 ///
 
