Added a warning about gather text using LOCTEXT when no namespace was defined

Fixed some code that was doing this

diff --git a/Engine/Source/Editor/UnrealEd/Private/Commandlets/GatherTextFromSourceCommandlet.cpp b/Engine/Source/Editor/UnrealEd/Private/Commandlets/GatherTextFromSourceCommandlet.cpp
--- a/Engine/Source/Editor/UnrealEd/Private/Commandlets/GatherTextFromSourceCommandlet.cpp
+++ b/Engine/Source/Editor/UnrealEd/Private/Commandlets/GatherTextFromSourceCommandlet.cpp
@@ -1499,10 +1499,15 @@ void UGatherTextFromSourceCommandlet::FStringMacroDescriptor::TryParse(const FSt
 			else
 			{
 				FString Identifier;
-				FString Namespace = Context.Namespace;
 				FString SourceLocation = FSourceLocation(Context.Filename, Context.LineNumber).ToString();
 				FString SourceText;
 
+				TOptional<FString> Namespace;
+				if (Context.WithinNamespaceDefine || !Context.Namespace.IsEmpty())
+				{
+					Namespace = Context.Namespace;
+				}
+
 				bool ArgParseError = false;
 				for (int32 ArgIdx=0; ArgIdx<Arguments.Num(); ArgIdx++)
 				{
@@ -1547,11 +1552,16 @@ void UGatherTextFromSourceCommandlet::FStringMacroDescriptor::TryParse(const FSt
 
 				if (!ArgParseError && !Identifier.IsEmpty() && !SourceText.IsEmpty())
 				{
+					if (!Namespace.IsSet())
+					{
+						UE_LOG(LogGatherTextFromSourceCommandlet, Warning, TEXT("Localization macro doesn't define a namespace and no external namespace was set. An empty namspace will be used. %s"), *SourceLocation );
+					}
+
 					FManifestContext MacroContext;
 					MacroContext.Key = Identifier;
 					MacroContext.SourceLocation = SourceLocation;
 
-					Context.AddManifestText( GetToken(), Namespace, SourceText, MacroContext );
+					Context.AddManifestText( GetToken(), Namespace.Get(FString()), SourceText, MacroContext );
 				}
 			}
 		}
diff --git a/Engine/Source/Runtime/Engine/Public/ShowFlags.h b/Engine/Source/Runtime/Engine/Public/ShowFlags.h
--- a/Engine/Source/Runtime/Engine/Public/ShowFlags.h
+++ b/Engine/Source/Runtime/Engine/Public/ShowFlags.h
@@ -76,9 +76,6 @@ struct FEngineShowFlags
 	};
 	// ---------------------------------------------------------
 
-	// namespace to match the original in the old loc system
-	#define LOCTEXT_NAMESPACE "UnrealEd"
-
 	/** Retrieve the localized display name for a named show flag */
 	static void FindShowFlagDisplayName(const FString& InName, FText& OutText)
 	{
@@ -91,22 +88,27 @@ struct FEngineShowFlags
 			#include "ShowFlagsValues.inl"
 
 			// Additional strings that don't correspond to a showflag variable:
-			LocNames.Add(TEXT("BlockingVolume"), LOCTEXT("BlockingVolumeSF", "Blocking"));
-			LocNames.Add(TEXT("CullDistanceVolume"), LOCTEXT("CullDistanceVolumeSF", "Cull Distance"));
-			LocNames.Add(TEXT("LevelStreamingVolume"), LOCTEXT("LevelStreamingVolumeSF", "Level Streaming"));
-			LocNames.Add(TEXT("LightmassCharacterIndirectDetailVolume"), LOCTEXT("LightmassCharacterIndirectDetailVolumeSF", "Lightmass Character Indirect Detail"));
-			LocNames.Add(TEXT("LightmassImportanceVolume"), LOCTEXT("LightmassImportanceVolumeSF", "Lightmass Importance"));
-			LocNames.Add(TEXT("MassiveLODOverrideVolume"), LOCTEXT("MassiveLODOverrideVolumeSF", "Massive LOD Override"));
-			LocNames.Add(TEXT("NavMeshBoundsVolume"), LOCTEXT("NavMeshBoundsVolumeSF", "NavMesh Bounds"));
-			LocNames.Add(TEXT("NavModifierVolume"), LOCTEXT("NavModifierVolumeSF", "Nav Modifier"));
-			LocNames.Add(TEXT("PainCausingVolume"), LOCTEXT("PainCausingVolumeSF", "Pain Causing"));
-			LocNames.Add(TEXT("PhysicsVolume"), LOCTEXT("PhysicsVolumeSF", "Physics"));
-			LocNames.Add(TEXT("PostProcessVolume"), LOCTEXT("PostProcessVolumeSF", "Post Process"));
-			LocNames.Add(TEXT("PrecomputedVisibilityOverrideVolume"), LOCTEXT("PrecomputedVisibilityOverrideVolumeSF", "Precomputed Visibility Override"));
-			LocNames.Add(TEXT("PrecomputedVisibilityVolume"), LOCTEXT("PrecomputedVisibilityVolumeSF", "Precomputed Visibility"));
-			LocNames.Add(TEXT("AudioVolume"), LOCTEXT("AudioVolumeSF", "Audio"));
-			LocNames.Add(TEXT("TriggerVolume"), LOCTEXT("TriggerVolumeSF", "Trigger"));
-			LocNames.Add(TEXT("VisualizeAO"), LOCTEXT("VisualizeAOSF", "Ambient Occlusion"));
+			LocNames.Add(TEXT("BlockingVolume"), NSLOCTEXT("UnrealEd", "BlockingVolumeSF", "Blocking"));
+			LocNames.Add(TEXT("CullDistanceVolume"), NSLOCTEXT("UnrealEd", "CullDistanceVolumeSF", "Cull Distance"));
+			LocNames.Add(TEXT("LevelStreamingVolume"), NSLOCTEXT("UnrealEd", "LevelStreamingVolumeSF", "Level Streaming"));
+			LocNames.Add(TEXT("LightmassCharacterIndirectDetailVolume"), NSLOCTEXT("UnrealEd", "LightmassCharacterIndirectDetailVolumeSF", "Lightmass Character Indirect Detail"));
+			LocNames.Add(TEXT("LightmassImportanceVolume"), NSLOCTEXT("UnrealEd", "LightmassImportanceVolumeSF", "Lightmass Importance"));
+			LocNames.Add(TEXT("MassiveLODOverrideVolume"), NSLOCTEXT("UnrealEd", "MassiveLODOverrideVolumeSF", "Massive LOD Override"));
+			LocNames.Add(TEXT("NavMeshBoundsVolume"), NSLOCTEXT("UnrealEd", "NavMeshBoundsVolumeSF", "NavMesh Bounds"));
+			LocNames.Add(TEXT("NavModifierVolume"), NSLOCTEXT("UnrealEd", "NavModifierVolumeSF", "Nav Modifier"));
+			LocNames.Add(TEXT("PainCausingVolume"), NSLOCTEXT("UnrealEd", "PainCausingVolumeSF", "Pain Causing"));
+			LocNames.Add(TEXT("PhysicsVolume"), NSLOCTEXT("UnrealEd", "PhysicsVolumeSF", "Physics"));
+			LocNames.Add(TEXT("PostProcessVolume"), NSLOCTEXT("UnrealEd", "PostProcessVolumeSF", "Post Process"));
+			LocNames.Add(TEXT("PrecomputedVisibilityOverrideVolume"), NSLOCTEXT("UnrealEd", "PrecomputedVisibilityOverrideVolumeSF", "Precomputed Visibility Override"));
+			LocNames.Add(TEXT("PrecomputedVisibilityVolume"), NSLOCTEXT("UnrealEd", "PrecomputedVisibilityVolumeSF", "Precomputed Visibility"));
+			LocNames.Add(TEXT("AudioVolume"), NSLOCTEXT("UnrealEd", "AudioVolumeSF", "Audio"));
+			LocNames.Add(TEXT("TriggerVolume"), NSLOCTEXT("UnrealEd", "TriggerVolumeSF", "Trigger"));
+			LocNames.Add(TEXT("CameraBlockingVolume"), NSLOCTEXT("UnrealEd", "CameraBlockingVolumeSF", "Camera Blocking"));
+			LocNames.Add(TEXT("HierarchicalLODVolume"), NSLOCTEXT("UnrealEd", "HierarchicalLODVolumeSF", "Hierarchical LOD"));
+			LocNames.Add(TEXT("KillZVolume"), NSLOCTEXT("UnrealEd", "KillZVolumeSF", "Kill Z"));
+			LocNames.Add(TEXT("MeshMergeCullingVolume"), NSLOCTEXT("UnrealEd", "MeshMergeCullingVolumeSF", "Mesh Merge Culling"));
+			LocNames.Add(TEXT("VolumetricLightmapDensityVolume"), NSLOCTEXT("UnrealEd", "VolumetricLightmapDensityVolumeSF", "Volumetric Lightmap Density"));
+			LocNames.Add(TEXT("VisualizeAO"), NSLOCTEXT("UnrealEd", "VisualizeAOSF", "Ambient Occlusion"));
 		}
 
 		check(InName.Len() > 0);
@@ -121,8 +123,6 @@ struct FEngineShowFlags
 		}
 	}
 
-	#undef LOCTEXT_NAMESPACE
-
 	/** this function defines the grouping in the editor. If a flag is not defined here it's assumed to be a postprocess flag */
 	static EShowFlagGroup FindShowFlagGroup(const TCHAR* Name)
 	{
diff --git a/Engine/Source/Runtime/Engine/Public/ShowFlagsValues.inl b/Engine/Source/Runtime/Engine/Public/ShowFlagsValues.inl
--- a/Engine/Source/Runtime/Engine/Public/ShowFlagsValues.inl
+++ b/Engine/Source/Runtime/Engine/Public/ShowFlagsValues.inl
@@ -17,318 +17,318 @@
 #endif
 
 /** Affects all postprocessing features, depending on viewmode this is on or off, for now SHOWFLAG_ALWAYS_ACCESSIBLE because it's used by ReflectionEnviromentCapture */
-SHOWFLAG_ALWAYS_ACCESSIBLE(PostProcessing, SFG_Hidden, LOCTEXT("PostProcessingSF", "Post-processing"))
+SHOWFLAG_ALWAYS_ACCESSIBLE(PostProcessing, SFG_Hidden, NSLOCTEXT("UnrealEd", "PostProcessingSF", "Post-processing"))
 /** Bloom, for now SHOWFLAG_ALWAYS_ACCESSIBLE because it's exposed in SceneCapture */
-SHOWFLAG_ALWAYS_ACCESSIBLE(Bloom, SFG_PostProcess, LOCTEXT("BloomSF", "Bloom"))
+SHOWFLAG_ALWAYS_ACCESSIBLE(Bloom, SFG_PostProcess, NSLOCTEXT("UnrealEd", "BloomSF", "Bloom"))
 /** HDR->LDR conversion is done through a tone mapper (otherwise linear mapping is used) */
-SHOWFLAG_FIXED_IN_SHIPPING(1, Tonemapper, SFG_PostProcess, LOCTEXT("TonemapperSF", "Tonemapper"))
+SHOWFLAG_FIXED_IN_SHIPPING(1, Tonemapper, SFG_PostProcess, NSLOCTEXT("UnrealEd", "TonemapperSF", "Tonemapper"))
 /** Any Anti-aliasing e.g. FXAA, Temporal AA, for now SHOWFLAG_ALWAYS_ACCESSIBLE because it's exposed in SceneCapture */
-SHOWFLAG_ALWAYS_ACCESSIBLE(AntiAliasing, SFG_Normal, LOCTEXT("AntiAliasingSF", "Anti-aliasing"))
+SHOWFLAG_ALWAYS_ACCESSIBLE(AntiAliasing, SFG_Normal, NSLOCTEXT("UnrealEd", "AntiAliasingSF", "Anti-aliasing"))
 /** Only used in AntiAliasing is on, true:uses Temporal AA, otherwise FXAA, for now SHOWFLAG_ALWAYS_ACCESSIBLE because it's exposed in SceneCapture  */
-SHOWFLAG_ALWAYS_ACCESSIBLE(TemporalAA, SFG_Advanced, LOCTEXT("TemporalAASF", "Temporal AA (instead FXAA)"))
+SHOWFLAG_ALWAYS_ACCESSIBLE(TemporalAA, SFG_Advanced, NSLOCTEXT("UnrealEd", "TemporalAASF", "Temporal AA (instead FXAA)"))
 /** e.g. Ambient cube map, for now SHOWFLAG_ALWAYS_ACCESSIBLE because it's exposed in SceneCapture */
-SHOWFLAG_ALWAYS_ACCESSIBLE(AmbientCubemap, SFG_LightingFeatures, LOCTEXT("AmbientCubemapSF", "Ambient Cubemap"))
+SHOWFLAG_ALWAYS_ACCESSIBLE(AmbientCubemap, SFG_LightingFeatures, NSLOCTEXT("UnrealEd", "AmbientCubemapSF", "Ambient Cubemap"))
 /** Human like eye simulation to adapt to the brightness of the view, for now SHOWFLAG_ALWAYS_ACCESSIBLE because it's exposed in SceneCapture */
-SHOWFLAG_ALWAYS_ACCESSIBLE(EyeAdaptation, SFG_PostProcess, LOCTEXT("EyeAdaptationSF", "Eye Adaptation"))
+SHOWFLAG_ALWAYS_ACCESSIBLE(EyeAdaptation, SFG_PostProcess, NSLOCTEXT("UnrealEd", "EyeAdaptationSF", "Eye Adaptation"))
 /** Display a histogram of the scene HDR color */
-SHOWFLAG_FIXED_IN_SHIPPING(0, VisualizeHDR, SFG_Visualize, LOCTEXT("VisualizeHDRSF", "HDR (Eye Adaptation)"))
+SHOWFLAG_FIXED_IN_SHIPPING(0, VisualizeHDR, SFG_Visualize, NSLOCTEXT("UnrealEd", "VisualizeHDRSF", "HDR (Eye Adaptation)"))
 /** Image based lens flares (Simulate artifact of reflections within a camera system) */
-SHOWFLAG_FIXED_IN_SHIPPING(1, LensFlares, SFG_PostProcess, LOCTEXT("LensFlaresSF", "Lens Flares"))
+SHOWFLAG_FIXED_IN_SHIPPING(1, LensFlares, SFG_PostProcess, NSLOCTEXT("UnrealEd", "LensFlaresSF", "Lens Flares"))
 /** show indirect lighting component, for now SHOWFLAG_ALWAYS_ACCESSIBLE because it's needed by r.GBuffer */
-SHOWFLAG_ALWAYS_ACCESSIBLE(GlobalIllumination, SFG_LightingComponents, LOCTEXT("GlobalIlluminationSF", "Global Illumination"))
+SHOWFLAG_ALWAYS_ACCESSIBLE(GlobalIllumination, SFG_LightingComponents, NSLOCTEXT("UnrealEd", "GlobalIlluminationSF", "Global Illumination"))
 /** Darkens the screen borders (Camera artifact and artistic effect) */
-SHOWFLAG_FIXED_IN_SHIPPING(1, Vignette, SFG_PostProcess, LOCTEXT("VignetteSF", "Vignette"))
+SHOWFLAG_FIXED_IN_SHIPPING(1, Vignette, SFG_PostProcess, NSLOCTEXT("UnrealEd", "VignetteSF", "Vignette"))
 /** Fine film grain */
-SHOWFLAG_FIXED_IN_SHIPPING(1, Grain, SFG_PostProcess, LOCTEXT("GrainSF", "Grain"))
+SHOWFLAG_FIXED_IN_SHIPPING(1, Grain, SFG_PostProcess, NSLOCTEXT("UnrealEd", "GrainSF", "Grain"))
 /** Screen Space Ambient Occlusion, for now SHOWFLAG_ALWAYS_ACCESSIBLE because it's exposed in SceneCapture */
-SHOWFLAG_ALWAYS_ACCESSIBLE(AmbientOcclusion, SFG_LightingComponents, LOCTEXT("AmbientOcclusionSF", "Ambient Occlusion"))
+SHOWFLAG_ALWAYS_ACCESSIBLE(AmbientOcclusion, SFG_LightingComponents, NSLOCTEXT("UnrealEd", "AmbientOcclusionSF", "Ambient Occlusion"))
 /** Decal rendering, for now SHOWFLAG_ALWAYS_ACCESSIBLE because it's exposed in SceneCapture */
-SHOWFLAG_ALWAYS_ACCESSIBLE(Decals, SFG_Normal, LOCTEXT("DecalsSF", "Decals"))
+SHOWFLAG_ALWAYS_ACCESSIBLE(Decals, SFG_Normal, NSLOCTEXT("UnrealEd", "DecalsSF", "Decals"))
 /** like bloom dirt mask */
-SHOWFLAG_FIXED_IN_SHIPPING(1, CameraImperfections, SFG_PostProcess, LOCTEXT("CameraImperfectionsSF", "Camera Imperfections"))
+SHOWFLAG_FIXED_IN_SHIPPING(1, CameraImperfections, SFG_PostProcess, NSLOCTEXT("UnrealEd", "CameraImperfectionsSF", "Camera Imperfections"))
 /** to allow to disable visualizetexture for some editor rendering (e.g. thumbnail rendering) */
-SHOWFLAG_ALWAYS_ACCESSIBLE(OnScreenDebug, SFG_Developer, LOCTEXT("OnScreenDebugSF", "On Screen Debug"))
+SHOWFLAG_ALWAYS_ACCESSIBLE(OnScreenDebug, SFG_Developer, NSLOCTEXT("UnrealEd", "OnScreenDebugSF", "On Screen Debug"))
 /** needed for VMI_Lit_DetailLighting, Whether to override material diffuse and specular with constants, used by the Detail Lighting viewmode. */
-SHOWFLAG_FIXED_IN_SHIPPING(0, OverrideDiffuseAndSpecular, SFG_Hidden, LOCTEXT("OverrideDiffuseAndSpecularSF", "Override Diffuse And Specular"))
+SHOWFLAG_FIXED_IN_SHIPPING(0, OverrideDiffuseAndSpecular, SFG_Hidden, NSLOCTEXT("UnrealEd", "OverrideDiffuseAndSpecularSF", "Override Diffuse And Specular"))
 /** needed for VMI_ReflectionOverride, Whether to override all materials to be smooth, mirror reflections. */
-SHOWFLAG_FIXED_IN_SHIPPING(0, ReflectionOverride, SFG_Hidden, LOCTEXT("ReflectionOverrideSF", "Reflections"))
+SHOWFLAG_FIXED_IN_SHIPPING(0, ReflectionOverride, SFG_Hidden, NSLOCTEXT("UnrealEd", "ReflectionOverrideSF", "Reflections"))
 /** needed for VMI_VisualizeBuffer, Whether to enable the buffer visualization mode. */
-SHOWFLAG_FIXED_IN_SHIPPING(0, VisualizeBuffer, SFG_Hidden, LOCTEXT("VisualizeBufferSF", "Buffer Visualization"))
+SHOWFLAG_FIXED_IN_SHIPPING(0, VisualizeBuffer, SFG_Hidden, NSLOCTEXT("UnrealEd", "VisualizeBufferSF", "Buffer Visualization"))
 /** Allows to disable all direct lighting (does not affect indirect light) */
-SHOWFLAG_FIXED_IN_SHIPPING(1, DirectLighting, SFG_LightingComponents, LOCTEXT("DirectLightingSF", "Direct Lighting"))
+SHOWFLAG_FIXED_IN_SHIPPING(1, DirectLighting, SFG_LightingComponents, NSLOCTEXT("UnrealEd", "DirectLightingSF", "Direct Lighting"))
 /** Allows to disable lighting from Directional Lights */
-SHOWFLAG_FIXED_IN_SHIPPING(1, DirectionalLights, SFG_LightTypes, LOCTEXT("DirectionalLightsSF", "Directional Lights"))
+SHOWFLAG_FIXED_IN_SHIPPING(1, DirectionalLights, SFG_LightTypes, NSLOCTEXT("UnrealEd", "DirectionalLightsSF", "Directional Lights"))
 /** Allows to disable lighting from Point Lights, for now SHOWFLAG_ALWAYS_ACCESSIBLE because it's exposed in SceneCapture */
-SHOWFLAG_ALWAYS_ACCESSIBLE(PointLights, SFG_LightTypes, LOCTEXT("PointLightsSF", "Point Lights"))
+SHOWFLAG_ALWAYS_ACCESSIBLE(PointLights, SFG_LightTypes, NSLOCTEXT("UnrealEd", "PointLightsSF", "Point Lights"))
 /** Allows to disable lighting from Spot Lights, for now SHOWFLAG_ALWAYS_ACCESSIBLE because it's exposed in SceneCapture */
-SHOWFLAG_ALWAYS_ACCESSIBLE(SpotLights, SFG_LightTypes, LOCTEXT("SpotLightsSF", "Spot Lights"))
+SHOWFLAG_ALWAYS_ACCESSIBLE(SpotLights, SFG_LightTypes, NSLOCTEXT("UnrealEd", "SpotLightsSF", "Spot Lights"))
 /** Color correction after tone mapping */
-SHOWFLAG_FIXED_IN_SHIPPING(1, ColorGrading, SFG_PostProcess, LOCTEXT("ColorGradingSF", "Color Grading"))
+SHOWFLAG_FIXED_IN_SHIPPING(1, ColorGrading, SFG_PostProcess, NSLOCTEXT("UnrealEd", "ColorGradingSF", "Color Grading"))
 /** Visualize vector fields. */
-SHOWFLAG_FIXED_IN_SHIPPING(0, VectorFields, SFG_Developer, LOCTEXT("VectorFieldsSF", "Vector Fields"))
+SHOWFLAG_FIXED_IN_SHIPPING(0, VectorFields, SFG_Developer, NSLOCTEXT("UnrealEd", "VectorFieldsSF", "Vector Fields"))
 /** Depth of Field */
-SHOWFLAG_FIXED_IN_SHIPPING(1, DepthOfField, SFG_PostProcess, LOCTEXT("DepthOfFieldSF", "Depth Of Field"))
+SHOWFLAG_FIXED_IN_SHIPPING(1, DepthOfField, SFG_PostProcess, NSLOCTEXT("UnrealEd", "DepthOfFieldSF", "Depth Of Field"))
 /** Highlight materials that indicate performance issues or show unrealistic materials */
-SHOWFLAG_FIXED_IN_SHIPPING(0, GBufferHints, SFG_Developer, LOCTEXT("GBufferHintsSF", "GBuffer Hints (material attributes)"))
+SHOWFLAG_FIXED_IN_SHIPPING(0, GBufferHints, SFG_Developer, NSLOCTEXT("UnrealEd", "GBufferHintsSF", "GBuffer Hints (material attributes)"))
 /** MotionBlur, for now only camera motion blur, for now SHOWFLAG_ALWAYS_ACCESSIBLE because it's exposed in SceneCapture */
-SHOWFLAG_ALWAYS_ACCESSIBLE(MotionBlur, SFG_PostProcess, LOCTEXT("MotionBlurSF", "Motion Blur"))
+SHOWFLAG_ALWAYS_ACCESSIBLE(MotionBlur, SFG_PostProcess, NSLOCTEXT("UnrealEd", "MotionBlurSF", "Motion Blur"))
 /** Whether to render the editor gizmos and other foreground editor widgets off screen and apply them after post process, only needed for the editor */
-SHOWFLAG_FIXED_IN_SHIPPING(0, CompositeEditorPrimitives, SFG_Developer, LOCTEXT("CompositeEditorPrimitivesSF", "Composite Editor Primitives"))
+SHOWFLAG_FIXED_IN_SHIPPING(0, CompositeEditorPrimitives, SFG_Developer, NSLOCTEXT("UnrealEd", "CompositeEditorPrimitivesSF", "Composite Editor Primitives"))
 /** Shows a test image that allows to tweak the monitor colors, borders and allows to judge image and temporal aliasing  */
-SHOWFLAG_FIXED_IN_SHIPPING(0, TestImage, SFG_Developer, LOCTEXT("TestImageSF", "Test Image"))
+SHOWFLAG_FIXED_IN_SHIPPING(0, TestImage, SFG_Developer, NSLOCTEXT("UnrealEd", "TestImageSF", "Test Image"))
 /** Helper to tweak depth of field */
-SHOWFLAG_FIXED_IN_SHIPPING(0, VisualizeDOF, SFG_Visualize, LOCTEXT("VisualizeDOFSF", "Depth of Field Layers"))
+SHOWFLAG_FIXED_IN_SHIPPING(0, VisualizeDOF, SFG_Visualize, NSLOCTEXT("UnrealEd", "VisualizeDOFSF", "Depth of Field Layers"))
 /** Helper to tweak depth of field */
-SHOWFLAG_FIXED_IN_SHIPPING(0, VisualizeAdaptiveDOF, SFG_Visualize, LOCTEXT("VisualizeAdaptiveDOFSF", "Adaptive Depth of Field"))
+SHOWFLAG_FIXED_IN_SHIPPING(0, VisualizeAdaptiveDOF, SFG_Visualize, NSLOCTEXT("UnrealEd", "VisualizeAdaptiveDOFSF", "Adaptive Depth of Field"))
 /** Show Vertex Colors */
-SHOWFLAG_FIXED_IN_SHIPPING(0, VertexColors, SFG_Advanced, LOCTEXT("VertexColorsSF", "Vertex Colors"))
+SHOWFLAG_FIXED_IN_SHIPPING(0, VertexColors, SFG_Advanced, NSLOCTEXT("UnrealEd", "VertexColorsSF", "Vertex Colors"))
 /** Render Post process (screen space) distortion/refraction */
-SHOWFLAG_FIXED_IN_SHIPPING(1, Refraction, SFG_Developer, LOCTEXT("RefractionSF", "Refraction"))
+SHOWFLAG_FIXED_IN_SHIPPING(1, Refraction, SFG_Developer, NSLOCTEXT("UnrealEd", "RefractionSF", "Refraction"))
 /** Usually set in game or when previewing Matinee but not in editor, used for motion blur or any kind of rendering features that rely on the former frame */
-SHOWFLAG_ALWAYS_ACCESSIBLE(CameraInterpolation, SFG_Hidden, LOCTEXT("CameraInterpolationSF", "Camera Interpolation"))
+SHOWFLAG_ALWAYS_ACCESSIBLE(CameraInterpolation, SFG_Hidden, NSLOCTEXT("UnrealEd", "CameraInterpolationSF", "Camera Interpolation"))
 /** Post processing color fringe (chromatic aberration) */
-SHOWFLAG_FIXED_IN_SHIPPING(1, SceneColorFringe, SFG_PostProcess, LOCTEXT("SceneColorFringeSF", "Scene Color Fringe"))
+SHOWFLAG_FIXED_IN_SHIPPING(1, SceneColorFringe, SFG_PostProcess, NSLOCTEXT("UnrealEd", "SceneColorFringeSF", "Scene Color Fringe"))
 /** If Translucency should be rendered into a separate RT and composited without DepthOfField, can be disabled in the materials (affects sorting), SHOWFLAG_ALWAYS_ACCESSIBLE for now because USceneCaptureComponent needs that */
-SHOWFLAG_ALWAYS_ACCESSIBLE(SeparateTranslucency, SFG_Advanced, LOCTEXT("SeparateTranslucencySF", "Separate Translucency"))
+SHOWFLAG_ALWAYS_ACCESSIBLE(SeparateTranslucency, SFG_Advanced, NSLOCTEXT("UnrealEd", "SeparateTranslucencySF", "Separate Translucency"))
 /** If Screen Percentage should be applied.  */
-SHOWFLAG_ALWAYS_ACCESSIBLE(ScreenPercentage, SFG_Hidden, LOCTEXT("ScreenPercentageSF", "Screen Percentage"))
+SHOWFLAG_ALWAYS_ACCESSIBLE(ScreenPercentage, SFG_Hidden, NSLOCTEXT("UnrealEd", "ScreenPercentageSF", "Screen Percentage"))
 /** Helper to tweak motion blur settings */
-SHOWFLAG_FIXED_IN_SHIPPING(0, VisualizeMotionBlur, SFG_Visualize, LOCTEXT("VisualizeMotionBlurSF", "Motion Blur"))
+SHOWFLAG_FIXED_IN_SHIPPING(0, VisualizeMotionBlur, SFG_Visualize, NSLOCTEXT("UnrealEd", "VisualizeMotionBlurSF", "Motion Blur"))
 /** Whether to display the Reflection Environment feature, which has local reflections from Reflection Capture actors, for now SHOWFLAG_ALWAYS_ACCESSIBLE because it's exposed in SceneCapture */
-SHOWFLAG_ALWAYS_ACCESSIBLE(ReflectionEnvironment, SFG_LightingFeatures, LOCTEXT("ReflectionEnvironmentSF", "Reflection Environment"))
+SHOWFLAG_ALWAYS_ACCESSIBLE(ReflectionEnvironment, SFG_LightingFeatures, NSLOCTEXT("UnrealEd", "ReflectionEnvironmentSF", "Reflection Environment"))
 /** Visualize pixels that are outside of their object's bounding box (content error). */
-SHOWFLAG_FIXED_IN_SHIPPING(0, VisualizeOutOfBoundsPixels, SFG_Visualize, LOCTEXT("VisualizeOutOfBoundsPixelsSF", "Out of Bounds Pixels"))
+SHOWFLAG_FIXED_IN_SHIPPING(0, VisualizeOutOfBoundsPixels, SFG_Visualize, NSLOCTEXT("UnrealEd", "VisualizeOutOfBoundsPixelsSF", "Out of Bounds Pixels"))
 /** Whether to display the scene's diffuse. */
-SHOWFLAG_FIXED_IN_SHIPPING(1, Diffuse, SFG_LightingComponents, LOCTEXT("DiffuseSF", "Diffuse"))
+SHOWFLAG_FIXED_IN_SHIPPING(1, Diffuse, SFG_LightingComponents, NSLOCTEXT("UnrealEd", "DiffuseSF", "Diffuse"))
 /** Whether to display the scene's specular, including reflections. */
-SHOWFLAG_ALWAYS_ACCESSIBLE(Specular, SFG_LightingComponents, LOCTEXT("SpecularSF", "Specular"))
+SHOWFLAG_ALWAYS_ACCESSIBLE(Specular, SFG_LightingComponents, NSLOCTEXT("UnrealEd", "SpecularSF", "Specular"))
 /** Outline around selected objects in the editor */
-SHOWFLAG_FIXED_IN_SHIPPING(0, SelectionOutline, SFG_Hidden, LOCTEXT("SelectionOutlineSF", "Selection Outline"))
+SHOWFLAG_FIXED_IN_SHIPPING(0, SelectionOutline, SFG_Hidden, NSLOCTEXT("UnrealEd", "SelectionOutlineSF", "Selection Outline"))
 /** If screen space reflections are enabled, for now SHOWFLAG_ALWAYS_ACCESSIBLE because it's exposed in SceneCapture */
-SHOWFLAG_ALWAYS_ACCESSIBLE(ScreenSpaceReflections, SFG_LightingFeatures, LOCTEXT("ScreenSpaceReflectionsSF", "Screen Space Reflections"))
+SHOWFLAG_ALWAYS_ACCESSIBLE(ScreenSpaceReflections, SFG_LightingFeatures, NSLOCTEXT("UnrealEd", "ScreenSpaceReflectionsSF", "Screen Space Reflections"))
 /** If Screen space contact shadows are enabled. */
-SHOWFLAG_ALWAYS_ACCESSIBLE(ContactShadows, SFG_LightingFeatures, LOCTEXT("ContactShadows", "Screen Space Contact Shadows"))
+SHOWFLAG_ALWAYS_ACCESSIBLE(ContactShadows, SFG_LightingFeatures, NSLOCTEXT("UnrealEd", "ContactShadows", "Screen Space Contact Shadows"))
 /** If RTDF shadows are enabled. */
-SHOWFLAG_ALWAYS_ACCESSIBLE(RayTracedDistanceFieldShadows, SFG_LightingFeatures, LOCTEXT("RayTracedDistanceFieldShadows", "Ray Traced Distance Field Shadows"))
+SHOWFLAG_ALWAYS_ACCESSIBLE(RayTracedDistanceFieldShadows, SFG_LightingFeatures, NSLOCTEXT("UnrealEd", "RayTracedDistanceFieldShadows", "Ray Traced Distance Field Shadows"))
 /** If Capsule shadows are enabled. */
-SHOWFLAG_ALWAYS_ACCESSIBLE(CapsuleShadows, SFG_LightingFeatures, LOCTEXT("CapsuleShadows", "Capsule Shadows"))
+SHOWFLAG_ALWAYS_ACCESSIBLE(CapsuleShadows, SFG_LightingFeatures, NSLOCTEXT("UnrealEd", "CapsuleShadows", "Capsule Shadows"))
 /** If Screen Space Subsurface Scattering enabled */
-SHOWFLAG_FIXED_IN_SHIPPING(1, SubsurfaceScattering, SFG_LightingFeatures, LOCTEXT("SubsurfaceScatteringSF", "Subsurface Scattering (Screen Space)"))
+SHOWFLAG_FIXED_IN_SHIPPING(1, SubsurfaceScattering, SFG_LightingFeatures, NSLOCTEXT("UnrealEd", "SubsurfaceScatteringSF", "Subsurface Scattering (Screen Space)"))
 /** If Screen Space Subsurface Scattering visualization is enabled */
-SHOWFLAG_FIXED_IN_SHIPPING(0, VisualizeSSS, SFG_Visualize, LOCTEXT("VisualizeSSSSF", "Subsurface Scattering (Screen Space)"))
+SHOWFLAG_FIXED_IN_SHIPPING(0, VisualizeSSS, SFG_Visualize, NSLOCTEXT("UnrealEd", "VisualizeSSSSF", "Subsurface Scattering (Screen Space)"))
 /** Whether to apply volumetric lightmap lighting, when present. */
-SHOWFLAG_ALWAYS_ACCESSIBLE(VolumetricLightmap, SFG_LightingFeatures, LOCTEXT("VolumetricLightmapSF", "Volumetric Lightmap"))
+SHOWFLAG_ALWAYS_ACCESSIBLE(VolumetricLightmap, SFG_LightingFeatures, NSLOCTEXT("UnrealEd", "VolumetricLightmapSF", "Volumetric Lightmap"))
 /** If the indirect lighting cache is enabled, for now SHOWFLAG_ALWAYS_ACCESSIBLE because it's exposed in SceneCapture */
-SHOWFLAG_ALWAYS_ACCESSIBLE(IndirectLightingCache, SFG_LightingFeatures, LOCTEXT("IndirectLightingCacheSF", "Indirect Lighting Cache"))
+SHOWFLAG_ALWAYS_ACCESSIBLE(IndirectLightingCache, SFG_LightingFeatures, NSLOCTEXT("UnrealEd", "IndirectLightingCacheSF", "Indirect Lighting Cache"))
 /** calls debug drawing for AIs */
-SHOWFLAG_FIXED_IN_SHIPPING(0, DebugAI, SFG_Developer, LOCTEXT("DebugAISF", "AI Debug"))
+SHOWFLAG_FIXED_IN_SHIPPING(0, DebugAI, SFG_Developer, NSLOCTEXT("UnrealEd", "DebugAISF", "AI Debug"))
 /** calls debug drawing for whatever LogVisualizer wants to draw */
-SHOWFLAG_FIXED_IN_SHIPPING(0, VisLog, SFG_Developer, LOCTEXT("VisLogSF", "Log Visualizer"))
+SHOWFLAG_FIXED_IN_SHIPPING(0, VisLog, SFG_Developer, NSLOCTEXT("UnrealEd", "VisLogSF", "Log Visualizer"))
 /** whether to draw navigation data */
-SHOWFLAG_FIXED_IN_SHIPPING(0, Navigation, SFG_Normal, LOCTEXT("NavigationSF", "Navigation"))
+SHOWFLAG_FIXED_IN_SHIPPING(0, Navigation, SFG_Normal, NSLOCTEXT("UnrealEd", "NavigationSF", "Navigation"))
 /** used by gameplay debugging components to debug-draw on screen */
-SHOWFLAG_FIXED_IN_SHIPPING(0, GameplayDebug, SFG_Developer, LOCTEXT("GameplayDebugSF", "Gameplay Debug"))
+SHOWFLAG_FIXED_IN_SHIPPING(0, GameplayDebug, SFG_Developer, NSLOCTEXT("UnrealEd", "GameplayDebugSF", "Gameplay Debug"))
 /** LightProfiles, usually 1d textures to have a light (IES), for now SHOWFLAG_ALWAYS_ACCESSIBLE because it's exposed in SceneCapture */
-SHOWFLAG_ALWAYS_ACCESSIBLE(TexturedLightProfiles,  SFG_LightingFeatures, LOCTEXT("TexturedLightProfilesSF", "Textured Light Profiles (IES Texture)"))
+SHOWFLAG_ALWAYS_ACCESSIBLE(TexturedLightProfiles,  SFG_LightingFeatures, NSLOCTEXT("UnrealEd", "TexturedLightProfilesSF", "Textured Light Profiles (IES Texture)"))
 /** LightFunctions (masking light sources with a material), for now SHOWFLAG_ALWAYS_ACCESSIBLE because it's exposed in SceneCapture */
-SHOWFLAG_ALWAYS_ACCESSIBLE(LightFunctions, SFG_LightingFeatures, LOCTEXT("LightFunctionsSF", "Light Functions"))
+SHOWFLAG_ALWAYS_ACCESSIBLE(LightFunctions, SFG_LightingFeatures, NSLOCTEXT("UnrealEd", "LightFunctionsSF", "Light Functions"))
 /** Hardware Tessellation (DX11 feature) */
-SHOWFLAG_FIXED_IN_SHIPPING(1, Tessellation, SFG_Advanced, LOCTEXT("TessellationSF", "Tessellation"))
+SHOWFLAG_FIXED_IN_SHIPPING(1, Tessellation, SFG_Advanced, NSLOCTEXT("UnrealEd", "TessellationSF", "Tessellation"))
 /** Draws instanced static meshes that are not foliage or grass, for now SHOWFLAG_ALWAYS_ACCESSIBLE because it's exposed in SceneCapture */
-SHOWFLAG_ALWAYS_ACCESSIBLE(InstancedStaticMeshes, SFG_Advanced, LOCTEXT("InstancedStaticMeshesSF", "Instanced Static Meshes"))
+SHOWFLAG_ALWAYS_ACCESSIBLE(InstancedStaticMeshes, SFG_Advanced, NSLOCTEXT("UnrealEd", "InstancedStaticMeshesSF", "Instanced Static Meshes"))
 /** Draws instanced foliage, for now SHOWFLAG_ALWAYS_ACCESSIBLE because it's exposed in SceneCapture */
-SHOWFLAG_ALWAYS_ACCESSIBLE(InstancedFoliage, SFG_Advanced, LOCTEXT("InstancedFoliageSF", "Foliage"))
+SHOWFLAG_ALWAYS_ACCESSIBLE(InstancedFoliage, SFG_Advanced, NSLOCTEXT("UnrealEd", "InstancedFoliageSF", "Foliage"))
 /** Allow to see the foliage bounds used in the occlusion test */
-SHOWFLAG_FIXED_IN_SHIPPING(0, FoliageOcclusionBounds, SFG_Advanced, LOCTEXT("FoliageOcclusionBoundsSF", "Foliage Occlusion Bounds"))
+SHOWFLAG_FIXED_IN_SHIPPING(0, FoliageOcclusionBounds, SFG_Advanced, NSLOCTEXT("UnrealEd", "FoliageOcclusionBoundsSF", "Foliage Occlusion Bounds"))
 /** Draws instanced grass, for now SHOWFLAG_ALWAYS_ACCESSIBLE because it's exposed in SceneCapture */
-SHOWFLAG_ALWAYS_ACCESSIBLE(InstancedGrass, SFG_Advanced, LOCTEXT("InstancedGrassSF", "Grass"))
+SHOWFLAG_ALWAYS_ACCESSIBLE(InstancedGrass, SFG_Advanced, NSLOCTEXT("UnrealEd", "InstancedGrassSF", "Grass"))
 /** non baked shadows, for now SHOWFLAG_ALWAYS_ACCESSIBLE because it's exposed in SceneCapture */
-SHOWFLAG_ALWAYS_ACCESSIBLE(DynamicShadows, SFG_LightingComponents, LOCTEXT("DynamicShadowsSF", "Dynamic Shadows"))
+SHOWFLAG_ALWAYS_ACCESSIBLE(DynamicShadows, SFG_LightingComponents, NSLOCTEXT("UnrealEd", "DynamicShadowsSF", "Dynamic Shadows"))
 /** Particles, for now SHOWFLAG_ALWAYS_ACCESSIBLE because it's exposed in SceneCapture */
-SHOWFLAG_ALWAYS_ACCESSIBLE(Particles, SFG_Normal, LOCTEXT("ParticlesSF", "Particle Sprites"))
+SHOWFLAG_ALWAYS_ACCESSIBLE(Particles, SFG_Normal, NSLOCTEXT("UnrealEd", "ParticlesSF", "Particle Sprites"))
 /** if SkeletalMeshes are getting rendered, for now SHOWFLAG_ALWAYS_ACCESSIBLE because it's exposed in SceneCapture */
-SHOWFLAG_ALWAYS_ACCESSIBLE(SkeletalMeshes, SFG_Normal, LOCTEXT("SkeletalMeshesSF", "Skeletal Meshes"))
+SHOWFLAG_ALWAYS_ACCESSIBLE(SkeletalMeshes, SFG_Normal, NSLOCTEXT("UnrealEd", "SkeletalMeshesSF", "Skeletal Meshes"))
 /** if the builder brush (editor) is getting rendered */
-SHOWFLAG_FIXED_IN_SHIPPING(0, BuilderBrush, SFG_Hidden, LOCTEXT("BuilderBrushSF", "Builder Brush"))
+SHOWFLAG_FIXED_IN_SHIPPING(0, BuilderBrush, SFG_Hidden, NSLOCTEXT("UnrealEd", "BuilderBrushSF", "Builder Brush"))
 /** Render translucency, for now SHOWFLAG_ALWAYS_ACCESSIBLE because it's exposed in SceneCapture */
-SHOWFLAG_ALWAYS_ACCESSIBLE(Translucency, SFG_Normal, LOCTEXT("TranslucencySF", "Translucency"))
+SHOWFLAG_ALWAYS_ACCESSIBLE(Translucency, SFG_Normal, NSLOCTEXT("UnrealEd", "TranslucencySF", "Translucency"))
 /** Draw billboard components */
-SHOWFLAG_FIXED_IN_SHIPPING(1, BillboardSprites, SFG_Advanced, LOCTEXT("BillboardSpritesSF", "Billboard Sprites"))
+SHOWFLAG_FIXED_IN_SHIPPING(1, BillboardSprites, SFG_Advanced, NSLOCTEXT("UnrealEd", "BillboardSpritesSF", "Billboard Sprites"))
 /** Use LOD parenting, MinDrawDistance, etc. If disabled, will show LOD parenting lines */
-SHOWFLAG_ALWAYS_ACCESSIBLE(LOD, SFG_Advanced, LOCTEXT("LODSF", "LOD Parenting"))
+SHOWFLAG_ALWAYS_ACCESSIBLE(LOD, SFG_Advanced, NSLOCTEXT("UnrealEd", "LODSF", "LOD Parenting"))
 /** needed for VMI_LightComplexity */
-SHOWFLAG_FIXED_IN_SHIPPING(0, LightComplexity, SFG_Hidden, LOCTEXT("LightComplexitySF", "Light Complexity"))
+SHOWFLAG_FIXED_IN_SHIPPING(0, LightComplexity, SFG_Hidden, NSLOCTEXT("UnrealEd", "LightComplexitySF", "Light Complexity"))
 /** needed for VMI_ShaderComplexity, render world colored by shader complexity */
-SHOWFLAG_FIXED_IN_SHIPPING(0, ShaderComplexity, SFG_Hidden, LOCTEXT("ShaderComplexitySF", "Shader Complexity"))
+SHOWFLAG_FIXED_IN_SHIPPING(0, ShaderComplexity, SFG_Hidden, NSLOCTEXT("UnrealEd", "ShaderComplexitySF", "Shader Complexity"))
 /** needed for VMI_StationaryLightOverlap, render world colored by stationary light overlap */
-SHOWFLAG_FIXED_IN_SHIPPING(0, StationaryLightOverlap,  SFG_Hidden, LOCTEXT("StationaryLightOverlapSF", "Stationary Light Overlap"))
+SHOWFLAG_FIXED_IN_SHIPPING(0, StationaryLightOverlap,  SFG_Hidden, NSLOCTEXT("UnrealEd", "StationaryLightOverlapSF", "Stationary Light Overlap"))
 /** needed for VMI_LightmapDensity and VMI_LitLightmapDensity, render checkerboard material with UVs scaled by lightmap resolution w. color tint for world-space lightmap density */
-SHOWFLAG_FIXED_IN_SHIPPING(0, LightMapDensity, SFG_Hidden, LOCTEXT("LightMapDensitySF", "Light Map Density"))
+SHOWFLAG_FIXED_IN_SHIPPING(0, LightMapDensity, SFG_Hidden, NSLOCTEXT("UnrealEd", "LightMapDensitySF", "Light Map Density"))
 /** Render streaming bounding volumes for the currently selected texture */
-SHOWFLAG_FIXED_IN_SHIPPING(0, StreamingBounds, SFG_Advanced, LOCTEXT("StreamingBoundsSF", "Streaming Bounds"))
+SHOWFLAG_FIXED_IN_SHIPPING(0, StreamingBounds, SFG_Advanced, NSLOCTEXT("UnrealEd", "StreamingBoundsSF", "Streaming Bounds"))
 /** Render joint limits */
-SHOWFLAG_FIXED_IN_SHIPPING(0, Constraints, SFG_Advanced, LOCTEXT("ConstraintsSF", "Constraints"))
+SHOWFLAG_FIXED_IN_SHIPPING(0, Constraints, SFG_Advanced, NSLOCTEXT("UnrealEd", "ConstraintsSF", "Constraints"))
 /** Render mass debug data */
-SHOWFLAG_FIXED_IN_SHIPPING(0, MassProperties, SFG_Advanced, LOCTEXT("MassPropertiesSF", "Mass Properties"))
+SHOWFLAG_FIXED_IN_SHIPPING(0, MassProperties, SFG_Advanced, NSLOCTEXT("UnrealEd", "MassPropertiesSF", "Mass Properties"))
 /** Draws camera frustums */
-SHOWFLAG_FIXED_IN_SHIPPING(0, CameraFrustums, SFG_Advanced, LOCTEXT("CameraFrustumsSF", "Camera Frustums"))
+SHOWFLAG_FIXED_IN_SHIPPING(0, CameraFrustums, SFG_Advanced, NSLOCTEXT("UnrealEd", "CameraFrustumsSF", "Camera Frustums"))
 /** Draw sound actor radii */
-SHOWFLAG_FIXED_IN_SHIPPING(0, AudioRadius, SFG_Advanced, LOCTEXT("AudioRadiusSF", "Audio Radius"))
+SHOWFLAG_FIXED_IN_SHIPPING(0, AudioRadius, SFG_Advanced, NSLOCTEXT("UnrealEd", "AudioRadiusSF", "Audio Radius"))
 /** Draw force feedback radii */
-SHOWFLAG_FIXED_IN_SHIPPING(0, ForceFeedbackRadius, SFG_Advanced, LOCTEXT("ForceFeedbackSF", "Force Feedback Radius"))
+SHOWFLAG_FIXED_IN_SHIPPING(0, ForceFeedbackRadius, SFG_Advanced, NSLOCTEXT("UnrealEd", "ForceFeedbackSF", "Force Feedback Radius"))
 /** Colors BSP based on model component association */
-SHOWFLAG_FIXED_IN_SHIPPING(0, BSPSplit, SFG_Advanced, LOCTEXT("BSPSplitSF", "BSP Split"))
+SHOWFLAG_FIXED_IN_SHIPPING(0, BSPSplit, SFG_Advanced, NSLOCTEXT("UnrealEd", "BSPSplitSF", "BSP Split"))
 /** show editor (wireframe) brushes, for now SHOWFLAG_ALWAYS_ACCESSIBLE because it's exposed in SceneCapture */
-SHOWFLAG_FIXED_IN_SHIPPING(0, Brushes, SFG_Hidden, LOCTEXT("BrushesSF", "Brushes"))
+SHOWFLAG_FIXED_IN_SHIPPING(0, Brushes, SFG_Hidden, NSLOCTEXT("UnrealEd", "BrushesSF", "Brushes"))
 /** Show the usual material light interaction */
-SHOWFLAG_ALWAYS_ACCESSIBLE(Lighting, SFG_Hidden, LOCTEXT("LightingSF", "Lighting"))
+SHOWFLAG_ALWAYS_ACCESSIBLE(Lighting, SFG_Hidden, NSLOCTEXT("UnrealEd", "LightingSF", "Lighting"))
 /** Execute the deferred light passes, can be disabled for debugging, for now SHOWFLAG_ALWAYS_ACCESSIBLE because it's exposed in SceneCapture */
-SHOWFLAG_ALWAYS_ACCESSIBLE(DeferredLighting, SFG_Advanced, LOCTEXT("DeferredLightingSF", "DeferredLighting"))
+SHOWFLAG_ALWAYS_ACCESSIBLE(DeferredLighting, SFG_Advanced, NSLOCTEXT("UnrealEd", "DeferredLightingSF", "DeferredLighting"))
 /** Special: Allows to hide objects in the editor, is evaluated per primitive */
-SHOWFLAG_FIXED_IN_SHIPPING(0, Editor, SFG_Hidden, LOCTEXT("EditorSF", "Editor"))
+SHOWFLAG_FIXED_IN_SHIPPING(0, Editor, SFG_Hidden, NSLOCTEXT("UnrealEd", "EditorSF", "Editor"))
 /** needed for VMI_BrushWireframe and VMI_LitLightmapDensity, Draws BSP triangles */
-SHOWFLAG_FIXED_IN_SHIPPING(1, BSPTriangles, SFG_Hidden, LOCTEXT("BSPTrianglesSF", "BSP Triangles"))
+SHOWFLAG_FIXED_IN_SHIPPING(1, BSPTriangles, SFG_Hidden, NSLOCTEXT("UnrealEd", "BSPTrianglesSF", "BSP Triangles"))
 /** Displays large clickable icons on static mesh vertices, only needed for the editor */
-SHOWFLAG_FIXED_IN_SHIPPING(0, LargeVertices, SFG_Advanced, LOCTEXT("LargeVerticesSF", "Large Vertices"))
+SHOWFLAG_FIXED_IN_SHIPPING(0, LargeVertices, SFG_Advanced, NSLOCTEXT("UnrealEd", "LargeVerticesSF", "Large Vertices"))
 /** To show the grid in editor (grey lines and red dots) */
-SHOWFLAG_FIXED_IN_SHIPPING(0, Grid, SFG_Normal, LOCTEXT("GridSF", "Grid"))
+SHOWFLAG_FIXED_IN_SHIPPING(0, Grid, SFG_Normal, NSLOCTEXT("UnrealEd", "GridSF", "Grid"))
 /** To show the snap in editor (only for editor view ports, red dots) */
-SHOWFLAG_FIXED_IN_SHIPPING(0, Snap, SFG_Hidden, LOCTEXT("SnapSF", "Snap"))
+SHOWFLAG_FIXED_IN_SHIPPING(0, Snap, SFG_Hidden, NSLOCTEXT("UnrealEd", "SnapSF", "Snap"))
 /** In the filled view modeModeWidgetss, render mesh edges as well as the filled surfaces. */
-SHOWFLAG_FIXED_IN_SHIPPING(0, MeshEdges, SFG_Advanced, LOCTEXT("MeshEdgesSF", "Mesh Edges"))
+SHOWFLAG_FIXED_IN_SHIPPING(0, MeshEdges, SFG_Advanced, NSLOCTEXT("UnrealEd", "MeshEdgesSF", "Mesh Edges"))
 /** Complex cover rendering */
-SHOWFLAG_FIXED_IN_SHIPPING(0, Cover, SFG_Hidden, LOCTEXT("CoverSF", "Cover"))
+SHOWFLAG_FIXED_IN_SHIPPING(0, Cover, SFG_Hidden, NSLOCTEXT("UnrealEd", "CoverSF", "Cover"))
 /** Spline rendering */
-SHOWFLAG_FIXED_IN_SHIPPING(0, Splines, SFG_Advanced, LOCTEXT("SplinesSF", "Splines"))
+SHOWFLAG_FIXED_IN_SHIPPING(0, Splines, SFG_Advanced, NSLOCTEXT("UnrealEd", "SplinesSF", "Splines"))
 /** Selection rendering, could be useful in game as well */
-SHOWFLAG_FIXED_IN_SHIPPING(0, Selection, SFG_Advanced, LOCTEXT("SelectionSF", "Selection"))
+SHOWFLAG_FIXED_IN_SHIPPING(0, Selection, SFG_Advanced, NSLOCTEXT("UnrealEd", "SelectionSF", "Selection"))
 /** Draws mode specific widgets and controls in the viewports (should only be set on viewport clients that are editing the level itself) */
-SHOWFLAG_FIXED_IN_SHIPPING(0, ModeWidgets, SFG_Advanced, LOCTEXT("ModeWidgetsSF", "Mode Widgets"))
+SHOWFLAG_FIXED_IN_SHIPPING(0, ModeWidgets, SFG_Advanced, NSLOCTEXT("UnrealEd", "ModeWidgetsSF", "Mode Widgets"))
 /**  */
-SHOWFLAG_FIXED_IN_SHIPPING(0, Bounds,  SFG_Advanced, LOCTEXT("BoundsSF", "Bounds"))
+SHOWFLAG_FIXED_IN_SHIPPING(0, Bounds,  SFG_Advanced, NSLOCTEXT("UnrealEd", "BoundsSF", "Bounds"))
 /** Draws each hit proxy in the scene with a different color, for now only available in the editor */
-SHOWFLAG_FIXED_IN_SHIPPING(0, HitProxies, SFG_Developer, LOCTEXT("HitProxiesSF", "Hit Proxies"))
+SHOWFLAG_FIXED_IN_SHIPPING(0, HitProxies, SFG_Developer, NSLOCTEXT("UnrealEd", "HitProxiesSF", "Hit Proxies"))
 /** Render objects with colors based on the property values */
-SHOWFLAG_FIXED_IN_SHIPPING(0, PropertyColoration, SFG_Advanced, LOCTEXT("PropertyColorationSF", "Property Coloration"))
+SHOWFLAG_FIXED_IN_SHIPPING(0, PropertyColoration, SFG_Advanced, NSLOCTEXT("UnrealEd", "PropertyColorationSF", "Property Coloration"))
 /** Draw lines to lights affecting this mesh if its selected. */
-SHOWFLAG_FIXED_IN_SHIPPING(0, LightInfluences, SFG_Advanced, LOCTEXT("LightInfluencesSF", "Light Influences"))
+SHOWFLAG_FIXED_IN_SHIPPING(0, LightInfluences, SFG_Advanced, NSLOCTEXT("UnrealEd", "LightInfluencesSF", "Light Influences"))
 /** for the Editor */
-SHOWFLAG_FIXED_IN_SHIPPING(0, Pivot, SFG_Hidden, LOCTEXT("PivotSF", "Pivot"))
+SHOWFLAG_FIXED_IN_SHIPPING(0, Pivot, SFG_Hidden, NSLOCTEXT("UnrealEd", "PivotSF", "Pivot"))
 /** Draws un-occluded shadow frustums in wireframe */
-SHOWFLAG_FIXED_IN_SHIPPING(0, ShadowFrustums, SFG_Advanced, LOCTEXT("ShadowFrustumsSF", "Shadow Frustums"))
+SHOWFLAG_FIXED_IN_SHIPPING(0, ShadowFrustums, SFG_Advanced, NSLOCTEXT("UnrealEd", "ShadowFrustumsSF", "Shadow Frustums"))
 /** needed for VMI_Wireframe and VMI_BrushWireframe */
-SHOWFLAG_FIXED_IN_SHIPPING(0, Wireframe, SFG_Hidden, LOCTEXT("WireframeSF", "Wireframe"))
+SHOWFLAG_FIXED_IN_SHIPPING(0, Wireframe, SFG_Hidden, NSLOCTEXT("UnrealEd", "WireframeSF", "Wireframe"))
 /**  */
-SHOWFLAG_FIXED_IN_SHIPPING(1, Materials, SFG_Hidden, LOCTEXT("MaterialsSF", "Materials"))
+SHOWFLAG_FIXED_IN_SHIPPING(1, Materials, SFG_Hidden, NSLOCTEXT("UnrealEd", "MaterialsSF", "Materials"))
 /** for now SHOWFLAG_ALWAYS_ACCESSIBLE because it's exposed in SceneCapture */
-SHOWFLAG_ALWAYS_ACCESSIBLE(StaticMeshes, SFG_Normal, LOCTEXT("StaticMeshesSF", "Static Meshes"))
+SHOWFLAG_ALWAYS_ACCESSIBLE(StaticMeshes, SFG_Normal, NSLOCTEXT("UnrealEd", "StaticMeshesSF", "Static Meshes"))
 /** for now SHOWFLAG_ALWAYS_ACCESSIBLE because it's exposed in SceneCapture */
-SHOWFLAG_ALWAYS_ACCESSIBLE(Landscape, SFG_Normal, LOCTEXT("LandscapeSF", "Landscape"))
+SHOWFLAG_ALWAYS_ACCESSIBLE(Landscape, SFG_Normal, NSLOCTEXT("UnrealEd", "LandscapeSF", "Landscape"))
 /**  */
-SHOWFLAG_FIXED_IN_SHIPPING(0, LightRadius, SFG_Advanced, LOCTEXT("LightRadiusSF", "Light Radius"))
+SHOWFLAG_FIXED_IN_SHIPPING(0, LightRadius, SFG_Advanced, NSLOCTEXT("UnrealEd", "LightRadiusSF", "Light Radius"))
 /** Draws fog, for now SHOWFLAG_ALWAYS_ACCESSIBLE because it's exposed in SceneCapture */
-SHOWFLAG_ALWAYS_ACCESSIBLE(Fog, SFG_Normal, LOCTEXT("FogSF", "Fog"))
+SHOWFLAG_ALWAYS_ACCESSIBLE(Fog, SFG_Normal, NSLOCTEXT("UnrealEd", "FogSF", "Fog"))
 /** Draws Volumes */
-SHOWFLAG_FIXED_IN_SHIPPING(0, Volumes, SFG_Advanced, LOCTEXT("VolumesSF", "Volumes"))
+SHOWFLAG_FIXED_IN_SHIPPING(0, Volumes, SFG_Advanced, NSLOCTEXT("UnrealEd", "VolumesSF", "Volumes"))
 /** if this is a game viewport, needed? */
-SHOWFLAG_ALWAYS_ACCESSIBLE(Game, SFG_Hidden, LOCTEXT("GameSF", "Game"))
+SHOWFLAG_ALWAYS_ACCESSIBLE(Game, SFG_Hidden, NSLOCTEXT("UnrealEd", "GameSF", "Game"))
 /** Render objects with colors based on what the level they belong to */
-SHOWFLAG_FIXED_IN_SHIPPING(0, LevelColoration, SFG_Advanced, LOCTEXT("LevelColorationSF", "Level Coloration"))
+SHOWFLAG_FIXED_IN_SHIPPING(0, LevelColoration, SFG_Advanced, NSLOCTEXT("UnrealEd", "LevelColorationSF", "Level Coloration"))
 /** Draws BSP brushes (in game or editor textured triangles usually with lightmaps), for now SHOWFLAG_ALWAYS_ACCESSIBLE because it's exposed in SceneCapture */
-SHOWFLAG_ALWAYS_ACCESSIBLE(BSP, SFG_Normal, LOCTEXT("BSPSF", "BSP"))
+SHOWFLAG_ALWAYS_ACCESSIBLE(BSP, SFG_Normal, NSLOCTEXT("UnrealEd", "BSPSF", "BSP"))
 /** Collision drawing */
-SHOWFLAG_FIXED_IN_SHIPPING(0, Collision, SFG_Normal, LOCTEXT("CollisionWireFrame", "Collision"))
+SHOWFLAG_FIXED_IN_SHIPPING(0, Collision, SFG_Normal, NSLOCTEXT("UnrealEd", "CollisionWireFrame", "Collision"))
 /** Collision blocking visibility against complex **/
-SHOWFLAG_FIXED_IN_SHIPPING(0, CollisionVisibility, SFG_Hidden, LOCTEXT("CollisionVisibility", "Visibility"))
+SHOWFLAG_FIXED_IN_SHIPPING(0, CollisionVisibility, SFG_Hidden, NSLOCTEXT("UnrealEd", "CollisionVisibility", "Visibility"))
 /** Collision blocking pawn against simple collision **/
-SHOWFLAG_FIXED_IN_SHIPPING(0, CollisionPawn, SFG_Hidden, LOCTEXT("CollisionPawn", "Pawn"))
+SHOWFLAG_FIXED_IN_SHIPPING(0, CollisionPawn, SFG_Hidden, NSLOCTEXT("UnrealEd", "CollisionPawn", "Pawn"))
 /** Render LightShafts, for now SHOWFLAG_ALWAYS_ACCESSIBLE because it's exposed in SceneCapture */
-SHOWFLAG_ALWAYS_ACCESSIBLE(LightShafts, SFG_LightingFeatures, LOCTEXT("LightShaftsSF", "Light Shafts"))
+SHOWFLAG_ALWAYS_ACCESSIBLE(LightShafts, SFG_LightingFeatures, NSLOCTEXT("UnrealEd", "LightShaftsSF", "Light Shafts"))
 /** Render the PostProcess Material */
-SHOWFLAG_FIXED_IN_SHIPPING(1, PostProcessMaterial, SFG_PostProcess, LOCTEXT("PostProcessMaterialSF", "Post Process Material"))
+SHOWFLAG_FIXED_IN_SHIPPING(1, PostProcessMaterial, SFG_PostProcess, NSLOCTEXT("UnrealEd", "PostProcessMaterialSF", "Post Process Material"))
 /** Render Atmospheric scattering (Atmospheric Fog), for now SHOWFLAG_ALWAYS_ACCESSIBLE because it's exposed in SceneCapture */
-SHOWFLAG_ALWAYS_ACCESSIBLE(AtmosphericFog, SFG_Advanced, LOCTEXT("AtmosphereSF", "Atmospheric Fog"))
+SHOWFLAG_ALWAYS_ACCESSIBLE(AtmosphericFog, SFG_Advanced, NSLOCTEXT("UnrealEd", "AtmosphereSF", "Atmospheric Fog"))
 /** Render safe frames bars*/
-SHOWFLAG_FIXED_IN_SHIPPING(0, CameraAspectRatioBars, SFG_Advanced, LOCTEXT("CameraAspectRatioBarsSF", "Camera Aspect Ratio Bars"))
+SHOWFLAG_FIXED_IN_SHIPPING(0, CameraAspectRatioBars, SFG_Advanced, NSLOCTEXT("UnrealEd", "CameraAspectRatioBarsSF", "Camera Aspect Ratio Bars"))
 /** Render safe frames */
-SHOWFLAG_FIXED_IN_SHIPPING(1, CameraSafeFrames, SFG_Advanced, LOCTEXT("CameraSafeFramesSF", "Camera Safe Frames"))
+SHOWFLAG_FIXED_IN_SHIPPING(1, CameraSafeFrames, SFG_Advanced, NSLOCTEXT("UnrealEd", "CameraSafeFramesSF", "Camera Safe Frames"))
 /** Render TextRenderComponents (3D text), for now SHOWFLAG_ALWAYS_ACCESSIBLE because it's exposed in SceneCapture */
-SHOWFLAG_ALWAYS_ACCESSIBLE(TextRender, SFG_Advanced, LOCTEXT("TextRenderSF", "Render (3D) Text"))
+SHOWFLAG_ALWAYS_ACCESSIBLE(TextRender, SFG_Advanced, NSLOCTEXT("UnrealEd", "TextRenderSF", "Render (3D) Text"))
 /** Any rendering/buffer clearing  (good for benchmarking and for pausing rendering while the app is not in focus to save cycles). */
-SHOWFLAG_ALWAYS_ACCESSIBLE(Rendering, SFG_Hidden, LOCTEXT("RenderingSF", "Any Rendering")) // do not make it FIXED_IN_SHIPPING, used by Oculus plugin.
+SHOWFLAG_ALWAYS_ACCESSIBLE(Rendering, SFG_Hidden, NSLOCTEXT("UnrealEd", "RenderingSF", "Any Rendering")) // do not make it FIXED_IN_SHIPPING, used by Oculus plugin.
 /** Show the current mask being used by the highres screenshot capture */
-SHOWFLAG_FIXED_IN_SHIPPING(0, HighResScreenshotMask, SFG_Hidden, LOCTEXT("HighResScreenshotMaskSF", "High Res Screenshot Mask"))
+SHOWFLAG_FIXED_IN_SHIPPING(0, HighResScreenshotMask, SFG_Hidden, NSLOCTEXT("UnrealEd", "HighResScreenshotMaskSF", "High Res Screenshot Mask"))
 /** Distortion of output for HMD devices, SHOWFLAG_ALWAYS_ACCESSIBLE for now because USceneCaptureComponent needs that */
-SHOWFLAG_ALWAYS_ACCESSIBLE(HMDDistortion, SFG_PostProcess, LOCTEXT("HMDDistortionSF", "HMD Distortion"))
+SHOWFLAG_ALWAYS_ACCESSIBLE(HMDDistortion, SFG_PostProcess, NSLOCTEXT("UnrealEd", "HMDDistortionSF", "HMD Distortion"))
 /** Whether to render in stereoscopic 3d, for now SHOWFLAG_ALWAYS_ACCESSIBLE because it's used by StereoRendering */
-SHOWFLAG_ALWAYS_ACCESSIBLE(StereoRendering, SFG_Hidden, LOCTEXT("StereoRenderingSF", "Stereoscopic Rendering"))
+SHOWFLAG_ALWAYS_ACCESSIBLE(StereoRendering, SFG_Hidden, NSLOCTEXT("UnrealEd", "StereoRenderingSF", "Stereoscopic Rendering"))
 /** Show objects even if they should be distance culled, for now SHOWFLAG_ALWAYS_ACCESSIBLE because it's exposed in SceneCapture */
-SHOWFLAG_ALWAYS_ACCESSIBLE(DistanceCulledPrimitives, SFG_Hidden, LOCTEXT("DistanceCulledPrimitivesSF", "Distance Culled Primitives"))
+SHOWFLAG_ALWAYS_ACCESSIBLE(DistanceCulledPrimitives, SFG_Hidden, NSLOCTEXT("UnrealEd", "DistanceCulledPrimitivesSF", "Distance Culled Primitives"))
 /** To visualize the culling in Tile Based Deferred Lighting, later for non tiled as well */
-SHOWFLAG_FIXED_IN_SHIPPING(0, VisualizeLightCulling, SFG_Hidden, LOCTEXT("VisualizeLightCullingSF", "Light Culling"))
+SHOWFLAG_FIXED_IN_SHIPPING(0, VisualizeLightCulling, SFG_Hidden, NSLOCTEXT("UnrealEd", "VisualizeLightCullingSF", "Light Culling"))
 /** To disable precomputed visibility */
-SHOWFLAG_FIXED_IN_SHIPPING(1, PrecomputedVisibility, SFG_Advanced, LOCTEXT("PrecomputedVisibilitySF", "Precomputed Visibility"))
+SHOWFLAG_FIXED_IN_SHIPPING(1, PrecomputedVisibility, SFG_Advanced, NSLOCTEXT("UnrealEd", "PrecomputedVisibilitySF", "Precomputed Visibility"))
 /** Contribution from sky light, for now SHOWFLAG_ALWAYS_ACCESSIBLE because it's exposed in SceneCapture */
-SHOWFLAG_ALWAYS_ACCESSIBLE(SkyLighting, SFG_LightTypes, LOCTEXT("SkyLightingSF", "Sky Lighting"))
+SHOWFLAG_ALWAYS_ACCESSIBLE(SkyLighting, SFG_LightTypes, NSLOCTEXT("UnrealEd", "SkyLightingSF", "Sky Lighting"))
 /** Visualize Light Propagation Volume, for developer (by default off): */
-SHOWFLAG_FIXED_IN_SHIPPING(0, VisualizeLPV, SFG_Visualize, LOCTEXT("VisualizeLPVSF", "Light Propagation Volume"))
+SHOWFLAG_FIXED_IN_SHIPPING(0, VisualizeLPV, SFG_Visualize, NSLOCTEXT("UnrealEd", "VisualizeLPVSF", "Light Propagation Volume"))
 /** Visualize preview shadow indicator */
-SHOWFLAG_FIXED_IN_SHIPPING(0, PreviewShadowsIndicator, SFG_Visualize, LOCTEXT("PreviewShadowIndicatorSF", "Preview Shadows Indicator"))
+SHOWFLAG_FIXED_IN_SHIPPING(0, PreviewShadowsIndicator, SFG_Visualize, NSLOCTEXT("UnrealEd", "PreviewShadowIndicatorSF", "Preview Shadows Indicator"))
 /** Visualize precomputed visibility cells */
-SHOWFLAG_FIXED_IN_SHIPPING(0, PrecomputedVisibilityCells, SFG_Visualize, LOCTEXT("PrecomputedVisibilityCellsSF", "Precomputed Visibility Cells"))
+SHOWFLAG_FIXED_IN_SHIPPING(0, PrecomputedVisibilityCells, SFG_Visualize, NSLOCTEXT("UnrealEd", "PrecomputedVisibilityCellsSF", "Precomputed Visibility Cells"))
 /** Visualize volumetric lightmap used for GI on dynamic objects */
-SHOWFLAG_FIXED_IN_SHIPPING(0, VisualizeVolumetricLightmap, SFG_Visualize, LOCTEXT("VisualizeVolumetricLightmapSF", "Volumetric Lightmap"))
+SHOWFLAG_FIXED_IN_SHIPPING(0, VisualizeVolumetricLightmap, SFG_Visualize, NSLOCTEXT("UnrealEd", "VisualizeVolumetricLightmapSF", "Volumetric Lightmap"))
 /** Visualize volume lighting samples used for GI on dynamic objects */
-SHOWFLAG_FIXED_IN_SHIPPING(0, VolumeLightingSamples, SFG_Visualize, LOCTEXT("VolumeLightingSamplesSF", "Volume Lighting Samples"))
+SHOWFLAG_FIXED_IN_SHIPPING(0, VolumeLightingSamples, SFG_Visualize, NSLOCTEXT("UnrealEd", "VolumeLightingSamplesSF", "Volume Lighting Samples"))
 /** Render Paper2D sprites, for now SHOWFLAG_ALWAYS_ACCESSIBLE because it's exposed in SceneCapture */
-SHOWFLAG_ALWAYS_ACCESSIBLE(Paper2DSprites, SFG_Advanced, LOCTEXT("Paper2DSpritesSF", "Paper 2D Sprites"))
+SHOWFLAG_ALWAYS_ACCESSIBLE(Paper2DSprites, SFG_Advanced, NSLOCTEXT("UnrealEd", "Paper2DSpritesSF", "Paper 2D Sprites"))
 /** Visualization of distance field AO */
-SHOWFLAG_FIXED_IN_SHIPPING(0, VisualizeDistanceFieldAO, SFG_Visualize, LOCTEXT("VisualizeDistanceFieldAOSF", "Distance Field Ambient Occlusion"))
+SHOWFLAG_FIXED_IN_SHIPPING(0, VisualizeDistanceFieldAO, SFG_Visualize, NSLOCTEXT("UnrealEd", "VisualizeDistanceFieldAOSF", "Distance Field Ambient Occlusion"))
 /** Visualization of distance field GI */
-SHOWFLAG_FIXED_IN_SHIPPING(0, VisualizeDistanceFieldGI, SFG_Hidden, LOCTEXT("VisualizeDistanceFieldGISF", "Distance Field Global Illumination"))
+SHOWFLAG_FIXED_IN_SHIPPING(0, VisualizeDistanceFieldGI, SFG_Hidden, NSLOCTEXT("UnrealEd", "VisualizeDistanceFieldGISF", "Distance Field Global Illumination"))
 /** Mesh Distance fields */
-SHOWFLAG_FIXED_IN_SHIPPING(0, VisualizeMeshDistanceFields, SFG_Visualize, LOCTEXT("MeshDistanceFieldsSF", "Mesh DistanceFields"))
+SHOWFLAG_FIXED_IN_SHIPPING(0, VisualizeMeshDistanceFields, SFG_Visualize, NSLOCTEXT("UnrealEd", "MeshDistanceFieldsSF", "Mesh DistanceFields"))
 /** Global Distance field */
-SHOWFLAG_FIXED_IN_SHIPPING(0, VisualizeGlobalDistanceField, SFG_Visualize, LOCTEXT("GlobalDistanceFieldSF", "Global DistanceField"))
+SHOWFLAG_FIXED_IN_SHIPPING(0, VisualizeGlobalDistanceField, SFG_Visualize, NSLOCTEXT("UnrealEd", "GlobalDistanceFieldSF", "Global DistanceField"))
 /** Screen space AO, for now SHOWFLAG_ALWAYS_ACCESSIBLE because r.GBuffer need that */
-SHOWFLAG_ALWAYS_ACCESSIBLE(ScreenSpaceAO, SFG_LightingFeatures, LOCTEXT("ScreenSpaceAOSF", "Screen Space Ambient Occlusion"))
+SHOWFLAG_ALWAYS_ACCESSIBLE(ScreenSpaceAO, SFG_LightingFeatures, NSLOCTEXT("UnrealEd", "ScreenSpaceAOSF", "Screen Space Ambient Occlusion"))
 /** Distance field AO, for now SHOWFLAG_ALWAYS_ACCESSIBLE because it's exposed in SceneCapture */
-SHOWFLAG_ALWAYS_ACCESSIBLE(DistanceFieldAO, SFG_LightingFeatures, LOCTEXT("DistanceFieldAOSF", "Distance Field Ambient Occlusion"))
+SHOWFLAG_ALWAYS_ACCESSIBLE(DistanceFieldAO, SFG_LightingFeatures, NSLOCTEXT("UnrealEd", "DistanceFieldAOSF", "Distance Field Ambient Occlusion"))
 /** Distance field GI */
-SHOWFLAG_FIXED_IN_SHIPPING(1, DistanceFieldGI, SFG_Hidden, LOCTEXT("DistanceFieldGISF", "Distance Field Global Illumination"))
+SHOWFLAG_FIXED_IN_SHIPPING(1, DistanceFieldGI, SFG_Hidden, NSLOCTEXT("UnrealEd", "DistanceFieldGISF", "Distance Field Global Illumination"))
 /** SHOWFLAG_ALWAYS_ACCESSIBLE because it's exposed in SceneCapture */
-SHOWFLAG_ALWAYS_ACCESSIBLE(VolumetricFog, SFG_LightingFeatures, LOCTEXT("VolumetricFogSF", "Volumetric Fog"))
+SHOWFLAG_ALWAYS_ACCESSIBLE(VolumetricFog, SFG_LightingFeatures, NSLOCTEXT("UnrealEd", "VolumetricFogSF", "Volumetric Fog"))
 /** Visualize screen space reflections, for developer (by default off): */
-SHOWFLAG_FIXED_IN_SHIPPING(0, VisualizeSSR, SFG_Visualize, LOCTEXT("VisualizeSSR", "Screen Space Reflections"))
+SHOWFLAG_FIXED_IN_SHIPPING(0, VisualizeSSR, SFG_Visualize, NSLOCTEXT("UnrealEd", "VisualizeSSR", "Screen Space Reflections"))
 /** Visualize the Shading Models, mostly or debugging and profiling */
-SHOWFLAG_FIXED_IN_SHIPPING(0, VisualizeShadingModels, SFG_Visualize, LOCTEXT("VisualizeShadingModels", "Shading Models"))
+SHOWFLAG_FIXED_IN_SHIPPING(0, VisualizeShadingModels, SFG_Visualize, NSLOCTEXT("UnrealEd", "VisualizeShadingModels", "Shading Models"))
 /** Visualize the senses configuration of AIs' PawnSensingComponent */
-SHOWFLAG_FIXED_IN_SHIPPING(0, VisualizeSenses, SFG_Advanced, LOCTEXT("VisualizeSenses", "Senses"))
+SHOWFLAG_FIXED_IN_SHIPPING(0, VisualizeSenses, SFG_Advanced, NSLOCTEXT("UnrealEd", "VisualizeSenses", "Senses"))
 /** Visualize the bloom, for developer (by default off): */
-SHOWFLAG_FIXED_IN_SHIPPING(0, VisualizeBloom, SFG_Visualize, LOCTEXT("VisualizeBloom", "Bloom"))
+SHOWFLAG_FIXED_IN_SHIPPING(0, VisualizeBloom, SFG_Visualize, NSLOCTEXT("UnrealEd", "VisualizeBloom", "Bloom"))
 /** Visualize LOD Coloration */
-SHOWFLAG_FIXED_IN_SHIPPING(0, LODColoration, SFG_Hidden, LOCTEXT("VisualizeLODColoration", "Visualize LOD Coloration"))
+SHOWFLAG_FIXED_IN_SHIPPING(0, LODColoration, SFG_Hidden, NSLOCTEXT("UnrealEd", "VisualizeLODColoration", "Visualize LOD Coloration"))
 /** Visualize HLOD Coloration */
-SHOWFLAG_FIXED_IN_SHIPPING(0, HLODColoration, SFG_Hidden, LOCTEXT("VisualizeHLODColoration", "Visualize HLOD Coloration"))
+SHOWFLAG_FIXED_IN_SHIPPING(0, HLODColoration, SFG_Hidden, NSLOCTEXT("UnrealEd", "VisualizeHLODColoration", "Visualize HLOD Coloration"))
 /** Visualize screen quads */
-SHOWFLAG_FIXED_IN_SHIPPING(0, QuadOverdraw, SFG_Hidden, LOCTEXT("QuadOverdrawSF", "Quad Overdraw"))
+SHOWFLAG_FIXED_IN_SHIPPING(0, QuadOverdraw, SFG_Hidden, NSLOCTEXT("UnrealEd", "QuadOverdrawSF", "Quad Overdraw"))
 /** Visualize the overhead of material quads */
-SHOWFLAG_FIXED_IN_SHIPPING(0, ShaderComplexityWithQuadOverdraw, SFG_Hidden, LOCTEXT("ShaderComplexityWithQuadOverdraw", "Shader Complexity With Quad Overdraw"))
+SHOWFLAG_FIXED_IN_SHIPPING(0, ShaderComplexityWithQuadOverdraw, SFG_Hidden, NSLOCTEXT("UnrealEd", "ShaderComplexityWithQuadOverdraw", "Shader Complexity With Quad Overdraw"))
 /** Visualize the accuracy of the primitive distance computed for texture streaming */
-SHOWFLAG_FIXED_IN_SHIPPING(0, PrimitiveDistanceAccuracy, SFG_Hidden, LOCTEXT("PrimitiveDistanceAccuracy", "Primitive Distance Accuracy"))
+SHOWFLAG_FIXED_IN_SHIPPING(0, PrimitiveDistanceAccuracy, SFG_Hidden, NSLOCTEXT("UnrealEd", "PrimitiveDistanceAccuracy", "Primitive Distance Accuracy"))
 /** Visualize the accuracy of the mesh UV density computed for texture streaming */
-SHOWFLAG_FIXED_IN_SHIPPING(0, MeshUVDensityAccuracy, SFG_Hidden, LOCTEXT("MeshUVDensityAccuracy", "Mesh UV Densities Accuracy"))
+SHOWFLAG_FIXED_IN_SHIPPING(0, MeshUVDensityAccuracy, SFG_Hidden, NSLOCTEXT("UnrealEd", "MeshUVDensityAccuracy", "Mesh UV Densities Accuracy"))
 /** Visualize the accuracy of CPU material texture scales when compared to the GPU values */
-SHOWFLAG_FIXED_IN_SHIPPING(0, MaterialTextureScaleAccuracy, SFG_Hidden, LOCTEXT("MaterialTextureScaleAccuracy", "Material Texture Scales Accuracy"))
+SHOWFLAG_FIXED_IN_SHIPPING(0, MaterialTextureScaleAccuracy, SFG_Hidden, NSLOCTEXT("UnrealEd", "MaterialTextureScaleAccuracy", "Material Texture Scales Accuracy"))
 /** Outputs the material texture scales. */
-SHOWFLAG_FIXED_IN_SHIPPING(0, OutputMaterialTextureScales, SFG_Hidden, LOCTEXT("OutputMaterialTextureScales", "Output Material Texture Scales"))
+SHOWFLAG_FIXED_IN_SHIPPING(0, OutputMaterialTextureScales, SFG_Hidden, NSLOCTEXT("UnrealEd", "OutputMaterialTextureScales", "Output Material Texture Scales"))
 /** Compare the required texture resolution to the actual resolution. */
-SHOWFLAG_FIXED_IN_SHIPPING(0, RequiredTextureResolution, SFG_Hidden, LOCTEXT("RequiredTextureResolution", "Required Texture Resolution"))
+SHOWFLAG_FIXED_IN_SHIPPING(0, RequiredTextureResolution, SFG_Hidden, NSLOCTEXT("UnrealEd", "RequiredTextureResolution", "Required Texture Resolution"))
 /** If WidgetComponents should be rendered in the scene */
-SHOWFLAG_ALWAYS_ACCESSIBLE(WidgetComponents, SFG_Normal, LOCTEXT("WidgetComponentsSF", "Widget Components"))
+SHOWFLAG_ALWAYS_ACCESSIBLE(WidgetComponents, SFG_Normal, NSLOCTEXT("UnrealEd", "WidgetComponentsSF", "Widget Components"))
 /** Draw the bones of all skeletal meshes */
-SHOWFLAG_FIXED_IN_SHIPPING(0, Bones, SFG_Developer, LOCTEXT("BoneSF", "Bones"))
+SHOWFLAG_FIXED_IN_SHIPPING(0, Bones, SFG_Developer, NSLOCTEXT("UnrealEd", "BoneSF", "Bones"))
 /** If media planes should be shown */
-SHOWFLAG_ALWAYS_ACCESSIBLE(MediaPlanes, SFG_Normal, LOCTEXT("MediaPlanesSF", "Media Planes"))
+SHOWFLAG_ALWAYS_ACCESSIBLE(MediaPlanes, SFG_Normal, NSLOCTEXT("UnrealEd", "MediaPlanesSF", "Media Planes"))
 #undef SHOWFLAG_ALWAYS_ACCESSIBLE
 #undef SHOWFLAG_FIXED_IN_SHIPPING
