Copying //UE4/Dev-Framework to //UE4/Dev-Main (Source: //UE4/Dev-Framework @ 3967517)

diff --git a/Engine/Source/Runtime/Engine/Classes/AI/Navigation/NavigationTypes.h b/Engine/Source/Runtime/Engine/Classes/AI/Navigation/NavigationTypes.h
--- a/Engine/Source/Runtime/Engine/Classes/AI/Navigation/NavigationTypes.h
+++ b/Engine/Source/Runtime/Engine/Classes/AI/Navigation/NavigationTypes.h
@@ -12,7 +12,7 @@
 #include "UObject/WeakObjectPtr.h"
 #include "Misc/CoreStats.h"
 #include "UObject/SoftObjectPath.h"
-#include "AI/Navigation/NavFilters/NavigationQueryFilter.h"
+#include "GameFramework/Actor.h"
 #include "NavigationTypes.generated.h"
 
 #define INVALID_NAVNODEREF (0)
@@ -33,11 +33,6 @@ struct FNavigationPath;
 /** uniform identifier type for navigation data elements may it be a polygon or graph node */
 typedef uint64 NavNodeRef;
 
-class AActor;
-class ANavigationData;
-class INavAgentInterface;
-class INavRelevantInterface;
-
 namespace FNavigationSystem
 {
 	/** used as a fallback value for navigation agent radius, when none specified via UNavigationSystem::SupportedAgents */
@@ -54,13 +49,6 @@ namespace FNavigationSystem
 	{
 		return TestLocation != InvalidLocation;
 	}
-
-	enum ECreateIfEmpty
-	{
-		Invalid = -1,
-		DontCreate = 0,
-		Create = 1,
-	};
 }
 
 UENUM()
@@ -479,13 +467,15 @@ struct ENGINE_API FNavAgentProperties : public FMovementProperties
 	float NavWalkingSearchHeightScale;
 
 	/** Type of navigation data used by agent, null means "any" */
-	UPROPERTY(EditAnywhere, BlueprintReadWrite, Category = MovementProperties)
-	TSubclassOf<ANavigationData> PreferredNavData;
-
+	/*UPROPERTY(EditAnywhere, BlueprintReadWrite, Category = MovementProperties)
+	TSubclassOf<ANavigationData> PreferredNavData;*/
+	UPROPERTY(EditAnywhere, BlueprintReadWrite, Category=MovementProperties, meta=(MetaClass = "NavigationData"))
+	FSoftClassPath PreferredNavData;
+	
 	FNavAgentProperties(float Radius = -1.f, float Height = -1.f)
 		: AgentRadius(Radius), AgentHeight(Height), AgentStepHeight(-1), NavWalkingSearchHeightScale(0.5f)
-	{
-	}
+	{}
+	FNavAgentProperties(const FNavAgentProperties& Other);
 
 	void UpdateWithCollisionComponent(class UShapeComponent* CollisionComponent);
 
@@ -498,7 +488,9 @@ struct ENGINE_API FNavAgentProperties : public FMovementProperties
 	{
 		return FGenericPlatformMath::Abs(AgentRadius - Other.AgentRadius) < Precision
 			&& FGenericPlatformMath::Abs(AgentHeight - Other.AgentHeight) < Precision
-			&& (HasStepHeightOverride() == false || FGenericPlatformMath::Abs(AgentStepHeight - Other.AgentStepHeight) < Precision)
+			&& ((HasStepHeightOverride() == false) 
+				|| (Other.HasStepHeightOverride() == false) 
+				|| FGenericPlatformMath::Abs(AgentStepHeight - Other.AgentStepHeight) < Precision)
 			&& IsNavDataMatching(Other);
 	}
 	
@@ -514,6 +506,11 @@ struct ENGINE_API FNavAgentProperties : public FMovementProperties
 			: INVALID_NAVEXTENT;
 	}
 
+	void SetPreferredNavData(TSubclassOf<AActor> NavDataClass)
+	{
+		PreferredNavData = FSoftClassPath(NavDataClass.Get());
+	}
+
 	static const FNavAgentProperties DefaultProperties;
 };
 
@@ -536,13 +533,15 @@ struct ENGINE_API FNavDataConfig : public FNavAgentProperties
 	UPROPERTY(EditAnywhere, BlueprintReadOnly, Category = Querying)
 	FVector DefaultQueryExtent;
 
-	UPROPERTY(Transient, EditAnywhere, BlueprintReadOnly, Category = Navigation)
-	TSubclassOf<ANavigationData> NavigationDataClass;
+	// mz@todo make sure we handle NavigationDataClass and NavigationDataClassName in PostEditChange
+	UPROPERTY(Transient)
+	TSubclassOf<AActor> NavigationDataClass;
 
-	UPROPERTY(config)
+	UPROPERTY(VisibleAnywhere, BlueprintReadOnly, Category=Navigation, meta=(MetaClass = "NavigationData"))
 	FSoftClassPath NavigationDataClassName;
 
 	FNavDataConfig(float Radius = FNavigationSystem::FallbackAgentRadius, float Height = FNavigationSystem::FallbackAgentHeight);
+	FNavDataConfig(const FNavDataConfig& Other);
 };
 
 struct FNavigationProjectionWork
@@ -606,51 +605,6 @@ namespace ENavigationQueryResult
 	};
 }
 
-struct ENGINE_API FPathFindingQueryData
-{
-	TWeakObjectPtr<const UObject> Owner;
-	FVector StartLocation;
-	FVector EndLocation;
-	FSharedConstNavQueryFilter QueryFilter;
-
-	/** additional flags passed to navigation data handling request */
-	int32 NavDataFlags;
-
-	/** if set, allow partial paths as a result */
-	uint32 bAllowPartialPaths : 1;
-
-	FPathFindingQueryData() : StartLocation(FNavigationSystem::InvalidLocation), EndLocation(FNavigationSystem::InvalidLocation), NavDataFlags(0), bAllowPartialPaths(true) {}
-	FPathFindingQueryData(const UObject* InOwner, const FVector& InStartLocation, const FVector& InEndLocation, FSharedConstNavQueryFilter InQueryFilter = nullptr, int32 InNavDataFlags = 0, bool bInAllowPartialPaths = true) :
-		Owner(InOwner), StartLocation(InStartLocation), EndLocation(InEndLocation), QueryFilter(InQueryFilter), NavDataFlags(InNavDataFlags), bAllowPartialPaths(bInAllowPartialPaths) {}
-};
-
-struct ENGINE_API FPathFindingQuery : public FPathFindingQueryData
-{
-	TWeakObjectPtr<const ANavigationData> NavData;
-	FNavPathSharedPtr PathInstanceToFill;
-	FNavAgentProperties NavAgentProperties;
-
-	FPathFindingQuery() : FPathFindingQueryData() {}
-	FPathFindingQuery(const FPathFindingQuery& Source);
-	FPathFindingQuery(const UObject* InOwner, const ANavigationData& InNavData, const FVector& Start, const FVector& End, FSharedConstNavQueryFilter SourceQueryFilter = NULL, FNavPathSharedPtr InPathInstanceToFill = NULL);
-	FPathFindingQuery(const INavAgentInterface& InNavAgent, const ANavigationData& InNavData, const FVector& Start, const FVector& End, FSharedConstNavQueryFilter SourceQueryFilter = NULL, FNavPathSharedPtr InPathInstanceToFill = NULL);
-
-	explicit FPathFindingQuery(FNavPathSharedRef PathToRecalculate, const ANavigationData* NavDataOverride = NULL);
-
-	FPathFindingQuery& SetPathInstanceToUpdate(FNavPathSharedPtr InPathInstanceToFill) { PathInstanceToFill = InPathInstanceToFill; return *this; }
-	FPathFindingQuery& SetAllowPartialPaths(bool bAllow) { bAllowPartialPaths = bAllow; return *this; }
-	FPathFindingQuery& SetNavAgentProperties(const FNavAgentProperties& InNavAgentProperties) { NavAgentProperties = InNavAgentProperties; return *this; }
-};
-
-namespace EPathFindingMode
-{
-	enum Type
-	{
-		Regular,
-		Hierarchical,
-	};
-};
-
 /**
 *	Delegate used to communicate that path finding query has been finished.
 *	@param uint32 unique Query ID of given query
@@ -661,25 +615,6 @@ namespace EPathFindingMode
 */
 DECLARE_DELEGATE_ThreeParams(FNavPathQueryDelegate, uint32, ENavigationQueryResult::Type, FNavPathSharedPtr);
 
-//////////////////////////////////////////////////////////////////////////
-// Custom path following data
-
-/** Custom data passed to movement requests. */
-struct ENGINE_API FMoveRequestCustomData
-{
-};
-
-typedef TSharedPtr<FMoveRequestCustomData, ESPMode::ThreadSafe> FCustomMoveSharedPtr;
-typedef TWeakPtr<FMoveRequestCustomData, ESPMode::ThreadSafe> FCustomMoveWeakPtr;
-
-UCLASS(Abstract, CustomConstructor)
-class UNavigationTypes : public UObject
-{
-	GENERATED_UCLASS_BODY()
-
-	UNavigationTypes(const FObjectInitializer& ObjectInitializer = FObjectInitializer::Get()) : Super(ObjectInitializer) { }
-};
-
 //////////////////////////////////////////////////////////////////////////
 // Memory stating
 
@@ -758,33 +693,6 @@ struct TContainerTraits<TNavStatArray<InElementType> > : public TContainerTraits
 	enum { MoveWillEmptyContainer = TContainerTraits<typename TNavStatArray<InElementType>::Super>::MoveWillEmptyContainer };
 };
 
-//----------------------------------------------------------------------//
-// Active tiles 
-//----------------------------------------------------------------------//
-struct FNavigationInvokerRaw
-{
-	FVector Location;
-	float RadiusMin;
-	float RadiusMax;
-
-	FNavigationInvokerRaw(const FVector& InLocation, float Min, float Max)
-		: Location(InLocation), RadiusMin(Min), RadiusMax(Max)
-	{}
-};
-
-struct FNavigationInvoker
-{
-	TWeakObjectPtr<AActor> Actor;
-
-	/** tiles GenerationRadius away or close will be generated if they're not already present */
-	float GenerationRadius;
-	/** tiles over RemovalRadius will get removed.
-	*	@Note needs to be >= GenerationRadius or will get clampped */
-	float RemovalRadius;
-
-	FNavigationInvoker();
-	FNavigationInvoker(AActor& InActor, float InGenerationRadius, float InRemovalRadius);
-};
 
 //----------------------------------------------------------------------//
 // generic "landscape" support
@@ -798,3 +706,4 @@ struct ENGINE_API FNavHeightfieldSamples
 
 	FORCEINLINE bool IsEmpty() const { return Heights.Num() == 0; }
 };
+
diff --git a/Engine/Source/Runtime/Engine/Classes/AI/NavigationSystemBase.h b/Engine/Source/Runtime/Engine/Classes/AI/NavigationSystemBase.h
new file mode 100644
--- /dev/null
+++ b/Engine/Source/Runtime/Engine/Classes/AI/NavigationSystemBase.h
@@ -0,0 +1,353 @@
+// Copyright 1998-2018 Epic Games, Inc. All Rights Reserved.
+
+#pragma once
+
+#include "UObject/Object.h"
+#include "Engine/Engine.h"
+#include "AI/Navigation/NavigationTypes.h"
+#include "AI/NavigationModifier.h"
+#include "Engine/World.h"
+#include "NavigationSystemBase.generated.h"
+
+class UNavigationSystemBase;
+class UNavigationSystemConfig;
+class AActor;
+class UActorComponent;
+class USceneComponent;
+class INavigationDataInterface;
+class IPathFollowingAgentInterface;
+class AWorldSettings;
+class ULevel;
+class AController;
+
+ENGINE_API DECLARE_LOG_CATEGORY_EXTERN(LogNavigation, Warning, All);
+
+UENUM()
+enum class FNavigationSystemRunMode : uint8
+{
+	InvalidMode,
+	GameMode,
+	EditorMode,
+	SimulationMode,
+	PIEMode,
+};
+
+namespace ENavigationLockReason
+{
+	enum Type
+	{
+		Unknown = 1 << 0,
+		AllowUnregister = 1 << 1,
+
+		MaterialUpdate = 1 << 2,
+		LightingUpdate = 1 << 3,
+		ContinuousEditorMove = 1 << 4,
+		SpawnOnDragEnter = 1 << 5,
+	};
+}
+
+class ENGINE_API FNavigationLockContext
+{
+public:
+	FNavigationLockContext(ENavigationLockReason::Type Reason = ENavigationLockReason::Unknown, bool bApplyLock = true)
+		: MyWorld(NULL), LockReason(Reason), bSingleWorld(false), bIsLocked(false)
+	{
+		if (bApplyLock)
+		{
+			LockUpdates();
+		}
+	}
+
+	FNavigationLockContext(UWorld* InWorld, ENavigationLockReason::Type Reason = ENavigationLockReason::Unknown, bool bApplyLock = true)
+		: MyWorld(InWorld), LockReason(Reason), bSingleWorld(true), bIsLocked(false)
+	{
+		if (bApplyLock)
+		{
+			LockUpdates();
+		}
+	}
+
+	~FNavigationLockContext()
+	{
+		UnlockUpdates();
+	}
+
+private:
+	UWorld* MyWorld;
+	uint8 LockReason;
+	uint8 bSingleWorld : 1;
+	uint8 bIsLocked : 1;
+
+	void LockUpdates();
+	void UnlockUpdates();
+};
+
+namespace FNavigationSystem
+{
+	/** Creates an instance of NavigationSystem (class being specified by WorldSetting's NavigationSystemConfig)
+	 *	A new instance will be created only if given WorldOwner doesn't have one yet.
+	 *	The new instance will be assigned to the given WorldOwner (via SetNavigationSystem 
+	 *	call) and depending on value of bInitializeForWorld the InitializeForWorld 
+	 *	function will be called on the new NavigationSystem instance.
+	 *	(@see UWorld.NavigationSystem)
+	 *	@param RunMode if set to a valid value (other than FNavigationSystemRunMode::InvalidMode) 
+	 *		will also configure the created NavigationSystem instance for that mode */
+	ENGINE_API void AddNavigationSystemToWorld(UWorld& WorldOwner, const FNavigationSystemRunMode RunMode = FNavigationSystemRunMode::InvalidMode, const bool bInitializeForWorld = true);
+
+	/** Discards all navigation data chunks in all sub-levels */
+	ENGINE_API void DiscardNavigationDataChunks(UWorld& InWorld);
+
+	template<typename TNavSys>
+	FORCEINLINE TNavSys* GetCurrent(UWorld* World)
+	{
+		return World ? Cast<TNavSys>(World->GetNavigationSystem()) : (TNavSys*)nullptr;
+	}
+
+	template<typename TNavSys>
+	FORCEINLINE const TNavSys* GetCurrent(const UWorld* World)
+	{
+		return World ? Cast<TNavSys>(World->GetNavigationSystem()) : (const TNavSys*)nullptr;
+	}
+
+	template<typename TNavSys>
+	TNavSys* GetCurrent(UObject* WorldContextObject)
+	{
+		UWorld* World = nullptr;
+
+		if (WorldContextObject != nullptr)
+		{
+			World = GEngine->GetWorldFromContextObject(WorldContextObject, EGetWorldErrorMode::LogAndReturnNull);
+		}
+
+		return GetCurrent<TNavSys>(World);
+	}
+
+	ENGINE_API void UpdateActorData(AActor& Actor);
+	ENGINE_API void UpdateComponentData(UActorComponent& Comp);
+	ENGINE_API void UpdateActorAndComponentData(AActor& Actor, bool bUpdateAttachedActors = true);
+	ENGINE_API void UpdateComponentDataAfterMove(USceneComponent& Comp);
+	//ENGINE_API bool HasComponentData(UActorComponent& Comp);
+	ENGINE_API void OnActorBoundsChanged(AActor& Actor);
+	ENGINE_API void OnPostEditActorMove(AActor& Actor);
+	ENGINE_API void OnComponentBoundsChanged(UActorComponent& Comp, const FBox& NewBounds, const FBox& DirtyArea);
+	ENGINE_API void OnComponentTransformChanged(USceneComponent& Comp);
+
+	ENGINE_API void OnActorRegistered(AActor& Actor);
+	ENGINE_API void OnActorUnregistered(AActor& Actor);
+
+	ENGINE_API void OnComponentRegistered(UActorComponent& Comp);
+	ENGINE_API void OnComponentUnregistered(UActorComponent& Comp);
+
+	ENGINE_API void RemoveActorData(AActor& Actor);
+
+	ENGINE_API bool HasComponentData(UActorComponent& Comp);
+	
+	ENGINE_API const FNavDataConfig& GetDefaultSupportedAgent();
+
+	ENGINE_API TSubclassOf<UNavAreaBase> GetDefaultWalkableArea();
+
+	ENGINE_API const FTransform& GetCoordTransformTo(const ENavigationCoordSystem::Type CoordType);
+	ENGINE_API const FTransform& GetCoordTransformFrom(const ENavigationCoordSystem::Type CoordType);
+
+	ENGINE_API bool WantsComponentChangeNotifies();
+
+	//ENGINE_API INavigationDataInterface* GetNavDataForProps(const FNavAgentProperties& AgentProperties);
+	ENGINE_API INavigationDataInterface* GetNavDataForActor(const AActor& Actor);
+	ENGINE_API TSubclassOf<AActor> GetDefaultNavDataClass();
+
+	ENGINE_API void VerifyNavigationRenderingComponents(UWorld& World, const bool bShow);
+	ENGINE_API void Build(UWorld& World);
+
+#if WITH_EDITOR
+	ENGINE_API void OnPIEStart(UWorld& World);
+	ENGINE_API void OnPIEEnd(UWorld& World);
+	ENGINE_API void SetNavigationAutoUpdateEnabled(const bool bNewEnable, UNavigationSystemBase* InNavigationSystem);
+	ENGINE_API void UpdateLevelCollision(ULevel& Level);
+#endif // WITH_EDITOR
+
+	enum class ECleanupMode : uint8
+	{
+		CleanupWithWorld,
+		CleanupUnsafe,
+	};
+
+	// pathfollowing
+	ENGINE_API bool IsFollowingAPath(const AController& Controller);
+	ENGINE_API void StopMovement(const AController& Controller);
+	ENGINE_API IPathFollowingAgentInterface* FindPathFollowingAgentForActor(const AActor& Actor);
+
+	DECLARE_DELEGATE_OneParam(FActorBasedSignature, AActor& /*Actor*/);
+	DECLARE_DELEGATE_OneParam(FActorComponentBasedSignature, UActorComponent& /*Comp*/);
+	DECLARE_DELEGATE_OneParam(FSceneComponentBasedSignature, USceneComponent& /*Comp*/);
+	DECLARE_DELEGATE_OneParam(FWorldBasedSignature, UWorld& /*World*/);
+	DECLARE_DELEGATE_OneParam(FLevelBasedSignature, ULevel& /*Level*/);
+	DECLARE_DELEGATE_OneParam(FControllerBasedSignature, const AController& /*Controller*/);
+	DECLARE_DELEGATE_TwoParams(FNavigationAutoUpdateEnableSignature, const bool /*bNewEnable*/, UNavigationSystemBase* /*InNavigationSystem*/);
+	DECLARE_DELEGATE_RetVal_OneParam(bool, FBoolControllerBasedSignature, const AController& /*World*/);
+	DECLARE_DELEGATE_RetVal_OneParam(bool, FBoolActorComponentBasedSignature, UActorComponent& /*Comp*/);
+	DECLARE_DELEGATE_RetVal(TSubclassOf<UNavAreaBase>, FNavAreaBasedSignature);
+	DECLARE_DELEGATE_RetVal(const FNavDataConfig&, FNavDatConfigBasedSignature);
+	DECLARE_DELEGATE_TwoParams(FWorldByteBasedSignature, UWorld& /*World*/, uint8 /*Flags*/);
+	DECLARE_DELEGATE_TwoParams(FActorBooleBasedSignature, AActor& /*Actor*/, bool /*bUpdateAttachedActors*/);
+	DECLARE_DELEGATE_ThreeParams(FComponentBoundsChangeSignature, UActorComponent& /*Comp*/, const FBox& /*NewBounds*/, const FBox& /*DirtyArea*/)
+	DECLARE_DELEGATE_RetVal_OneParam(INavigationDataInterface*, FNavDataForPropsSignature, const FNavAgentProperties& /*AgentProperties*/);
+	DECLARE_DELEGATE_RetVal_OneParam(INavigationDataInterface*, FNavDataForActorSignature, const AActor& /*Actor*/);
+	DECLARE_DELEGATE_RetVal(TSubclassOf<AActor>, FNavDataClassFetchSignature);
+	DECLARE_DELEGATE_TwoParams(FWorldBoolBasedSignature, UWorld& /*World*/, const bool /*bShow*/);
+}
+
+
+UCLASS(Abstract, config = Engine, defaultconfig)
+class ENGINE_API UNavigationSystemBase : public UObject
+{
+	GENERATED_BODY()
+
+public:
+	virtual ~UNavigationSystemBase(){}
+
+	virtual void Tick(float DeltaSeconds) PURE_VIRTUAL(UNavigationSystemBase::Tick, );
+	virtual void CleanUp(const FNavigationSystem::ECleanupMode Mode) PURE_VIRTUAL(UNavigationSystemBase::CleanUp, );
+	virtual void Configure(const UNavigationSystemConfig& Config) PURE_VIRTUAL(UNavigationSystemBase::Configure, );
+
+	// possibly editor-only
+	
+
+	/**
+	*	Called when owner-UWorld initializes actors
+	*/
+	virtual void OnInitializeActors() {}
+
+	virtual bool IsNavigationBuilt(const AWorldSettings* Settings) const { return false; }
+
+	virtual void ApplyWorldOffset(const FVector& InOffset, bool bWorldShift) PURE_VIRTUAL(UNavigationSystemBase::ApplyWorldOffset, );
+
+	virtual void InitializeForWorld(UWorld& World, FNavigationSystemRunMode Mode) PURE_VIRTUAL(UNavigationSystemBase::InitializeForWorld, );
+
+	/** 
+	 *	If you're using NavigationSysstem module consider calling 
+	 *	FNavigationSystem::GetCurrent<UNavigationSystemV1>()->GetDefaultNavDataInstance 
+	 *	instead.
+	 */
+	virtual INavigationDataInterface* GetMainNavData() const { return nullptr; }
+
+	DEPRECATED(4.20, "GetMainNavData is deprecated. Use FNavigationSystem::GetCurrent<UNavigationSystemV1>()->GetDefaultNavDataInstance instead")
+	INavigationDataInterface* GetMainNavData(int) { return nullptr; }
+
+protected:
+	
+	static void SetCoordTransformTo(const ENavigationCoordSystem::Type CoordType, const FTransform& Transform);
+	static void SetCoordTransformFrom(const ENavigationCoordSystem::Type CoordType, const FTransform& Transform);
+
+	static FNavigationSystem::FActorBasedSignature& UpdateActorDataDelegate();
+	static FNavigationSystem::FActorComponentBasedSignature& UpdateComponentDataDelegate();
+	static FNavigationSystem::FSceneComponentBasedSignature& UpdateComponentDataAfterMoveDelegate();
+	static FNavigationSystem::FActorBasedSignature& OnActorBoundsChangedDelegate();
+	static FNavigationSystem::FActorBasedSignature& OnPostEditActorMoveDelegate();
+	static FNavigationSystem::FSceneComponentBasedSignature& OnComponentTransformChangedDelegate();
+	static FNavigationSystem::FActorBasedSignature& OnActorRegisteredDelegate();
+	static FNavigationSystem::FActorBasedSignature& OnActorUnregisteredDelegate();
+	static FNavigationSystem::FActorComponentBasedSignature& OnComponentRegisteredDelegate();
+	static FNavigationSystem::FActorComponentBasedSignature& OnComponentUnregisteredDelegate();
+	static FNavigationSystem::FActorBasedSignature& RemoveActorDataDelegate();
+	static FNavigationSystem::FBoolActorComponentBasedSignature& HasComponentDataDelegate();
+	static FNavigationSystem::FNavAreaBasedSignature& GetDefaultWalkableAreaDelegate();
+	static FNavigationSystem::FNavDatConfigBasedSignature& GetDefaultSupportedAgentDelegate();
+	static FNavigationSystem::FActorBooleBasedSignature& UpdateActorAndComponentDataDelegate();
+	static FNavigationSystem::FComponentBoundsChangeSignature& OnComponentBoundsChangedDelegate();
+	//static FNavigationSystem::FNavDataForPropsSignature& GetNavDataForPropsDelegate();
+	static FNavigationSystem::FNavDataForActorSignature& GetNavDataForActorDelegate();
+	static FNavigationSystem::FNavDataClassFetchSignature& GetDefaultNavDataClassDelegate();
+	static FNavigationSystem::FWorldBoolBasedSignature& VerifyNavigationRenderingComponentsDelegate();
+	static FNavigationSystem::FWorldBasedSignature& BuildDelegate();
+#if WITH_EDITOR
+	static FNavigationSystem::FWorldBasedSignature& OnPIEStartDelegate();
+	static FNavigationSystem::FWorldBasedSignature& OnPIEEndDelegate();
+	static FNavigationSystem::FLevelBasedSignature& UpdateLevelCollisionDelegate();
+	static FNavigationSystem::FNavigationAutoUpdateEnableSignature& SetNavigationAutoUpdateEnableDelegate();
+	static FNavigationSystem::FWorldByteBasedSignature& AddNavigationUpdateLockDelegate();
+	static FNavigationSystem::FWorldByteBasedSignature& RemoveNavigationUpdateLockDelegate();
+#endif // WITH_EDITOR
+};
+
+
+class ENGINE_API IPathFollowingManagerInterface
+{
+protected:
+	static FNavigationSystem::FControllerBasedSignature& StopMovementDelegate();
+	static FNavigationSystem::FBoolControllerBasedSignature& IsFollowingAPathDelegate();
+};
+
+
+UCLASS(Abstract, config = Engine, defaultconfig)
+class ENGINE_API UNavigationSystem : public UObject
+{
+	GENERATED_BODY()
+public:
+	UNavigationSystem(const FObjectInitializer& ObjectInitializer = FObjectInitializer::Get());
+
+	DEPRECATED(4.20, "UpdateActorInNavOctree is deprecated. Use FNavigationSystem::RemoveActorData instead")
+	static void ClearNavOctreeAll(AActor* Actor);
+	DEPRECATED(4.20, "UpdateActorInNavOctree is deprecated. Use FNavigationSystem::UpdateActorData instead")
+	static void UpdateActorInNavOctree(AActor& Actor);
+	DEPRECATED(4.20, "UpdateComponentInNavOctree is deprecated. Use FNavigationSystem::UpdateComponentData instead")
+	static void UpdateComponentInNavOctree(UActorComponent& Comp);
+	DEPRECATED(4.20, "UpdateActorAndComponentsInNavOctree is deprecated. Use FNavigationSystem::UpdateActorAndComponentData instead")
+	static void UpdateActorAndComponentsInNavOctree(AActor& Actor, bool bUpdateAttachedActors = true);
+	DEPRECATED(4.20, "UpdateNavOctreeAfterMove is deprecated. Use FNavigationSystem::UpdateComponentDataAfterMove instead")
+	static void UpdateNavOctreeAfterMove(USceneComponent* Comp);
+	DEPRECATED(4.20, "UpdateNavOctreeBounds is deprecated. Use FNavigationSystem::OnActorBoundsChanged instead")
+	static void UpdateNavOctreeBounds(AActor* Actor);
+	DEPRECATED(4.20, "InitializeForWorld is deprecated. Use FNavigationSystem::CreateNavigationSystem instead")
+	static void InitializeForWorld(UWorld* World, FNavigationSystemRunMode Mode);
+	DEPRECATED(4.20, "CreateNavigationSystem is deprecated. Use FNavigationSystem::CreateNavigationSystem instead")
+	static UNavigationSystem* CreateNavigationSystem(UWorld* WorldOwner);
+	DEPRECATED(4.20, "OnComponentRegistered is deprecated. Use FNavigationSystem::OnComponentRegistered instead")
+	static void OnComponentRegistered(UActorComponent* Comp);
+	DEPRECATED(4.20, "OnComponentUnregistered is deprecated. Use FNavigationSystem::OnComponentUnregistered instead")
+	static void OnComponentUnregistered(UActorComponent* Comp);
+	DEPRECATED(4.20, "OnActorRegistered is deprecated. Use FNavigationSystem::OnActorRegistered instead")
+	static void OnActorRegistered(AActor* Actor);
+	DEPRECATED(4.20, "OnActorUnregistered is deprecated. Use FNavigationSystem::OnActorUnregistered instead")
+	static void OnActorUnregistered(AActor* Actor);
+	DEPRECATED(4.20, "GetCurrent is deprecated. Use FNavigationSystem::GetCurrent instead")
+	static UNavigationSystem* GetCurrent(UWorld* World);
+	DEPRECATED(4.20, "GetCurrent is deprecated. Use FNavigationSystem::GetCurrent instead")
+	static UNavigationSystem* GetCurrent(UObject* WorldContextObject);
+	DEPRECATED(4.20, "ShouldUpdateNavOctreeOnComponentChange is deprecated. Use FNavigationSystem::WantsComponentChangeNotifies instead")
+	static bool ShouldUpdateNavOctreeOnComponentChange();
+	DEPRECATED(4.20, "GetDefaultWalkableArea is deprecated. Use FNavigationSystem::GetDefaultWalkableArea instead")
+	static TSubclassOf<UNavAreaBase> GetDefaultWalkableArea();
+	DEPRECATED(4.20, "UNavigationSystem::GetNavDataForProps is deprecated. Use UNavigationSystemV1::GetDefaultObstacleArea instead")
+	static TSubclassOf<UNavAreaBase> GetDefaultObstacleArea();
+	DEPRECATED(4.20, "UNavigationSystem::K2_GetRandomReachablePointInRadius is deprecated. Use UNavigationSystemV1::K2_GetRandomReachablePointInRadius instead")
+	static bool K2_GetRandomReachablePointInRadius(UObject* WorldContextObject, const FVector& Origin, FVector& RandomLocation, float Radius, UObject* NavData = NULL, TSubclassOf<UObject> FilterClass = NULL) { return false; }
+
+	DEPRECATED(4.20, "UpdateNavOctreeElementBounds is deprecated. Use FNavigationSystem::OnComponentBoundsChanged instead")
+	bool UpdateNavOctreeElementBounds(UActorComponent* Comp, const FBox& NewBounds, const FBox& DirtyArea);
+
+	DEPRECATED(4.20, "GetObjectsNavOctreeId is deprecated. You need to access NavigationSystem's modul NavigationSystem to access this functionality")
+	const class FOctreeElementId* GetObjectsNavOctreeId(const UObject* Object) const { return nullptr; }
+
+	DEPRECATED(4.20, "HasPendingObjectNavOctreeId is deprecated. You need to access NavigationSystem's modul NavigationSystem class to access this functionality")
+	bool HasPendingObjectNavOctreeId(UObject* Object) const { return false; }
+
+	DEPRECATED(4.20, "GetDefaultSupportedAgentConfig is deprecated. Use FNavigationSystem::GetDefaultSupportedAgent instead")
+	const FNavDataConfig& GetDefaultSupportedAgentConfig() const;
+	
+	DEPRECATED(4.20, "GetNavDataForProps is deprecated. Use FNavigationSystem::GetNavDataForProps instead")
+	INavigationDataInterface* GetNavDataForProps(const FNavAgentProperties& AgentProperties);
+
+	DEPRECATED(4.20, "GetSupportedAgents has been moved. Use UNavigationSystemV1::GetSupportedAgents instead")
+	const TArray<FNavDataConfig>& GetSupportedAgents() const { return FakeSupportedAgents; }
+
+	DEPRECATED(4.20, "GetMainNavData has been moved. Use FNavigationSystem::GetCurrent<UNavigationSystemV1>()->GetDefaultNavDataInstance instead")
+	UObject* GetMainNavData(int) { return nullptr; }
+	DEPRECATED(4.20, "GetMainNavData has been moved. Use FNavigationSystem::GetCurrent<UNavigationSystemV1>()->GetDefaultNavDataInstance instead")
+	UObject* GetMainNavData() const { return nullptr; }
+
+	DEPRECATED(4.20, "ProjectPointToNavigation has been deprecated along with the whole UNavigationSystem class. Use UNavigationSystemV1 instead")
+	bool ProjectPointToNavigation(const FVector& Point, FNavLocation& OutLocation, const FVector& Extent = INVALID_NAVEXTENT, const UObject* NavData = NULL, void* QueryFilter = NULL) const { return false; }
+
+private:
+	const TArray<FNavDataConfig> FakeSupportedAgents;
+};
\ No newline at end of file
diff --git a/Engine/Source/Runtime/Engine/Classes/AI/NavigationSystemConfig.h b/Engine/Source/Runtime/Engine/Classes/AI/NavigationSystemConfig.h
new file mode 100644
--- /dev/null
+++ b/Engine/Source/Runtime/Engine/Classes/AI/NavigationSystemConfig.h
@@ -0,0 +1,46 @@
+// Copyright 1998-2018 Epic Games, Inc. All Rights Reserved.
+
+#pragma once
+
+#include "CoreMinimal.h"
+#include "UObject/Object.h"
+#include "UObject/SoftObjectPath.h"
+#include "NavigationSystemConfig.generated.h"
+
+class UNavigationSystemBase;
+class UWorld;
+
+
+UCLASS(config = Engine, defaultconfig, EditInlineNew, DisplayName = "Generic Navigation System Config", collapseCategories)
+class ENGINE_API UNavigationSystemConfig : public UObject
+{
+	GENERATED_BODY()
+
+public:
+	UPROPERTY(EditAnywhere, Category=Navigation, meta = (MetaClass = "NavigationSystemBase", NoResetToDefault))
+	FSoftClassPath NavigationSystemClass;
+
+	/** Whether at game runtime we expect any kind of dynamic navigation generation */
+	UPROPERTY(EditAnywhere, Category = Navigation)
+	uint32 bStrictlyStatic : 1;
+
+	UPROPERTY(EditAnywhere, Category = Navigation)
+	uint32 bCreateOnClient : 1;
+
+public:
+	UNavigationSystemConfig(const FObjectInitializer& ObjectInitializer = FObjectInitializer::Get());
+
+	UNavigationSystemBase* CreateAndConfigureNavigationSystem(UWorld& World) const;
+
+#if WITH_EDITOR
+	virtual void PostEditChangeProperty(FPropertyChangedEvent& PropertyChangedEvent) override;
+#endif // WITH_EDITOR
+};
+
+UCLASS(MinimalAPI, HideCategories=Navigation)
+class UNullNavSysConfig : public UNavigationSystemConfig
+{
+	GENERATED_BODY()
+public:
+	UNullNavSysConfig(const FObjectInitializer& ObjectInitializer = FObjectInitializer::Get());
+};
diff --git a/Engine/Source/Runtime/Engine/Classes/Components/ShapeComponent.h b/Engine/Source/Runtime/Engine/Classes/Components/ShapeComponent.h
--- a/Engine/Source/Runtime/Engine/Classes/Components/ShapeComponent.h
+++ b/Engine/Source/Runtime/Engine/Classes/Components/ShapeComponent.h
@@ -71,7 +71,7 @@ public:
 
 	/** Navigation area type (empty = default obstacle) */
 	UPROPERTY(EditAnywhere, Category = Navigation)
-	TSubclassOf<class UNavArea> AreaClass;
+	TSubclassOf<class UNavAreaBase> AreaClass;
 
 	//~ Begin UPrimitiveComponent Interface.
 	virtual FPrimitiveSceneProxy* CreateSceneProxy() override;
diff --git a/Engine/Source/Runtime/Engine/Classes/EdGraph/EdGraphPin.h b/Engine/Source/Runtime/Engine/Classes/EdGraph/EdGraphPin.h
--- a/Engine/Source/Runtime/Engine/Classes/EdGraph/EdGraphPin.h
+++ b/Engine/Source/Runtime/Engine/Classes/EdGraph/EdGraphPin.h
@@ -342,17 +342,64 @@ public:
 	/** Name of this pin. */
 	FName PinName;
 
+	/** Direction of flow of this pin (input or output) */
+	TEnumAsByte<enum EEdGraphPinDirection> Direction;
+
+#if WITH_EDITORONLY_DATA
+	////////////////////////////////////////////////////////////////////////////////////
+	// ONLY PUT BITFIELD PROPERTIES AFTER THIS TO ENSURE GOOD MEMORY ALIGNMENT
+	////////////////////////////////////////////////////////////////////////////////////
+
+	/** If true, this connector is currently hidden. */
+	uint8 bHidden:1;
+
+	/** If true, this connector is unconnectable, and present only to allow the editing of the default text. */
+	uint8 bNotConnectable:1;
+
+	/** If true, the default value of this connector is fixed and cannot be modified by the user (it's visible for reference only). */
+	uint8 bDefaultValueIsReadOnly:1;
+
+	/** If true, the default value on this pin is ignored and should not be set. */
+	uint8 bDefaultValueIsIgnored:1;
+
+	/** If true, this pin is the focus of a diff. This is transient. */
+	uint8 bIsDiffing:1;
+
+	/** If true, the pin may be hidden by user. */
+	uint8 bAdvancedView:1;
+
+	/** If true, the pin is displayed as ref. This is transient. */
+	uint8 bDisplayAsMutableRef:1;
+
+	/** 
+	* If true, this pin existed on an older version of the owning node, but when the node was reconstructed a matching pin was not found.
+	* This pin must be linked to other pins or have a non-default value and will be removed if disconnected, reset to default, or the node is refreshed.
+	**/
+	uint8 bOrphanedPin:1;
+
+	bool ShouldSavePinIfOrphaned() const { return bSavePinIfOrphaned; }
+	ENGINE_API void SetSavePinIfOrphaned(bool bShouldSave);
+
+private:
+	/** If true, this pin will be retained when reconstructing a node if there is no matching pin on the new version of the pin. This is transient. */
+	uint8 bSavePinIfOrphaned:1;
+
+public:
+	/** Older content sometimes had an empty autogenerated default value string in cases where that does not mean the property default value (0, none, false, etc.) */
+	uint8 bUseBackwardsCompatForEmptyAutogeneratedValue:1;
+
+#endif // WITH_EDITORONLY_DATA
+
+	/** True when InvalidateAndTrash was called. This pin is intended to be discarded and destroyed. */
+	uint8 bWasTrashed:1;	/** Used as the display name if set. */
+
 #if WITH_EDITORONLY_DATA
-	/** Used as the display name if set. */
 	FText PinFriendlyName;
 #endif
 
 	/** The tool-tip describing this pin's purpose */
 	FString PinToolTip;
 
-	/** Direction of flow of this pin (input or output) */
-	TEnumAsByte<enum EEdGraphPinDirection> Direction;
-
 	/** The type of information carried on this pin */
 	FEdGraphPinType PinType;
 
@@ -384,47 +431,7 @@ public:
 	/** Pin name could be changed, so whenever possible it's good to have a persistent GUID identifying Pin to reconstruct Node seamlessly */
 	FGuid PersistentGuid;
 
-////////////////////////////////////////////////////////////////////////////////////
-// ONLY PUT BITFIELD PROPERTIES AFTER THIS TO ENSURE GOOD MEMORY ALIGNMENT
-////////////////////////////////////////////////////////////////////////////////////
-
-	/** If true, this connector is currently hidden. */
-	uint32 bHidden:1;
-
-	/** If true, this connector is unconnectable, and present only to allow the editing of the default text. */
-	uint32 bNotConnectable:1;
-
-	/** If true, the default value of this connector is fixed and cannot be modified by the user (it's visible for reference only). */
-	uint32 bDefaultValueIsReadOnly:1;
-
-	/** If true, the default value on this pin is ignored and should not be set. */
-	uint32 bDefaultValueIsIgnored:1;
-
-	/** If true, this pin is the focus of a diff. This is transient. */
-	uint32 bIsDiffing:1;
-
-	/** If true, the pin may be hidden by user. */
-	uint32 bAdvancedView:1;
-
-	/** If true, the pin is displayed as ref. This is transient. */
-	uint32 bDisplayAsMutableRef:1;
-
-	/** 
-	 * If true, this pin existed on an older version of the owning node, but when the node was reconstructed a matching pin was not found.
-	 * This pin must be linked to other pins or have a non-default value and will be removed if disconnected, reset to default, or the node is refreshed.
-	 **/
-	uint32 bOrphanedPin:1;
-
-	/** If true, this pin will be retained when reconstructing a node if there is no matching pin on the new version of the pin. This is transient. */
-	uint32 bSavePinIfOrphaned:1;
-
-	/** Older content sometimes had an empty autogenerated default value string in cases where that does not mean the property default value (0, none, false, etc.) */
-	uint32 bUseBackwardsCompatForEmptyAutogeneratedValue:1;
-
-#endif // WITH_EDITORONLY_DATA
-
-	/** True when InvalidateAndTrash was called. This pin is intended to be discarded and destroyed. */
-	uint32 bWasTrashed:1;
+#endif
 
 public:
 
diff --git a/Engine/Source/Runtime/Engine/Classes/Engine/Blueprint.h b/Engine/Source/Runtime/Engine/Classes/Engine/Blueprint.h
--- a/Engine/Source/Runtime/Engine/Classes/Engine/Blueprint.h
+++ b/Engine/Source/Runtime/Engine/Classes/Engine/Blueprint.h
@@ -296,8 +296,9 @@ struct FEditedDocumentInfo
 {
 	GENERATED_USTRUCT_BODY()
 
+	/** Edited object */
 	UPROPERTY()
-	UObject* EditedObject;
+	FSoftObjectPath EditedObjectPath;
 
 	/** Saved view position */
 	UPROPERTY()
@@ -308,29 +309,78 @@ struct FEditedDocumentInfo
 	float SavedZoomAmount;
 
 	FEditedDocumentInfo()
-		: EditedObject(nullptr)
-		, SavedViewOffset(0.0f, 0.0f)
+		: SavedViewOffset(0.0f, 0.0f)
 		, SavedZoomAmount(-1.0f)
+		, EditedObject_DEPRECATED(nullptr)
 	{ }
 
 	FEditedDocumentInfo(UObject* InEditedObject)
-		: EditedObject(InEditedObject)
+		: EditedObjectPath(InEditedObject)
 		, SavedViewOffset(0.0f, 0.0f)
 		, SavedZoomAmount(-1.0f)
+		, EditedObject_DEPRECATED(nullptr)
 	{ }
 
 	FEditedDocumentInfo(UObject* InEditedObject, FVector2D& InSavedViewOffset, float InSavedZoomAmount)
-		: EditedObject(InEditedObject)
+		: EditedObjectPath(InEditedObject)
 		, SavedViewOffset(InSavedViewOffset)
 		, SavedZoomAmount(InSavedZoomAmount)
+		, EditedObject_DEPRECATED(nullptr)
 	{ }
 
-	friend bool operator==( const FEditedDocumentInfo& LHS, const FEditedDocumentInfo& RHS )
+	void PostSerialize(const FArchive& Ar)
 	{
-		return LHS.EditedObject == RHS.EditedObject && LHS.SavedViewOffset == RHS.SavedViewOffset && LHS.SavedZoomAmount == RHS.SavedZoomAmount;
+		if (Ar.IsLoading() && EditedObject_DEPRECATED)
+		{
+			// Convert hard to soft reference.
+			EditedObjectPath = EditedObject_DEPRECATED;
+			EditedObject_DEPRECATED = nullptr;
+		}
 	}
+
+	friend bool operator==(const FEditedDocumentInfo& LHS, const FEditedDocumentInfo& RHS)
+	{
+		return LHS.EditedObjectPath == RHS.EditedObjectPath && LHS.SavedViewOffset == RHS.SavedViewOffset && LHS.SavedZoomAmount == RHS.SavedZoomAmount;
+	}
+
+private:
+	// Legacy hard reference is now serialized as a soft reference (see above).
+	UPROPERTY()
+	UObject* EditedObject_DEPRECATED;
 };
 
+template<>
+struct TStructOpsTypeTraits<FEditedDocumentInfo> : public TStructOpsTypeTraitsBase2<FEditedDocumentInfo>
+{
+	enum
+	{
+		WithPostSerialize = true
+	};
+};
+
+/** Bookmark node info */
+USTRUCT()
+struct FBPEditorBookmarkNode
+{
+	GENERATED_USTRUCT_BODY()
+
+	/** Node ID */
+	UPROPERTY()
+	FGuid NodeGuid;
+
+	/** Parent ID */
+	UPROPERTY()
+	FGuid ParentGuid;
+
+	/** Display name */
+	UPROPERTY()
+	FText DisplayName;
+
+	friend bool operator==(const FBPEditorBookmarkNode& LHS, const FBPEditorBookmarkNode& RHS)
+	{
+		return LHS.NodeGuid == RHS.NodeGuid;
+	}
+};
 
 UENUM()
 enum class EBlueprintNativizationFlag : uint8
@@ -340,6 +390,7 @@ enum class EBlueprintNativizationFlag : uint8
 	ExplicitlyEnabled
 };
 
+
 /**
  * Blueprints are special assets that provide an intuitive, node-based interface that can be used to create new types of Actors
  * and script level events; giving designers and gameplay programmers the tools to quickly create and iterate gameplay from
@@ -529,6 +580,14 @@ class ENGINE_API UBlueprint : public UBlueprintCore
 	UPROPERTY()
 	TArray<struct FEditedDocumentInfo> LastEditedDocuments;
 
+	/** Bookmark data */
+	UPROPERTY()
+	TMap<FGuid, struct FEditedDocumentInfo> Bookmarks;
+
+	/** Bookmark nodes (for display) */
+	UPROPERTY()
+	TArray<FBPEditorBookmarkNode> BookmarkNodes;
+
 	/** Persistent debugging options */
 	UPROPERTY()
 	TArray<class UBreakpoint*> Breakpoints;
diff --git a/Engine/Source/Runtime/Engine/Classes/Engine/Engine.h b/Engine/Source/Runtime/Engine/Classes/Engine/Engine.h
--- a/Engine/Source/Runtime/Engine/Classes/Engine/Engine.h
+++ b/Engine/Source/Runtime/Engine/Classes/Engine/Engine.h
@@ -710,7 +710,7 @@ public:
 
 	/** The class for NavigationSystem **/
 	UPROPERTY()
-	TSubclassOf<class UNavigationSystem>  NavigationSystemClass;
+	TSubclassOf<class UNavigationSystemBase>  NavigationSystemClass;
 	
 	/** Name of behavior tree manager class */
 	UPROPERTY(globalconfig, noclear, meta=(MetaClass="AvoidanceManager", DisplayName="Avoidance Manager Class"))
diff --git a/Engine/Source/Runtime/Engine/Classes/Engine/Level.h b/Engine/Source/Runtime/Engine/Classes/Engine/Level.h
--- a/Engine/Source/Runtime/Engine/Classes/Engine/Level.h
+++ b/Engine/Source/Runtime/Engine/Classes/Engine/Level.h
@@ -651,7 +651,7 @@ public:
 	virtual void BeginDestroy() override;
 	virtual bool IsReadyForFinishDestroy() override;
 	virtual void FinishDestroy() override;
-	virtual UWorld* GetWorld() const override;
+	ENGINE_API virtual UWorld* GetWorld() const override final;
 
 #if	WITH_EDITOR
 	virtual void PreEditUndo() override;
diff --git a/Engine/Source/Runtime/Engine/Classes/Engine/StaticMesh.h b/Engine/Source/Runtime/Engine/Classes/Engine/StaticMesh.h
--- a/Engine/Source/Runtime/Engine/Classes/Engine/StaticMesh.h
+++ b/Engine/Source/Runtime/Engine/Classes/Engine/StaticMesh.h
@@ -24,6 +24,7 @@
 class FSpeedTreeWind;
 class UAssetUserData;
 class UMaterialInterface;
+class UNavCollisionBase;
 struct FStaticMeshLODResources;
 
 /*-----------------------------------------------------------------------------
@@ -645,7 +646,7 @@ protected:
 public:
 	/** Pre-build navigation collision */
 	UPROPERTY(VisibleAnywhere, transient, duplicatetransient, Instanced, Category = Navigation)
-	class UNavCollision* NavCollision;
+	UNavCollisionBase* NavCollision;
 public:
 	/**
 	 * Default constructor
@@ -813,9 +814,9 @@ public:
 	 */
 	ENGINE_API void CreateNavCollision(const bool bIsUpdate = false);
 
-	FORCEINLINE const UNavCollision* GetNavCollision() const { return NavCollision; }
+	FORCEINLINE const UNavCollisionBase* GetNavCollision() const { return NavCollision; }
 
-	/** Configures this SM as bHasNavigationData = false and clears stored UNavCollision */
+	/** Configures this SM as bHasNavigationData = false and clears stored NavCollision */
 	ENGINE_API void MarkAsNotHavingNavigationData();
 
 	const FGuid& GetLightingGuid() const
diff --git a/Engine/Source/Runtime/Engine/Classes/Engine/StreamableManager.h b/Engine/Source/Runtime/Engine/Classes/Engine/StreamableManager.h
--- a/Engine/Source/Runtime/Engine/Classes/Engine/StreamableManager.h
+++ b/Engine/Source/Runtime/Engine/Classes/Engine/StreamableManager.h
@@ -64,11 +64,18 @@ struct ENGINE_API FStreamableHandle : public TSharedFromThis<FStreamableHandle,
 		return Priority;
 	}
 
-	/** Release this handle. This can be called from normal gameplay code to indicate that the loaded assets are no longer needed. Will be called implicitly if all shared pointers
-	to this handle are destroyed. If called before the completion delegate, the release will be delayed until after completion. */
+	/**
+	 * Release this handle. This can be called from normal gameplay code to indicate that the loaded assets are no longer needed
+	 * This will be called implicitly if all shared pointers to this handle are destroyed
+	 * If called before the completion delegate, the release will be delayed until after completion
+	 */
 	void ReleaseHandle();
 
-	/** Cancel a request, callable from within the manager or externally. Will stop delegate from being called */
+	/**
+	 * Cancel a request, callable from within the manager or externally
+	 * This will immediately release the handle even if it is still in progress
+	 * This will normally stop the completion delegate from getting called, unless that delegate is already in the callback delay queue
+	 */
 	void CancelHandle();
 
 	/** Tells a stalled handle to start its actual request. */
@@ -285,6 +292,9 @@ struct ENGINE_API FStreamableManager : public FGCObject
 	/** This will release any managed active handles pointing to the target string asset reference, even if they include other requested assets in the same load */
 	void Unload(const FSoftObjectPath& Target);
 
+	/** Checks for any redirectors that were previously loaded, and returns the redirected target if found. This will not handle redirects that it doesn't yet know about */
+	FSoftObjectPath ResolveRedirects(const FSoftObjectPath& Target) const;
+
 	DEPRECATED(4.16, "Call LoadSynchronous with bManageActiveHandle=true instead if you want the manager to keep the handle alive")
 	UObject* SynchronousLoad(FSoftObjectPath const& Target);
 
@@ -317,7 +327,6 @@ private:
 
 	void RemoveReferencedAsset(const FSoftObjectPath& Target, TSharedRef<FStreamableHandle> Handle);
 	void StartHandleRequests(TSharedRef<FStreamableHandle> Handle);
-	FSoftObjectPath ResolveRedirects(const FSoftObjectPath& Target) const;
 	void FindInMemory(FSoftObjectPath& InOutTarget, struct FStreamable* Existing);
 	struct FStreamable* FindStreamable(const FSoftObjectPath& Target) const;
 	struct FStreamable* StreamInternal(const FSoftObjectPath& Target, TAsyncLoadPriority Priority, TSharedRef<FStreamableHandle> Handle);
diff --git a/Engine/Source/Runtime/Engine/Classes/Engine/World.h b/Engine/Source/Runtime/Engine/Classes/Engine/World.h
--- a/Engine/Source/Runtime/Engine/Classes/Engine/World.h
+++ b/Engine/Source/Runtime/Engine/Classes/Engine/World.h
@@ -896,9 +896,9 @@ public:
 	uint32 bTriggerPostLoadMap:1;
 
 private:
-	/** The world's navmesh */
+	/** The world's navigation data manager */
 	UPROPERTY(Transient)
-	class UNavigationSystem*					NavigationSystem;
+	class UNavigationSystemBase*				NavigationSystem;
 
 	/** The current GameMode, valid only on the server */
 	UPROPERTY(Transient)
@@ -961,8 +961,8 @@ private:
 	UCanvas* CanvasForDrawMaterialToRenderTarget;
 
 public:
-	/** Set the pointer to the Navgation system. */
-	void SetNavigationSystem( UNavigationSystem* InNavigationSystem);
+	/** Set the pointer to the Navigation System instance. */
+	void SetNavigationSystem(UNavigationSystemBase* InNavigationSystem);
 
 	/** The interface to the scene manager for this world. */
 	class FSceneInterface*						Scene;
@@ -1934,9 +1934,9 @@ public:
 	bool IsTraceHandleValid(const FTraceHandle& Handle, bool bOverlapTrace);
 
 	/** NavigationSystem getter */
-	FORCEINLINE UNavigationSystem* GetNavigationSystem() { return NavigationSystem; }
+	FORCEINLINE UNavigationSystemBase* GetNavigationSystem() { return NavigationSystem; }
 	/** NavigationSystem const getter */
-	FORCEINLINE const UNavigationSystem* GetNavigationSystem() const { return NavigationSystem; }
+	FORCEINLINE const UNavigationSystemBase* GetNavigationSystem() const { return NavigationSystem; }
 
 	/** AISystem getter. if AISystem is missing it tries to create one and returns the result.
 	 *	@NOTE the result can be NULL, for example on client games or if no AI module or AISystem class have not been specified
diff --git a/Engine/Source/Runtime/Engine/Classes/GameFramework/Actor.h b/Engine/Source/Runtime/Engine/Classes/GameFramework/Actor.h
--- a/Engine/Source/Runtime/Engine/Classes/GameFramework/Actor.h
+++ b/Engine/Source/Runtime/Engine/Classes/GameFramework/Actor.h
@@ -38,6 +38,7 @@ ENGINE_API DECLARE_LOG_CATEGORY_EXTERN(LogActor, Log, Warning);
 // Delegate signatures
 DECLARE_DYNAMIC_MULTICAST_DELEGATE_FiveParams( FTakeAnyDamageSignature, AActor*, DamagedActor, float, Damage, const class UDamageType*, DamageType, class AController*, InstigatedBy, AActor*, DamageCauser );
 DECLARE_DYNAMIC_MULTICAST_DELEGATE_NineParams( FTakePointDamageSignature, AActor*, DamagedActor, float, Damage, class AController*, InstigatedBy, FVector, HitLocation, class UPrimitiveComponent*, FHitComponent, FName, BoneName, FVector, ShotFromDirection, const class UDamageType*, DamageType, AActor*, DamageCauser );
+DECLARE_DYNAMIC_MULTICAST_DELEGATE_SevenParams( FTakeRadialDamageSignature, AActor*, DamagedActor, float, Damage, const class UDamageType*, DamageType, FVector, Origin, FHitResult, HitInfo, class AController*, InstigatedBy, AActor*, DamageCauser );
 DECLARE_DYNAMIC_MULTICAST_DELEGATE_TwoParams( FActorBeginOverlapSignature, AActor*, OverlappedActor, AActor*, OtherActor );
 DECLARE_DYNAMIC_MULTICAST_DELEGATE_TwoParams( FActorEndOverlapSignature, AActor*, OverlappedActor, AActor*, OtherActor );
 DECLARE_DYNAMIC_MULTICAST_DELEGATE_FourParams( FActorHitSignature, AActor*, SelfActor, AActor*, OtherActor, FVector, NormalImpulse, const FHitResult&, Hit );
@@ -621,6 +622,10 @@ public:
 	/** Called when the actor is damaged by point damage. */
 	UPROPERTY(BlueprintAssignable, Category="Game|Damage")
 	FTakePointDamageSignature OnTakePointDamage;
+
+	/** Called when the actor is damaged by radial damage. */
+	UPROPERTY(BlueprintAssignable, Category="Game|Damage")
+	FTakeRadialDamageSignature OnTakeRadialDamage;
 	
 	/** 
 	 * Called when another actor begins to overlap this actor, for example a player walking into a trigger.
@@ -2678,7 +2683,7 @@ public:
 			{
 				OutComponents.Add(Component);
 			}
-			else if (bIncludeFromChildActors)
+			if (bIncludeFromChildActors)
 			{
 				if (UChildActorComponent* ChildActorComponent = Cast<UChildActorComponent>(OwnedComponent))
 				{
@@ -2726,12 +2731,13 @@ public:
 			if (Component)
 			{
 				OutComponents.Add(Component);
-			}
-			else if (bIncludeFromChildActors)
-			{
-				if (UChildActorComponent* ChildActorComponent = Cast<UChildActorComponent>(Component))
+
+				if (bIncludeFromChildActors)
 				{
-					ChildActorComponents.Add(ChildActorComponent);
+					if (UChildActorComponent* ChildActorComponent = Cast<UChildActorComponent>(Component))
+					{
+						ChildActorComponents.Add(ChildActorComponent);
+					}
 				}
 			}
 		}
diff --git a/Engine/Source/Runtime/Engine/Classes/GameFramework/CharacterMovementComponent.h b/Engine/Source/Runtime/Engine/Classes/GameFramework/CharacterMovementComponent.h
--- a/Engine/Source/Runtime/Engine/Classes/GameFramework/CharacterMovementComponent.h
+++ b/Engine/Source/Runtime/Engine/Classes/GameFramework/CharacterMovementComponent.h
@@ -24,6 +24,7 @@ class FDebugDisplayInfo;
 class FNetworkPredictionData_Server_Character;
 class FSavedMove_Character;
 class UPrimitiveComponent;
+class INavigationData;
 class UCharacterMovementComponent;
 
 DECLARE_DELEGATE_RetVal_TwoParams(FTransform, FOnProcessRootMotion, const FTransform&, UCharacterMovementComponent*)
@@ -486,6 +487,14 @@ public:
 	UPROPERTY(Category="Character Movement (Networking)", EditDefaultsOnly)
 	uint32 bNetworkSkipProxyPredictionOnNetUpdate:1;
 
+	/**
+	 * Flag used on the server to determine whether to always replicate ReplicatedServerLastTransformUpdateTimeStamp to clients.
+	 * Normally this is only sent when the network smoothing mode on character movement is set to Linear smoothing (on the server), to save bandwidth.
+	 * Setting this to true will force the timestamp to replicate regardless, in case the server doesn't know about the smoothing mode, or if the timestamp is used for another purpose.
+	 */
+	UPROPERTY(Category="Character Movement (Networking)", EditDefaultsOnly, AdvancedDisplay)
+	uint32 bNetworkAlwaysReplicateTransformUpdateTimestamp:1;
+
 public:
 
 	/** true to update CharacterOwner and UpdatedComponent after movement ends */
@@ -1497,7 +1506,7 @@ protected:
 	virtual void SetNavWalkingPhysics(bool bEnable);
 
 	/** Get Navigation data for the Character. Returns null if there is no associated nav data. */
-	const class ANavigationData* GetNavData() const;
+	const class INavigationDataInterface* GetNavData() const;
 
 	/** 
 	 * Checks to see if the current location is not encroaching blocking geometry so the character can leave NavWalking.
diff --git a/Engine/Source/Runtime/Engine/Classes/GameFramework/CheatManager.h b/Engine/Source/Runtime/Engine/Classes/GameFramework/CheatManager.h
--- a/Engine/Source/Runtime/Engine/Classes/GameFramework/CheatManager.h
+++ b/Engine/Source/Runtime/Engine/Classes/GameFramework/CheatManager.h
@@ -252,14 +252,6 @@ class ENGINE_API UCheatManager : public UObject
 	UFUNCTION(exec)
 	virtual void TestCollisionDistance();
 
-	/** Builds the navigation mesh (or rebuilds it). **/
-	UFUNCTION(exec)
-	virtual void RebuildNavigation();
-
-	/** Sets navigation drawing distance. Relevant only in non-editor modes. **/
-	UFUNCTION(exec)
-	void SetNavDrawDistance(float DrawDistance);
-
 	/** Dump online session information */
 	UFUNCTION(exec)
 	virtual void DumpOnlineSessionState();
diff --git a/Engine/Source/Runtime/Engine/Classes/GameFramework/Controller.h b/Engine/Source/Runtime/Engine/Classes/GameFramework/Controller.h
--- a/Engine/Source/Runtime/Engine/Classes/GameFramework/Controller.h
+++ b/Engine/Source/Runtime/Engine/Classes/GameFramework/Controller.h
@@ -14,10 +14,10 @@ class APawn;
 class APlayerState;
 class FDebugDisplayInfo;
 class UDamageType;
-class UPathFollowingComponent;
 
 UDELEGATE(BlueprintAuthorityOnly)
 DECLARE_DYNAMIC_MULTICAST_DELEGATE_FourParams( FInstigatedAnyDamageSignature, float, Damage, const UDamageType*, DamageType, AActor*, DamagedActor, AActor*, DamageCauser );
+DECLARE_MULTICAST_DELEGATE_OneParam(FPawnChangedSignature, APawn* /*NewPawn*/);
 
 /**
  * Controllers are non-physical actors that can possess a Pawn to control
@@ -80,6 +80,8 @@ private:
 	USceneComponent* TransformComponent;
 
 protected:
+	FPawnChangedSignature OnNewPawn;
+
 	/** The control rotation of the Controller. See GetControlRotation. */
 	UPROPERTY()
 	FRotator ControlRotation;
@@ -187,6 +189,10 @@ public:
 	UFUNCTION(BlueprintCallable, Category=Pawn, meta=(DisplayName="Get Controlled Pawn", ScriptName="GetControlledPawn"))
 	APawn* K2_GetPawn() const;
 
+	FPawnChangedSignature GetOnNewPawnNotifier() { return OnNewPawn; }
+
+public:
+
 	virtual void GetLifetimeReplicatedProps(TArray<FLifetimeProperty>& OutLifetimeProps) const override;
 
 	//~ Begin AActor Interface
@@ -289,14 +295,9 @@ public:
 	virtual void GetMoveGoalReachTest(const AActor* MovingActor, const FVector& MoveOffset, FVector& GoalOffset, float& GoalRadius, float& GoalHalfHeight) const override;
 	virtual bool ShouldPostponePathUpdates() const override;
 	virtual bool IsFollowingAPath() const override;
+	virtual IPathFollowingAgentInterface* GetPathFollowingAgent() const override;
 	//~ End INavAgentInterface Interface
-
-	/** Prepares path following component */
-	virtual void InitNavigationControl(UPathFollowingComponent*& PathFollowingComp);
-
-	/** If controller has any navigation-related components then this function makes them update their cached data */
-	virtual void UpdateNavigationComponents();
-
+	
 	/** Aborts the move the controller is currently performing */
 	UFUNCTION(BlueprintCallable, Category="AI|Navigation")
 	virtual void StopMovement();
@@ -352,6 +353,20 @@ protected:
 private:
 	// Hidden functions that don't make sense to use on this class.
 	HIDE_ACTOR_TRANSFORM_FUNCTIONS();
+
+public:
+
+	/** prepares path following component */
+	DEPRECATED(4.20, "InitNavigationControl is depracated. The function lost its meaning due to better Engine-AIModule code separation")
+	virtual void InitNavigationControl(UObject*& PathFollowingComp) {}
+
+	/** If controller has any navigation-related components then this function
+	 *	makes them update their cached data 
+	 *	DEPRECATED: this functionality has been taken over by PathFollowingComponent 
+	 *		by observing newly possessed pawns (via OnNewPawn)
+	 */
+	DEPRECATED(4.20, "UpdateNavigationComponents is depracated. The function lost its meaning due to better Engine-AIModule code separation.")
+	virtual void UpdateNavigationComponents() {}
 };
 
 
diff --git a/Engine/Source/Runtime/Engine/Classes/GameFramework/GameNetworkManager.h b/Engine/Source/Runtime/Engine/Classes/GameFramework/GameNetworkManager.h
--- a/Engine/Source/Runtime/Engine/Classes/GameFramework/GameNetworkManager.h
+++ b/Engine/Source/Runtime/Engine/Classes/GameFramework/GameNetworkManager.h
@@ -129,6 +129,10 @@ class ENGINE_API AGameNetworkManager : public AInfo
 	UPROPERTY(GlobalConfig)
 	float ClientNetSendMoveDeltaTimeThrottled;
 
+	/** ClientNetSendMoveDeltaTimeStationary is used when players are determined to not be moving or changing their view. See ClientNetSendMoveDeltaTime for more info. */
+	UPROPERTY(GlobalConfig)
+	float ClientNetSendMoveDeltaTimeStationary;
+
 	/** When player net speed (CurrentNetSpeed, based on ConfiguredInternetSpeed or ConfiguredLanSpeed) is less than or equal to this amount, ClientNetSendMoveDeltaTimeThrottled is used instead of ClientNetSendMoveDeltaTime. */
 	UPROPERTY(GlobalConfig)
 	int32 ClientNetSendMoveThrottleAtNetSpeed;
diff --git a/Engine/Source/Runtime/Engine/Classes/GameFramework/NavMovementComponent.h b/Engine/Source/Runtime/Engine/Classes/GameFramework/NavMovementComponent.h
--- a/Engine/Source/Runtime/Engine/Classes/GameFramework/NavMovementComponent.h
+++ b/Engine/Source/Runtime/Engine/Classes/GameFramework/NavMovementComponent.h
@@ -12,11 +12,11 @@
 #include "GameFramework/Actor.h"
 #include "AI/Navigation/NavigationTypes.h"
 #include "AI/Navigation/NavAgentInterface.h"
+#include "AI/Navigation/PathFollowingAgentInterface.h"
 #include "GameFramework/MovementComponent.h"
 #include "NavMovementComponent.generated.h"
 
 class UCapsuleComponent;
-class UPathFollowingComponent;
 
 /**
  * NavMovementComponent defines base functionality for MovementComponents that move any 'agent' that may be involved in AI pathfinding.
@@ -47,18 +47,20 @@ protected:
 	UPROPERTY(EditAnywhere, Category = NavMovement, meta = (EditCondition = "bUseAccelerationForPaths"))
 	uint32 bUseFixedBrakingDistanceForPaths : 1;
 
-protected:
 	/** If set, StopActiveMovement call will abort current path following request */
 	uint32 bStopMovementAbortPaths:1;
 
+private:
+	/** object implementing IPathFollowingAgentInterface. Is private to control access to it.
+	 *	@see SetPathFollowingAgent, GetPathFollowingAgent */
+	UPROPERTY()
+	UObject* PathFollowingComp;
+
 public:
 	/** Expresses runtime state of character's movement. Put all temporal changes to movement properties here */
 	UPROPERTY()
 	FMovementProperties MovementState;
 
-	/** associated path following component */
-	TWeakObjectPtr<class UPathFollowingComponent> PathFollowingComp;
-
 	/** Stops applying further movement (usually zeros acceleration). */
 	UFUNCTION(BlueprintCallable, Category="Pawn|Components|PawnMovement")
 	virtual void StopActiveMovement();
@@ -98,6 +100,10 @@ public:
 	/** @returns braking distance for acceleration driven path following */
 	virtual float GetPathFollowingBrakingDistance(float MaxSpeed) const;
 
+	void SetPathFollowingAgent(IPathFollowingAgentInterface* InPathFollowingAgent) { PathFollowingComp = Cast<UObject>(InPathFollowingAgent); }
+	IPathFollowingAgentInterface* GetPathFollowingAgent() { return Cast<IPathFollowingAgentInterface>(PathFollowingComp); }
+	const IPathFollowingAgentInterface* GetPathFollowingAgent() const { return Cast<const IPathFollowingAgentInterface>(PathFollowingComp); }
+
 	/** set fixed braking distance */
 	void SetFixedBrakingDistance(float DistanceToEndOfPath);
 
diff --git a/Engine/Source/Runtime/Engine/Classes/GameFramework/WorldSettings.h b/Engine/Source/Runtime/Engine/Classes/GameFramework/WorldSettings.h
--- a/Engine/Source/Runtime/Engine/Classes/GameFramework/WorldSettings.h
+++ b/Engine/Source/Runtime/Engine/Classes/GameFramework/WorldSettings.h
@@ -16,6 +16,7 @@
 
 class UAssetUserData;
 class UNetConnection;
+class UNavigationSystemConfig;
 
 UENUM()
 enum EVisibilityAggressiveness
@@ -387,10 +388,15 @@ class ENGINE_API AWorldSettings : public AInfo, public IInterface_AssetUserData
 	UPROPERTY(EditAnywhere, BlueprintReadOnly, Category=World, AdvancedDisplay)
 	uint32 bEnableWorldBoundsChecks:1;
 
-	/** if set to false navigation system will not get created (and all navigation functionality won't be accessible)*/
-	UPROPERTY(EditAnywhere, BlueprintReadOnly, config, Category=World, AdvancedDisplay)
+protected:
+	/** if set to false navigation system will not get created (and all 
+	 *	navigation functionality won't be accessible).
+	 *	This flag is now deprecated. Use NavigationSystemConfig property to 
+	 *	determine navigation system's properties or existence */
+	UPROPERTY(VisibleAnywhere, BlueprintReadOnly, config, Category=World, AdvancedDisplay, meta=(DisplayName = "DEPRECATED_bEnableNavigationSystem"))
 	uint32 bEnableNavigationSystem:1;
 
+public:
 	/** if set to false AI system will not get created. Use it to disable all AI-related activity on a map */
 	UPROPERTY(EditAnywhere, BlueprintReadOnly, config, Category = World, AdvancedDisplay)
 	uint32 bEnableAISystem:1;
@@ -421,6 +427,11 @@ class ENGINE_API AWorldSettings : public AInfo, public IInterface_AssetUserData
 	UPROPERTY(EditAnywhere, BlueprintReadOnly, meta=(DisplayName = "Override World Gravity"), Category = Physics)
 	uint32 bGlobalGravitySet:1;
 
+	/** Holds parameters for NavigationSystem's creation. Set to Null will result
+	 *	in NavigationSystem instance not being created for this world */
+	UPROPERTY(EditAnywhere, BlueprintReadOnly, Category=World, AdvancedDisplay, Instanced, NoClear, meta=(NoResetToDefault))
+	UNavigationSystemConfig* NavigationSystemConfig;
+
 	// any actor falling below this level gets destroyed
 	UPROPERTY(EditAnywhere, BlueprintReadOnly, Category=World, meta=(editcondition = "bEnableWorldBoundsChecks"))
 	float KillZ;
@@ -675,6 +686,7 @@ public:
 #if WITH_EDITOR
 	virtual void CheckForErrors() override;
 #endif // WITH_EDITOR
+	virtual void PostInitProperties() override;
 	virtual void PreInitializeComponents() override;
 	virtual void PostInitializeComponents() override;
 	virtual void PostRegisterAllComponents() override;
@@ -701,6 +713,13 @@ public:
 	/** Sets the global time dilation value (subject to clamping). Returns the final value that was set. */
 	virtual float SetTimeDilation(float NewTimeDilation);
 
+	/** @return configuration for NavigationSystem's creation. Null means 
+	 *	no navigation system will be created*/
+	UNavigationSystemConfig* const GetNavigationSystemConfig() const { return NavigationSystemConfig; }
+
+	/** @return whether given world is configured to host any NavigationSystem */
+	bool IsNavigationSystemEnabled() const;
+
 	/**
 	 * Called from GameStateBase, calls BeginPlay on all actors
 	 */
diff --git a/Engine/Source/Runtime/Engine/Classes/Interfaces/Interface_CollisionDataProvider.h b/Engine/Source/Runtime/Engine/Classes/Interfaces/Interface_CollisionDataProvider.h
--- a/Engine/Source/Runtime/Engine/Classes/Interfaces/Interface_CollisionDataProvider.h
+++ b/Engine/Source/Runtime/Engine/Classes/Interfaces/Interface_CollisionDataProvider.h
@@ -60,7 +60,7 @@ struct FTriMeshCollisionData
 	}
 };
 
-UINTERFACE(meta=(CannotImplementInterfaceInBlueprint))
+UINTERFACE(MinimalAPI, meta=(CannotImplementInterfaceInBlueprint))
 class UInterface_CollisionDataProvider : public UInterface
 {
 	GENERATED_UINTERFACE_BODY()
diff --git a/Engine/Source/Runtime/Engine/Classes/Kismet/KismetArrayLibrary.h b/Engine/Source/Runtime/Engine/Classes/Kismet/KismetArrayLibrary.h
--- a/Engine/Source/Runtime/Engine/Classes/Kismet/KismetArrayLibrary.h
+++ b/Engine/Source/Runtime/Engine/Classes/Kismet/KismetArrayLibrary.h
@@ -1,4 +1,4 @@
-﻿// Copyright 1998-2018 Epic Games, Inc. All Rights Reserved.
+// Copyright 1998-2018 Epic Games, Inc. All Rights Reserved.
 
 #pragma once
 
@@ -249,11 +249,10 @@ public:
  
  		Stack.MostRecentPropertyAddress = NULL;
  		Stack.StepCompiledIn<UProperty>(StorageSpace);
-		void* NewItemPtr = (Stack.MostRecentPropertyAddress != NULL && Stack.MostRecentProperty->GetClass() == InnerProp->GetClass()) ? Stack.MostRecentPropertyAddress : StorageSpace;
  
  		P_FINISH;
 		P_NATIVE_BEGIN;
-		*(int32*)RESULT_PARAM = GenericArray_Add(ArrayAddr, ArrayProperty, NewItemPtr);
+		*(int32*)RESULT_PARAM = GenericArray_Add(ArrayAddr, ArrayProperty, StorageSpace);
 		P_NATIVE_END;
 		InnerProp->DestroyValue(StorageSpace);
 	}
@@ -278,11 +277,10 @@ public:
 
 		Stack.MostRecentPropertyAddress = NULL;
 		Stack.StepCompiledIn<UProperty>(StorageSpace);
-		void* NewItemPtr = (Stack.MostRecentPropertyAddress != NULL && Stack.MostRecentProperty->GetClass() == InnerProp->GetClass()) ? Stack.MostRecentPropertyAddress : StorageSpace;
 
 		P_FINISH;
 		P_NATIVE_BEGIN;
-		*(int32*)RESULT_PARAM = GenericArray_AddUnique(ArrayAddr, ArrayProperty, NewItemPtr);
+		*(int32*)RESULT_PARAM = GenericArray_AddUnique(ArrayAddr, ArrayProperty, StorageSpace);
 		P_NATIVE_END;
 		InnerProp->DestroyValue(StorageSpace);
 	}
@@ -354,12 +352,11 @@ public:
 
 		Stack.MostRecentPropertyAddress = NULL;
 		Stack.StepCompiledIn<UProperty>(StorageSpace);
-		void* NewItemPtr = (Stack.MostRecentPropertyAddress != NULL && Stack.MostRecentProperty->GetClass() == InnerProp->GetClass()) ? Stack.MostRecentPropertyAddress : StorageSpace;
 
 		P_GET_PROPERTY(UIntProperty, Index);
 		P_FINISH;
 		P_NATIVE_BEGIN;
-		GenericArray_Insert(ArrayAddr, ArrayProperty, NewItemPtr, Index);
+		GenericArray_Insert(ArrayAddr, ArrayProperty, StorageSpace, Index);
 		P_NATIVE_END;
 		InnerProp->DestroyValue(StorageSpace);
 	}
@@ -535,14 +532,13 @@ public:
 
 		Stack.MostRecentPropertyAddress = NULL;
 		Stack.StepCompiledIn<UProperty>(StorageSpace);
-		void* NewItemPtr = (Stack.MostRecentPropertyAddress != NULL && Stack.MostRecentProperty->GetClass() == InnerProp->GetClass()) ? Stack.MostRecentPropertyAddress : StorageSpace;
 
 		P_GET_UBOOL(bSizeToFit);
 
 		P_FINISH;
 
 		P_NATIVE_BEGIN;
-		GenericArray_Set(ArrayAddr, ArrayProperty, Index, NewItemPtr, bSizeToFit);
+		GenericArray_Set(ArrayAddr, ArrayProperty, Index, StorageSpace, bSizeToFit);
 		P_NATIVE_END;
 		InnerProp->DestroyValue(StorageSpace);
 	}
diff --git a/Engine/Source/Runtime/Engine/Classes/Kismet/KismetMathLibrary.h b/Engine/Source/Runtime/Engine/Classes/Kismet/KismetMathLibrary.h
--- a/Engine/Source/Runtime/Engine/Classes/Kismet/KismetMathLibrary.h
+++ b/Engine/Source/Runtime/Engine/Classes/Kismet/KismetMathLibrary.h
@@ -313,7 +313,7 @@ class ENGINE_API UKismetMathLibrary : public UBlueprintFunctionLibrary
 	static int32 Not_Int(int32 A);
 
 	/** Sign (integer, returns -1 if A < 0, 0 if A is zero, and +1 if A > 0) */
-	UFUNCTION(BlueprintPure, meta=(DisplayName = "Sign (int)"), Category="Math|Integer")
+	UFUNCTION(BlueprintPure, meta=(DisplayName = "Sign (integer)"), Category="Math|Integer")
 	static int32 SignOfInteger(int32 A);
 
 	/** Returns a uniformly distributed random number between 0 and Max - 1 */
@@ -325,19 +325,19 @@ class ENGINE_API UKismetMathLibrary : public UBlueprintFunctionLibrary
 	static int32 RandomIntegerInRange(int32 Min, int32 Max);
 
 	/** Returns the minimum value of A and B */
-	UFUNCTION(BlueprintPure, meta=(DisplayName = "Min (int)", CompactNodeTitle = "MIN", CommutativeAssociativeBinaryOperator = "true"), Category="Math|Integer")
+	UFUNCTION(BlueprintPure, meta=(DisplayName = "Min (integer)", CompactNodeTitle = "MIN", CommutativeAssociativeBinaryOperator = "true"), Category="Math|Integer")
 	static int32 Min(int32 A, int32 B);
 
 	/** Returns the maximum value of A and B */
-	UFUNCTION(BlueprintPure, meta=(DisplayName = "Max (int)", CompactNodeTitle = "MAX", CommutativeAssociativeBinaryOperator = "true"), Category="Math|Integer")
+	UFUNCTION(BlueprintPure, meta=(DisplayName = "Max (integer)", CompactNodeTitle = "MAX", CommutativeAssociativeBinaryOperator = "true"), Category="Math|Integer")
 	static int32 Max(int32 A, int32 B);
 
 	/** Returns Value clamped to be between A and B (inclusive) */
-	UFUNCTION(BlueprintPure, meta=(DisplayName = "Clamp (int)"), Category="Math|Integer")
+	UFUNCTION(BlueprintPure, meta=(DisplayName = "Clamp (integer)"), Category="Math|Integer")
 	static int32 Clamp(int32 Value, int32 Min, int32 Max);
 
 	/** Returns the absolute (positive) value of A */
-	UFUNCTION(BlueprintPure, meta=(DisplayName = "Absolute (int)", CompactNodeTitle = "ABS"), Category="Math|Integer")
+	UFUNCTION(BlueprintPure, meta=(DisplayName = "Absolute (integer)", CompactNodeTitle = "ABS"), Category="Math|Integer")
 	static int32 Abs_Int(int32 A);
 
 	//
@@ -1016,7 +1016,7 @@ class ENGINE_API UKismetMathLibrary : public UBlueprintFunctionLibrary
 	static FRotator Multiply_RotatorFloat(FRotator A, float B);
 	
 	/** Returns rotator representing rotator A scaled by B */
-	UFUNCTION(BlueprintPure, meta=(DisplayName = "ScaleRotator (int)", CompactNodeTitle = "*", Keywords = "* multiply rotate rotation"), Category="Math|Rotator")
+	UFUNCTION(BlueprintPure, meta=(DisplayName = "ScaleRotator (integer)", CompactNodeTitle = "*", Keywords = "* multiply rotate rotation"), Category="Math|Rotator")
 	static FRotator Multiply_RotatorInt(FRotator A, int32 B);
 
 	/** Combine 2 rotations to give you the resulting rotation of first applying A, then B. */
@@ -1403,19 +1403,19 @@ class ENGINE_API UKismetMathLibrary : public UBlueprintFunctionLibrary
 	static float Conv_ByteToFloat(uint8 InByte);
 
 	/** Converts an integer to a float */
-	UFUNCTION(BlueprintPure, meta=(DisplayName = "ToFloat (int)", CompactNodeTitle = "->", Keywords="cast convert", BlueprintAutocast), Category="Math|Conversions")
+	UFUNCTION(BlueprintPure, meta=(DisplayName = "ToFloat (integer)", CompactNodeTitle = "->", Keywords="cast convert", BlueprintAutocast), Category="Math|Conversions")
 	static float Conv_IntToFloat(int32 InInt);
 
 	/** Converts an integer to a byte (if the integer is too large, returns the low 8 bits) */
-	UFUNCTION(BlueprintPure, meta=(DisplayName = "ToByte (int)", CompactNodeTitle = "->", Keywords="cast convert", BlueprintAutocast), Category="Math|Conversions")
+	UFUNCTION(BlueprintPure, meta=(DisplayName = "ToByte (integer)", CompactNodeTitle = "->", Keywords="cast convert", BlueprintAutocast), Category="Math|Conversions")
 	static uint8 Conv_IntToByte(int32 InInt);
 
 	/** Converts an integer to an IntVector*/
-	UFUNCTION(BlueprintPure, meta = (DisplayName = "ToIntVector (int)", CompactNodeTitle = "->", Keywords = "cast convert", BlueprintAutocast), Category = "Math|Conversions")
+	UFUNCTION(BlueprintPure, meta = (DisplayName = "ToIntVector (integer)", CompactNodeTitle = "->", Keywords = "cast convert", BlueprintAutocast), Category = "Math|Conversions")
 	static FIntVector Conv_IntToIntVector(int32 InInt);
 
 	/** Converts a int to a bool*/
-	UFUNCTION(BlueprintPure, meta=(DisplayName = "ToBool (int)", CompactNodeTitle = "->", Keywords="cast convert", BlueprintAutocast), Category="Math|Conversions")
+	UFUNCTION(BlueprintPure, meta=(DisplayName = "ToBool (integer)", CompactNodeTitle = "->", Keywords="cast convert", BlueprintAutocast), Category="Math|Conversions")
 	static bool Conv_IntToBool(int32 InInt);
 
 	/** Converts a bool to an int */
diff --git a/Engine/Source/Runtime/Engine/Classes/Kismet/KismetStringLibrary.h b/Engine/Source/Runtime/Engine/Classes/Kismet/KismetStringLibrary.h
--- a/Engine/Source/Runtime/Engine/Classes/Kismet/KismetStringLibrary.h
+++ b/Engine/Source/Runtime/Engine/Classes/Kismet/KismetStringLibrary.h
@@ -17,7 +17,7 @@ class ENGINE_API UKismetStringLibrary : public UBlueprintFunctionLibrary
 	static FString Conv_FloatToString(float InFloat);
 
 	/** Converts an integer value to a string */
-	UFUNCTION(BlueprintPure, meta=(DisplayName = "ToString (int)", CompactNodeTitle = "->", BlueprintAutocast), Category="Utilities|String")
+	UFUNCTION(BlueprintPure, meta=(DisplayName = "ToString (integer)", CompactNodeTitle = "->", BlueprintAutocast), Category="Utilities|String")
 	static FString Conv_IntToString(int32 InInt);
 
 	/** Converts a byte value to a string */
@@ -107,7 +107,7 @@ class ENGINE_API UKismetStringLibrary : public UBlueprintFunctionLibrary
 	 * @param Suffix - A suffix to append to the end of the conversion string
 	 * @return A new string built from the passed parameters
 	 */
-	UFUNCTION(BlueprintPure, meta=(DisplayName = "BuildString (int)"), Category="Utilities|String")
+	UFUNCTION(BlueprintPure, meta=(DisplayName = "BuildString (integer)"), Category="Utilities|String")
 	static FString BuildString_Int(const FString& AppendTo, const FString& Prefix, int32 InInt, const FString& Suffix);
 
 	/** 
diff --git a/Engine/Source/Runtime/Engine/Classes/Kismet/KismetTextLibrary.h b/Engine/Source/Runtime/Engine/Classes/Kismet/KismetTextLibrary.h
--- a/Engine/Source/Runtime/Engine/Classes/Kismet/KismetTextLibrary.h
+++ b/Engine/Source/Runtime/Engine/Classes/Kismet/KismetTextLibrary.h
@@ -193,7 +193,7 @@ class ENGINE_API UKismetTextLibrary : public UBlueprintFunctionLibrary
 
 	// Default values are duplicated from FNumberFormattingOptions and should be replicated in all functions and in the struct when changed!
 	/* Converts a passed in integer to text based on formatting options */
-	UFUNCTION(BlueprintPure, meta=(DisplayName = "ToText (int)", AdvancedDisplay = "1", BlueprintAutocast), Category="Utilities|Text")
+	UFUNCTION(BlueprintPure, meta=(DisplayName = "ToText (integer)", AdvancedDisplay = "1", BlueprintAutocast), Category="Utilities|Text")
 	static FText Conv_IntToText(int32 Value, bool bAlwaysSign = false, bool bUseGrouping = true, int32 MinimumIntegralDigits = 1, int32 MaximumIntegralDigits = 324);
 
 	// Default values are duplicated from FNumberFormattingOptions and should be replicated in all functions and in the struct when changed!
@@ -212,7 +212,7 @@ class ENGINE_API UKismetTextLibrary : public UBlueprintFunctionLibrary
 
 	// Default values are duplicated from FNumberFormattingOptions and should be replicated in all functions and in the struct when changed!
 	/* Converts a passed in integer to a text formatted as a currency */
-	UFUNCTION(BlueprintPure, meta=(DisplayName = "AsCurrency (int) - DEPRECATED (use AsCurrency)"), Category="Utilities|Text")
+	UFUNCTION(BlueprintPure, meta=(DisplayName = "AsCurrency (integer) - DEPRECATED (use AsCurrency)"), Category="Utilities|Text")
 	static FText AsCurrency_Integer(int32 Value, TEnumAsByte<ERoundingMode> RoundingMode, bool bAlwaysSign = false, bool bUseGrouping = true, int32 MinimumIntegralDigits = 1, int32 MaximumIntegralDigits = 324, int32 MinimumFractionalDigits = 0, int32 MaximumFractionalDigits = 3, const FString& CurrencyCode = TEXT(""));
 
 	// Default values are duplicated from FNumberFormattingOptions and should be replicated in all functions and in the struct when changed!
diff --git a/Engine/Source/Runtime/Engine/Engine.Build.cs b/Engine/Source/Runtime/Engine/Engine.Build.cs
--- a/Engine/Source/Runtime/Engine/Engine.Build.cs
+++ b/Engine/Source/Runtime/Engine/Engine.Build.cs
@@ -76,7 +76,6 @@ public class Engine : ModuleRules
 				"EngineSettings",
 				"SynthBenchmark",
 				"GameplayTags",
-				"AIModule",
 				"DatabaseSupport",
 				"PacketHandler",
 				"HardwareSurvey",
@@ -162,7 +161,6 @@ public class Engine : ModuleRules
 		}
 
 		CircularlyReferencedDependentModules.Add("GameplayTags");
-		CircularlyReferencedDependentModules.Add("AIModule");
 		CircularlyReferencedDependentModules.Add("Landscape");
 		CircularlyReferencedDependentModules.Add("UMG");
 		CircularlyReferencedDependentModules.Add("MaterialShaderQualitySettings");
@@ -415,19 +413,6 @@ public class Engine : ModuleRules
 				"libOpus"
 				);
 		}
-
-		if (Target.bCompileRecast)
-		{
-			PrivateDependencyModuleNames.Add("Navmesh");
-			PublicDefinitions.Add("WITH_RECAST=1");
-		}
-		else
-		{
-			// Because we test WITH_RECAST in public Engine header files, we need to make sure that modules
-			// that import this also have this definition set appropriately.  Recast is a private dependency
-			// module, so it's definitions won't propagate to modules that import Engine.
-			PublicDefinitions.Add("WITH_RECAST=0");
-		}
 /*
 		ConfigHierarchy Ini = ConfigCache.ReadHierarchy(ConfigHierarchyType.Engine, System.IO.DirectoryReference.FromFile(Target.ProjectFile), Target.Platform);
 		bool bLocalVectorFieldOnly = false;
diff --git a/Engine/Source/Runtime/Engine/Private/AI/Navigation/NavigationInterfaces.cpp b/Engine/Source/Runtime/Engine/Private/AI/Navigation/NavigationInterfaces.cpp
--- a/Engine/Source/Runtime/Engine/Private/AI/Navigation/NavigationInterfaces.cpp
+++ b/Engine/Source/Runtime/Engine/Private/AI/Navigation/NavigationInterfaces.cpp
@@ -5,16 +5,11 @@
 #include "Templates/Casts.h"
 #include "AI/Navigation/NavAgentInterface.h"
 #include "AI/Navigation/NavRelevantInterface.h"
-#include "AI/Navigation/NavigationPathGenerator.h"
-#include "AI/Navigation/NavNodeInterface.h"
-#include "AI/Navigation/NavLinkDefinition.h"
-#include "AI/Navigation/NavLinkHostInterface.h"
 #include "AI/Navigation/NavPathObserverInterface.h"
-#include "AI/Navigation/NavLinkCustomInterface.h"
 #include "AI/Navigation/NavEdgeProviderInterface.h"
 #include "AI/RVOAvoidanceInterface.h"
-
-uint32 INavLinkCustomInterface::NextUniqueId = 1;
+#include "AI/Navigation/NavigationDataInterface.h"
+#include "AI/Navigation/PathFollowingAgentInterface.h"
 
 URVOAvoidanceInterface::URVOAvoidanceInterface(const FObjectInitializer& ObjectInitializer)
 	: Super(ObjectInitializer)
@@ -26,65 +21,27 @@ UNavAgentInterface::UNavAgentInterface(const FObjectInitializer& ObjectInitializ
 {
 }
 
-UNavigationPathGenerator::UNavigationPathGenerator(const FObjectInitializer& ObjectInitializer)
-	: Super(ObjectInitializer)
-{
-}
-
-UNavNodeInterface::UNavNodeInterface(const FObjectInitializer& ObjectInitializer)
+UNavRelevantInterface::UNavRelevantInterface(const FObjectInitializer& ObjectInitializer)
 	: Super(ObjectInitializer)
 {
 }
 
-UNavLinkHostInterface::UNavLinkHostInterface(const FObjectInitializer& ObjectInitializer)
+UNavEdgeProviderInterface::UNavEdgeProviderInterface(const FObjectInitializer& ObjectInitializer) 
 	: Super(ObjectInitializer)
 {
 }
 
-UNavLinkCustomInterface::UNavLinkCustomInterface(const FObjectInitializer& ObjectInitializer)
-: Super(ObjectInitializer)
-{
-}
-
-UObject* INavLinkCustomInterface::GetLinkOwner() const
-{
-	return Cast<UObject>((INavLinkCustomInterface*)this);
-}
-
-uint32 INavLinkCustomInterface::GetUniqueId()
-{
-	return NextUniqueId++;
-}
-
-void INavLinkCustomInterface::UpdateUniqueId(uint32 AlreadyUsedId)
-{
-	NextUniqueId = FMath::Max(NextUniqueId, AlreadyUsedId + 1);
-}
-
-FNavigationLink INavLinkCustomInterface::GetModifier(const INavLinkCustomInterface* CustomNavLink)
-{
-	FNavigationLink LinkMod;
-	LinkMod.SetAreaClass(CustomNavLink->GetLinkAreaClass());
-	LinkMod.UserId = CustomNavLink->GetLinkId();
-
-	ENavLinkDirection::Type LinkDirection = ENavLinkDirection::BothWays;
-	CustomNavLink->GetLinkData(LinkMod.Left, LinkMod.Right, LinkDirection);
-	LinkMod.Direction = LinkDirection;
-
-	return LinkMod;
-}
-
 UNavPathObserverInterface::UNavPathObserverInterface(const FObjectInitializer& ObjectInitializer)
 	: Super(ObjectInitializer)
 {
 }
 
-UNavRelevantInterface::UNavRelevantInterface(const FObjectInitializer& ObjectInitializer)
+UNavigationDataInterface::UNavigationDataInterface(const FObjectInitializer& ObjectInitializer)
 	: Super(ObjectInitializer)
 {
 }
 
-UNavEdgeProviderInterface::UNavEdgeProviderInterface(const FObjectInitializer& ObjectInitializer) : Super(ObjectInitializer)
+UPathFollowingAgentInterface::UPathFollowingAgentInterface(const FObjectInitializer& ObjectInitializer)
+	: Super(ObjectInitializer)
 {
 }
-
diff --git a/Engine/Source/Runtime/Engine/Private/AI/Navigation/NavigationModifier.cpp b/Engine/Source/Runtime/Engine/Private/AI/Navigation/NavigationModifier.cpp
--- a/Engine/Source/Runtime/Engine/Private/AI/Navigation/NavigationModifier.cpp
+++ b/Engine/Source/Runtime/Engine/Private/AI/Navigation/NavigationModifier.cpp
@@ -5,12 +5,9 @@
 #include "EngineStats.h"
 #include "GameFramework/Actor.h"
 #include "Components/BrushComponent.h"
-#include "AI/Navigation/RecastHelpers.h"
-#include "AI/Navigation/NavLinkTrivial.h"
-#include "AI/Navigation/NavAreas/NavAreaMeta.h"
 #include "PhysicsEngine/BodySetup.h"
-#include "AI/Navigation/NavigationSystem.h"
-#include "AI/Navigation/NavAreas/NavArea_LowHeight.h"
+#include "AI/NavigationSystemBase.h"
+#include "AI/Navigation/NavAreaBase.h"
 
 // if square distance between two points is less than this the those points
 // will be considered identical when calculating convex hull
@@ -42,7 +39,7 @@ void FNavigationLinkBase::SetAreaClass(UClass* InAreaClass)
 UClass* FNavigationLinkBase::GetAreaClass() const
 {
 	UClass* ClassOb = AreaClassOb.Get();
-	return ClassOb ? ClassOb : (UClass*)UNavigationSystem::GetDefaultWalkableArea();
+	return ClassOb ? ClassOb : *FNavigationSystem::GetDefaultWalkableArea();
 }
 
 void FNavigationLinkBase::InitializeAreaClass(const bool bForceRefresh)
@@ -52,8 +49,7 @@ void FNavigationLinkBase::InitializeAreaClass(const bool bForceRefresh)
 
 bool FNavigationLinkBase::HasMetaArea() const
 {
-	UClass* ClassOb = GetAreaClass();
-	return ClassOb && ClassOb->IsChildOf(UNavAreaMeta::StaticClass());
+	return AreaClass && AreaClass->GetDefaultObject<UNavAreaBase>()->IsMetaArea();
 }
 
 void FNavigationLinkBase::PostSerialize(const FArchive& Ar)
@@ -279,7 +275,7 @@ bool UNavLinkDefinition::HasAdjustableLinks() const
 // FAreaNavModifier
 //----------------------------------------------------------------------//
 
-FAreaNavModifier::FAreaNavModifier(float Radius, float Height, const FTransform& LocalToWorld, const TSubclassOf<UNavArea> InAreaClass)
+FAreaNavModifier::FAreaNavModifier(float Radius, float Height, const FTransform& LocalToWorld, const TSubclassOf<UNavAreaBase> InAreaClass)
 {
 	Init(InAreaClass);
 	
@@ -296,25 +292,25 @@ FAreaNavModifier::FAreaNavModifier(float Radius, float Height, const FTransform&
 	Bounds = FBox::BuildAABB(LocalToWorld.GetLocation(), FVector(Radius, Radius, Height));
 }
 
-FAreaNavModifier::FAreaNavModifier(const FVector& Extent, const FTransform& LocalToWorld, const TSubclassOf<UNavArea> InAreaClass)
+FAreaNavModifier::FAreaNavModifier(const FVector& Extent, const FTransform& LocalToWorld, const TSubclassOf<UNavAreaBase> InAreaClass)
 {
 	Init(InAreaClass);
 	SetBox(FBox::BuildAABB(FVector::ZeroVector, Extent), LocalToWorld);
 }
 
-FAreaNavModifier::FAreaNavModifier(const FBox& Box, const FTransform& LocalToWorld, const TSubclassOf<UNavArea> InAreaClass)
+FAreaNavModifier::FAreaNavModifier(const FBox& Box, const FTransform& LocalToWorld, const TSubclassOf<UNavAreaBase> InAreaClass)
 {
 	Init(InAreaClass);
 	SetBox(Box, LocalToWorld);
 }
 
-FAreaNavModifier::FAreaNavModifier(const TArray<FVector>& InPoints, ENavigationCoordSystem::Type CoordType, const FTransform& LocalToWorld, const TSubclassOf<UNavArea> InAreaClass)
+FAreaNavModifier::FAreaNavModifier(const TArray<FVector>& InPoints, ENavigationCoordSystem::Type CoordType, const FTransform& LocalToWorld, const TSubclassOf<UNavAreaBase> InAreaClass)
 {
 	Init(InAreaClass);
 	SetConvex(InPoints.GetData(), 0, InPoints.Num(), CoordType, LocalToWorld);
 }
 
-FAreaNavModifier::FAreaNavModifier(const TArray<FVector>& InPoints, const int32 FirstIndex, const int32 LastIndex, ENavigationCoordSystem::Type CoordType, const FTransform& LocalToWorld, const TSubclassOf<UNavArea> InAreaClass)
+FAreaNavModifier::FAreaNavModifier(const TArray<FVector>& InPoints, const int32 FirstIndex, const int32 LastIndex, ENavigationCoordSystem::Type CoordType, const FTransform& LocalToWorld, const TSubclassOf<UNavAreaBase> InAreaClass)
 {
 	check(InPoints.IsValidIndex(FirstIndex) && InPoints.IsValidIndex(LastIndex-1));
 
@@ -322,7 +318,7 @@ FAreaNavModifier::FAreaNavModifier(const TArray<FVector>& InPoints, const int32
 	SetConvex(InPoints.GetData(), FirstIndex, LastIndex, CoordType, LocalToWorld);
 }
 
-FAreaNavModifier::FAreaNavModifier(const TNavStatArray<FVector>& InPoints, const int32 FirstIndex, const int32 LastIndex, ENavigationCoordSystem::Type CoordType, const FTransform& LocalToWorld, const TSubclassOf<UNavArea> InAreaClass)
+FAreaNavModifier::FAreaNavModifier(const TNavStatArray<FVector>& InPoints, const int32 FirstIndex, const int32 LastIndex, ENavigationCoordSystem::Type CoordType, const FTransform& LocalToWorld, const TSubclassOf<UNavAreaBase> InAreaClass)
 {
 	check(InPoints.IsValidIndex(FirstIndex) && InPoints.IsValidIndex(LastIndex-1));
 
@@ -330,7 +326,7 @@ FAreaNavModifier::FAreaNavModifier(const TNavStatArray<FVector>& InPoints, const
 	SetConvex(InPoints.GetData(), FirstIndex, LastIndex, CoordType, LocalToWorld);
 }
 
-FAreaNavModifier::FAreaNavModifier(const UBrushComponent* BrushComponent, const TSubclassOf<UNavArea> InAreaClass)
+FAreaNavModifier::FAreaNavModifier(const UBrushComponent* BrushComponent, const TSubclassOf<UNavAreaBase> InAreaClass)
 {
 	check(BrushComponent != NULL);
 
@@ -375,7 +371,7 @@ void FAreaNavModifier::GetConvex(FConvexNavAreaData& Data) const
 	Data.MaxZ = LastPoint.Y;
 }
 
-void FAreaNavModifier::Init(const TSubclassOf<UNavArea> InAreaClass)
+void FAreaNavModifier::Init(const TSubclassOf<UNavAreaBase> InAreaClass)
 {
 	bIncludeAgentHeight = false;
 	ApplyMode = ENavigationAreaMode::Apply;
@@ -385,24 +381,32 @@ void FAreaNavModifier::Init(const TSubclassOf<UNavArea> InAreaClass)
 	SetAreaClass(InAreaClass);
 }
 
-void FAreaNavModifier::SetAreaClass(const TSubclassOf<UNavArea> InAreaClass)
+bool IsMetaAreaClass(UClass& AreaClass)
+{
+	const UNavAreaBase* AreaClassCDO = GetDefault<UNavAreaBase>(&AreaClass);
+	return AreaClassCDO && AreaClassCDO->IsMetaArea();
+}
+
+void FAreaNavModifier::SetAreaClass(const TSubclassOf<UNavAreaBase> InAreaClass)
 {
 	AreaClassOb = (UClass*)InAreaClass;
 
-	const UClass* AreaClass1 = AreaClassOb.Get();
-	const UClass* AreaClass2 = ReplaceAreaClassOb.Get();
-	bHasMetaAreas = (AreaClass1 && AreaClass1->IsChildOf(UNavAreaMeta::StaticClass())) || (AreaClass2 && AreaClass2->IsChildOf(UNavAreaMeta::StaticClass()));
+	UClass* AreaClass1 = AreaClassOb.Get();
+	UClass* AreaClass2 = ReplaceAreaClassOb.Get();
+	bHasMetaAreas = (AreaClass1 && IsMetaAreaClass(*AreaClass1))
+		|| (AreaClass2 && IsMetaAreaClass(*AreaClass2));
 }
 
-void FAreaNavModifier::SetAreaClassToReplace(const TSubclassOf<UNavArea> InAreaClass)
+void FAreaNavModifier::SetAreaClassToReplace(const TSubclassOf<UNavAreaBase> InAreaClass)
 {
 	ReplaceAreaClassOb = (UClass*)InAreaClass;
 
-	const UClass* AreaClass1 = AreaClassOb.Get();
-	const UClass* AreaClass2 = ReplaceAreaClassOb.Get();
-	bHasMetaAreas = (AreaClass1 && AreaClass1->IsChildOf(UNavAreaMeta::StaticClass())) || (AreaClass2 && AreaClass2->IsChildOf(UNavAreaMeta::StaticClass()));
+	UClass* AreaClass1 = AreaClassOb.Get();
+	UClass* AreaClass2 = ReplaceAreaClassOb.Get();
+	bHasMetaAreas = (AreaClass1 && IsMetaAreaClass(*AreaClass1))
+		|| (AreaClass2 && IsMetaAreaClass(*AreaClass2));
 
-	bIsLowAreaModifier = (AreaClass2 == UNavArea_LowHeight::StaticClass());
+	bIsLowAreaModifier = (AreaClass2 && AreaClass2->GetDefaultObject<UNavAreaBase>()->IsLowArea());
 	ApplyMode = bIsLowAreaModifier ? ENavigationAreaMode::ReplaceInLowPass : AreaClass2 ? ENavigationAreaMode::Replace : ENavigationAreaMode::Apply;
 }
 
@@ -464,9 +468,12 @@ void FAreaNavModifier::SetConvex(const FVector* InPoints, const int32 FirstIndex
 	TArray<FVector, TInlineAllocator<MaxConvexPoints>> HullVertices;
 	HullVertices.Empty(MaxConvexPoints);
 
+	const FTransform& NavCoordTransform = FNavigationSystem::GetCoordTransformTo(CoordType);
+
 	for (int32 i = FirstIndex; i < LastIndex; i++)
 	{
-		const FVector Point = (CoordType == ENavigationCoordSystem::Recast) ? Recast2UnrealPoint(InPoints[i]) : InPoints[i];
+		//const FVector Point = (CoordType == ENavigationCoordSystem::Recast) ? Recast2UnrealPoint(InPoints[i]) : InPoints[i];
+		const FVector Point = NavCoordTransform.TransformPosition(InPoints[i]);
 
 		FVector TransformedPoint = LocalToWorld.TransformPosition(Point);
 		ConvexData.MinZ = FMath::Min( ConvexData.MinZ, TransformedPoint.Z );
@@ -705,8 +712,8 @@ FCompositeNavModifier FCompositeNavModifier::GetInstantiatedMetaModifier(const F
 		{
 			if (Area->HasMetaAreas())
 			{
-				Area->SetAreaClass(UNavAreaMeta::PickAreaClass(Area->GetAreaClass(), ActorOwner, *NavAgent));
-				Area->SetAreaClassToReplace(UNavAreaMeta::PickAreaClass(Area->GetAreaClassToReplace(), ActorOwner, *NavAgent));
+				Area->SetAreaClass(UNavAreaBase::PickAreaClassForAgent(Area->GetAreaClass(), *ActorOwner, *NavAgent));
+				Area->SetAreaClassToReplace(UNavAreaBase::PickAreaClassForAgent(Area->GetAreaClassToReplace(), *ActorOwner, *NavAgent));
 			}
 		}
 	}
@@ -720,13 +727,13 @@ FCompositeNavModifier FCompositeNavModifier::GetInstantiatedMetaModifier(const F
 				for (int32 LinkIndex = 0; LinkIndex < SimpleLink->Links.Num(); ++LinkIndex)
 				{
 					FNavigationLink& Link = SimpleLink->Links[LinkIndex];
-					Link.SetAreaClass(UNavAreaMeta::PickAreaClass(Link.GetAreaClass(), ActorOwner, *NavAgent));
+					Link.SetAreaClass(UNavAreaBase::PickAreaClassForAgent(Link.GetAreaClass(), *ActorOwner, *NavAgent));
 				}
 
 				for (int32 LinkIndex = 0; LinkIndex < SimpleLink->SegmentLinks.Num(); ++LinkIndex)
 				{
 					FNavigationSegmentLink& Link = SimpleLink->SegmentLinks[LinkIndex];
-					Link.SetAreaClass(UNavAreaMeta::PickAreaClass(Link.GetAreaClass(), ActorOwner, *NavAgent));
+					Link.SetAreaClass(UNavAreaBase::PickAreaClassForAgent(Link.GetAreaClass(), *ActorOwner, *NavAgent));
 				}
 			}
 		}
@@ -754,7 +761,7 @@ FCompositeNavModifier FCompositeNavModifier::GetInstantiatedMetaModifier(const F
 				{
 					const int32 AddedIdx = SimpleLink.Links.Add(Links[LinkIndex]);
 					FNavigationLink& NavLink = SimpleLink.Links[AddedIdx];
-					NavLink.SetAreaClass(UNavAreaMeta::PickAreaClass(NavLink.GetAreaClass(), ActorOwner, *NavAgent));
+					NavLink.SetAreaClass(UNavAreaBase::PickAreaClassForAgent(NavLink.GetAreaClass(), *ActorOwner, *NavAgent));
 				}				
 
 				const TArray<FNavigationSegmentLink>& SegmentLinks = UNavLinkDefinition::GetSegmentLinksDefinition(CustomLink->GetNavLinkClass());
@@ -769,7 +776,7 @@ FCompositeNavModifier FCompositeNavModifier::GetInstantiatedMetaModifier(const F
 				{
 					const int32 AddedIdx = SimpleSegLink.SegmentLinks.Add(SegmentLinks[LinkIndex]);
 					FNavigationSegmentLink& NavLink = SimpleSegLink.SegmentLinks[AddedIdx];
-					NavLink.SetAreaClass(UNavAreaMeta::PickAreaClass(NavLink.GetAreaClass(), ActorOwner, *NavAgent));
+					NavLink.SetAreaClass(UNavAreaBase::PickAreaClassForAgent(NavLink.GetAreaClass(), *ActorOwner, *NavAgent));
 				}
 
 				Result.CustomLinks.RemoveAtSwap(Index, 1, false);
@@ -780,7 +787,7 @@ FCompositeNavModifier FCompositeNavModifier::GetInstantiatedMetaModifier(const F
 	return Result;
 }
 
-void FCompositeNavModifier::CreateAreaModifiers(const UPrimitiveComponent* PrimComp, const TSubclassOf<UNavArea> AreaClass)
+void FCompositeNavModifier::CreateAreaModifiers(const UPrimitiveComponent* PrimComp, const TSubclassOf<UNavAreaBase> AreaClass)
 {
 	UBodySetup* BodySetup = PrimComp ? ((UPrimitiveComponent*)PrimComp)->GetBodySetup() : nullptr;
 	if (BodySetup == nullptr)
@@ -841,11 +848,3 @@ bool FCompositeNavModifier::HasPerInstanceTransforms() const
 {
 	return NavDataPerInstanceTransformDelegate.IsBound();
 }
-
-//----------------------------------------------------------------------//
-// UNavLinkTrivial
-//----------------------------------------------------------------------//
-UNavLinkTrivial::UNavLinkTrivial(const FObjectInitializer& ObjectInitializer) : Super(ObjectInitializer)
-{
-	FNavigationLink& Link = Links[Links.Add(FNavigationLink(FVector(0,100,0), FVector(0,-100,0)))];
-}
diff --git a/Engine/Source/Runtime/Engine/Private/Actor.cpp b/Engine/Source/Runtime/Engine/Private/Actor.cpp
--- a/Engine/Source/Runtime/Engine/Private/Actor.cpp
+++ b/Engine/Source/Runtime/Engine/Private/Actor.cpp
@@ -9,7 +9,7 @@
 #include "TimerManager.h"
 #include "GameFramework/Pawn.h"
 #include "Components/PrimitiveComponent.h"
-#include "AI/Navigation/NavigationSystem.h"
+#include "AI/NavigationSystemBase.h"
 #include "Components/InputComponent.h"
 #include "Engine/Engine.h"
 #include "UObject/UObjectHash.h"
@@ -752,8 +752,8 @@ void AActor::ApplyWorldOffset(const FVector& InOffset, bool bWorldShift)
 	{
 		if (RootComponent != nullptr && RootComponent->IsRegistered())
 		{
-			UNavigationSystem::UpdateNavOctreeBounds(this);
-			UNavigationSystem::UpdateActorAndComponentsInNavOctree(*this);
+			FNavigationSystem::OnActorBoundsChanged(*this);
+			FNavigationSystem::UpdateActorAndComponentData(*this);
 		}
 	}
 }
@@ -850,23 +850,41 @@ float AActor::GetActorTickInterval() const
 bool AActor::Rename( const TCHAR* InName, UObject* NewOuter, ERenameFlags Flags )
 {
 	const bool bRenameTest = ((Flags & REN_Test) != 0);
+	const bool bChangingOuters = (NewOuter && (NewOuter != GetOuter()));
 
-	if (!bRenameTest && NewOuter)
+	if (!bRenameTest && bChangingOuters)
 	{
 		RegisterAllActorTickFunctions(false, true); // unregister all tick functions
 		UnregisterAllComponents();
+
+		if (ULevel* MyLevel = GetLevel())
+		{
+			int32 ActorIndex;
+			if (MyLevel->Actors.Find(this, ActorIndex))
+			{
+				MyLevel->Actors[ActorIndex] = nullptr;
+				MyLevel->ActorsForGC.Remove(this);
+				// TODO: There may need to be some consideration about removing this actor from the level cluster, but that would probably require destroying the entire cluster, so defer for now
+			}
+		}
 	}
 
 	const bool bSuccess = Super::Rename( InName, NewOuter, Flags );
 
-	if (!bRenameTest && NewOuter && NewOuter->IsA<ULevel>())
+	if (!bRenameTest && bChangingOuters)
 	{
-		UWorld* World = NewOuter->GetWorld();
-		if (World && World->bIsWorldInitialized)
+		if (ULevel* MyLevel = GetLevel())
 		{
-			RegisterAllComponents();
+			MyLevel->Actors.Add(this);
+			MyLevel->ActorsForGC.Add(this);
+
+			UWorld* World = MyLevel->GetWorld();
+			if (World && World->bIsWorldInitialized)
+			{
+				RegisterAllComponents();
+			}
+			RegisterAllActorTickFunctions(true, true); // register all tick functions
 		}
-		RegisterAllActorTickFunctions(true, true); // register all tick functions
 	}
 	return bSuccess;
 }
@@ -1485,7 +1503,6 @@ void AActor::SetOwner( AActor *NewOwner )
 		}
 
 		// Sets this actor's parent to the specified actor.
-		AActor* OldOwner = Owner;
 		if( Owner != nullptr )
 		{
 			// remove from old owner's Children array
@@ -1979,7 +1996,7 @@ void AActor::RouteEndPlay(const EEndPlayReason::Type EndPlayReason)
 			SetLifeSpan(0.f);
 		}
 
-		UNavigationSystem::OnActorUnregistered(this);
+		FNavigationSystem::OnActorUnregistered(*this);
 	}
 
 	UninitializeComponents();
@@ -2135,6 +2152,7 @@ float AActor::TakeDamage(float DamageAmount, FDamageEvent const& DamageEvent, AC
 		{
 			FHitResult const& Hit = (RadialDamageEvent->ComponentHits.Num() > 0) ? RadialDamageEvent->ComponentHits[0] : FHitResult();
 			ReceiveRadialDamage(ActualDamage, DamageTypeCDO, RadialDamageEvent->Origin, Hit, EventInstigator, DamageCauser);
+			OnTakeRadialDamage.Broadcast(this, ActualDamage, DamageTypeCDO, RadialDamageEvent->Origin, Hit, EventInstigator, DamageCauser);
 
 			// add any desired physics impulses to our components
 			for (int HitIdx = 0; HitIdx < RadialDamageEvent->ComponentHits.Num(); ++HitIdx)
@@ -4406,7 +4424,7 @@ void AActor::PostInitializeComponents()
 	{
 		bActorInitialized = true;
 
-		UNavigationSystem::OnActorRegistered(this);
+		FNavigationSystem::OnActorRegistered(*this);
 		
 		UpdateAllReplicatedComponents();
 	}
diff --git a/Engine/Source/Runtime/Engine/Private/ActorEditor.cpp b/Engine/Source/Runtime/Engine/Private/ActorEditor.cpp
--- a/Engine/Source/Runtime/Engine/Private/ActorEditor.cpp
+++ b/Engine/Source/Runtime/Engine/Private/ActorEditor.cpp
@@ -10,7 +10,7 @@
 #include "GameFramework/Actor.h"
 #include "Components/ChildActorComponent.h"
 #include "Components/PrimitiveComponent.h"
-#include "AI/Navigation/NavigationSystem.h"
+#include "AI/NavigationSystemBase.h"
 #include "Engine/BlueprintGeneratedClass.h"
 #include "EditorSupportDelegates.h"
 #include "Logging/TokenizedMessage.h"
@@ -172,19 +172,7 @@ void AActor::PostEditMove(bool bFinished)
 
 	if (bFinished)
 	{
-		// update actor and all its components in navigation system after finishing move
-		// USceneComponent::UpdateNavigationData works only in game world
-		UNavigationSystem::UpdateNavOctreeBounds(this);
-
-		TArray<AActor*> ParentedActors;
-		GetAttachedActors(ParentedActors);
-		for (int32 Idx = 0; Idx < ParentedActors.Num(); Idx++)
-		{
-			UNavigationSystem::UpdateNavOctreeBounds(ParentedActors[Idx]);
-		}
-
-		// not doing manual update of all attached actors since UpdateActorAndComponentsInNavOctree should take care of it
-		UNavigationSystem::UpdateActorAndComponentsInNavOctree(*this);
+		FNavigationSystem::OnPostEditActorMove(*this);
 	}
 }
 
@@ -372,7 +360,7 @@ void AActor::PreEditUndo()
 	}
 
 	// let navigation system know to not care about this actor anymore
-	UNavigationSystem::ClearNavOctreeAll(this);
+	FNavigationSystem::RemoveActorData(*this);
 
 	Super::PreEditUndo();
 }
@@ -411,11 +399,11 @@ bool AActor::InternalPostEditUndo()
 		BlueprintCreatedComponents.Reset();
 
 		// notify navigation system
-		UNavigationSystem::UpdateActorAndComponentsInNavOctree(*this);
+		FNavigationSystem::UpdateActorAndComponentData(*this);
 	}
 	else
 	{
-		UNavigationSystem::ClearNavOctreeAll(this);
+		FNavigationSystem::RemoveActorData(*this);
 	}
 
 	// This is a normal undo, so call super
diff --git a/Engine/Source/Runtime/Engine/Private/Blueprint.cpp b/Engine/Source/Runtime/Engine/Private/Blueprint.cpp
--- a/Engine/Source/Runtime/Engine/Private/Blueprint.cpp
+++ b/Engine/Source/Runtime/Engine/Private/Blueprint.cpp
@@ -1810,7 +1810,7 @@ UEdGraph* UBlueprint::GetLastEditedUberGraph() const
 {
 	for ( int32 LastEditedIndex = LastEditedDocuments.Num() - 1; LastEditedIndex >= 0; LastEditedIndex-- )
 	{
-		if ( UObject* Obj = LastEditedDocuments[LastEditedIndex].EditedObject )
+		if ( UObject* Obj = LastEditedDocuments[LastEditedIndex].EditedObjectPath.ResolveObject() )
 		{
 			if ( UEdGraph* Graph = Cast<UEdGraph>(Obj) )
 			{
diff --git a/Engine/Source/Runtime/Engine/Private/BlueprintGeneratedClass.cpp b/Engine/Source/Runtime/Engine/Private/BlueprintGeneratedClass.cpp
--- a/Engine/Source/Runtime/Engine/Private/BlueprintGeneratedClass.cpp
+++ b/Engine/Source/Runtime/Engine/Private/BlueprintGeneratedClass.cpp
@@ -30,7 +30,7 @@
 DEFINE_STAT(STAT_PersistentUberGraphFrameMemory);
 
 int32 GBlueprintClusteringEnabled = 0;
-static FAutoConsoleVariableRef CVarUseBackgroundLevelStreaming(
+static FAutoConsoleVariableRef CVarBlueprintClusteringEnabled(
 	TEXT("gc.BlueprintClusteringEnabled"),
 	GBlueprintClusteringEnabled,
 	TEXT("Whether to allow Blueprint classes to create GC clusters."),
diff --git a/Engine/Source/Runtime/Engine/Private/Character.cpp b/Engine/Source/Runtime/Engine/Private/Character.cpp
--- a/Engine/Source/Runtime/Engine/Private/Character.cpp
+++ b/Engine/Source/Runtime/Engine/Private/Character.cpp
@@ -1407,6 +1407,12 @@ void ACharacter::PreReplication( IRepChangedPropertyTracker & ChangedPropertyTra
 			ReplicatedBasedMovement.Rotation = GetActorRotation();
 		}
 	}
+
+	// Save bandwidth by not replicating this value unless it is necessary, since it changes every update.
+	if ((CharacterMovement->NetworkSmoothingMode != ENetworkSmoothingMode::Linear) && !CharacterMovement->bNetworkAlwaysReplicateTransformUpdateTimestamp)
+	{
+		ReplicatedServerLastTransformUpdateTimeStamp = 0.f;
+	}
 }
 
 void ACharacter::PreReplicationForReplay(IRepChangedPropertyTracker & ChangedPropertyTracker)
diff --git a/Engine/Source/Runtime/Engine/Private/CheatManager.cpp b/Engine/Source/Runtime/Engine/Private/CheatManager.cpp
--- a/Engine/Source/Runtime/Engine/Private/CheatManager.cpp
+++ b/Engine/Source/Runtime/Engine/Private/CheatManager.cpp
@@ -17,7 +17,7 @@
 #include "CollisionQueryParams.h"
 #include "WorldCollision.h"
 #include "Engine/World.h"
-#include "AI/Navigation/NavigationSystem.h"
+#include "AI/NavigationSystemBase.h"
 #include "UObject/Package.h"
 #include "GameFramework/PlayerController.h"
 #include "GameFramework/Volume.h"
@@ -29,7 +29,6 @@
 #include "EngineUtils.h"
 #include "Net/OnlineEngineInterface.h"
 #include "VisualLogger/VisualLogger.h"
-#include "AI/Navigation/RecastNavMesh.h"
 #include "GameFramework/Character.h"
 #include "Engine/Console.h"
 #include "Engine/DebugCameraController.h"
@@ -915,28 +914,6 @@ void UCheatManager::TestCollisionDistance()
 #endif
 }
 
-void UCheatManager::RebuildNavigation()
-{
-	UNavigationSystem* NavSys = UNavigationSystem::GetCurrent(GetWorld());
-	if (NavSys)
-	{
-		NavSys->Build();
-	}
-}
-
-void UCheatManager::SetNavDrawDistance(float DrawDistance)
-{
-	if (GIsEditor)
-	{
-		APlayerController* PC = GetOuterAPlayerController();
-		if (PC != NULL)
-		{
-			PC->ClientMessage(TEXT("Setting Nav Rendering Draw Distance is not supported while in Edior"));
-		}
-	}
-	ARecastNavMesh::SetDrawDistance(DrawDistance);
-}
-
 void UCheatManager::DumpOnlineSessionState()
 {
 	UOnlineEngineInterface::Get()->DumpSessionState(GetWorld());
@@ -967,6 +944,11 @@ void UCheatManager::DumpVoiceMutingState()
 
 UWorld* UCheatManager::GetWorld() const
 {
+	if (IsTemplate())
+	{
+		return nullptr;
+	}
+
 	return GetOuterAPlayerController()->GetWorld();
 }
 
