Copying //UE4/Dev-Framework to //UE4/Dev-Main (Source: //UE4/Dev-Framework @ 3967517)

diff --git a/Engine/Source/Runtime/Engine/Private/Components/ActorComponent.cpp b/Engine/Source/Runtime/Engine/Private/Components/ActorComponent.cpp
--- a/Engine/Source/Runtime/Engine/Private/Components/ActorComponent.cpp
+++ b/Engine/Source/Runtime/Engine/Private/Components/ActorComponent.cpp
@@ -1,4 +1,4 @@
-﻿// Copyright 1998-2018 Epic Games, Inc. All Rights Reserved.
+// Copyright 1998-2018 Epic Games, Inc. All Rights Reserved.
 // ActorComponent.cpp: Actor component implementation.
 
 #include "Components/ActorComponent.h"
@@ -11,7 +11,7 @@
 #include "GameFramework/Actor.h"
 #include "Engine/World.h"
 #include "Components/PrimitiveComponent.h"
-#include "AI/Navigation/NavigationSystem.h"
+#include "AI/NavigationSystemBase.h"
 #include "Engine/BlueprintGeneratedClass.h"
 #include "ContentStreaming.h"
 #include "ComponentReregisterContext.h"
@@ -1017,10 +1017,12 @@ void UActorComponent::RegisterComponentWithWorld(UWorld* InWorld)
 
 		for (UObject* Child : Children)
 		{
-			UActorComponent* ChildComponent = Cast<UActorComponent>(Child);
-			if (ChildComponent && !ChildComponent->IsRegistered() && ChildComponent->GetOwner() == MyOwner)
+			if (UActorComponent* ChildComponent = Cast<UActorComponent>(Child))
 			{
-				ChildComponent->RegisterComponentWithWorld(InWorld);
+				if (ChildComponent->bAutoRegister && !ChildComponent->IsRegistered() && ChildComponent->GetOwner() == MyOwner)
+				{
+					ChildComponent->RegisterComponentWithWorld(InWorld);
+				}
 			}
 		}
 
@@ -1780,11 +1782,11 @@ void UActorComponent::HandleCanEverAffectNavigationChange(bool bForceUpdate)
 		if (bCanEverAffectNavigation)
 		{
 			bNavigationRelevant = IsNavigationRelevant();
-			UNavigationSystem::OnComponentRegistered(this);
+			FNavigationSystem::OnComponentRegistered(*this);
 		}
 		else
 		{
-			UNavigationSystem::OnComponentUnregistered(this);
+			FNavigationSystem::OnComponentUnregistered(*this);
 		}
 	}
 }
diff --git a/Engine/Source/Runtime/Engine/Private/Components/CharacterMovementComponent.cpp b/Engine/Source/Runtime/Engine/Private/Components/CharacterMovementComponent.cpp
--- a/Engine/Source/Runtime/Engine/Private/Components/CharacterMovementComponent.cpp
+++ b/Engine/Source/Runtime/Engine/Private/Components/CharacterMovementComponent.cpp
@@ -8,7 +8,8 @@
 #include "GameFramework/CharacterMovementComponent.h"
 #include "EngineStats.h"
 #include "Components/PrimitiveComponent.h"
-#include "AI/Navigation/NavigationSystem.h"
+#include "AI/NavigationSystemBase.h"
+#include "AI/Navigation/NavigationDataInterface.h"
 #include "UObject/Package.h"
 #include "GameFramework/PlayerController.h"
 #include "GameFramework/PhysicsVolume.h"
@@ -20,10 +21,7 @@
 #include "Components/CapsuleComponent.h"
 #include "GameFramework/GameStateBase.h"
 #include "Engine/Canvas.h"
-
-// @todo this is here only due to circular dependency to AIModule. To be removed
-#include "Navigation/PathFollowingComponent.h"
-#include "AI/Navigation/RecastNavMesh.h"
+#include "AI/Navigation/PathFollowingAgentInterface.h"
 #include "AI/Navigation/AvoidanceManager.h"
 #include "Components/BrushComponent.h"
 
@@ -1012,9 +1010,13 @@ void UCharacterMovementComponent::OnMovementModeChanged(EMovementMode PreviousMo
 		}
 	}
 
-	if (MovementMode == MOVE_Falling && PreviousMovementMode != MOVE_Falling && PathFollowingComp.IsValid())
+	if (MovementMode == MOVE_Falling && PreviousMovementMode != MOVE_Falling)
 	{
-		PathFollowingComp->OnStartedFalling();
+		IPathFollowingAgentInterface* PFAgent = GetPathFollowingAgent();
+		if (PFAgent)
+		{
+			PFAgent->OnStartedFalling();
+		}
 	}
 
 	CharacterOwner->OnMovementModeChanged(PreviousMovementMode, PreviousCustomMode);
@@ -2382,7 +2384,7 @@ bool UCharacterMovementComponent::CanCrouchInCurrentState() const
 		return false;
 	}
 
-	return IsFalling() || IsMovingOnGround();
+	return (IsFalling() || IsMovingOnGround()) && UpdatedComponent && !UpdatedComponent->IsSimulatingPhysics();
 }
 
 
@@ -2626,12 +2628,12 @@ void UCharacterMovementComponent::UnCrouch(bool bClientSimulation)
 void UCharacterMovementComponent::UpdateCharacterStateBeforeMovement()
 {
 	// Check for a change in crouch state. Players toggle crouch by changing bWantsToCrouch.
-	const bool bAllowedToCrouch = CanCrouchInCurrentState();
-	if ((!bAllowedToCrouch || !bWantsToCrouch) && IsCrouching())
+	const bool bIsCrouching = IsCrouching();
+	if (bIsCrouching && (!bWantsToCrouch || !CanCrouchInCurrentState()))
 	{
 		UnCrouch(false);
 	}
-	else if (bWantsToCrouch && bAllowedToCrouch && !IsCrouching())
+	else if (!bIsCrouching && bWantsToCrouch && CanCrouchInCurrentState())
 	{
 		Crouch(false);
 	}
@@ -3278,15 +3280,12 @@ void UCharacterMovementComponent::UpdateDefaultAvoidance()
 	SCOPE_CYCLE_COUNTER(STAT_AI_ObstacleAvoidance);
 
 	UAvoidanceManager* AvoidanceManager = GetWorld()->GetAvoidanceManager();
-	if (AvoidanceManager && !bWasAvoidanceUpdated)
+	if (AvoidanceManager && !bWasAvoidanceUpdated && GetCharacterOwner()->GetCapsuleComponent())
 	{
-		if (UCapsuleComponent *OurCapsule = GetCharacterOwner()->GetCapsuleComponent())
-		{
-			AvoidanceManager->UpdateRVO(this);
+		AvoidanceManager->UpdateRVO(this);
 
-			//Consider this a clean move because we didn't even try to avoid.
-			SetAvoidanceVelocityLock(AvoidanceManager, AvoidanceManager->LockTimeAfterClean);
-		}
+		//Consider this a clean move because we didn't even try to avoid.
+		SetAvoidanceVelocityLock(AvoidanceManager, AvoidanceManager->LockTimeAfterClean);
 	}
 
 	bWasAvoidanceUpdated = false;		//Reset for next frame
@@ -4860,7 +4859,7 @@ void UCharacterMovementComponent::PhysNavWalking(float deltaTime, int32 Iteratio
 
 bool UCharacterMovementComponent::FindNavFloor(const FVector& TestLocation, FNavLocation& NavFloorLocation) const
 {
-	const ANavigationData* NavData = GetNavData();
+	const INavigationDataInterface* NavData = GetNavData();
 	if (NavData == nullptr)
 	{
 		return false;
@@ -5029,32 +5028,17 @@ void UCharacterMovementComponent::FindBestNavMeshLocation(const FVector& TraceSt
 	}
 }
 
-const ANavigationData* UCharacterMovementComponent::GetNavData() const
+const INavigationDataInterface* UCharacterMovementComponent::GetNavData() const
 {
-	UNavigationSystem* NavSys = UNavigationSystem::GetCurrent(GetWorld());
-	if (NavSys == nullptr || !HasValidData())
+	const UWorld* World = GetWorld();
+	if (World == nullptr || World->GetNavigationSystem() == nullptr 
+		|| !HasValidData()
+		|| CharacterOwner == nullptr)
 	{
 		return nullptr;
 	}
 
-	const ANavigationData* NavData = nullptr;
-	INavAgentInterface* MyNavAgent = CastChecked<INavAgentInterface>(CharacterOwner);
-	if (MyNavAgent)
-	{
-		const FNavAgentProperties& AgentProps = MyNavAgent->GetNavAgentPropertiesRef();
-		NavData = NavSys->GetNavDataForProps(AgentProps);
-	}
-	if (NavData == nullptr)
-	{
-		NavData = NavSys->GetMainNavData();
-	}
-
-	// Only RecastNavMesh supported
-	const ARecastNavMesh* NavMeshData = Cast<const ARecastNavMesh>(NavData);
-	if (NavMeshData == nullptr)
-	{
-		return nullptr;
-	}
+	const INavigationDataInterface* NavData = FNavigationSystem::GetNavDataForActor(*CharacterOwner);
 
 	return NavData;
 }
@@ -5178,9 +5162,11 @@ void UCharacterMovementComponent::ProcessLanded(const FHitResult& Hit, float rem
 
 		SetPostLandedPhysics(Hit);
 	}
-	if (PathFollowingComp.IsValid())
+	
+	IPathFollowingAgentInterface* PFAgent = GetPathFollowingAgent();
+	if (PFAgent)
 	{
-		PathFollowingComp->OnLanded();
+		PFAgent->OnLanded();
 	}
 
 	StartNewPhysics(remainingTime, Iterations);
@@ -5461,9 +5447,11 @@ void UCharacterMovementComponent::PhysicsVolumeChanged( APhysicsVolume* NewVolum
 		if ( !CanEverSwim() )
 		{
 			// AI needs to stop any current moves
-			if (PathFollowingComp.IsValid())
+			IPathFollowingAgentInterface* PFAgent = GetPathFollowingAgent();
+			if (PFAgent)
 			{
-				PathFollowingComp->AbortMove(*this, FPathFollowingResultFlags::MovementStop);
+				//PathFollowingComp->AbortMove(*this, FPathFollowingResultFlags::MovementStop);
+				PFAgent->OnUnableToMove(*this);
 			}			
 		}
 		else if ( !IsSwimming() )
@@ -6424,9 +6412,11 @@ void UCharacterMovementComponent::HandleImpact(const FHitResult& Impact, float T
 		CharacterOwner->MoveBlockedBy(Impact);
 	}
 
-	if (PathFollowingComp.IsValid())
-	{	// Also notify path following!
-		PathFollowingComp->OnMoveBlockedBy(Impact);
+	IPathFollowingAgentInterface* PFAgent = GetPathFollowingAgent();
+	if (PFAgent)
+	{
+		// Also notify path following!
+		PFAgent->OnMoveBlockedBy(Impact);
 	}
 
 	APawn* OtherPawn = Cast<APawn>(Impact.GetActor());
@@ -10189,7 +10179,6 @@ void FSavedMove_Character::PostUpdate(ACharacter* Character, FSavedMove_Characte
 		if ((StartActorOverlapCounter != EndActorOverlapCounter) || (StartComponentOverlapCounter != EndComponentOverlapCounter))
 		{
 			bForceNoCombine = true;
-			UE_LOG(LogNetPlayerMovement, Log, TEXT("SavedMove: Setting bForceNoCombine true because of an overlap event!"));
 		}
 	}
 	else if (PostUpdateMode == PostUpdate_Replay)
@@ -10256,14 +10245,23 @@ float UCharacterMovementComponent::GetClientNetSendDeltaTime(const APlayerContro
 	const AGameNetworkManager* const GameNetworkManager = GetDefault<AGameNetworkManager>();
 	float NetMoveDelta = GameNetworkManager->ClientNetSendMoveDeltaTime;
 
-	// send moves more frequently in small games where server isn't likely to be saturated
-	if (Player && (Player->CurrentNetSpeed > GameNetworkManager->ClientNetSendMoveThrottleAtNetSpeed) && (GameState != nullptr) && (GameState->PlayerArray.Num() <= GameNetworkManager->ClientNetSendMoveThrottleOverPlayerCount))
-	{
-		NetMoveDelta = GameNetworkManager->ClientNetSendMoveDeltaTime;
-	}
-	else if (Player)
+	if (PC && Player)
 	{
-		NetMoveDelta = FMath::Max(GameNetworkManager->ClientNetSendMoveDeltaTimeThrottled, 2 * GameNetworkManager->MoveRepSize / Player->CurrentNetSpeed);
+		// send moves more frequently in small games where server isn't likely to be saturated
+		if ((Player->CurrentNetSpeed > GameNetworkManager->ClientNetSendMoveThrottleAtNetSpeed) && (GameState != nullptr) && (GameState->PlayerArray.Num() <= GameNetworkManager->ClientNetSendMoveThrottleOverPlayerCount))
+		{
+			NetMoveDelta = GameNetworkManager->ClientNetSendMoveDeltaTime;
+		}
+		else
+		{
+			NetMoveDelta = FMath::Max(GameNetworkManager->ClientNetSendMoveDeltaTimeThrottled, 2 * GameNetworkManager->MoveRepSize / Player->CurrentNetSpeed);
+		}
+
+		// Lower frequency for standing still and not rotating camera
+		if (Acceleration.IsZero() && Velocity.IsZero() && ClientData->LastAckedMove.IsValid() && ClientData->LastAckedMove->StartControlRotation.Equals(PC->GetControlRotation()))
+		{
+			NetMoveDelta = FMath::Max(GameNetworkManager->ClientNetSendMoveDeltaTimeStationary, NetMoveDelta);
+		}
 	}
 	
 	return NetMoveDelta;
diff --git a/Engine/Source/Runtime/Engine/Private/Components/NavMovementComponent.cpp b/Engine/Source/Runtime/Engine/Private/Components/NavMovementComponent.cpp
--- a/Engine/Source/Runtime/Engine/Private/Components/NavMovementComponent.cpp
+++ b/Engine/Source/Runtime/Engine/Private/Components/NavMovementComponent.cpp
@@ -1,25 +1,10 @@
 // Copyright 1998-2018 Epic Games, Inc. All Rights Reserved.
 
 #include "GameFramework/NavMovementComponent.h"
-#include "AI/Navigation/NavigationSystem.h"
+#include "AI/NavigationSystemBase.h"
 #include "Components/CapsuleComponent.h"
-#include "Navigation/PathFollowingComponent.h"
+#include "AI/Navigation/PathFollowingAgentInterface.h"
 
-//----------------------------------------------------------------------//
-// FNavAgentProperties
-//----------------------------------------------------------------------//
-const FNavAgentProperties FNavAgentProperties::DefaultProperties;
-
-void FNavAgentProperties::UpdateWithCollisionComponent(UShapeComponent* CollisionComponent)
-{
-	check( CollisionComponent != NULL);
-	AgentRadius = CollisionComponent->Bounds.SphereRadius;
-}
-
-bool FNavAgentProperties::IsNavDataMatching(const FNavAgentProperties& Other) const
-{
-	return (PreferredNavData == Other.PreferredNavData || PreferredNavData == nullptr || Other.PreferredNavData == nullptr);
-}
 
 //----------------------------------------------------------------------//
 // UMovementComponent
@@ -74,9 +59,15 @@ void UNavMovementComponent::ClearFixedBrakingDistance()
 
 void UNavMovementComponent::StopActiveMovement()
 {
-	if (PathFollowingComp.IsValid() && bStopMovementAbortPaths)
+	if (!bStopMovementAbortPaths)
+	{
+		return;
+	}
+
+	IPathFollowingAgentInterface* PFAgent = GetPathFollowingAgent();
+	if (PFAgent)
 	{
-		PathFollowingComp->AbortMove(*this, FPathFollowingResultFlags::MovementStop);
+		PFAgent->OnUnableToMove(*this);
 	}
 }
 
@@ -89,12 +80,8 @@ void UNavMovementComponent::UpdateNavAgent(const AActor& Owner)
 	}
 
 	// initialize properties from navigation system
-	UNavigationSystem* NavSys = UNavigationSystem::GetCurrent(GetWorld());
-	if (NavSys != nullptr)
-	{
-		NavAgentProps.NavWalkingSearchHeightScale = NavSys->GetDefaultSupportedAgentConfig().NavWalkingSearchHeightScale;
-	}
-
+	NavAgentProps.NavWalkingSearchHeightScale = FNavigationSystem::GetDefaultSupportedAgent().NavWalkingSearchHeightScale;
+	
 	// Can't call GetSimpleCollisionCylinder(), because no components will be registered.
 	float BoundRadius, BoundHalfHeight;	
 	Owner.GetSimpleCollisionCylinder(BoundRadius, BoundHalfHeight);
@@ -110,11 +97,7 @@ void UNavMovementComponent::UpdateNavAgent(const UCapsuleComponent& CapsuleCompo
 	}
 
 	// initialize properties from navigation system
-	UNavigationSystem* NavSys = UNavigationSystem::GetCurrent(GetWorld());
-	if (NavSys != nullptr)
-	{
-		NavAgentProps.NavWalkingSearchHeightScale = NavSys->GetDefaultSupportedAgentConfig().NavWalkingSearchHeightScale;
-	}
+	NavAgentProps.NavWalkingSearchHeightScale = FNavigationSystem::GetDefaultSupportedAgent().NavWalkingSearchHeightScale;
 
 	NavAgentProps.AgentRadius = CapsuleComponent.GetScaledCapsuleRadius();
 	NavAgentProps.AgentHeight = CapsuleComponent.GetScaledCapsuleHalfHeight() * 2.f;
@@ -124,3 +107,4 @@ void UNavMovementComponent::SetUpdateNavAgentWithOwnersCollisions(bool bUpdateWi
 {
 	bUpdateNavAgentWithOwnersCollision = bUpdateWithOwner;
 }
+
diff --git a/Engine/Source/Runtime/Engine/Private/Components/PrimitiveComponent.cpp b/Engine/Source/Runtime/Engine/Private/Components/PrimitiveComponent.cpp
--- a/Engine/Source/Runtime/Engine/Private/Components/PrimitiveComponent.cpp
+++ b/Engine/Source/Runtime/Engine/Private/Components/PrimitiveComponent.cpp
@@ -9,7 +9,7 @@
 #include "GameFramework/DamageType.h"
 #include "GameFramework/Pawn.h"
 #include "WorldCollision.h"
-#include "AI/Navigation/NavigationSystem.h"
+#include "AI/NavigationSystemBase.h"
 #include "GameFramework/PlayerController.h"
 #include "GameFramework/PhysicsVolume.h"
 #include "GameFramework/WorldSettings.h"
@@ -536,7 +536,7 @@ void UPrimitiveComponent::OnRegister()
 		const bool bNavRelevant = bNavigationRelevant = IsNavigationRelevant();
 		if (bNavRelevant)
 		{
-			UNavigationSystem::OnComponentRegistered(this);
+			FNavigationSystem::OnComponentRegistered(*this);
 		}
 	}
 	else
@@ -568,7 +568,7 @@ void UPrimitiveComponent::OnUnregister()
 
 	if (bCanEverAffectNavigation)
 	{
-		UNavigationSystem::OnComponentUnregistered(this);
+		FNavigationSystem::OnComponentUnregistered(*this);
 	}
 }
 
diff --git a/Engine/Source/Runtime/Engine/Private/Components/SceneComponent.cpp b/Engine/Source/Runtime/Engine/Private/Components/SceneComponent.cpp
--- a/Engine/Source/Runtime/Engine/Private/Components/SceneComponent.cpp
+++ b/Engine/Source/Runtime/Engine/Private/Components/SceneComponent.cpp
@@ -13,7 +13,7 @@
 #include "WorldCollision.h"
 #include "Components/PrimitiveComponent.h"
 #include "Components/StaticMeshComponent.h"
-#include "AI/Navigation/NavigationSystem.h"
+#include "AI/NavigationSystemBase.h"
 #include "Engine/MapBuildDataRegistry.h"
 #include "GameFramework/PhysicsVolume.h"
 #include "Components/BillboardComponent.h"
@@ -508,7 +508,7 @@ void USceneComponent::PostEditChangeProperty(FPropertyChangedEvent& PropertyChan
 
 	if (PropertyName == LocationName || PropertyName == RotationName || PropertyName == ScaleName)
 	{
-		UNavigationSystem::UpdateComponentInNavOctree(*this);
+		FNavigationSystem::UpdateComponentData(*this);
 	}
 }
 
@@ -1824,11 +1824,11 @@ bool USceneComponent::AttachToComponent(USceneComponent* Parent, const FAttachme
 		{
 			if (Parent->IsTemplate())
 			{
-				ensureMsgf(false, TEXT("Template Mismatch during attachment. Attaching instanced component to template component. Parent '%s' Self '%s'"), *Parent->GetName(), *GetName());
+				ensureMsgf(false, TEXT("Template Mismatch during attachment. Attaching instanced component to template component. Parent '%s' (Owner '%s') Self '%s' (Owner '%s')."), *Parent->GetName(), *GetNameSafe(Parent->GetOwner()), *GetName(), *GetNameSafe(GetOwner()));
 			}
 			else
 			{
-				ensureMsgf(false, TEXT("Template Mismatch during attachment. Attaching template component to instanced component. Parent '%s' Self '%s'"), *Parent->GetName(), *GetName());
+				ensureMsgf(false, TEXT("Template Mismatch during attachment. Attaching template component to instanced component. Parent '%s' (Owner '%s') Self '%s' (Owner '%s')."), *Parent->GetName(), *GetNameSafe(Parent->GetOwner()), *GetName(), *GetNameSafe(GetOwner()));
 			}
 			return false;
 		}
@@ -2728,11 +2728,7 @@ bool USceneComponent::InternalSetWorldLocationAndRotation(FVector NewLocation, c
 		UpdateComponentToWorldWithParent(GetAttachParent(),GetAttachSocketName(), SkipPhysicsToEnum(bNoPhysics), RelativeRotationCache.GetCachedQuat(), Teleport);
 
 		// we need to call this even if this component itself is not navigation relevant
-		// checking ShouldUpdateNavOctreeOnComponentChange here is an optimization for static navigation users
-		if (UNavigationSystem::ShouldUpdateNavOctreeOnComponentChange())
-		{
-			PostUpdateNavigationData();
-		}
+		PostUpdateNavigationData();
 
 		return true;
 	}
@@ -3357,7 +3353,7 @@ bool FScopedMovementUpdate::IsTransformDirty() const
 {
 	if (IsValid(Owner))
 	{
-		return !InitialTransform.Equals(Owner->GetComponentToWorld());
+		return !InitialTransform.Equals(Owner->GetComponentToWorld(), SMALL_NUMBER);
 	}
 
 	return false;
@@ -3579,13 +3575,13 @@ void USceneComponent::UpdateNavigationData()
 {
 	SCOPE_CYCLE_COUNTER(STAT_ComponentUpdateNavData);
 
-	if (UNavigationSystem::ShouldUpdateNavOctreeOnComponentChange() && IsRegistered())
+	if (IsRegistered())
 	{
 		UWorld* MyWorld = GetWorld();
-		if ((MyWorld == nullptr) || !MyWorld->IsGameWorld() || !MyWorld->IsNetMode(ENetMode::NM_Client))
+		if ((MyWorld != nullptr) && (!MyWorld->IsGameWorld() || !MyWorld->IsNetMode(ENetMode::NM_Client)))
 		{
 			// use propagated component's transform update in editor OR server game with additional navsys check
-			UNavigationSystem::UpdateComponentInNavOctree(*this);
+			FNavigationSystem::UpdateComponentData(*this);
 		}
 	}
 }
@@ -3594,15 +3590,9 @@ void USceneComponent::PostUpdateNavigationData()
 {
 	SCOPE_CYCLE_COUNTER(STAT_ComponentPostUpdateNavData);
 
-	if (UNavigationSystem::ShouldUpdateNavOctreeOnComponentChange() && IsRegistered())
+	if (IsRegistered())
 	{
-		UWorld* MyWorld = GetWorld();
-		if (MyWorld != nullptr && MyWorld->GetNavigationSystem() != nullptr 
-			&& (MyWorld->GetNavigationSystem()->ShouldAllowClientSideNavigation() || !MyWorld->IsNetMode(ENetMode::NM_Client)))
-		{
-			// use propagated component's transform update in editor OR server game with additional navsys check
-			UNavigationSystem::UpdateNavOctreeAfterMove(this);
-		}
+		FNavigationSystem::OnComponentTransformChanged(*this);
 	}
 }
 
diff --git a/Engine/Source/Runtime/Engine/Private/Components/ShapeComponent.cpp b/Engine/Source/Runtime/Engine/Private/Components/ShapeComponent.cpp
--- a/Engine/Source/Runtime/Engine/Private/Components/ShapeComponent.cpp
+++ b/Engine/Source/Runtime/Engine/Private/Components/ShapeComponent.cpp
@@ -2,8 +2,7 @@
 
 
 #include "Components/ShapeComponent.h"
-#include "AI/Navigation/NavAreas/NavArea_Obstacle.h"
-#include "AI/NavigationOctree.h"
+#include "AI/Navigation/NavAreaBase.h"
 #include "PhysicsEngine/BoxElem.h"
 #include "PhysicsEngine/SphereElem.h"
 #include "PhysicsEngine/SphylElem.h"
@@ -39,7 +38,8 @@ UShapeComponent::UShapeComponent(const FObjectInitializer& ObjectInitializer)
 	bHasCustomNavigableGeometry = EHasCustomNavigableGeometry::Yes;
 	bCanEverAffectNavigation = true;
 	bDynamicObstacle = false;
-	AreaClass = UNavArea_Obstacle::StaticClass();
+	// mz@todo
+	//AreaClass = UNavArea_Obstacle::StaticClass();
 }
 
 FPrimitiveSceneProxy* UShapeComponent::CreateSceneProxy()
diff --git a/Engine/Source/Runtime/Engine/Private/Components/SplineMeshComponent.cpp b/Engine/Source/Runtime/Engine/Private/Components/SplineMeshComponent.cpp
--- a/Engine/Source/Runtime/Engine/Private/Components/SplineMeshComponent.cpp
+++ b/Engine/Source/Runtime/Engine/Private/Components/SplineMeshComponent.cpp
@@ -11,7 +11,7 @@
 #include "SplineMeshSceneProxy.h"
 #include "ShaderParameterUtils.h"
 #include "AI/NavigationSystemHelpers.h"
-#include "AI/Navigation/NavCollision.h"
+#include "AI/Navigation/NavCollisionBase.h"
 #include "Engine/StaticMeshSocket.h"
 #include "Engine/StaticMesh.h"
 #include "PhysicsEngine/ConvexElem.h"
@@ -941,37 +941,37 @@ bool USplineMeshComponent::DoCustomNavigableGeometryExport(FNavigableGeometryExp
 
 	if (GetStaticMesh() != nullptr && GetStaticMesh()->NavCollision != nullptr)
 	{
-		UNavCollision* NavCollision = GetStaticMesh()->NavCollision;
+		UNavCollisionBase* NavCollision = GetStaticMesh()->NavCollision;
 		
-		if (ensure(!NavCollision->bIsDynamicObstacle))
+		if (ensure(!NavCollision->IsDynamicObstacle()))
 		{
-			if (NavCollision->bHasConvexGeometry)
+			if (NavCollision->HasConvexGeometry())
 			{
 				FVector Mask = FVector(1, 1, 1);
 				GetAxisValue(Mask, ForwardAxis) = 0;
 
 				TArray<FVector> VertexBuffer;
-				VertexBuffer.Reserve(FMath::Max(NavCollision->ConvexCollision.VertexBuffer.Num(), NavCollision->TriMeshCollision.VertexBuffer.Num()));
+				VertexBuffer.Reserve(FMath::Max(NavCollision->GetConvexCollision().VertexBuffer.Num(), NavCollision->GetTriMeshCollision().VertexBuffer.Num()));
 
-				for (int32 i = 0; i < NavCollision->ConvexCollision.VertexBuffer.Num(); ++i)
+				for (int32 i = 0; i < NavCollision->GetConvexCollision().VertexBuffer.Num(); ++i)
 				{
-					FVector Vertex = NavCollision->ConvexCollision.VertexBuffer[i];
+					FVector Vertex = NavCollision->GetConvexCollision().VertexBuffer[i];
 					Vertex = CalcSliceTransform(GetAxisValue(Vertex, ForwardAxis)).TransformPosition(Vertex * Mask);
 					VertexBuffer.Add(Vertex);
 				}
 				GeomExport.ExportCustomMesh(VertexBuffer.GetData(), VertexBuffer.Num(),
-					NavCollision->ConvexCollision.IndexBuffer.GetData(), NavCollision->ConvexCollision.IndexBuffer.Num(),
+					NavCollision->GetConvexCollision().IndexBuffer.GetData(), NavCollision->GetConvexCollision().IndexBuffer.Num(),
 					GetComponentTransform());
 
 				VertexBuffer.Reset();
-				for (int32 i = 0; i < NavCollision->TriMeshCollision.VertexBuffer.Num(); ++i)
+				for (int32 i = 0; i < NavCollision->GetTriMeshCollision().VertexBuffer.Num(); ++i)
 				{
-					FVector Vertex = NavCollision->TriMeshCollision.VertexBuffer[i];
+					FVector Vertex = NavCollision->GetTriMeshCollision().VertexBuffer[i];
 					Vertex = CalcSliceTransform(GetAxisValue(Vertex, ForwardAxis)).TransformPosition(Vertex * Mask);
 					VertexBuffer.Add(Vertex);
 				}
 				GeomExport.ExportCustomMesh(VertexBuffer.GetData(), VertexBuffer.Num(),
-					NavCollision->TriMeshCollision.IndexBuffer.GetData(), NavCollision->TriMeshCollision.IndexBuffer.Num(),
+					NavCollision->GetTriMeshCollision().IndexBuffer.GetData(), NavCollision->GetTriMeshCollision().IndexBuffer.Num(),
 					GetComponentTransform());
 
 				return false;
diff --git a/Engine/Source/Runtime/Engine/Private/Components/StaticMeshComponent.cpp b/Engine/Source/Runtime/Engine/Private/Components/StaticMeshComponent.cpp
--- a/Engine/Source/Runtime/Engine/Private/Components/StaticMeshComponent.cpp
+++ b/Engine/Source/Runtime/Engine/Private/Components/StaticMeshComponent.cpp
@@ -28,11 +28,10 @@
 #include "LightMap.h"
 #include "ShadowMap.h"
 #include "Engine/ShadowMapTexture2D.h"
-#include "AI/Navigation/NavCollision.h"
+#include "AI/Navigation/NavCollisionBase.h"
 #include "Engine/StaticMeshSocket.h"
 #include "AI/NavigationSystemHelpers.h"
-#include "AI/NavigationOctree.h"
-#include "AI/Navigation/NavigationSystem.h"
+#include "AI/NavigationSystemBase.h"
 #include "PhysicsEngine/BodySetup.h"
 #include "EngineGlobals.h"
 #include "ComponentRecreateRenderStateContext.h"
@@ -612,14 +611,14 @@ void UStaticMeshComponent::OnCreatePhysicsState()
 	Super::OnCreatePhysicsState();
 
 	bNavigationRelevant = IsNavigationRelevant();
-	UNavigationSystem::UpdateComponentInNavOctree(*this);
+	FNavigationSystem::UpdateComponentData(*this);
 }
 
 void UStaticMeshComponent::OnDestroyPhysicsState()
 {
 	Super::OnDestroyPhysicsState();
 
-	UNavigationSystem::UpdateComponentInNavOctree(*this);
+	FNavigationSystem::UpdateComponentData(*this);
 	bNavigationRelevant = IsNavigationRelevant();
 }
 
@@ -2240,8 +2239,8 @@ bool UStaticMeshComponent::DoCustomNavigableGeometryExport(FNavigableGeometryExp
 	const FVector Scale3D = GetComponentToWorld().GetScale3D();
 	if (GetStaticMesh() && GetStaticMesh()->NavCollision && !Scale3D.IsZero())
 	{
-		UNavCollision* NavCollision = GetStaticMesh()->NavCollision;
-		const bool bExportAsObstacle = bOverrideNavigationExport ? bForceNavigationObstacle : NavCollision->bIsDynamicObstacle;
+		UNavCollisionBase* NavCollision = GetStaticMesh()->NavCollision;
+		const bool bExportAsObstacle = bOverrideNavigationExport ? bForceNavigationObstacle : NavCollision->IsDynamicObstacle();
 
 		if (bExportAsObstacle)
 		{
@@ -2313,8 +2312,8 @@ void UStaticMeshComponent::GetNavigationData(FNavigationRelevantData& Data) cons
 	const FVector Scale3D = GetComponentToWorld().GetScale3D();
 	if (GetStaticMesh() && GetStaticMesh()->NavCollision && !Scale3D.IsZero())
 	{
-		UNavCollision* NavCollision = GetStaticMesh()->NavCollision;
-		const bool bExportAsObstacle = bOverrideNavigationExport ? bForceNavigationObstacle : NavCollision->bIsDynamicObstacle;
+		UNavCollisionBase* NavCollision = GetStaticMesh()->NavCollision;
+		const bool bExportAsObstacle = bOverrideNavigationExport ? bForceNavigationObstacle : NavCollision->IsDynamicObstacle();
 
 		if (bExportAsObstacle)
 		{
diff --git a/Engine/Source/Runtime/Engine/Private/Controller.cpp b/Engine/Source/Runtime/Engine/Private/Controller.cpp
--- a/Engine/Source/Runtime/Engine/Private/Controller.cpp
+++ b/Engine/Source/Runtime/Engine/Private/Controller.cpp
@@ -8,7 +8,7 @@
 #include "GameFramework/Pawn.h"
 #include "CollisionQueryParams.h"
 #include "Engine/World.h"
-#include "AI/Navigation/NavigationSystem.h"
+#include "AI/NavigationSystemBase.h"
 #include "GameFramework/PlayerController.h"
 #include "Net/UnrealNetwork.h"
 #include "NetworkingDistanceConstants.h"
@@ -21,7 +21,6 @@
 
 // @todo this is here only due to circular dependency to AIModule. To be removed
 #include "Engine/Canvas.h"
-#include "Navigation/PathFollowingComponent.h"
 
 #include "GameFramework/PlayerState.h"
 
@@ -286,9 +285,11 @@ void AController::Possess(APawn* InPawn)
 
 	REDIRECT_OBJECT_TO_VLOG(InPawn, this);
 
+	const bool bNewPawn = GetPawn() != InPawn;
+
 	if (InPawn != NULL)
 	{
-		if (GetPawn() && GetPawn() != InPawn)
+		if (bNewPawn && GetPawn())
 		{
 			UnPossess();
 		}
@@ -306,6 +307,11 @@ void AController::Possess(APawn* InPawn)
 
 		Pawn->Restart();
 	}
+
+	if (bNewPawn)
+	{
+		OnNewPawn.Broadcast(GetPawn());
+	}
 }
 
 void AController::UnPossess()
@@ -641,39 +647,17 @@ bool AController::ShouldPostponePathUpdates() const
 
 bool AController::IsFollowingAPath() const
 {
-	UPathFollowingComponent* PathFollowingComp = FindComponentByClass<UPathFollowingComponent>();
-	return (PathFollowingComp != nullptr) && (PathFollowingComp->GetStatus() != EPathFollowingStatus::Idle);
-}
-
-void AController::UpdateNavigationComponents()
-{
-	UPathFollowingComponent* PathFollowingComp = FindComponentByClass<UPathFollowingComponent>();
-	if (PathFollowingComp != NULL)
-	{
-		PathFollowingComp->UpdateCachedComponents();
-	}
+	return FNavigationSystem::IsFollowingAPath(*this);
 }
 
-void AController::InitNavigationControl(UPathFollowingComponent*& PathFollowingComp)
+IPathFollowingAgentInterface* AController::GetPathFollowingAgent() const
 {
-	PathFollowingComp = FindComponentByClass<UPathFollowingComponent>();
-	if (PathFollowingComp == NULL)
-	{
-		PathFollowingComp = NewObject<UPathFollowingComponent>(this);
-		PathFollowingComp->RegisterComponentWithWorld(GetWorld());
-		PathFollowingComp->Initialize();
-	}
+	return FNavigationSystem::FindPathFollowingAgentForActor(*this);
 }
 
 void AController::StopMovement()
 {
-	UE_VLOG(this, LogNavigation, Log, TEXT("AController::StopMovement: %s STOP MOVEMENT"), *GetNameSafe(GetPawn()));
-
-	UPathFollowingComponent* PathFollowingComp = FindComponentByClass<UPathFollowingComponent>();
-	if (PathFollowingComp != NULL)
-	{
-		PathFollowingComp->AbortMove(*this, FPathFollowingResultFlags::MovementStop);
-	}
+	FNavigationSystem::StopMovement(*this);
 }
 
 void AController::GetLifetimeReplicatedProps( TArray< FLifetimeProperty > & OutLifetimeProps ) const
diff --git a/Engine/Source/Runtime/Engine/Private/EdGraph/EdGraphPin.cpp b/Engine/Source/Runtime/Engine/Private/EdGraph/EdGraphPin.cpp
--- a/Engine/Source/Runtime/Engine/Private/EdGraph/EdGraphPin.cpp
+++ b/Engine/Source/Runtime/Engine/Private/EdGraph/EdGraphPin.cpp
@@ -1312,11 +1312,24 @@ UEdGraphPin::UEdGraphPin(UEdGraphNode* InOwningNode, const FGuid& PinIdGuid)
 	: OwningNode(InOwningNode)
 	, PinId(PinIdGuid)
 	, PinName()
+	, Direction(EGPD_Input)
+#if WITH_EDITORONLY_DATA
+	, bHidden(false)
+	, bNotConnectable(false)
+	, bDefaultValueIsReadOnly(false)
+	, bDefaultValueIsIgnored(false)
+	, bIsDiffing(false)
+	, bAdvancedView(false)
+	, bDisplayAsMutableRef(false)
+	, bOrphanedPin(false)
+	, bSavePinIfOrphaned(true)
+	, bUseBackwardsCompatForEmptyAutogeneratedValue(false)
+#endif
+	, bWasTrashed(false)
 #if WITH_EDITORONLY_DATA
 	, PinFriendlyName()
 #endif // WITH_EDITORONLY_DATA
 	, PinToolTip()
-	, Direction(EGPD_Input)
 	, PinType()
 	, DefaultValue()
 	, AutogeneratedDefaultValue()
@@ -1328,17 +1341,7 @@ UEdGraphPin::UEdGraphPin(UEdGraphNode* InOwningNode, const FGuid& PinIdGuid)
 	, ReferencePassThroughConnection(nullptr)
 #if WITH_EDITORONLY_DATA
 	, PersistentGuid()
-	, bHidden(false)
-	, bNotConnectable(false)
-	, bDefaultValueIsReadOnly(false)
-	, bDefaultValueIsIgnored(false)
-	, bIsDiffing(false)
-	, bAdvancedView(false)
-	, bDisplayAsMutableRef(false)
-	, bOrphanedPin(false)
-	, bSavePinIfOrphaned(true)
 #endif
-	, bWasTrashed(false)
 {
 	PinHelpers::NumPinsInMemory++;
 #ifdef TRACK_PINS
@@ -2101,6 +2104,24 @@ void UEdGraphPin::SanitizePinsPostUndoRedo()
 	GraphConnectionSanitizer.SanitizeConnectionsForCorruptGraphs();
 }
 
+#if WITH_EDITORONLY_DATA
+void UEdGraphPin::SetSavePinIfOrphaned(const bool bShouldSave)
+{
+	if (bSavePinIfOrphaned != bShouldSave)
+	{
+		bSavePinIfOrphaned = bShouldSave;
+		// If we are not saving this pin, we cannot save any children either
+		if (!bSavePinIfOrphaned)
+		{
+			for (UEdGraphPin* SubPin : SubPins)
+			{
+				SubPin->SetSavePinIfOrphaned(false);
+			}
+		}
+	}
+}
+#endif
+
 /////////////////////////////////////////////////////
 // UEdGraphPin
 
diff --git a/Engine/Source/Runtime/Engine/Private/EnginePrivatePCH.h b/Engine/Source/Runtime/Engine/Private/EnginePrivatePCH.h
--- a/Engine/Source/Runtime/Engine/Private/EnginePrivatePCH.h
+++ b/Engine/Source/Runtime/Engine/Private/EnginePrivatePCH.h
@@ -392,7 +392,7 @@
 #include "EngineDefines.h"
 #include "CollisionQueryParams.h"
 #include "AI/Navigation/NavigationTypes.h"
-#include "AI/Navigation/NavFilters/NavigationQueryFilter.h"
+#include "AI/Navigation/NavQueryFilter.h"
 #include "WorldCollision.h"
 #include "AI/Navigation/NavAgentInterface.h"
 #include "Engine/Blueprint.h"
@@ -499,9 +499,8 @@
 #include "Materials/MaterialFunction.h"
 #include "GameFramework/PlayerController.h"
 #include "GameFramework/PlayerMuteList.h"
-#include "AI/Navigation/NavigationData.h"
 #include "Animation/AnimSequence.h"
-#include "AI/Navigation/NavigationSystem.h"
+#include "AI/NavigationSystemBase.h"
 #include "Sound/SoundConcurrency.h"
 #include "ClothSimData.h"
 #include "AnimEncoding.h"
diff --git a/Engine/Source/Runtime/Engine/Private/GameInstance.cpp b/Engine/Source/Runtime/Engine/Private/GameInstance.cpp
--- a/Engine/Source/Runtime/Engine/Private/GameInstance.cpp
+++ b/Engine/Source/Runtime/Engine/Private/GameInstance.cpp
@@ -9,7 +9,7 @@
 #include "TimerManager.h"
 #include "Engine/LatentActionManager.h"
 #include "Engine/World.h"
-#include "AI/Navigation/NavigationSystem.h"
+#include "AI/NavigationSystemBase.h"
 #include "Misc/Paths.h"
 #include "UObject/CoreOnline.h"
 #include "GameFramework/PlayerController.h"
@@ -320,7 +320,7 @@ FGameInstancePIEResult UGameInstance::StartPlayInEditorGameInstance(ULocalPlayer
 
 		PlayWorld->InitializeActorsForPlay(URL);
 		// calling it after InitializeActorsForPlay has been called to have all potential bounding boxed initialized
-		UNavigationSystem::InitializeForWorld(PlayWorld, LocalPlayers.Num() > 0 ? FNavigationSystemRunMode::PIEMode : FNavigationSystemRunMode::SimulationMode);
+		FNavigationSystem::AddNavigationSystemToWorld(*PlayWorld, LocalPlayers.Num() > 0 ? FNavigationSystemRunMode::PIEMode : FNavigationSystemRunMode::SimulationMode);
 
 		// @todo, just use WorldContext.GamePlayer[0]?
 		if (LocalPlayer)
diff --git a/Engine/Source/Runtime/Engine/Private/GameNetworkManager.cpp b/Engine/Source/Runtime/Engine/Private/GameNetworkManager.cpp
--- a/Engine/Source/Runtime/Engine/Private/GameNetworkManager.cpp
+++ b/Engine/Source/Runtime/Engine/Private/GameNetworkManager.cpp
@@ -23,8 +23,9 @@ AGameNetworkManager::AGameNetworkManager(const FObjectInitializer& ObjectInitial
 	CLIENTADJUSTUPDATECOST = 180.0f;
 	MAXCLIENTUPDATEINTERVAL = 0.25f;
 	MaxMoveDeltaTime = 0.125f;
-	ClientNetSendMoveDeltaTime = 0.0111f;
+	ClientNetSendMoveDeltaTime = 0.0166f;
 	ClientNetSendMoveDeltaTimeThrottled = 0.0222f;
+	ClientNetSendMoveDeltaTimeStationary = 0.0833f;
 	ClientNetSendMoveThrottleAtNetSpeed = 10000;
 	ClientNetSendMoveThrottleOverPlayerCount = 10;
 	ClientAuthorativePosition = false;
diff --git a/Engine/Source/Runtime/Engine/Private/GameViewportClient.cpp b/Engine/Source/Runtime/Engine/Private/GameViewportClient.cpp
--- a/Engine/Source/Runtime/Engine/Private/GameViewportClient.cpp
+++ b/Engine/Source/Runtime/Engine/Private/GameViewportClient.cpp
@@ -12,7 +12,7 @@
 #include "RenderingThread.h"
 #include "SceneView.h"
 #include "LegacyScreenPercentageDriver.h"
-#include "AI/Navigation/NavigationSystem.h"
+#include "AI/NavigationSystemBase.h"
 #include "CanvasItem.h"
 #include "Engine/Canvas.h"
 #include "GameFramework/Volume.h"
@@ -2314,24 +2314,9 @@ void UGameViewportClient::VerifyPathRenderingComponents()
 
 	UWorld* const ViewportWorld = GetWorld();
 
-	// make sure nav mesh has a rendering component
-	ANavigationData* const NavData = (ViewportWorld && ViewportWorld->GetNavigationSystem() != nullptr)
-		? ViewportWorld->GetNavigationSystem()->GetMainNavData(FNavigationSystem::DontCreate)
-		: NULL;
-
-	if(NavData && NavData->RenderingComp == NULL)
-	{
-		NavData->RenderingComp = NavData->ConstructRenderingComponent();
-		if (NavData->RenderingComp)
-		{
-			NavData->RenderingComp->SetVisibility(bShowPaths);
-			NavData->RenderingComp->RegisterComponent();
-		}
-	}
-
-	if(NavData == NULL)
+	if (ViewportWorld)
 	{
-		UE_LOG(LogPlayerManagement, Warning, TEXT("No NavData found when calling UGameViewportClient::VerifyPathRenderingComponents()"));
+		FNavigationSystem::VerifyNavigationRenderingComponents(*ViewportWorld, bShowPaths);
 	}
 }
 
diff --git a/Engine/Source/Runtime/Engine/Private/HierarchicalInstancedStaticMesh.cpp b/Engine/Source/Runtime/Engine/Private/HierarchicalInstancedStaticMesh.cpp
--- a/Engine/Source/Runtime/Engine/Private/HierarchicalInstancedStaticMesh.cpp
+++ b/Engine/Source/Runtime/Engine/Private/HierarchicalInstancedStaticMesh.cpp
@@ -15,7 +15,7 @@
 #include "Async/AsyncWork.h"
 #include "PrimitiveViewRelevance.h"
 #include "ConvexVolume.h"
-#include "AI/Navigation/NavigationSystem.h"
+#include "AI/NavigationSystemBase.h"
 #include "Engine/MapBuildDataRegistry.h"
 #include "MaterialShared.h"
 #include "UObject/UObjectIterator.h"
@@ -2741,14 +2741,12 @@ void UHierarchicalInstancedStaticMeshComponent::PartialNavigationUpdate(int32 In
 	if (InstanceIdx == INDEX_NONE)
 	{
 		AccumulatedNavigationDirtyArea.Init();
-		UNavigationSystem::UpdateComponentInNavOctree(*this);
+		FNavigationSystem::UpdateComponentData(*this);
 	}
 	else if (GetStaticMesh())
 	{
 		// Accumulate dirty areas and send them to navigation system once cluster tree is rebuilt
-		UNavigationSystem* NavSys = UNavigationSystem::GetCurrent(GetWorld());
-		// Check if this component is registered in navigation system
-		if (NavSys && (NavSys->GetObjectsNavOctreeId(this) || NavSys->HasPendingObjectNavOctreeId(this)))
+		if (FNavigationSystem::HasComponentData(*this))
 		{
 			FTransform InstanceTransform(PerInstanceSMData[InstanceIdx].Transform);
 			FBox InstanceBox = GetStaticMesh()->GetBounds().TransformBy(InstanceTransform*GetComponentTransform()).GetBox(); // in world space
@@ -2764,12 +2762,10 @@ void UHierarchicalInstancedStaticMeshComponent::FlushAccumulatedNavigationUpdate
 		QUICK_SCOPE_CYCLE_COUNTER(STAT_UHierarchicalInstancedStaticMeshComponent_FlushAccumulatedNavigationUpdates);
 
 		const TArray<FClusterNode>& ClusterTree = *ClusterTreePtr;
-		UNavigationSystem* NavSys = UNavigationSystem::GetCurrent(GetWorld());
-		// Check if this component is registered in navigation system
-		if (ClusterTree.Num() && NavSys && NavSys->GetObjectsNavOctreeId(this))
+		if (ClusterTree.Num())
 		{
 			FBox NewBounds = FBox(ClusterTree[0].BoundMin, ClusterTree[0].BoundMax).TransformBy(GetComponentTransform());
-			NavSys->UpdateNavOctreeElementBounds(this, NewBounds, AccumulatedNavigationDirtyArea);
+			FNavigationSystem::OnComponentBoundsChanged(*this, NewBounds, AccumulatedNavigationDirtyArea);
 		}
 			
 		AccumulatedNavigationDirtyArea.Init();
diff --git a/Engine/Source/Runtime/Engine/Private/InstancedStaticMesh.cpp b/Engine/Source/Runtime/Engine/Private/InstancedStaticMesh.cpp
--- a/Engine/Source/Runtime/Engine/Private/InstancedStaticMesh.cpp
+++ b/Engine/Source/Runtime/Engine/Private/InstancedStaticMesh.cpp
@@ -5,15 +5,14 @@
 =============================================================================*/
 
 #include "InstancedStaticMesh.h"
-#include "AI/Navigation/NavigationSystem.h"
+#include "AI/NavigationSystemBase.h"
 #include "Engine/MapBuildDataRegistry.h"
 #include "Components/LightComponent.h"
 #include "Logging/TokenizedMessage.h"
 #include "Logging/MessageLog.h"
 #include "UnrealEngine.h"
 #include "AI/NavigationSystemHelpers.h"
-#include "AI/Navigation/NavCollision.h"
-#include "AI/NavigationOctree.h"
+#include "AI/Navigation/NavCollisionBase.h"
 #include "ShaderParameterUtils.h"
 #include "Misc/UObjectToken.h"
 #include "PhysXPublic.h"
@@ -24,6 +23,7 @@
 #include "SceneManagement.h"
 #include "Algo/Transform.h"
 #include "UObject/MobileObjectVersion.h"
+#include "EngineStats.h"
 
 const int32 InstancedStaticMeshMaxTexCoord = 8;
 
@@ -1904,7 +1904,7 @@ void UInstancedStaticMeshComponent::ClearInstances()
 	InitPerInstanceRenderData(false);
 	MarkRenderStateDirty();
 
-	UNavigationSystem::UpdateComponentInNavOctree(*this);
+	FNavigationSystem::UpdateComponentData(*this);
 }
 
 int32 UInstancedStaticMeshComponent::GetInstanceCount() const
@@ -2003,26 +2003,26 @@ void UInstancedStaticMeshComponent::PostLoad()
 void UInstancedStaticMeshComponent::PartialNavigationUpdate(int32 InstanceIdx)
 {
 	// Just update everything
-	UNavigationSystem::UpdateComponentInNavOctree(*this);
+	FNavigationSystem::UpdateComponentData(*this);
 }
 
 bool UInstancedStaticMeshComponent::DoCustomNavigableGeometryExport(FNavigableGeometryExport& GeomExport) const
 {
 	if (GetStaticMesh() && GetStaticMesh()->NavCollision)
 	{
-		UNavCollision* NavCollision = GetStaticMesh()->NavCollision;
-		if (GetStaticMesh()->NavCollision->bIsDynamicObstacle)
+		UNavCollisionBase* NavCollision = GetStaticMesh()->NavCollision;
+		if (GetStaticMesh()->NavCollision->IsDynamicObstacle())
 		{
 			return false;
 		}
 		
-		if (NavCollision->bHasConvexGeometry)
+		if (NavCollision->HasConvexGeometry())
 		{
-			GeomExport.ExportCustomMesh(NavCollision->ConvexCollision.VertexBuffer.GetData(), NavCollision->ConvexCollision.VertexBuffer.Num(),
-				NavCollision->ConvexCollision.IndexBuffer.GetData(), NavCollision->ConvexCollision.IndexBuffer.Num(), FTransform::Identity);
+			GeomExport.ExportCustomMesh(NavCollision->GetConvexCollision().VertexBuffer.GetData(), NavCollision->GetConvexCollision().VertexBuffer.Num(),
+				NavCollision->GetConvexCollision().IndexBuffer.GetData(), NavCollision->GetConvexCollision().IndexBuffer.Num(), FTransform::Identity);
 
-			GeomExport.ExportCustomMesh(NavCollision->TriMeshCollision.VertexBuffer.GetData(), NavCollision->TriMeshCollision.VertexBuffer.Num(),
-				NavCollision->TriMeshCollision.IndexBuffer.GetData(), NavCollision->TriMeshCollision.IndexBuffer.Num(), FTransform::Identity);
+			GeomExport.ExportCustomMesh(NavCollision->GetTriMeshCollision().VertexBuffer.GetData(), NavCollision->GetTriMeshCollision().VertexBuffer.Num(),
+				NavCollision->GetTriMeshCollision().IndexBuffer.GetData(), NavCollision->GetTriMeshCollision().IndexBuffer.Num(), FTransform::Identity);
 		}
 		else
 		{
@@ -2045,8 +2045,8 @@ void UInstancedStaticMeshComponent::GetNavigationData(FNavigationRelevantData& D
 {
 	if (GetStaticMesh() && GetStaticMesh()->NavCollision)
 	{
-		UNavCollision* NavCollision = GetStaticMesh()->NavCollision;
-		if (NavCollision->bIsDynamicObstacle)
+		UNavCollisionBase* NavCollision = GetStaticMesh()->NavCollision;
+		if (NavCollision->IsDynamicObstacle())
 		{
 			NavCollision->GetNavigationModifier(Data.Modifiers, FTransform::Identity);
 
@@ -2166,7 +2166,7 @@ void UInstancedStaticMeshComponent::PostEditUndo()
 {
 	Super::PostEditUndo();
 
-	UNavigationSystem::UpdateComponentInNavOctree(*this);
+	FNavigationSystem::UpdateComponentData(*this);
 }
 #endif
 
diff --git a/Engine/Source/Runtime/Engine/Private/KismetGuidLibrary.cpp b/Engine/Source/Runtime/Engine/Private/KismetGuidLibrary.cpp
--- a/Engine/Source/Runtime/Engine/Private/KismetGuidLibrary.cpp
+++ b/Engine/Source/Runtime/Engine/Private/KismetGuidLibrary.cpp
@@ -35,7 +35,7 @@ bool UKismetGuidLibrary::IsValid_Guid( const FGuid& InGuid )
 }
 
 
-void UKismetGuidLibrary::Invalidate_Guid( UPARAM(ref) FGuid& InGuid )
+void UKismetGuidLibrary::Invalidate_Guid( FGuid& InGuid )
 {
 	InGuid.Invalidate();
 }
diff --git a/Engine/Source/Runtime/Engine/Private/Level.cpp b/Engine/Source/Runtime/Engine/Private/Level.cpp
--- a/Engine/Source/Runtime/Engine/Private/Level.cpp
+++ b/Engine/Source/Runtime/Engine/Private/Level.cpp
@@ -19,7 +19,6 @@ Level.cpp: Level-related functions
 #include "GameFramework/Pawn.h"
 #include "Engine/World.h"
 #include "SceneInterface.h"
-#include "AI/Navigation/NavigationData.h"
 #include "PrecomputedLightVolume.h"
 #include "PrecomputedVolumetricLightmap.h"
 #include "Engine/MapBuildDataRegistry.h"
@@ -2114,13 +2113,9 @@ void ULevel::ApplyWorldOffset(const FVector& InWorldOffset, bool bWorldShift)
 			AActor* Actor = Actors[ActorIndex];
 			if (Actor)
 			{
-				FVector Offset = (bWorldShift && Actor->bIgnoresOriginShifting) ? FVector::ZeroVector : InWorldOffset;
-
-				if (!Actor->IsA(ANavigationData::StaticClass())) // Navigation data will be moved in NavigationSystem
-				{
-					FScopeCycleCounterUObject Context(Actor);
-					Actor->ApplyWorldOffset(Offset, bWorldShift);
-				}
+				const FVector Offset = (bWorldShift && Actor->bIgnoresOriginShifting) ? FVector::ZeroVector : InWorldOffset;
+				FScopeCycleCounterUObject Context(Actor);
+				Actor->ApplyWorldOffset(Offset, bWorldShift);
 			}
 		}
 	}
diff --git a/Engine/Source/Runtime/Engine/Private/LevelActor.cpp b/Engine/Source/Runtime/Engine/Private/LevelActor.cpp
--- a/Engine/Source/Runtime/Engine/Private/LevelActor.cpp
+++ b/Engine/Source/Runtime/Engine/Private/LevelActor.cpp
@@ -20,7 +20,7 @@
 #include "WorldCollision.h"
 #include "Engine/World.h"
 #include "Components/PrimitiveComponent.h"
-#include "AI/Navigation/NavigationSystem.h"
+#include "AI/NavigationSystemBase.h"
 #include "Engine/Brush.h"
 #include "UObject/LinkerLoad.h"
 #include "UObject/CoreOnline.h"
@@ -505,7 +505,7 @@ ABrush* UWorld::SpawnBrush()
  */
 bool UWorld::EditorDestroyActor( AActor* ThisActor, bool bShouldModifyLevel )
 {
-	UNavigationSystem::OnActorUnregistered(ThisActor);
+	FNavigationSystem::OnActorUnregistered(*ThisActor);
 
 	bool bReturnValue = DestroyActor( ThisActor, false, bShouldModifyLevel );
 	ThisActor->GetWorld()->BroadcastLevelsChanged();
diff --git a/Engine/Source/Runtime/Engine/Private/LevelTick.cpp b/Engine/Source/Runtime/Engine/Private/LevelTick.cpp
--- a/Engine/Source/Runtime/Engine/Private/LevelTick.cpp
+++ b/Engine/Source/Runtime/Engine/Private/LevelTick.cpp
@@ -23,7 +23,7 @@
 #include "RenderingThread.h"
 #include "Engine/World.h"
 #include "GameFramework/Controller.h"
-#include "AI/Navigation/NavigationSystem.h"
+#include "AI/NavigationSystemBase.h"
 #include "GameFramework/PlayerController.h"
 #include "SceneUtils.h"
 #include "ParticleHelper.h"
@@ -1383,13 +1383,10 @@ void UWorld::Tick( ELevelTick TickType, float DeltaSeconds )
 	}
 	
 	// update world's subsystems (NavigationSystem for now)
-	if ( !bIsPaused )
+	if (NavigationSystem != NULL)
 	{
-		if (NavigationSystem != NULL)
-		{
-			SCOPE_CYCLE_COUNTER(STAT_NavWorldTickTime);
-			NavigationSystem->Tick(DeltaSeconds);
-		}
+		SCOPE_CYCLE_COUNTER(STAT_NavWorldTickTime);
+		NavigationSystem->Tick(DeltaSeconds);
 	}
 
 	bool bDoingActorTicks = 
diff --git a/Engine/Source/Runtime/Engine/Private/Pawn.cpp b/Engine/Source/Runtime/Engine/Private/Pawn.cpp
--- a/Engine/Source/Runtime/Engine/Private/Pawn.cpp
+++ b/Engine/Source/Runtime/Engine/Private/Pawn.cpp
@@ -11,7 +11,7 @@
 #include "Engine/World.h"
 #include "GameFramework/Controller.h"
 #include "Components/PrimitiveComponent.h"
-#include "AI/Navigation/NavigationSystem.h"
+#include "AI/NavigationSystemBase.h"
 #include "Components/InputComponent.h"
 #include "GameFramework/PlayerController.h"
 #include "Engine/Engine.h"
@@ -167,8 +167,8 @@ void APawn::SetCanAffectNavigationGeneration(bool bNewValue, bool bForceUpdate)
 		// update components 
 		UpdateNavigationRelevance();
 
-		// update entries in navigation octree 
-		UNavigationSystem::UpdateActorAndComponentsInNavOctree(*this);
+		// update entries in navigation system
+		FNavigationSystem::UpdateActorAndComponentData(*this);
 	}
 }
 
diff --git a/Engine/Source/Runtime/Engine/Private/PlayerController.cpp b/Engine/Source/Runtime/Engine/Private/PlayerController.cpp
--- a/Engine/Source/Runtime/Engine/Private/PlayerController.cpp
+++ b/Engine/Source/Runtime/Engine/Private/PlayerController.cpp
@@ -806,7 +806,9 @@ void APlayerController::Possess(APawn* PawnToPossess)
 			AutoManageActiveCameraTarget(GetPawn());
 			ResetCameraMode();
 		}
-		UpdateNavigationComponents();
+		// not calling UpdateNavigationComponents() anymore. The
+		// PathFollowingComponent is now observing newly possessed
+		// pawns (via OnNewPawn)
 	}
 }
 
diff --git a/Engine/Source/Runtime/Engine/Private/PrimitiveComponentPhysics.cpp b/Engine/Source/Runtime/Engine/Private/PrimitiveComponentPhysics.cpp
--- a/Engine/Source/Runtime/Engine/Private/PrimitiveComponentPhysics.cpp
+++ b/Engine/Source/Runtime/Engine/Private/PrimitiveComponentPhysics.cpp
@@ -6,7 +6,7 @@
 #include "CollisionQueryParams.h"
 #include "Engine/World.h"
 #include "Components/PrimitiveComponent.h"
-#include "AI/Navigation/NavigationSystem.h"
+#include "AI/NavigationSystemBase.h"
 #include "Components/LineBatchComponent.h"
 #include "Logging/MessageLog.h"
 #include "PhysicsEngine/BodySetup.h"
@@ -950,7 +950,7 @@ void UPrimitiveComponent::OnComponentCollisionSettingsChanged()
 		if (bNavigationRelevant != bNewNavRelevant)
 		{
 			bNavigationRelevant = bNewNavRelevant;
-			UNavigationSystem::UpdateComponentInNavOctree(*this);
+			FNavigationSystem::UpdateComponentData(*this);
 		}
 
 		OnComponentCollisionSettingsChangedEvent.Broadcast(this);
diff --git a/Engine/Source/Runtime/Engine/Private/StaticMesh.cpp b/Engine/Source/Runtime/Engine/Private/StaticMesh.cpp
--- a/Engine/Source/Runtime/Engine/Private/StaticMesh.cpp
+++ b/Engine/Source/Runtime/Engine/Private/StaticMesh.cpp
@@ -50,8 +50,8 @@
 
 #include "Engine/StaticMeshSocket.h"
 #include "EditorFramework/AssetImportData.h"
-#include "AI/Navigation/NavCollision.h"
-#include "AI/Navigation/NavigationSystem.h"
+#include "AI/Navigation/NavCollisionBase.h"
+#include "AI/NavigationSystemBase.h"
 #include "AI/NavigationSystemHelpers.h"
 #include "ProfilingDebugging/CookStats.h"
 #include "UObject/ReleaseObjectVersion.h"
@@ -2114,7 +2114,7 @@ void UStaticMesh::SetLODGroup(FName NewGroup, bool bRebuildImmediately)
 
 void UStaticMesh::BroadcastNavCollisionChange()
 {
-	if (UNavigationSystem::ShouldUpdateNavOctreeOnComponentChange())
+	if (FNavigationSystem::WantsComponentChangeNotifies())
 	{
 		for (FObjectIterator Iter(UStaticMeshComponent::StaticClass()); Iter; ++Iter)
 		{
@@ -2123,7 +2123,7 @@ void UStaticMesh::BroadcastNavCollisionChange()
 			if (StaticMeshComponent->GetStaticMesh() == this)
 			{
 				StaticMeshComponent->bNavigationRelevant = StaticMeshComponent->IsNavigationRelevant();
-				UNavigationSystem::UpdateComponentInNavOctree(*StaticMeshComponent);
+				FNavigationSystem::UpdateComponentData(*StaticMeshComponent);
 			}
 		}
 	}
@@ -3235,17 +3235,19 @@ void UStaticMesh::CreateNavCollision(const bool bIsUpdate)
 	{
 		if (NavCollision == nullptr)
 		{
-			NavCollision = NewObject<UNavCollision>(this);
+			NavCollision = UNavCollisionBase::ConstructNew(*this);
 		}
 
-#if WITH_EDITOR
-		if (bIsUpdate)
+		if (NavCollision)
 		{
-			NavCollision->InvalidateCollision();
-		}
+#if WITH_EDITOR
+			if (bIsUpdate)
+			{
+				NavCollision->InvalidateCollision();
+			}
 #endif // WITH_EDITOR
-
-		NavCollision->Setup(BodySetup);	
+			NavCollision->Setup(BodySetup);
+		}
 	}
 	else
 	{
diff --git a/Engine/Source/Runtime/Engine/Private/StaticMeshRender.cpp b/Engine/Source/Runtime/Engine/Private/StaticMeshRender.cpp
--- a/Engine/Source/Runtime/Engine/Private/StaticMeshRender.cpp
+++ b/Engine/Source/Runtime/Engine/Private/StaticMeshRender.cpp
@@ -34,7 +34,7 @@
 #include "TessellationRendering.h"
 #include "DistanceFieldAtlas.h"
 #include "Components/BrushComponent.h"
-#include "AI/Navigation/NavCollision.h"
+#include "AI/Navigation/NavCollisionBase.h"
 #include "ComponentRecreateRenderStateContext.h"
 #include "PhysicsEngine/BodySetup.h"
 #include "Engine/LODActor.h"
@@ -1249,7 +1249,7 @@ void FStaticMeshSceneProxy::GetDynamicMeshElements(const TArray<const FSceneView
 
 
 					// The simple nav geometry is only used by dynamic obstacles for now
-					if (StaticMesh->NavCollision && StaticMesh->NavCollision->bIsDynamicObstacle)
+					if (StaticMesh->NavCollision && StaticMesh->NavCollision->IsDynamicObstacle())
 					{
 						// Draw the static mesh's body setup (simple collision)
 						FTransform GeomTransform(GetLocalToWorld());
diff --git a/Engine/Source/Runtime/Engine/Private/TickTaskManager.cpp b/Engine/Source/Runtime/Engine/Private/TickTaskManager.cpp
--- a/Engine/Source/Runtime/Engine/Private/TickTaskManager.cpp
+++ b/Engine/Source/Runtime/Engine/Private/TickTaskManager.cpp
@@ -1659,12 +1659,15 @@ void FTickFunction::UnRegisterTickFunction()
 /** Enables or disables this tick function. **/
 void FTickFunction::SetTickFunctionEnable(bool bInEnabled)
 {
-	if (bRegistered && (bInEnabled == (TickState == ETickState::Disabled)))
+	if (bRegistered)
 	{
-		check(TickTaskLevel);
-		TickTaskLevel->RemoveTickFunction(this);
-		TickState = (bInEnabled ? ETickState::Enabled : ETickState::Disabled);
-		TickTaskLevel->AddTickFunction(this);
+		if (bInEnabled == (TickState == ETickState::Disabled))
+		{
+			check(TickTaskLevel);
+			TickTaskLevel->RemoveTickFunction(this);
+			TickState = (bInEnabled ? ETickState::Enabled : ETickState::Disabled);
+			TickTaskLevel->AddTickFunction(this);
+		}
 	}
 	else
 	{
diff --git a/Engine/Source/Runtime/Engine/Private/UnrealEngine.cpp b/Engine/Source/Runtime/Engine/Private/UnrealEngine.cpp
--- a/Engine/Source/Runtime/Engine/Private/UnrealEngine.cpp
+++ b/Engine/Source/Runtime/Engine/Private/UnrealEngine.cpp
@@ -18,7 +18,7 @@ UnrealEngine.cpp: Implements the UEngine class and helpers.
 #include "RHIStaticStates.h"
 #include "Engine/TextureStreamingTypes.h"
 #include "Components/PrimitiveComponent.h"
-#include "AI/Navigation/NavigationSystem.h"
+#include "AI/NavigationSystemBase.h"
 #include "Misc/MessageDialog.h"
 #include "HAL/FileManager.h"
 #include "Misc/CommandLine.h"
@@ -1935,11 +1935,19 @@ void LoadSpecialMaterial(const FString& MaterialName, const FString& MaterialNam
 
 
 template<typename ClassType>
-void LoadEngineClass(const FSoftClassPath& ClassName, TSubclassOf<ClassType>& EngineClassRef)
+void LoadEngineClass(FSoftClassPath& ClassName, TSubclassOf<ClassType>& EngineClassRef)
 {
 	if ( EngineClassRef == nullptr )
 	{
 		EngineClassRef = LoadClass<ClassType>(nullptr, *ClassName.ToString());
+
+		// try redirects
+		if (EngineClassRef == nullptr)
+		{
+			ClassName.FixupCoreRedirects();
+			EngineClassRef = LoadClass<ClassType>(nullptr, *ClassName.ToString());
+		}
+
 		if (EngineClassRef == nullptr)
 		{
 			EngineClassRef = ClassType::StaticClass();
@@ -2098,7 +2106,7 @@ void UEngine::InitializeObjectReferences()
 	LoadEngineClass<UGameViewportClient>(GameViewportClientClassName, GameViewportClientClass);
 	LoadEngineClass<ULocalPlayer>(LocalPlayerClassName, LocalPlayerClass);
 	LoadEngineClass<AWorldSettings>(WorldSettingsClassName, WorldSettingsClass);
-	LoadEngineClass<UNavigationSystem>(NavigationSystemClassName, NavigationSystemClass);
+	LoadEngineClass<UNavigationSystemBase>(NavigationSystemClassName, NavigationSystemClass);
 	LoadEngineClass<UAvoidanceManager>(AvoidanceManagerClassName, AvoidanceManagerClass);
 	LoadEngineClass<UPhysicsCollisionHandler>(PhysicsCollisionHandlerClassName, PhysicsCollisionHandlerClass);
 	LoadEngineClass<UGameUserSettings>(GameUserSettingsClassName, GameUserSettingsClass);
@@ -9066,22 +9074,6 @@ float DrawMapWarnings(UWorld* World, FViewport* Viewport, FCanvas* Canvas, UCanv
 
 	}
 
-	// check navmesh
-#if WITH_EDITOR
-	const bool bIsNavigationAutoUpdateEnabled = UNavigationSystem::GetIsNavigationAutoUpdateEnabled();
-#else
-	const bool bIsNavigationAutoUpdateEnabled = true;
-#endif
-	UNavigationSystem* NavSys = UNavigationSystem::GetCurrent(World);
-	if (NavSys && NavSys->IsNavigationDirty() &&
-		(!bIsNavigationAutoUpdateEnabled || !NavSys->SupportsNavigationGeneration() || !NavSys->CanRebuildDirtyNavigation()))
-	{
-		SmallTextItem.SetColor(FLinearColor::White);
-		SmallTextItem.Text = LOCTEXT("NAVMESHERROR", "NAVMESH NEEDS TO BE REBUILT");
-		Canvas->DrawItem(SmallTextItem, FVector2D(MessageX, MessageY));
-		MessageY += FontSizeY;
-	}
-
 	if (World->bKismetScriptError)
 	{
 		SmallTextItem.Text = LOCTEXT("BlueprintInLevelHadCompileErrorMessage", "BLUEPRINT COMPILE ERROR");
@@ -11722,7 +11714,7 @@ bool UEngine::LoadMap( FWorldContext& WorldContext, FURL URL, class UPendingNetG
 	WorldContext.World()->InitializeActorsForPlay(URL);		
 
 	// calling it after InitializeActorsForPlay has been called to have all potential bounding boxed initialized
-	UNavigationSystem::InitializeForWorld(WorldContext.World(), FNavigationSystemRunMode::GameMode);
+	FNavigationSystem::AddNavigationSystemToWorld(*WorldContext.World(), FNavigationSystemRunMode::GameMode);
 
 	// Remember the URL. Put this before spawning player controllers so that
 	// a player controller can get the map name during initialization and
diff --git a/Engine/Source/Runtime/Engine/Private/VisualLogger/VisualLogger.cpp b/Engine/Source/Runtime/Engine/Private/VisualLogger/VisualLogger.cpp
--- a/Engine/Source/Runtime/Engine/Private/VisualLogger/VisualLogger.cpp
+++ b/Engine/Source/Runtime/Engine/Private/VisualLogger/VisualLogger.cpp
@@ -9,7 +9,7 @@
 #include "Modules/ModuleManager.h"
 #include "EngineGlobals.h"
 #include "AI/NavDataGenerator.h"
-#include "AI/Navigation/NavigationSystem.h"
+#include "AI/NavigationSystemBase.h"
 #include "Framework/Docking/TabManager.h"
 #include "VisualLogger/VisualLoggerBinaryFileDevice.h"
 #include "VisualLogger/VisualLoggerDebugSnapshotInterface.h"
@@ -50,6 +50,7 @@ namespace
 
 TMap<const UWorld*, FVisualLogger::RedirectionMapType> FVisualLogger::WorldToRedirectionMap;
 int32 FVisualLogger::bIsRecording = false;
+FVisualLogger::FNavigationDataDump FVisualLogger::NavigationDataDumpDelegate;
 
 bool FVisualLogger::CheckVisualLogInputInternal(const UObject* Object, const FName& CategoryName, ELogVerbosity::Type Verbosity, UWorld **World, FVisualLogEntry **CurrentEntry)
 {
@@ -337,19 +338,7 @@ void FVisualLogger::NavigationDataDump(const UObject* Object, const FLogCategory
 void FVisualLogger::NavigationDataDump(const UObject* Object, const FName& CategoryName, const ELogVerbosity::Type Verbosity, const FBox& Box)
 {
 	SCOPE_CYCLE_COUNTER(STAT_VisualLog);
-	UWorld* World = nullptr;
-	FVisualLogEntry* CurrentEntry = nullptr;
-	if (CheckVisualLogInputInternal(Object, CategoryName, Verbosity, &World, &CurrentEntry) == false || CurrentEntry == nullptr)
-	{
-		return;
-	}
-
-	const ANavigationData* MainNavData = World ? UNavigationSystem::GetNavigationSystem(World)->GetMainNavData(FNavigationSystem::ECreateIfEmpty::DontCreate) : nullptr;
-	const FNavDataGenerator* Generator = MainNavData ? MainNavData->GetGenerator() : nullptr;
-	if (Generator)
-	{
-		Generator->GrabDebugSnapshot(CurrentEntry, FMath::IsNearlyZero(Box.GetVolume()) ? MainNavData->GetBounds().ExpandBy(FVector(20,20,20)) : Box, CategoryName, Verbosity);
-	}
+	NavigationDataDumpDelegate.Broadcast(Object, CategoryName, Verbosity, Box);
 }
 
 
diff --git a/Engine/Source/Runtime/Engine/Private/VisualLogger/VisualLoggerTypes.cpp b/Engine/Source/Runtime/Engine/Private/VisualLogger/VisualLoggerTypes.cpp
--- a/Engine/Source/Runtime/Engine/Private/VisualLogger/VisualLoggerTypes.cpp
+++ b/Engine/Source/Runtime/Engine/Private/VisualLogger/VisualLoggerTypes.cpp
@@ -540,29 +540,29 @@ FArchive& FVisualLoggerHelpers::Serialize(FArchive& Ar, TArray<FVisualLogDevice:
 
 void FVisualLoggerHelpers::GetCategories(const FVisualLogEntry& EntryItem, TArray<FVisualLoggerCategoryVerbosityPair>& OutCategories)
 {
-	for (const auto& CurrentEvent : EntryItem.Events)
+	for (const auto& Element : EntryItem.Events)
 	{
-		OutCategories.AddUnique(FVisualLoggerCategoryVerbosityPair(*CurrentEvent.Name, ELogVerbosity::All));
+		OutCategories.AddUnique(FVisualLoggerCategoryVerbosityPair(*Element.Name, Element.Verbosity));
 	}
 
-	for (const auto& CurrentLine : EntryItem.LogLines)
+	for (const auto& Element : EntryItem.LogLines)
 	{
-		OutCategories.AddUnique(FVisualLoggerCategoryVerbosityPair(CurrentLine.Category, ELogVerbosity::All));
+		OutCategories.AddUnique(FVisualLoggerCategoryVerbosityPair(Element.Category, Element.Verbosity));
 	}
 
-	for (const auto& CurrentElement : EntryItem.ElementsToDraw)
+	for (const auto& Element : EntryItem.ElementsToDraw)
 	{
-		OutCategories.AddUnique(FVisualLoggerCategoryVerbosityPair(CurrentElement.Category, ELogVerbosity::All));
+		OutCategories.AddUnique(FVisualLoggerCategoryVerbosityPair(Element.Category, Element.Verbosity));
 	}
 
-	for (const auto& CurrentSample : EntryItem.HistogramSamples)
+	for (const auto& Element : EntryItem.HistogramSamples)
 	{
-		OutCategories.AddUnique(FVisualLoggerCategoryVerbosityPair(CurrentSample.Category, ELogVerbosity::All));
+		OutCategories.AddUnique(FVisualLoggerCategoryVerbosityPair(Element.Category, Element.Verbosity));
 	}
 
-	for (const auto& CurrentBlock : EntryItem.DataBlocks)
+	for (const auto& Element : EntryItem.DataBlocks)
 	{
-		OutCategories.AddUnique(FVisualLoggerCategoryVerbosityPair(CurrentBlock.Category, ELogVerbosity::All));
+		OutCategories.AddUnique(FVisualLoggerCategoryVerbosityPair(Element.Category, Element.Verbosity));
 	}
 }
 
diff --git a/Engine/Source/Runtime/Engine/Private/World.cpp b/Engine/Source/Runtime/Engine/Private/World.cpp
--- a/Engine/Source/Runtime/Engine/Private/World.cpp
+++ b/Engine/Source/Runtime/Engine/Private/World.cpp
@@ -24,7 +24,7 @@
 #include "TimerManager.h"
 #include "Materials/MaterialInterface.h"
 #include "GameFramework/Controller.h"
-#include "AI/Navigation/NavigationSystem.h"
+#include "AI/NavigationSystemBase.h"
 #include "Engine/MapBuildDataRegistry.h"
 #include "Model.h"
 #include "Engine/Brush.h"
@@ -1086,11 +1086,11 @@ UAISystemBase* UWorld::CreateAISystem()
 	// create navigation system for editor and server targets, but remove it from game clients
 	if (AISystem == NULL && UAISystemBase::ShouldInstantiateInNetMode(GetNetMode()) && PersistentLevel)
 	{
-		FName AIModuleName = UAISystemBase::GetAISystemModuleName();
+		const FName AIModuleName = UAISystemBase::GetAISystemModuleName();
 		const AWorldSettings* WorldSettings = PersistentLevel->GetWorldSettings(false);
 		if (AIModuleName.IsNone() == false && WorldSettings && WorldSettings->bEnableAISystem)
 		{
-			IAISystemModule* AISystemModule = FModuleManager::LoadModulePtr<IAISystemModule>(UAISystemBase::GetAISystemModuleName());
+			IAISystemModule* AISystemModule = FModuleManager::LoadModulePtr<IAISystemModule>(AIModuleName);
 			if (AISystemModule)
 			{
 				AISystem = AISystemModule->CreateAISystemInstance(this);
@@ -1193,14 +1193,14 @@ void UWorld::InitWorld(const InitializationValues IVS)
 		AWorldSettings* WorldSettings = GetWorldSettings();
 		if (WorldSettings)
 		{
-			if (IVS.bCreateNavigation && WorldSettings->bEnableNavigationSystem)
-		{
-			UNavigationSystem::CreateNavigationSystem(this);
-		}
+			if (IVS.bCreateNavigation)
+			{
+				FNavigationSystem::AddNavigationSystemToWorld(*this, FNavigationSystemRunMode::InvalidMode, /*bInitializeForWorld=*/false);
+			}
 			if (IVS.bCreateAISystem && WorldSettings->bEnableAISystem)
-	{
-		CreateAISystem();
-	}
+			{
+				CreateAISystem();
+			}
 		}
 	}
 	
@@ -2340,16 +2340,18 @@ void UWorld::RemoveFromWorld( ULevel* Level, bool bAllowIncrementalRemoval )
 			// Remove any pawns from the pawn list that are about to be streamed out
 			for( FConstPawnIterator Iterator = GetPawnIterator(); Iterator; ++Iterator )
 			{
-				APawn* Pawn = Iterator->Get();
-				if (Pawn->IsInLevel(Level))
-				{
-					RemovePawn(Pawn);
-					--Iterator;
-				}
-				else if (UCharacterMovementComponent* CharacterMovement = Cast<UCharacterMovementComponent>(Pawn->GetMovementComponent()))
+				if (APawn* Pawn = Iterator->Get())
 				{
-					// otherwise force floor check in case the floor was streamed out from under it
-					CharacterMovement->bForceNextFloorCheck = true;
+					if (Pawn->IsInLevel(Level))
+					{
+						RemovePawn(Pawn);
+						--Iterator;
+					}
+					else if (UCharacterMovementComponent* CharacterMovement = Cast<UCharacterMovementComponent>(Pawn->GetMovementComponent()))
+					{
+						// otherwise force floor check in case the floor was streamed out from under it
+						CharacterMovement->bForceNextFloorCheck = true;
+					}
 				}
 			}
 
@@ -5606,9 +5608,6 @@ UWorld* FSeamlessTravelHandler::Tick()
 					{
 						KeptGameState = static_cast<AGameStateBase*>(TheActor);
 					}
-					// add to new world's actor list and remove from old
-					LoadedWorld->PersistentLevel->Actors.Add(TheActor);
-					LoadedWorld->PersistentLevel->ActorsForGC.Add(TheActor);
 
 					TheActor->bActorSeamlessTraveled = true;
 				}
@@ -5749,7 +5748,7 @@ UWorld* FSeamlessTravelHandler::Tick()
 			LoadedWorld->InitializeActorsForPlay(PendingTravelURL, false);
 
 			// calling it after InitializeActorsForPlay has been called to have all potential bounding boxed initialized
-			UNavigationSystem::InitializeForWorld(LoadedWorld, FNavigationSystemRunMode::GameMode);
+			FNavigationSystem::AddNavigationSystemToWorld(*LoadedWorld, FNavigationSystemRunMode::GameMode);
 
 			// send loading complete notifications for all local players
 			for (FLocalPlayerIterator It(GEngine, LoadedWorld); It; ++It)
@@ -6233,11 +6232,11 @@ bool UWorld::ServerTravel(const FString& FURL, bool bAbsolute, bool bShouldSkipG
 	return true;
 }
 
-void UWorld::SetNavigationSystem( UNavigationSystem* InNavigationSystem)
+void UWorld::SetNavigationSystem(UNavigationSystemBase* InNavigationSystem)
 {
 	if (NavigationSystem != NULL && NavigationSystem != InNavigationSystem)
 	{
-		NavigationSystem->CleanUp(UNavigationSystem::CleanupWithWorld);
+		NavigationSystem->CleanUp(FNavigationSystem::ECleanupMode::CleanupWithWorld);
 	}
 
 	NavigationSystem = InNavigationSystem;
diff --git a/Engine/Source/Runtime/Engine/Private/WorldSettings.cpp b/Engine/Source/Runtime/Engine/Private/WorldSettings.cpp
--- a/Engine/Source/Runtime/Engine/Private/WorldSettings.cpp
+++ b/Engine/Source/Runtime/Engine/Private/WorldSettings.cpp
@@ -22,6 +22,8 @@
 #include "PhysicsEngine/PhysicsSettings.h"
 #include "UObject/ReleaseObjectVersion.h"
 #include "SceneManagement.h"
+#include "AI/NavigationSystemConfig.h"
+#include "AI/NavigationSystemBase.h"
 
 #if WITH_EDITOR
 #include "Editor.h"
@@ -52,6 +54,7 @@ AWorldSettings::AWorldSettings(const FObjectInitializer& ObjectInitializer)
 
 	bEnableWorldBoundsChecks = true;
 	bEnableNavigationSystem = true;
+	NavigationSystemConfig = nullptr;
 	bEnableAISystem = true;
 	bEnableWorldComposition = false;
 	bEnableWorldOriginRebasing = false;
@@ -94,6 +97,15 @@ AWorldSettings::AWorldSettings(const FObjectInitializer& ObjectInitializer)
 #endif // WITH_EDITORONLY_DATA
 }
 
+void AWorldSettings::PostInitProperties()
+{
+	Super::PostInitProperties();
+	if (HasAnyFlags(RF_NeedLoad|RF_WasLoaded|RF_ClassDefaultObject) == false)
+	{
+		NavigationSystemConfig = NewObject<UNavigationSystemConfig>(this);
+	}
+}
+
 void AWorldSettings::PreInitializeComponents()
 {
 	Super::PreInitializeComponents();
@@ -373,6 +385,21 @@ void AWorldSettings::PostLoad()
 
 	SetIsTemporarilyHiddenInEditor(true);
 #endif// WITH_EDITOR
+
+	if (bEnableNavigationSystem && NavigationSystemConfig == nullptr)
+	{
+		ULevel* Level = GetLevel();
+		if (Level && Level->IsPersistentLevel())
+		{
+			NavigationSystemConfig = NewObject<UNavigationSystemConfig>(this);
+			bEnableNavigationSystem = false;
+		}
+	}
+}
+
+bool AWorldSettings::IsNavigationSystemEnabled() const
+{
+	return NavigationSystemConfig && NavigationSystemConfig->NavigationSystemClass.IsValid();
 }
 
 
@@ -475,12 +502,12 @@ void AWorldSettings::PostEditChangeProperty(FPropertyChangedEvent& PropertyChang
 	UProperty* PropertyThatChanged = PropertyChangedEvent.Property;
 	if (PropertyThatChanged)
 	{
-		if (PropertyThatChanged->GetFName()==GET_MEMBER_NAME_CHECKED(AWorldSettings,bForceNoPrecomputedLighting) && bForceNoPrecomputedLighting)
+		const FName PropertyName = PropertyThatChanged->GetFName();
+		if (PropertyName == GET_MEMBER_NAME_CHECKED(AWorldSettings,bForceNoPrecomputedLighting) && bForceNoPrecomputedLighting)
 		{
 			FMessageDialog::Open( EAppMsgType::Ok, LOCTEXT("bForceNoPrecomputedLightingIsEnabled", "bForceNoPrecomputedLighting is now enabled, build lighting once to propagate the change (will remove existing precomputed lighting data)."));
 		}
-
-		else if (PropertyThatChanged->GetFName()==GET_MEMBER_NAME_CHECKED(AWorldSettings,bEnableWorldComposition))
+		else if (PropertyName == GET_MEMBER_NAME_CHECKED(AWorldSettings,bEnableWorldComposition))
 		{
 			if (UWorldComposition::EnableWorldCompositionEvent.IsBound())
 			{
@@ -491,6 +518,18 @@ void AWorldSettings::PostEditChangeProperty(FPropertyChangedEvent& PropertyChang
 				bEnableWorldComposition = false;
 			}
 		}
+		else if (PropertyName == GET_MEMBER_NAME_CHECKED(AWorldSettings, NavigationSystemConfig))
+		{
+			UWorld* World = GetWorld();
+			if (World)
+			{
+				World->SetNavigationSystem(nullptr);
+				if (NavigationSystemConfig)
+				{
+					FNavigationSystem::AddNavigationSystemToWorld(*World, FNavigationSystemRunMode::EditorMode);
+				}
+			}
+		}
 	}
 
 	LightmassSettings.NumIndirectLightingBounces = FMath::Clamp(LightmassSettings.NumIndirectLightingBounces, 0, 100);
diff --git a/Engine/Source/Runtime/Engine/Public/AI/Navigation/NavQueryFilter.h b/Engine/Source/Runtime/Engine/Public/AI/Navigation/NavQueryFilter.h
new file mode 100644
--- /dev/null
+++ b/Engine/Source/Runtime/Engine/Public/AI/Navigation/NavQueryFilter.h
@@ -0,0 +1,121 @@
+// Copyright 1998-2018 Epic Games, Inc. All Rights Reserved.
+
+#pragma once
+
+#include "CoreMinimal.h"
+
+class INavigationQueryFilterInterface
+{
+public:
+	virtual ~INavigationQueryFilterInterface(){}
+
+	virtual void Reset() = 0;
+
+	virtual void SetAreaCost(uint8 AreaType, float Cost) = 0;
+	virtual void SetFixedAreaEnteringCost(uint8 AreaType, float Cost) = 0;
+	virtual void SetExcludedArea(uint8 AreaType) = 0;
+	virtual void SetAllAreaCosts(const float* CostArray, const int32 Count) = 0;
+	virtual void GetAllAreaCosts(float* CostArray, float* FixedCostArray, const int32 Count) const = 0;
+	virtual void SetBacktrackingEnabled(const bool bBacktracking) = 0;
+	virtual bool IsBacktrackingEnabled() const = 0;
+	virtual bool IsEqual(const INavigationQueryFilterInterface* Other) const = 0;
+	virtual void SetIncludeFlags(uint16 Flags) = 0;
+	virtual uint16 GetIncludeFlags() const = 0;
+	virtual void SetExcludeFlags(uint16 Flags) = 0;
+	virtual uint16 GetExcludeFlags() const = 0;
+	virtual FVector GetAdjustedEndLocation(const FVector& EndLocation) const { return EndLocation; }
+
+	virtual INavigationQueryFilterInterface* CreateCopy() const = 0;
+};
+
+struct FNavigationQueryFilter;
+typedef TSharedPtr<FNavigationQueryFilter, ESPMode::ThreadSafe> FSharedNavQueryFilter;
+typedef TSharedPtr<const FNavigationQueryFilter, ESPMode::ThreadSafe> FSharedConstNavQueryFilter;
+
+struct ENGINE_API FNavigationQueryFilter : public TSharedFromThis<FNavigationQueryFilter, ESPMode::ThreadSafe>
+{
+	FNavigationQueryFilter() : QueryFilterImpl(NULL), MaxSearchNodes(DefaultMaxSearchNodes) {}
+private:
+	FNavigationQueryFilter(const FNavigationQueryFilter& Source);
+	FNavigationQueryFilter(const FNavigationQueryFilter* Source);
+	FNavigationQueryFilter(const FSharedNavQueryFilter Source);
+	FNavigationQueryFilter& operator=(const FNavigationQueryFilter& Source);
+public:
+
+	/** set travel cost for area */
+	void SetAreaCost(uint8 AreaType, float Cost);
+
+	/** set entering cost for area */
+	void SetFixedAreaEnteringCost(uint8 AreaType, float Cost);
+
+	/** mark area as excluded from path finding */
+	void SetExcludedArea(uint8 AreaType);
+
+	/** set travel cost for all areas */
+	void SetAllAreaCosts(const TArray<float>& CostArray);
+	void SetAllAreaCosts(const float* CostArray, const int32 Count);
+
+	/** get travel & entering costs for all areas */
+	void GetAllAreaCosts(float* CostArray, float* FixedCostArray, const int32 Count) const;
+
+	/** set required flags of navigation nodes */
+	void SetIncludeFlags(uint16 Flags);
+
+	/** get required flags of navigation nodes */
+	uint16 GetIncludeFlags() const;
+
+	/** set forbidden flags of navigation nodes */
+	void SetExcludeFlags(uint16 Flags);
+
+	/** get forbidden flags of navigation nodes */
+	uint16 GetExcludeFlags() const;
+
+	/** set node limit for A* loop */
+	void SetMaxSearchNodes(const uint32 MaxNodes) { MaxSearchNodes = MaxNodes; }
+
+	/** get node limit for A* loop */
+	FORCEINLINE uint32 GetMaxSearchNodes() const { return MaxSearchNodes; }
+
+	/** mark filter as backtracking - parse directional links in opposite direction
+	*  (find path from End to Start, but all links works like on path from Start to End) */
+	void SetBacktrackingEnabled(const bool bBacktracking) { QueryFilterImpl->SetBacktrackingEnabled(bBacktracking); }
+
+	/** get backtracking status */
+	bool IsBacktrackingEnabled() const { return QueryFilterImpl->IsBacktrackingEnabled(); }
+
+	/** post processing for pathfinding's end point */
+	FVector GetAdjustedEndLocation(const FVector& EndPoint) const { return QueryFilterImpl->GetAdjustedEndLocation(EndPoint);  }
+
+	template<typename FilterType>
+	void SetFilterType()
+	{
+		QueryFilterImpl = MakeShareable(new FilterType());
+	}
+
+	FORCEINLINE_DEBUGGABLE void SetFilterImplementation(const INavigationQueryFilterInterface* InQueryFilterImpl)
+	{
+		QueryFilterImpl = MakeShareable(InQueryFilterImpl->CreateCopy());
+	}
+
+	FORCEINLINE const INavigationQueryFilterInterface* GetImplementation() const { return QueryFilterImpl.Get(); }
+	FORCEINLINE INavigationQueryFilterInterface* GetImplementation() { return QueryFilterImpl.Get(); }
+	void Reset() { GetImplementation()->Reset(); }
+
+	FSharedNavQueryFilter GetCopy() const;
+
+	FORCEINLINE bool operator==(const FNavigationQueryFilter& Other) const
+	{
+		const INavigationQueryFilterInterface* Impl0 = GetImplementation();
+		const INavigationQueryFilterInterface* Impl1 = Other.GetImplementation();
+		return Impl0 && Impl1 && Impl0->IsEqual(Impl1);
+	}
+
+	static const uint32 DefaultMaxSearchNodes;
+
+protected:
+	void Assign(const FNavigationQueryFilter& Source);
+
+	TSharedPtr<INavigationQueryFilterInterface, ESPMode::ThreadSafe> QueryFilterImpl;
+	uint32 MaxSearchNodes;
+};
+
diff --git a/Engine/Source/Runtime/Engine/Public/AI/NavigationModifier.h b/Engine/Source/Runtime/Engine/Public/AI/NavigationModifier.h
--- a/Engine/Source/Runtime/Engine/Public/AI/NavigationModifier.h
+++ b/Engine/Source/Runtime/Engine/Public/AI/NavigationModifier.h
@@ -4,11 +4,12 @@
 
 #include "CoreMinimal.h"
 #include "Templates/SubclassOf.h"
+#include "UObject/WeakObjectPtrTemplates.h"
 #include "AI/Navigation/NavLinkDefinition.h"
-#include "AI/Navigation/NavAreas/NavArea.h"
 
 class UBrushComponent;
 class UPrimitiveComponent;
+class UNavAreaBase;
 
 template<typename InElementType> class TNavStatArray;
 
@@ -18,7 +19,7 @@ struct ENGINE_API FNavigationModifier
 	FORCEINLINE bool HasMetaAreas() const { return !!bHasMetaAreas; }
 
 protected:
-	/** set to true if any of areas used by this modifier is a meta nav area (UNavAreaMeta subclass)*/
+	/** set to true if any of areas used by this modifier is a meta nav area (UNavArea::IsMetaArea is true)*/
 	int32 bHasMetaAreas : 1;
 };
 
@@ -59,6 +60,8 @@ namespace ENavigationCoordSystem
 	{
 		Unreal,
 		Recast,
+
+		MAX
 	};
 }
 
@@ -92,31 +95,31 @@ struct ENGINE_API FAreaNavModifier : public FNavigationModifier
 	float FixedCost;
 
 	FAreaNavModifier() : Cost(0.0f), FixedCost(0.0f), Bounds(ForceInitToZero), ShapeType(ENavigationShapeType::Unknown), ApplyMode(ENavigationAreaMode::Apply), bIncludeAgentHeight(false), bIsLowAreaModifier(false) {}
-	FAreaNavModifier(float Radius, float Height, const FTransform& LocalToWorld, const TSubclassOf<UNavArea> AreaClass);
-	FAreaNavModifier(const FVector& Extent, const FTransform& LocalToWorld, const TSubclassOf<UNavArea> AreaClass);
-	FAreaNavModifier(const FBox& Box, const FTransform& LocalToWorld, const TSubclassOf<UNavArea> AreaClass);
-	FAreaNavModifier(const TArray<FVector>& Points, ENavigationCoordSystem::Type CoordType, const FTransform& LocalToWorld, const TSubclassOf<UNavArea> AreaClass);
-	FAreaNavModifier(const TArray<FVector>& Points, const int32 FirstIndex, const int32 LastIndex, ENavigationCoordSystem::Type CoordType, const FTransform& LocalToWorld, const TSubclassOf<UNavArea> AreaClass);
-	FAreaNavModifier(const TNavStatArray<FVector>& Points, const int32 FirstIndex, const int32 LastIndex, ENavigationCoordSystem::Type CoordType, const FTransform& LocalToWorld, const TSubclassOf<UNavArea> AreaClass);
-	FAreaNavModifier(const UBrushComponent* BrushComponent, const TSubclassOf<UNavArea> AreaClass);
+	FAreaNavModifier(float Radius, float Height, const FTransform& LocalToWorld, const TSubclassOf<UNavAreaBase> AreaClass);
+	FAreaNavModifier(const FVector& Extent, const FTransform& LocalToWorld, const TSubclassOf<UNavAreaBase> AreaClass);
+	FAreaNavModifier(const FBox& Box, const FTransform& LocalToWorld, const TSubclassOf<UNavAreaBase> AreaClass);
+	FAreaNavModifier(const TArray<FVector>& Points, ENavigationCoordSystem::Type CoordType, const FTransform& LocalToWorld, const TSubclassOf<UNavAreaBase> AreaClass);
+	FAreaNavModifier(const TArray<FVector>& Points, const int32 FirstIndex, const int32 LastIndex, ENavigationCoordSystem::Type CoordType, const FTransform& LocalToWorld, const TSubclassOf<UNavAreaBase> AreaClass);
+	FAreaNavModifier(const TNavStatArray<FVector>& Points, const int32 FirstIndex, const int32 LastIndex, ENavigationCoordSystem::Type CoordType, const FTransform& LocalToWorld, const TSubclassOf<UNavAreaBase> AreaClass);
+	FAreaNavModifier(const UBrushComponent* BrushComponent, const TSubclassOf<UNavAreaBase> AreaClass);
 
 	FORCEINLINE const FBox& GetBounds() const { return Bounds; }
 	FORCEINLINE ENavigationShapeType::Type GetShapeType() const { return ShapeType; }
 	FORCEINLINE ENavigationAreaMode::Type GetApplyMode() const { return ApplyMode; }
 	FORCEINLINE bool IsLowAreaModifier() const { return bIsLowAreaModifier; }
 	FORCEINLINE bool ShouldIncludeAgentHeight() const { return bIncludeAgentHeight; }
-	FORCEINLINE void SetIncludeAgentHeight(bool bInclude) { bIncludeAgentHeight = bInclude; }
-	FORCEINLINE const TSubclassOf<UNavArea> GetAreaClass() const { return TSubclassOf<UNavArea>(AreaClassOb.Get()); }
-	FORCEINLINE const TSubclassOf<UNavArea> GetAreaClassToReplace() const { return TSubclassOf<UNavArea>(ReplaceAreaClassOb.Get()); }
+	FORCEINLINE FAreaNavModifier& SetIncludeAgentHeight(bool bInclude) { bIncludeAgentHeight = bInclude; return *this; }
+	FORCEINLINE const TSubclassOf<UNavAreaBase> GetAreaClass() const { return TSubclassOf<UNavAreaBase>(AreaClassOb.Get()); }
+	FORCEINLINE const TSubclassOf<UNavAreaBase> GetAreaClassToReplace() const { return TSubclassOf<UNavAreaBase>(ReplaceAreaClassOb.Get()); }
 
 	/** navigation area applied by this modifier */
-	void SetAreaClass(const TSubclassOf<UNavArea> AreaClass);
+	void SetAreaClass(const TSubclassOf<UNavAreaBase> AreaClass);
 
 	/** operation mode, ReplaceInLowPass will always automatically use UNavArea_LowHeight as ReplaceAreaClass! */
 	void SetApplyMode(ENavigationAreaMode::Type InApplyMode);
 	
 	/** additional class for used by some ApplyModes, setting it will automatically change ApplyMode to keep backwards compatibility! */
-	void SetAreaClassToReplace(const TSubclassOf<UNavArea> AreaClass);
+	void SetAreaClassToReplace(const TSubclassOf<UNavAreaBase> AreaClass);
 
 	void GetCylinder(FCylinderNavAreaData& Data) const;
 	void GetBox(FBoxNavAreaData& Data) const;
@@ -138,7 +141,7 @@ protected:
 	/** set when this modifier affects low spans in navmesh generation step */
 	uint8 bIsLowAreaModifier : 1;
 
-	void Init(const TSubclassOf<UNavArea> InAreaClass);
+	void Init(const TSubclassOf<UNavAreaBase> InAreaClass);
 	void SetConvex(const FVector* InPoints, const int32 FirstIndex, const int32 LastIndex, ENavigationCoordSystem::Type CoordType, const FTransform& LocalToWorld);
 	void SetBox(const FBox& Box, const FTransform& LocalToWorld);
 };
@@ -278,7 +281,7 @@ struct ENGINE_API FCompositeNavModifier : public FNavigationModifier
 		bHasLowAreaModifiers |= Modifiers.HasLowAreaModifiers();
 	}
 
-	void CreateAreaModifiers(const UPrimitiveComponent* PrimComp, const TSubclassOf<UNavArea> AreaClass);
+	void CreateAreaModifiers(const UPrimitiveComponent* PrimComp, const TSubclassOf<UNavAreaBase> AreaClass);
 
 	FORCEINLINE const TArray<FAreaNavModifier>& GetAreas() const { return Areas; }
 	FORCEINLINE const TArray<FSimpleLinkNavModifier>& GetSimpleLinks() const { return SimpleLinks; }
diff --git a/Engine/Source/Runtime/Engine/Public/AI/NavigationSystemHelpers.h b/Engine/Source/Runtime/Engine/Public/AI/NavigationSystemHelpers.h
--- a/Engine/Source/Runtime/Engine/Public/AI/NavigationSystemHelpers.h
+++ b/Engine/Source/Runtime/Engine/Public/AI/NavigationSystemHelpers.h
@@ -4,13 +4,13 @@
 
 #include "CoreMinimal.h"
 #include "EngineDefines.h"
-#include "AI/Navigation/NavLinkDefinition.h"
+#include "AI/Navigation/NavigationTypes.h"
 
 class AActor;
 class UBodySetup;
-class UNavCollision;
 struct FCompositeNavModifier;
 struct FKAggregateGeom;
+struct FNavHeightfieldSamples;
 
 template<typename InElementType> class TNavStatArray;
 
@@ -41,62 +41,3 @@ struct FNavigableGeometryExport
 	// Optional delegate for geometry per instance transforms
 	virtual void SetNavDataPerInstanceTransformDelegate(const FNavDataPerInstanceTransformDelegate& InDelegate) = 0;
 };
-
-
-namespace NavigationHelper
-{
-	struct ENGINE_API FNavLinkOwnerData
-	{
-		const AActor* Actor;
-		FTransform LinkToWorld;
-
-		FNavLinkOwnerData() : Actor(nullptr) {}
-		FNavLinkOwnerData(const AActor& InActor);
-		FNavLinkOwnerData(const USceneComponent& InComponent);
-	};
-
-	void GatherCollision(UBodySetup* RigidBody, TNavStatArray<FVector>& OutVertexBuffer, TNavStatArray<int32>& OutIndexBuffer, const FTransform& ComponentToWorld = FTransform::Identity);
-	void GatherCollision(UBodySetup* RigidBody, UNavCollision* NavCollision);
-
-	/** gather collisions from aggregated geom, convex and tri mesh elements are not supported - use override with full UBodySetup param instead */
-	void GatherCollision(const FKAggregateGeom& AggGeom, UNavCollision& NavCollision);
-
-	DECLARE_DELEGATE_ThreeParams(FNavLinkProcessorDelegate, FCompositeNavModifier*, const AActor*, const TArray<FNavigationLink>&);
-	DECLARE_DELEGATE_ThreeParams(FNavLinkSegmentProcessorDelegate, FCompositeNavModifier*, const AActor*, const TArray<FNavigationSegmentLink>&);
-
-	DECLARE_DELEGATE_ThreeParams(FNavLinkProcessorDataDelegate, FCompositeNavModifier*, const FNavLinkOwnerData&, const TArray<FNavigationLink>&);
-	DECLARE_DELEGATE_ThreeParams(FNavLinkSegmentProcessorDataDelegate, FCompositeNavModifier*, const FNavLinkOwnerData&, const TArray<FNavigationSegmentLink>&);
-
-	/** Set new implementation of nav link processor, a function that will be
-	 *	be used to process/transform links before adding them to CompositeModifier.
-	 *	This function is supposed to be called once during the engine/game 
-	 *	setup phase. Not intended to be toggled at runtime */
-	ENGINE_API void SetNavLinkProcessorDelegate(const FNavLinkProcessorDataDelegate& NewDelegate);
-	ENGINE_API void SetNavLinkSegmentProcessorDelegate(const FNavLinkSegmentProcessorDataDelegate& NewDelegate);
-
-	/** called to do any necessary processing on NavLinks and put results in CompositeModifier */
-	ENGINE_API void ProcessNavLinkAndAppend(FCompositeNavModifier* OUT CompositeModifier, const AActor* Actor, const TArray<FNavigationLink>& IN NavLinks);
-	ENGINE_API void ProcessNavLinkAndAppend(FCompositeNavModifier* OUT CompositeModifier, const FNavLinkOwnerData& OwnerData, const TArray<FNavigationLink>& IN NavLinks);
-
-	/** called to do any necessary processing on NavLinks and put results in CompositeModifier */
-	ENGINE_API void ProcessNavLinkSegmentAndAppend(FCompositeNavModifier* OUT CompositeModifier, const AActor* Actor, const TArray<FNavigationSegmentLink>& IN NavLinks);
-	ENGINE_API void ProcessNavLinkSegmentAndAppend(FCompositeNavModifier* OUT CompositeModifier, const FNavLinkOwnerData& OwnerData, const TArray<FNavigationSegmentLink>& IN NavLinks);
-
-	ENGINE_API void DefaultNavLinkProcessorImpl(FCompositeNavModifier* OUT CompositeModifier, const FNavLinkOwnerData& OwnerData, const TArray<FNavigationLink>& IN NavLinks);
-	ENGINE_API void DefaultNavLinkSegmentProcessorImpl(FCompositeNavModifier* OUT CompositeModifier, const FNavLinkOwnerData& OwnerData, const TArray<FNavigationSegmentLink>& IN NavLinks);
-
-	ENGINE_API bool IsBodyNavigationRelevant(const UBodySetup& IN BodySetup);
-
-	// deprecated functions
-	DEPRECATED_FORGAME(4.17, "FNavLinkProcessorDelegate type is now deprecated, please use FNavLinkProcessorDataDelegate instead.")
-	ENGINE_API void SetNavLinkProcessorDelegate(const FNavLinkProcessorDelegate& NewDelegate);
-	
-	DEPRECATED_FORGAME(4.17, "FNavLinkSegmentProcessorDelegate type is now deprecated, please use FNavLinkSegmentProcessorDataDelegate instead.")
-	ENGINE_API void SetNavLinkSegmentProcessorDelegate(const FNavLinkSegmentProcessorDelegate& NewDelegate);
-	
-	DEPRECATED_FORGAME(4.17, "This function is now deprecated, please use override with OwnerData argument.")
-	ENGINE_API void DefaultNavLinkProcessorImpl(FCompositeNavModifier* OUT CompositeModifier, const AActor* Actor, const TArray<FNavigationLink>& IN NavLinks);
-	
-	DEPRECATED_FORGAME(4.17, "This function is now deprecated, please use override with OwnerData argument.")
-	ENGINE_API void DefaultNavLinkSegmentProcessorImpl(FCompositeNavModifier* OUT CompositeModifier, const AActor* Actor, const TArray<FNavigationSegmentLink>& IN NavLinks);
-}
diff --git a/Engine/Source/Runtime/Engine/Public/EngineSharedPCH.h b/Engine/Source/Runtime/Engine/Public/EngineSharedPCH.h
--- a/Engine/Source/Runtime/Engine/Public/EngineSharedPCH.h
+++ b/Engine/Source/Runtime/Engine/Public/EngineSharedPCH.h
@@ -480,7 +480,7 @@
 #include "Engine/BlueprintCore.h"
 #include "EngineDefines.h"
 #include "AI/Navigation/NavigationTypes.h"
-#include "AI/Navigation/NavFilters/NavigationQueryFilter.h"
+#include "AI/Navigation/NavQueryFilter.h"
 #include "CollisionQueryParams.h"
 #include "Engine/LatentActionManager.h"
 #include "AI/Navigation/NavAgentInterface.h"
diff --git a/Engine/Source/Runtime/Engine/Public/GeneratedCodeHelpers.h b/Engine/Source/Runtime/Engine/Public/GeneratedCodeHelpers.h
--- a/Engine/Source/Runtime/Engine/Public/GeneratedCodeHelpers.h
+++ b/Engine/Source/Runtime/Engine/Public/GeneratedCodeHelpers.h
@@ -6,10 +6,10 @@
 // THIS FILE SHOULD BE USED ONLY BY AUTOMATICALLY GENERATED CODE. 
 
 // Common includes
-#include "Runtime/Core/Public/Core.h"
 #include "UObject/Stack.h"
 #include "Blueprint/BlueprintSupport.h"
 #include "Engine/BlueprintGeneratedClass.h"
+#include "Engine/UserDefinedStruct.h"
 
 // Common libraries
 #include "Kismet/KismetArrayLibrary.h"
diff --git a/Engine/Source/Runtime/Engine/Public/VisualLogger/VisualLogger.h b/Engine/Source/Runtime/Engine/Public/VisualLogger/VisualLogger.h
--- a/Engine/Source/Runtime/Engine/Public/VisualLogger/VisualLogger.h
+++ b/Engine/Source/Runtime/Engine/Public/VisualLogger/VisualLogger.h
@@ -370,6 +370,9 @@ public:
 	static void NavigationDataDump(const UObject* LogOwner, const FLogCategoryBase& Category, const ELogVerbosity::Type Verbosity, const FBox& Box);
 	static void NavigationDataDump(const UObject* LogOwner, const FName& CategoryName, const ELogVerbosity::Type Verbosity, const FBox& Box);
 
+	DECLARE_MULTICAST_DELEGATE_FourParams(FNavigationDataDump, const UObject* /*LogOwner*/, const FName& /*CategoryName*/, const ELogVerbosity::Type /*Verbosity*/, const FBox& /*Box*/);
+	static FNavigationDataDump NavigationDataDumpDelegate;
+
 	/** Log events */
 	static void EventLog(const UObject* LogOwner, const FName EventTag1, const FVisualLogEventBase& Event1, const FName EventTag2 = NAME_None, const FName EventTag3 = NAME_None, const FName EventTag4 = NAME_None, const FName EventTag5 = NAME_None, const FName EventTag6 = NAME_None);
 	static void EventLog(const UObject* LogOwner, const FName EventTag1, const FVisualLogEventBase& Event1, const FVisualLogEventBase& Event2);
diff --git a/Engine/Source/Runtime/Landscape/Private/LandscapeCollision.cpp b/Engine/Source/Runtime/Landscape/Private/LandscapeCollision.cpp
--- a/Engine/Source/Runtime/Landscape/Private/LandscapeCollision.cpp
+++ b/Engine/Source/Runtime/Landscape/Private/LandscapeCollision.cpp
@@ -18,7 +18,7 @@
 #include "LandscapeProxy.h"
 #include "LandscapeInfo.h"
 #include "Interfaces/Interface_CollisionDataProvider.h"
-#include "AI/Navigation/NavigationSystem.h"
+#include "AI/NavigationSystemBase.h"
 #include "LandscapeComponent.h"
 #include "LandscapeLayerInfoObject.h"
 #include "LandscapePrivate.h"
@@ -1526,7 +1526,7 @@ void ULandscapeHeightfieldCollisionComponent::PostEditUndo()
 		RecreateCollision();
 	}
 
-	UNavigationSystem::UpdateComponentInNavOctree(*this);
+	FNavigationSystem::UpdateComponentData(*this);
 }
 
 bool ULandscapeHeightfieldCollisionComponent::ComponentIsTouchingSelectionBox(const FBox& InSelBBox, const FEngineShowFlags& ShowFlags, const bool bConsiderOnlyBSP, const bool bMustEncompassEntireComponent) const
diff --git a/Engine/Source/Runtime/Landscape/Private/LandscapeSplineRaster.cpp b/Engine/Source/Runtime/Landscape/Private/LandscapeSplineRaster.cpp
--- a/Engine/Source/Runtime/Landscape/Private/LandscapeSplineRaster.cpp
+++ b/Engine/Source/Runtime/Landscape/Private/LandscapeSplineRaster.cpp
@@ -8,7 +8,7 @@
 #include "LandscapeSplineRaster.h"
 #include "LandscapeProxy.h"
 #include "LandscapeInfo.h"
-#include "AI/Navigation/NavigationSystem.h"
+#include "AI/NavigationSystemBase.h"
 #include "LandscapeComponent.h"
 #include "LandscapeLayerInfoObject.h"
 #include "LandscapeHeightfieldCollisionComponent.h"
@@ -588,11 +588,7 @@ bool ULandscapeInfo::ApplySplinesInternal(bool bOnlySelected, ALandscapeProxy* L
 		if (CollisionComponent)
 		{
 			CollisionComponent->RecreateCollision();
-			UNavigationSystem* NavSys = UNavigationSystem::GetCurrent(Component);
-			if (NavSys)
-			{
-				NavSys->UpdateComponentInNavOctree(*CollisionComponent);
-			}
+			FNavigationSystem::UpdateComponentData(*CollisionComponent);
 		}
 	}
 
@@ -683,11 +679,7 @@ namespace LandscapeSplineRaster
 			if (CollisionComponent)
 			{
 				CollisionComponent->RecreateCollision();
-				UNavigationSystem* NavSys = UNavigationSystem::GetCurrent(Component);
-				if (NavSys)
-				{
-					NavSys->UpdateComponentInNavOctree(*CollisionComponent);
-				}
+				FNavigationSystem::UpdateComponentData(*CollisionComponent);
 			}
 		}
 	}
diff --git a/Engine/Source/Runtime/NavigationSystem/NavigationSystem.Build.cs b/Engine/Source/Runtime/NavigationSystem/NavigationSystem.Build.cs
new file mode 100644
--- /dev/null
+++ b/Engine/Source/Runtime/NavigationSystem/NavigationSystem.Build.cs
@@ -0,0 +1,74 @@
+// Copyright 1998-2018 Epic Games, Inc. All Rights Reserved.
+
+namespace UnrealBuildTool.Rules
+{
+    public class NavigationSystem : ModuleRules
+    {
+        public NavigationSystem(ReadOnlyTargetRules Target) : base(Target)
+        {
+            PublicIncludePaths.AddRange(
+                new string[] {
+                    "Runtime/NavigationSystem/Public",
+                }
+                );
+
+            PrivateIncludePaths.AddRange(
+                new string[] {
+                    "Runtime/NavigationSystem/Private",
+                    "Runtime/Engine/Private",
+                    "Developer/DerivedDataCache/Public",
+                }
+                );
+
+            PublicDependencyModuleNames.AddRange(
+                new string[] {
+                    "Core",
+                    "CoreUObject",
+                    "Engine",
+                }
+                );
+
+            PrivateDependencyModuleNames.AddRange(
+                new string[] {
+                    "RHI",
+                    "RenderCore",
+                    "ShaderCore",
+                }
+                );
+
+			PrivateIncludePathModuleNames.AddRange(
+				new string[]
+				{
+					"TargetPlatform",
+				}
+				);
+            
+            if (!Target.bBuildRequiresCookedData && Target.bCompileAgainstEngine)
+            {
+                DynamicallyLoadedModuleNames.Add("DerivedDataCache");
+            }
+
+            SetupModulePhysXAPEXSupport(Target);
+
+            if (Target.bCompileRecast)
+            {
+                PrivateDependencyModuleNames.Add("Navmesh");
+                PublicDefinitions.Add("WITH_RECAST=1");
+            }
+            else
+            {
+                // Because we test WITH_RECAST in public Engine header files, we need to make sure that modules
+                // that import us also have this definition set appropriately.  Recast is a private dependency
+                // module, so it's definitions won't propagate to modules that import Engine.
+                PublicDefinitions.Add("WITH_RECAST=0");
+            }
+
+            if (Target.bBuildEditor == true)
+            {
+                // @todo api: Only public because of WITH_EDITOR and UNREALED_API
+                PublicDependencyModuleNames.Add("UnrealEd");
+                CircularlyReferencedDependentModules.Add("UnrealEd");
+            }
+        }
+    }
+}
diff --git a/Engine/Source/Runtime/NavigationSystem/Private/NavAreas/NavAreaMeta.cpp b/Engine/Source/Runtime/NavigationSystem/Private/NavAreas/NavAreaMeta.cpp
new file mode 100644
--- /dev/null
+++ b/Engine/Source/Runtime/NavigationSystem/Private/NavAreas/NavAreaMeta.cpp
@@ -0,0 +1,23 @@
+// Copyright 1998-2018 Epic Games, Inc. All Rights Reserved.
+
+#include "NavAreas/NavAreaMeta.h"
+
+UNavAreaMeta::UNavAreaMeta(const FObjectInitializer& ObjectInitializer) 
+	: Super(ObjectInitializer)
+{
+	bIsMetaArea = true;
+}
+
+TSubclassOf<UNavArea> UNavAreaMeta::PickAreaClass(TSubclassOf<UNavArea> AreaClass, const AActor* Actor, const FNavAgentProperties& NavAgent)
+{
+	return Actor
+		? TSubclassOf<UNavArea>(UNavAreaBase::PickAreaClassForAgent(AreaClass, *Actor, NavAgent))
+		: AreaClass;
+}
+
+TSubclassOf<UNavArea> UNavAreaMeta::PickAreaClass(const AActor* Actor, const FNavAgentProperties& NavAgent)
+{
+	return Actor
+		? TSubclassOf<UNavArea>(PickAreaClassForAgent(*Actor, NavAgent))
+		: TSubclassOf<UNavArea>(GetClass());
+}
diff --git a/Engine/Source/Runtime/NavigationSystem/Private/NavAreas/NavArea_Null.cpp b/Engine/Source/Runtime/NavigationSystem/Private/NavAreas/NavArea_Null.cpp
new file mode 100644
--- /dev/null
+++ b/Engine/Source/Runtime/NavigationSystem/Private/NavAreas/NavArea_Null.cpp
@@ -0,0 +1,9 @@
+// Copyright 1998-2018 Epic Games, Inc. All Rights Reserved.
+
+#include "NavAreas/NavArea_Null.h"
+
+UNavArea_Null::UNavArea_Null(const FObjectInitializer& ObjectInitializer) : Super(ObjectInitializer)
+{
+	DefaultCost = BIG_NUMBER;
+	AreaFlags = 0;
+}
diff --git a/Engine/Source/Runtime/NavigationSystem/Private/NavGraph/NavGraphGenerator.h b/Engine/Source/Runtime/NavigationSystem/Private/NavGraph/NavGraphGenerator.h
new file mode 100644
--- /dev/null
+++ b/Engine/Source/Runtime/NavigationSystem/Private/NavGraph/NavGraphGenerator.h
@@ -0,0 +1,46 @@
+// Copyright 1998-2018 Epic Games, Inc. All Rights Reserved.
+
+#pragma once 
+
+#include "CoreMinimal.h"
+#include "GameFramework/Volume.h"
+#include "AI/NavDataGenerator.h"
+
+class ANavigationGraph;
+
+/**
+ * Class that handles generation of the ANavigationGraph data
+ */
+class FNavGraphGenerator : public FNavDataGenerator
+{
+public:
+	FNavGraphGenerator(ANavigationGraph* InDestNavGraph);
+	virtual ~FNavGraphGenerator();
+
+private:
+	/** Prevent copying. */
+	FNavGraphGenerator(FNavGraphGenerator const& NoCopy) { check(0); };
+	FNavGraphGenerator& operator=(FNavGraphGenerator const& NoCopy) { check(0); return *this; }
+
+public:
+	//----------------------------------------------------------------------//
+	// FNavDataGenerator overrides
+	//----------------------------------------------------------------------//
+	virtual bool IsBuildInProgress(bool bCheckDirtyToo = false) const override;
+
+private:
+	// Performs initial setup of member variables so that generator is ready to do its thing from this point on
+	void Init();
+	void CleanUpIntermediateData();
+	void UpdateBuilding();
+
+private:
+	/** Bounding geometry definition. */
+	TArray<AVolume const*> InclusionVolumes;
+
+	FCriticalSection GraphChangingLock;
+
+	ANavigationGraph* DestNavGraph;
+
+	uint32 bInitialized:1;
+};
