Copying //UE4/Dev-Enterprise to //UE4/Dev-Main (Source: //UE4/Dev-Enterprise @ 3972172)

diff --git a/Engine/Source/Runtime/Engine/Classes/Exporters/Exporter.h b/Engine/Source/Runtime/Engine/Classes/Exporters/Exporter.h
--- a/Engine/Source/Runtime/Engine/Classes/Exporters/Exporter.h
+++ b/Engine/Source/Runtime/Engine/Classes/Exporters/Exporter.h
@@ -16,6 +16,7 @@
 
 class FExportObjectInnerContext;
 class UActorComponent;
+class UAssetExportTask;
 
 UCLASS(abstract, transient, MinimalAPI)
 class UExporter : public UObject
@@ -61,6 +62,8 @@ class UExporter : public UObject
 	UPROPERTY()
 	uint32 bForceFileOperations:1;
 
+	UPROPERTY()
+	class UAssetExportTask* ExportTask;
 
 	ENGINE_API static		FString	CurrentFilename;
 	/** (debugging purposes only) */
@@ -180,7 +183,7 @@ class UExporter : public UObject
 
 	/**
 	 * Export this object to a file.  Child classes do not override this, but they do provide an Export() function
-	 * to do the resoource-specific export work.
+	 * to do the resource-specific export work.
 	 * 
 	 * @param	Object				the object to export
 	 * @param	InExporter			exporter to use for exporting this object.  If NULL, attempts to create a valid exporter.
@@ -189,7 +192,7 @@ class UExporter : public UObject
 	 * @param	NoReplaceIdentical	false if we always want to overwrite any existing files, even if they're identical
 	 * @param	Prompt				true if the user should be prompted to checkout/overwrite the output file if it already exists and is read-only
 	 *
-	 * @return	1 if the the object was successfully exported, 0 if a fatal error was encountered during export, or -1 if a non fatal error was encountered
+	 * @return	1 if the the object was successfully exported, 0 if a fatal error was encountered during export
 	 */
 	ENGINE_API static int32 ExportToFile( UObject* Object, UExporter* Exporter, const TCHAR* Filename, bool InSelectedOnly, bool NoReplaceIdentical=false, bool Prompt=false );
 
@@ -244,8 +247,31 @@ class UExporter : public UObject
 	 *
 	 * @return	1 if the the object was successfully exported, 0 if a fatal error was encountered during export, or -1 if a non fatal error was encountered
 	 */
+	DEPRECATED(4.20, "ExportToFileEx was deprecated, please use RunAssetExportTask.")
 	ENGINE_API static int32 ExportToFileEx( FExportToFileParams& ExportParams );
 
+	/**
+	* Export the given object to file.  Child classes do not override this, but they do provide an Export() function
+	* to do the resource-specific export work.
+	*
+	* @param	Task		The task to export.
+	*
+	* @return	true if the the object was successfully exported
+	*/
+	UFUNCTION(BlueprintCallable, Category = "Miscellaneous")
+	static ENGINE_API bool RunAssetExportTask( class UAssetExportTask* Task );
+
+	/**
+	 * Export the given objects to files.  Child classes do not override this, but they do provide an Export() function
+	 * to do the resource-specific export work.
+	 * 
+	 * @param	ExportTasks		The array of tasks to export.
+	 *
+	 * @return	true if all tasks ran without error
+	 */
+	UFUNCTION(BlueprintCallable, Category = "Miscellaneous")
+	static ENGINE_API bool RunAssetExportTasks( const TArray<UAssetExportTask*>& ExportTasks );
+	
 	/**
 	 * Single entry point to export an object's subobjects, its components, and its properties
 	 *
@@ -284,6 +310,7 @@ protected:
 	 */
 	ENGINE_API void EmitEndObject( FOutputDevice& Ar );
 
+protected:
 	/** The set of registered exporters. */
 	static TSet< TWeakObjectPtr<UExporter> > RegisteredExporters;
 
diff --git a/Engine/Source/Runtime/Engine/Classes/Kismet/DataTableFunctionLibrary.h b/Engine/Source/Runtime/Engine/Classes/Kismet/DataTableFunctionLibrary.h
--- a/Engine/Source/Runtime/Engine/Classes/Kismet/DataTableFunctionLibrary.h
+++ b/Engine/Source/Runtime/Engine/Classes/Kismet/DataTableFunctionLibrary.h
@@ -101,4 +101,39 @@ class ENGINE_API UDataTableFunctionLibrary : public UBlueprintFunctionLibrary
 		}
 		*(bool*)RESULT_PARAM = bSuccess;
     }
+
+#if WITH_EDITOR
+	/** 
+	 * Empty and fill a Data Table from CSV string.
+	 * @param	CSVString	The Data that representing the contents of a CSV file.
+	 * @return	True if the operation succeeds, check the log for errors if it didn't succeed.
+	 */
+	UFUNCTION(BlueprintCallable, Category = "Editor Scripting | DataTable", DisplayName="Fill Data Table from CSV String")
+	static bool FillDataTableFromCSVString(UDataTable* DataTable, const FString& CSVString);
+
+	/** 
+	 * Empty and fill a Data Table from CSV file.
+	 * @param	CSVFilePath	The file path of the CSV file.
+	 * @return	True if the operation succeeds, check the log for errors if it didn't succeed.
+	 */
+	UFUNCTION(BlueprintCallable, Category = "Editor Scripting | DataTable", DisplayName = "Fill Data Table from CSV File")
+	static bool FillDataTableFromCSVFile(UDataTable* DataTable, const FString& CSVFilePath);
+
+	/** 
+	 * Empty and fill a Data Table from JSON string.
+	 * @param	JSONString	The Data that representing the contents of a JSON file.
+	 * @return	True if the operation succeeds, check the log for errors if it didn't succeed.
+	 */
+	UFUNCTION(BlueprintCallable, Category = "Editor Scripting | DataTable", DisplayName = "Fill Data Table from JSON String")
+	static bool FillDataTableFromJSONString(UDataTable* DataTable, const FString& JSONString);
+
+
+	/** 
+	 * Empty and fill a Data Table from JSON file.
+	 * @param	JSONFilePath	The file path of the JSON file.
+	 * @return	True if the operation succeeds, check the log for errors if it didn't succeed.
+	 */
+	UFUNCTION(BlueprintCallable, Category = "Editor Scripting | DataTable", DisplayName = "Fill Data Table from JSON File")
+	static bool FillDataTableFromJSONFile(UDataTable* DataTable, const FString& JSONFilePath);
+#endif //WITH_EDITOR
 };
diff --git a/Engine/Source/Runtime/Engine/Classes/Materials/MaterialExpressionFunctionInput.h b/Engine/Source/Runtime/Engine/Classes/Materials/MaterialExpressionFunctionInput.h
--- a/Engine/Source/Runtime/Engine/Classes/Materials/MaterialExpressionFunctionInput.h
+++ b/Engine/Source/Runtime/Engine/Classes/Materials/MaterialExpressionFunctionInput.h
@@ -25,6 +25,7 @@ enum EFunctionInputType
 	FunctionInput_TextureCube,
 	FunctionInput_StaticBool,
 	FunctionInput_MaterialAttributes,
+	FunctionInput_TextureExternal,
 	FunctionInput_MAX,
 };
 
diff --git a/Engine/Source/Runtime/Engine/Classes/Materials/MaterialInterface.h b/Engine/Source/Runtime/Engine/Classes/Materials/MaterialInterface.h
--- a/Engine/Source/Runtime/Engine/Classes/Materials/MaterialInterface.h
+++ b/Engine/Source/Runtime/Engine/Classes/Materials/MaterialInterface.h
@@ -15,6 +15,7 @@
 #include "RHI.h"
 #include "Engine/BlendableInterface.h"
 #include "Materials/MaterialLayersFunctions.h"
+#include "Interfaces/Interface_AssetUserData.h"
 #include "MaterialInterface.generated.h"
 
 class FMaterialCompiler;
@@ -200,7 +201,7 @@ struct FMaterialTextureInfo
 };
 
 UCLASS(abstract, BlueprintType,MinimalAPI)
-class UMaterialInterface : public UObject, public IBlendableInterface
+class UMaterialInterface : public UObject, public IBlendableInterface, public IInterface_AssetUserData
 {
 	GENERATED_UCLASS_BODY()
 
@@ -230,8 +231,18 @@ protected:
 	UPROPERTY()
 	TArray<FMaterialTextureInfo> TextureStreamingData;
 
+	/** Array of user data stored with the asset */
+	UPROPERTY(EditAnywhere, AdvancedDisplay, Instanced, Category = Material)
+	TArray<UAssetUserData*> AssetUserData;
+
 public:
 
+	//~ Begin IInterface_AssetUserData Interface
+	ENGINE_API virtual void AddAssetUserData(UAssetUserData* InUserData) override;
+	ENGINE_API virtual void RemoveUserDataOfClass(TSubclassOf<UAssetUserData> InUserDataClass) override;
+	ENGINE_API virtual UAssetUserData* GetAssetUserDataOfClass(TSubclassOf<UAssetUserData> InUserDataClass) override;
+	//~ End IInterface_AssetUserData Interface
+
 #if WITH_EDITORONLY_DATA
 	/** The mesh used by the material editor to preview the material.*/
 	UPROPERTY(EditAnywhere, Category=Previewing, meta=(AllowedClasses="StaticMesh,SkeletalMesh", ExactClass="true"))
@@ -245,7 +256,7 @@ public:
 	TMap<FString, bool> LayerParameterExpansion;
 
 	/** Importing data and options used for this material */
-	UPROPERTY(EditAnywhere, Instanced, Category = ImportSettings, AssetRegistrySearchable)
+	UPROPERTY(EditAnywhere, Instanced, Category = ImportSettings)
 	class UAssetImportData* AssetImportData;
 
 private:
@@ -277,6 +288,7 @@ public:
 
 #if WITH_EDITOR
 	ENGINE_API virtual void PostEditChangeProperty(FPropertyChangedEvent& PropertyChangedEvent) override;
+	ENGINE_API virtual void GetAssetRegistryTags(TArray<FAssetRegistryTag>& OutTags) const override;
 #endif // WITH_EDITOR
 	//~ End UObject Interface.
 
diff --git a/Engine/Source/Runtime/Engine/Private/AssetManager.cpp b/Engine/Source/Runtime/Engine/Private/AssetManager.cpp
--- a/Engine/Source/Runtime/Engine/Private/AssetManager.cpp
+++ b/Engine/Source/Runtime/Engine/Private/AssetManager.cpp
@@ -24,7 +24,7 @@
 
 #define LOCTEXT_NAMESPACE "AssetManager"
 
-DEFINE_LOG_CATEGORY_STATIC(LogAssetManager, Log, All);
+DEFINE_LOG_CATEGORY(LogAssetManager);
 
 /** Structure defining the current loading state of an asset */
 struct FPrimaryAssetLoadState
diff --git a/Engine/Source/Runtime/Engine/Private/AssetManagerTypes.cpp b/Engine/Source/Runtime/Engine/Private/AssetManagerTypes.cpp
--- a/Engine/Source/Runtime/Engine/Private/AssetManagerTypes.cpp
+++ b/Engine/Source/Runtime/Engine/Private/AssetManagerTypes.cpp
@@ -92,10 +92,85 @@ void FPrimaryAssetRules::PropagateCookRules(const FPrimaryAssetRules& ParentRule
 }
 
 #if WITH_EDITOR
+void UAssetManagerSettings::PostInitProperties()
+{
+	Super::PostInitProperties();
+
+	ApplyMetaDataTagsSettings();
+}
+
+void UAssetManagerSettings::ApplyMetaDataTagsSettings()
+{
+	TSet<FName>& GlobalTagsForAssetRegistry = UObject::GetMetaDataTagsForAssetRegistry();
+	for (FName Tag : MetaDataTagsForAssetRegistry)
+	{
+		if (!Tag.IsNone())
+		{
+			if (!GlobalTagsForAssetRegistry.Contains(Tag))
+			{
+				GlobalTagsForAssetRegistry.Add(Tag);
+			}
+			else
+			{
+				// To catch the case where the same tag is used by different users and their settings are synced after edition
+				UE_LOG(LogAssetManager, Warning, TEXT("Cannot use duplicate metadata tag '%s' for Asset Registry"), *Tag.ToString());
+			}
+		}
+	}
+}
+
+void UAssetManagerSettings::ClearMetaDataTagsSettings()
+{
+	TSet<FName>& GlobalTagsForAssetRegistry = UObject::GetMetaDataTagsForAssetRegistry();
+	for (FName Tag : MetaDataTagsForAssetRegistry)
+	{
+		if (!Tag.IsNone())
+		{
+			GlobalTagsForAssetRegistry.Remove(Tag);
+		}
+	}
+}
+
+void UAssetManagerSettings::PreEditChange(UProperty* PropertyAboutToChange)
+{
+	Super::PreEditChange(PropertyAboutToChange);
+
+	FName PropertyName = PropertyAboutToChange->GetFName();
+	if (PropertyName == GET_MEMBER_NAME_CHECKED(UAssetManagerSettings, MetaDataTagsForAssetRegistry))
+	{
+		ClearMetaDataTagsSettings();
+	}
+}
+
 void UAssetManagerSettings::PostEditChangeProperty(FPropertyChangedEvent& PropertyChangedEvent)
 {
 	Super::PostEditChangeProperty(PropertyChangedEvent);
-	if (PropertyChangedEvent.Property && UAssetManager::IsValid())
+
+	FName PropertyName = (PropertyChangedEvent.MemberProperty != NULL) ? PropertyChangedEvent.MemberProperty->GetFName() : NAME_None;
+	if (PropertyName == GET_MEMBER_NAME_CHECKED(UAssetManagerSettings, MetaDataTagsForAssetRegistry))
+	{
+		if (PropertyChangedEvent.ChangeType == EPropertyChangeType::ValueSet)
+		{
+			// Check if the new value already exists in the global tags list
+			int32 Index = PropertyChangedEvent.GetArrayIndex(PropertyName.ToString());
+			if (Index > 0)
+			{
+				TSet<FName>::TIterator It = MetaDataTagsForAssetRegistry.CreateIterator();
+				for (int32 i = 0; i < Index; ++i)
+				{
+					++It;
+				}
+				FName NewValue = It ? *It : FName();
+				if (UObject::GetMetaDataTagsForAssetRegistry().Contains(NewValue))
+				{
+					*It = FName();
+					UE_LOG(LogAssetManager, Warning, TEXT("Cannot use duplicate metadata tag '%s' for Asset Registry"), *NewValue.ToString());
+				}
+			}
+		}
+		ApplyMetaDataTagsSettings();
+	}
+	else if (PropertyChangedEvent.Property && UAssetManager::IsValid())
 	{
 		UAssetManager::Get().ReinitializeFromConfig();
 	}
diff --git a/Engine/Source/Runtime/Engine/Private/DataTableFunctionLibrary.cpp b/Engine/Source/Runtime/Engine/Private/DataTableFunctionLibrary.cpp
--- a/Engine/Source/Runtime/Engine/Private/DataTableFunctionLibrary.cpp
+++ b/Engine/Source/Runtime/Engine/Private/DataTableFunctionLibrary.cpp
@@ -3,6 +3,10 @@
 #include "Kismet/DataTableFunctionLibrary.h"
 #include "Engine/CurveTable.h"
 
+#if WITH_EDITOR
+#include "Misc/FileHelper.h"
+#endif //WITH_EDITOR
+
 UDataTableFunctionLibrary::UDataTableFunctionLibrary(const FObjectInitializer& ObjectInitializer)
 	: Super(ObjectInitializer)
 {
@@ -77,3 +81,98 @@ void UDataTableFunctionLibrary::GetDataTableRowNames(UDataTable* Table, TArray<F
 		OutRowNames.Empty();
 	}
 }
+
+
+#if WITH_EDITOR
+bool UDataTableFunctionLibrary::FillDataTableFromCSVString(UDataTable* DataTable, const FString& InString)
+{
+	if (!DataTable)
+	{
+		UE_LOG(LogDataTable, Error, TEXT("Can't fill an invalid DataTable."));
+		return false;
+	}
+
+	bool bResult = true;
+	if (InString.Len() == 0)
+	{
+		DataTable->EmptyTable();
+	}
+	else
+	{
+		TArray<FString> Errors = DataTable->CreateTableFromCSVString(InString);
+		if (Errors.Num())
+		{
+			for (const FString& Error : Errors)
+			{
+				UE_LOG(LogDataTable, Warning, TEXT("%s"), *Error);
+			}
+		}
+		bResult = Errors.Num() == 0;
+	}
+	return bResult;
+}
+
+bool UDataTableFunctionLibrary::FillDataTableFromCSVFile(UDataTable* DataTable, const FString& InFilePath)
+{
+	if (!DataTable)
+	{
+		UE_LOG(LogDataTable, Error, TEXT("Can't fill an invalid DataTable."));
+		return false;
+	}
+
+	FString Data;
+	if (!FFileHelper::LoadFileToString(Data, *InFilePath))
+	{
+		UE_LOG(LogDataTable, Error, TEXT("Can't load the file '%s'."), *InFilePath);
+		return false;
+	}
+
+	return FillDataTableFromCSVString(DataTable, Data);
+}
+
+bool UDataTableFunctionLibrary::FillDataTableFromJSONString(UDataTable* DataTable, const FString& InString)
+{
+	if (!DataTable)
+	{
+		UE_LOG(LogDataTable, Error, TEXT("Can't fill an invalid DataTable."));
+		return false;
+	}
+
+	bool bResult = true;
+	if (InString.Len() == 0)
+	{
+		DataTable->EmptyTable();
+	}
+	else
+	{
+		TArray<FString> Errors = DataTable->CreateTableFromJSONString(InString);
+		if (Errors.Num())
+		{
+			for (const FString& Error : Errors)
+			{
+				UE_LOG(LogDataTable, Warning, TEXT("%s"), *Error);
+			}
+		}
+		bResult = Errors.Num() == 0;
+	}
+	return bResult;
+}
+
+bool UDataTableFunctionLibrary::FillDataTableFromJSONFile(UDataTable* DataTable, const FString& InFilePath)
+{
+	if (!DataTable)
+	{
+		UE_LOG(LogDataTable, Error, TEXT("Can't fill an invalid DataTable."));
+		return false;
+	}
+
+	FString Data;
+	if (!FFileHelper::LoadFileToString(Data, *InFilePath))
+	{
+		UE_LOG(LogDataTable, Error, TEXT("Can't load the file '%s'."), *InFilePath);
+		return false;
+	}
+
+	return FillDataTableFromJSONString(DataTable, Data);
+}
+#endif //WITH_EDITOR
\ No newline at end of file
diff --git a/Engine/Source/Runtime/Engine/Private/EditorFramework/AssetImportData.cpp b/Engine/Source/Runtime/Engine/Private/EditorFramework/AssetImportData.cpp
--- a/Engine/Source/Runtime/Engine/Private/EditorFramework/AssetImportData.cpp
+++ b/Engine/Source/Runtime/Engine/Private/EditorFramework/AssetImportData.cpp
@@ -145,11 +145,25 @@ void UAssetImportData::Update(const FString& InPath, const FMD5Hash InPreCompute
 }
 //@third party END SIMPLYGON
 
+#if WITH_EDITOR
+FString UAssetImportData::K2_GetFirstFilename() const
+{
+	return GetFirstFilename();
+}
+#endif
+
 FString UAssetImportData::GetFirstFilename() const
 {
 	return SourceData.SourceFiles.Num() > 0 ? ResolveImportFilename(SourceData.SourceFiles[0].RelativeFilename) : FString();
 }
 
+#if WITH_EDITOR
+TArray<FString> UAssetImportData::K2_ExtractFilenames() const
+{
+	return ExtractFilenames();
+}
+#endif
+
 void UAssetImportData::ExtractFilenames(TArray<FString>& AbsoluteFilenames) const
 {
 	for (const auto& File : SourceData.SourceFiles)
diff --git a/Engine/Source/Runtime/Engine/Private/GameInstance.cpp b/Engine/Source/Runtime/Engine/Private/GameInstance.cpp
--- a/Engine/Source/Runtime/Engine/Private/GameInstance.cpp
+++ b/Engine/Source/Runtime/Engine/Private/GameInstance.cpp
@@ -152,6 +152,50 @@ void UGameInstance::InitializeStandalone()
 	Init();
 }
 
+void UGameInstance::InitializeForMinimalNetRPC(const FName InPackageName)
+{
+	check(!InPackageName.IsNone());
+
+	// Creates the world context. This should be the only WorldContext that ever gets created for this GameInstance.
+	WorldContext = &GetEngine()->CreateNewWorldContext(EWorldType::GameRPC);
+	WorldContext->OwningGameInstance = this;
+
+	UPackage* NetWorldPackage = nullptr;
+	UWorld* NetWorld = nullptr;
+	CreateMinimalNetRPCWorld(InPackageName, NetWorldPackage, NetWorld);
+
+	NetWorld->SetGameInstance(this);
+	WorldContext->SetCurrentWorld(NetWorld);
+
+	Init();
+}
+
+void UGameInstance::CreateMinimalNetRPCWorld(const FName InPackageName, UPackage*& OutWorldPackage, UWorld*& OutWorld)
+{
+	check(!InPackageName.IsNone());
+
+	const FName WorldName = FPackageName::GetShortFName(InPackageName);
+
+	// Create the empty named world within the given package. This will be the world name used with "Browse" to initialize the RPC server and connect the client(s).
+	OutWorldPackage = NewObject<UPackage>(nullptr, InPackageName, RF_Transient);
+	OutWorldPackage->ThisContainsMap();
+
+	OutWorld = NewObject<UWorld>(OutWorldPackage, WorldName);
+	OutWorld->WorldType = EWorldType::GameRPC;
+	OutWorld->InitializeNewWorld(UWorld::InitializationValues()
+		.InitializeScenes(false)
+		.AllowAudioPlayback(false)
+		.RequiresHitProxies(false)
+		.CreatePhysicsScene(false)
+		.CreateNavigation(false)
+		.CreateAISystem(false)
+		.ShouldSimulatePhysics(false)
+		.EnableTraceCollision(false)
+		.SetTransactional(false)
+		.CreateFXSystem(false)
+		);
+}
+
 #if WITH_EDITOR
 
 FGameInstancePIEResult UGameInstance::InitializeForPlayInEditor(int32 PIEInstanceIndex, const FGameInstancePIEParameters& Params)
@@ -1082,28 +1126,27 @@ void UGameInstance::PreloadContentForURL(FURL InURL)
 
 AGameModeBase* UGameInstance::CreateGameModeForURL(FURL InURL)
 {
-	UWorld* World = GetWorld();
 	// Init the game info.
-	FString Options(TEXT(""));
-	TCHAR GameParam[256] = TEXT("");
-	FString	Error = TEXT("");
-	AWorldSettings* Settings = World->GetWorldSettings();
+	FString Options;
+	FString GameParam;
 	for (int32 i = 0; i < InURL.Op.Num(); i++)
 	{
 		Options += TEXT("?");
 		Options += InURL.Op[i];
-		FParse::Value(*InURL.Op[i], TEXT("GAME="), GameParam, ARRAY_COUNT(GameParam));
+		FParse::Value(*InURL.Op[i], TEXT("GAME="), GameParam);
 	}
 
+	UWorld* World = GetWorld();
+	AWorldSettings* Settings = World->GetWorldSettings();
 	UGameEngine* const GameEngine = Cast<UGameEngine>(GEngine);
 
 	// Get the GameMode class. Start by using the default game type specified in the map's worldsettings.  It may be overridden by settings below.
 	TSubclassOf<AGameModeBase> GameClass = Settings->DefaultGameMode;
 
 	// If there is a GameMode parameter in the URL, allow it to override the default game type
-	if (GameParam[0])
+	if (!GameParam.IsEmpty())
 	{
-		FString const GameClassName = UGameMapsSettings::GetGameModeForName(FString(GameParam));
+		FString const GameClassName = UGameMapsSettings::GetGameModeForName(GameParam);
 
 		// If the gamename was specified, we can use it to fully load the pergame PreLoadClass packages
 		if (GameEngine)
diff --git a/Engine/Source/Runtime/Engine/Private/Materials/HLSLMaterialTranslator.h b/Engine/Source/Runtime/Engine/Private/Materials/HLSLMaterialTranslator.h
--- a/Engine/Source/Runtime/Engine/Private/Materials/HLSLMaterialTranslator.h
+++ b/Engine/Source/Runtime/Engine/Private/Materials/HLSLMaterialTranslator.h
@@ -1374,6 +1374,7 @@ protected:
 		case MCT_TextureCube:	return TEXT("textureCube");
 		case MCT_StaticBool:	return TEXT("static bool");
 		case MCT_MaterialAttributes:	return TEXT("MaterialAttributes");
+		case MCT_TextureExternal:	return TEXT("TextureExternal");
 		default:				return TEXT("unknown");
 		};
 	}
@@ -1392,6 +1393,7 @@ protected:
 		case MCT_TextureCube:	return TEXT("textureCube");
 		case MCT_StaticBool:	return TEXT("static bool");
 		case MCT_MaterialAttributes:	return TEXT("MaterialAttributes");
+		case MCT_TextureExternal:	return TEXT("TextureExternal");
 		default:				return TEXT("unknown");
 		};
 	}
@@ -5279,6 +5281,13 @@ protected:
 				InputParamDecl += InputNameStr;
 				InputParamDecl += TEXT("Sampler ");
 				break;
+			case MCT_TextureExternal:
+				InputParamDecl += TEXT("TextureExternal ");
+				InputParamDecl += InputNameStr;
+				InputParamDecl += TEXT(", SamplerState ");
+				InputParamDecl += InputNameStr;
+				InputParamDecl += TEXT("Sampler ");
+				break;
 			default:
 				return Errorf(TEXT("Bad type %s for %s input %s"),DescribeType(GetParameterType(CompiledInputs[i])), *Custom->Description, *InputNameStr);
 				break;
@@ -5312,7 +5321,7 @@ protected:
 
 			CodeChunk += TEXT(",");
 			CodeChunk += *ParamCode;
-			if (ParamType == MCT_Texture2D || ParamType == MCT_TextureCube)
+			if (ParamType == MCT_Texture2D || ParamType == MCT_TextureCube || ParamType == MCT_TextureExternal)
 			{
 				CodeChunk += TEXT(",");
 				CodeChunk += *ParamCode;
diff --git a/Engine/Source/Runtime/Engine/Private/Materials/MaterialExpressions.cpp b/Engine/Source/Runtime/Engine/Private/Materials/MaterialExpressions.cpp
--- a/Engine/Source/Runtime/Engine/Private/Materials/MaterialExpressions.cpp
+++ b/Engine/Source/Runtime/Engine/Private/Materials/MaterialExpressions.cpp
@@ -10714,7 +10714,8 @@ static const TCHAR* GetInputTypeName(uint8 InputType)
 		TEXT("T2d"),
 		TEXT("TCube"),
 		TEXT("B"),
-		TEXT("MA")
+		TEXT("MA"),
+		TEXT("TExt")
 	};
 
 	check(InputType < FunctionInput_MAX);
@@ -11220,7 +11221,8 @@ void UMaterialExpressionFunctionInput::GetCaption(TArray<FString>& OutCaptions)
 		TEXT("Texture2D"),
 		TEXT("TextureCube"),
 		TEXT("StaticBool"),
-		TEXT("MaterialAttributes")
+		TEXT("MaterialAttributes"),
+		TEXT("External")
 	};
 	check(InputType < FunctionInput_MAX);
 	OutCaptions.Add(FString(TEXT("Input ")) + InputName.ToString() + TEXT(" (") + TypeNames[InputType] + TEXT(")"));
@@ -11256,6 +11258,7 @@ int32 UMaterialExpressionFunctionInput::CompilePreviewValue(FMaterialCompiler* C
 			return FMaterialAttributeDefinitionMap::CompileDefaultExpression(Compiler, AttributeID);
 		case FunctionInput_Texture2D:
 		case FunctionInput_TextureCube:
+		case FunctionInput_TextureExternal:
 		case FunctionInput_StaticBool:
 			return Compiler->Errorf(TEXT("Missing Preview connection for function input '%s'"), *InputName.ToString());
 		default:
@@ -11275,7 +11278,8 @@ int32 UMaterialExpressionFunctionInput::Compile(class FMaterialCompiler* Compile
 		MCT_Texture2D,
 		MCT_TextureCube,
 		MCT_StaticBool,
-		MCT_MaterialAttributes
+		MCT_MaterialAttributes,
+		MCT_TextureExternal
 	};
 	check(InputType < FunctionInput_MAX);
 
@@ -11412,6 +11416,8 @@ uint32 UMaterialExpressionFunctionInput::GetInputType(int32 InputIndex)
 		return MCT_Texture2D;
 	case FunctionInput_TextureCube:
 		return MCT_TextureCube;
+	case FunctionInput_TextureExternal:
+		return MCT_TextureExternal;
 	case FunctionInput_StaticBool:
 		return MCT_StaticBool;
 	case FunctionInput_MaterialAttributes:
diff --git a/Engine/Source/Runtime/Engine/Private/Materials/MaterialInterface.cpp b/Engine/Source/Runtime/Engine/Private/Materials/MaterialInterface.cpp
--- a/Engine/Source/Runtime/Engine/Private/Materials/MaterialInterface.cpp
+++ b/Engine/Source/Runtime/Engine/Private/Materials/MaterialInterface.cpp
@@ -11,6 +11,8 @@
 #include "Materials/Material.h"
 #include "UObject/UObjectHash.h"
 #include "UObject/UObjectIterator.h"
+#include "EditorFramework/AssetImportData.h"
+#include "Engine/AssetUserData.h"
 #include "Engine/Texture2D.h"
 #include "Engine/SubsurfaceProfile.h"
 #include "Engine/TextureStreamingTypes.h"
@@ -236,6 +238,16 @@ void UMaterialInterface::PostEditChangeProperty(FPropertyChangedEvent& PropertyC
 
 	Super::PostEditChangeProperty(PropertyChangedEvent);
 }
+
+void UMaterialInterface::GetAssetRegistryTags(TArray<FAssetRegistryTag>& OutTags) const
+{
+	if (AssetImportData)
+	{
+		OutTags.Add( FAssetRegistryTag(SourceFileTagName(), AssetImportData->GetSourceData().ToJson(), FAssetRegistryTag::TT_Hidden) );
+	}
+
+	Super::GetAssetRegistryTags(OutTags);
+}
 #endif // WITH_EDITOR
 
 void UMaterialInterface::GetLightingGuidChain(bool bIncludeTextures, TArray<FGuid>& OutGuids) const
@@ -593,3 +605,41 @@ void UMaterialInterface::PreSave(const class ITargetPlatform* TargetPlatform)
 	}
 }
 
+void UMaterialInterface::AddAssetUserData(UAssetUserData* InUserData)
+{
+	if (InUserData != nullptr)
+	{
+		UAssetUserData* ExistingData = GetAssetUserDataOfClass(InUserData->GetClass());
+		if (ExistingData != nullptr)
+		{
+			AssetUserData.Remove(ExistingData);
+		}
+		AssetUserData.Add(InUserData);
+	}
+}
+
+UAssetUserData* UMaterialInterface::GetAssetUserDataOfClass(TSubclassOf<UAssetUserData> InUserDataClass)
+{
+	for (int32 DataIdx = 0; DataIdx < AssetUserData.Num(); DataIdx++)
+	{
+		UAssetUserData* Datum = AssetUserData[DataIdx];
+		if (Datum != nullptr && Datum->IsA(InUserDataClass))
+		{
+			return Datum;
+		}
+	}
+	return nullptr;
+}
+
+void UMaterialInterface::RemoveUserDataOfClass(TSubclassOf<UAssetUserData> InUserDataClass)
+{
+	for (int32 DataIdx = 0; DataIdx < AssetUserData.Num(); DataIdx++)
+	{
+		UAssetUserData* Datum = AssetUserData[DataIdx];
+		if (Datum != nullptr && Datum->IsA(InUserDataClass))
+		{
+			AssetUserData.RemoveAt(DataIdx);
+			return;
+		}
+	}
+}
diff --git a/Engine/Source/Runtime/Engine/Private/NetworkDriver.cpp b/Engine/Source/Runtime/Engine/Private/NetworkDriver.cpp
--- a/Engine/Source/Runtime/Engine/Private/NetworkDriver.cpp
+++ b/Engine/Source/Runtime/Engine/Private/NetworkDriver.cpp
@@ -1185,6 +1185,30 @@ bool UNetDriver::InitConnectionClass(void)
 
 bool UNetDriver::InitBase(bool bInitAsClient, FNetworkNotify* InNotify, const FURL& URL, bool bReuseAddressAndPort, FString& Error)
 {
+	// Read any timeout overrides from the URL
+	if (const TCHAR* InitialConnectTimeoutOverride = URL.GetOption(TEXT("InitialConnectTimeout="), nullptr))
+	{
+		float ParsedValue;
+		Lex::FromString(ParsedValue, InitialConnectTimeoutOverride);
+		if (ParsedValue != 0.0f)
+		{
+			InitialConnectTimeout = ParsedValue;
+		}
+	}
+	if (const TCHAR* ConnectionTimeoutOverride = URL.GetOption(TEXT("ConnectionTimeout="), nullptr))
+	{
+		float ParsedValue;
+		Lex::FromString(ParsedValue, ConnectionTimeoutOverride);
+		if (ParsedValue != 0.0f)
+		{
+			ConnectionTimeout = ParsedValue;
+		}
+	}
+	if (URL.HasOption(TEXT("NoTimeouts")))
+	{
+		bNoTimeouts = true;
+	}
+
 	LastTickDispatchRealtime = FPlatformTime::Seconds();
 	bool bSuccess = InitConnectionClass();
 
diff --git a/Engine/Source/Runtime/Engine/Private/RendererSettings.cpp b/Engine/Source/Runtime/Engine/Private/RendererSettings.cpp
--- a/Engine/Source/Runtime/Engine/Private/RendererSettings.cpp
+++ b/Engine/Source/Runtime/Engine/Private/RendererSettings.cpp
@@ -1,6 +1,7 @@
 // Copyright 1998-2018 Epic Games, Inc. All Rights Reserved.
 
 #include "Engine/RendererSettings.h"
+#include "PixelFormat.h"
 
 #if WITH_EDITOR
 #include "Editor/EditorEngine.h"
@@ -9,6 +10,15 @@
 extern UNREALED_API class UEditorEngine* GEditor;
 #endif // #if WITH_EDITOR
 
+namespace EDefaultBackBufferPixelFormat
+{
+	ENGINE_API EPixelFormat Convert2PixelFormat(int32 InDefaultBackBufferPixelFormat)
+	{
+		static EPixelFormat SPixelFormat[] = { PF_B8G8R8A8, PF_A16B16G16R16, PF_FloatRGB, PF_FloatRGBA, PF_A2B10G10R10 };
+		return SPixelFormat[FMath::Clamp(InDefaultBackBufferPixelFormat, 0, DBBPF_MAX - 1)];
+	}
+}
+
 URendererSettings::URendererSettings(const FObjectInitializer& ObjectInitializer)
 	: Super(ObjectInitializer)
 {
diff --git a/Engine/Source/Runtime/Engine/Private/ShaderCompiler/ShaderCompiler.cpp b/Engine/Source/Runtime/Engine/Private/ShaderCompiler/ShaderCompiler.cpp
--- a/Engine/Source/Runtime/Engine/Private/ShaderCompiler/ShaderCompiler.cpp
+++ b/Engine/Source/Runtime/Engine/Private/ShaderCompiler/ShaderCompiler.cpp
@@ -3104,7 +3104,12 @@ void GlobalBeginCompileShader(
 
 	{
 		static IConsoleVariable* CVar = IConsoleManager::Get().FindConsoleVariable(TEXT("r.PostProcessing.PropagateAlpha"));
-		Input.Environment.SetDefine(TEXT("POST_PROCESS_ALPHA"), CVar ? (CVar->GetInt() != 0) : 0);
+		int32 PropagateAlpha = CVar->GetInt();
+		if (PropagateAlpha < 0 || PropagateAlpha > 2)
+		{
+			PropagateAlpha = 0;
+		}
+		Input.Environment.SetDefine(TEXT("POST_PROCESS_ALPHA"), PropagateAlpha);
 	}
 
 	{
diff --git a/Engine/Source/Runtime/Engine/Private/StaticMesh.cpp b/Engine/Source/Runtime/Engine/Private/StaticMesh.cpp
--- a/Engine/Source/Runtime/Engine/Private/StaticMesh.cpp
+++ b/Engine/Source/Runtime/Engine/Private/StaticMesh.cpp
@@ -46,6 +46,7 @@
 #include "MeshUtilities.h"
 #include "DerivedDataCacheInterface.h"
 #include "PlatformInfo.h"
+#include "ScopedTransaction.h"
 #endif // #if WITH_EDITOR
 
 #include "Engine/StaticMeshSocket.h"
@@ -57,6 +58,7 @@
 #include "UObject/ReleaseObjectVersion.h"
 #include "Streaming/UVChannelDensity.h"
 
+#define LOCTEXT_NAMESPACE "StaticMesh"
 DEFINE_LOG_CATEGORY(LogStaticMesh);	
 
 DECLARE_MEMORY_STAT( TEXT( "StaticMesh Total Memory" ), STAT_StaticMeshTotalMemory2, STATGROUP_MemoryStaticMesh );
@@ -2310,7 +2312,7 @@ void FMeshSectionInfoMap::Clear()
 
 int32 FMeshSectionInfoMap::GetSectionNumber(int32 LODIndex) const
 {
-	int SectionCount = 0;
+	int32 SectionCount = 0;
 	for (auto kvp : Map)
 	{
 		if (((kvp.Key & 0xffff0000) >> 16) == LODIndex)
@@ -3681,6 +3683,53 @@ int32 UStaticMesh::GetMaterialIndex(FName MaterialSlotName) const
 	return -1;
 }
 
+#if WITH_EDITOR
+void UStaticMesh::SetMaterial(int32 MaterialIndex, UMaterialInterface* NewMaterial)
+{
+	static FName NAME_StaticMaterials = GET_MEMBER_NAME_CHECKED(UStaticMesh, StaticMaterials);
+
+	if (StaticMaterials.IsValidIndex(MaterialIndex))
+	{
+		FScopedTransaction ScopeTransaction(LOCTEXT("StaticMeshMaterialChanged", "StaticMesh: Material changed"));
+
+		// flag the property (Materials) we're modifying so that not all of the object is rebuilt.
+		UProperty* ChangedProperty = FindField<UProperty>(UStaticMesh::StaticClass(), NAME_StaticMaterials);
+		check(ChangedProperty);
+		PreEditChange(ChangedProperty);
+
+		StaticMaterials[MaterialIndex].MaterialInterface = NewMaterial;
+		if (NewMaterial != nullptr)
+		{
+			//Set the Material slot name to a good default one
+			if (StaticMaterials[MaterialIndex].MaterialSlotName == NAME_None)
+			{
+				StaticMaterials[MaterialIndex].MaterialSlotName = NewMaterial->GetFName();
+			}
+			//Set the original fbx material name so we can re-import correctly
+			if (StaticMaterials[MaterialIndex].ImportedMaterialSlotName == NAME_None)
+			{
+				StaticMaterials[MaterialIndex].ImportedMaterialSlotName = NewMaterial->GetFName();
+			}
+		}
+
+		if (ChangedProperty)
+		{
+			FPropertyChangedEvent PropertyUpdateStruct(ChangedProperty);
+			PostEditChangeProperty(PropertyUpdateStruct);
+		}
+		else
+		{
+			Modify();
+			PostEditChange();
+		}
+		if (BodySetup)
+		{
+			BodySetup->CreatePhysicsMeshes();
+		}
+	}
+}
+#endif //WITH_EDITOR
+
 /**
  * Returns the render data to use for exporting the specified LOD. This method should always
  * be called when exporting a static mesh.
@@ -3934,3 +3983,5 @@ void UStaticMeshSocket::Serialize(FArchive& Ar)
 		RelativeScale = FVector(1.0f, 1.0f, 1.0f);
 	}
 }
+
+#undef LOCTEXT_NAMESPACE
diff --git a/Engine/Source/Runtime/Engine/Private/UnrealEngine.cpp b/Engine/Source/Runtime/Engine/Private/UnrealEngine.cpp
--- a/Engine/Source/Runtime/Engine/Private/UnrealEngine.cpp
+++ b/Engine/Source/Runtime/Engine/Private/UnrealEngine.cpp
@@ -83,6 +83,7 @@ UnrealEngine.cpp: Implements the UEngine class and helpers.
 #include "EngineUtils.h"
 #include "Framework/Application/SlateApplication.h"
 #include "Widgets/Input/SButton.h"
+#include "Engine/EngineCustomTimeStep.h"
 #include "Engine/TextureLODSettings.h"
 #include "Engine/LevelStreamingPersistent.h"
 #include "Engine/ObjectReferencer.h"
@@ -1529,6 +1530,8 @@ void UEngine::PreExit()
 
 	delete ScreenSaverInhibitorRunnable;
 
+	SetCustomTimeStep(nullptr);
+
 	ShutdownHMD();
 }
 
@@ -1663,6 +1666,12 @@ void UEngine::UpdateTimeAndHandleMaxTickRate()
 	FTimedMemReport::Get().PumpTimedMemoryReports();
 #endif
 
+	if (CustomTimeStep)
+	{
+		CustomTimeStep->UpdateTimeStep(this);
+		return;
+	}
+
 	// This is always in realtime and is not adjusted by fixed framerate. Start slightly below current real time
 	static double LastRealTime = FPlatformTime::Seconds() - 0.0001;
 	static bool bTimeWasManipulated = false;
@@ -1852,6 +1861,31 @@ void UEngine::UpdateTimeAndHandleMaxTickRate()
 #endif // !UE_BUILD_SHIPPING
 }
 
+bool UEngine::SetCustomTimeStep(UEngineCustomTimeStep* InCustomTimeStep)
+{
+	bool bResult = true;
+
+	if (InCustomTimeStep != CustomTimeStep)
+	{
+		if (CustomTimeStep)
+		{
+			CustomTimeStep->Shutdown(this);
+		}
+
+		CustomTimeStep = InCustomTimeStep;
+
+		if (CustomTimeStep)
+		{
+			bResult = CustomTimeStep->Initialize(this);
+			if (!bResult)
+			{
+				CustomTimeStep = nullptr;
+			}
+		}
+	}
+
+	return bResult;
+}
 
 void UEngine::ParseCommandline()
 {
@@ -2715,6 +2749,11 @@ bool UEngine::IsSplitScreen(UWorld *InWorld)
 		// If no specified world, return true if any world context has multiple local players
 		for (auto It = WorldList.CreateIterator(); It; ++It)
 		{
+			if (It->WorldType == EWorldType::GameRPC)
+			{
+				continue;
+			}
+
 			if (It->OwningGameInstance != NULL && It->OwningGameInstance->GetNumLocalPlayers() > 1)
 			{
 				return true;
@@ -10060,6 +10099,11 @@ ULocalPlayer* UEngine::GetDebugLocalPlayer()
 {
 	for (auto It = WorldList.CreateConstIterator(); It; ++It)
 	{
+		if (It->WorldType == EWorldType::GameRPC)
+		{
+			continue;
+		}
+
 		if (It->OwningGameInstance != NULL && It->OwningGameInstance->GetFirstGamePlayer() != NULL )
 		{
 			return It->OwningGameInstance->GetFirstGamePlayer();
@@ -10969,24 +11013,58 @@ EBrowseReturnVal::Type UEngine::Browse( FWorldContext& WorldContext, FURL URL, F
 	}
 	else if (URL.HasOption(TEXT("failed")) || URL.HasOption(TEXT("closed")))
 	{
-		// Browsing after a failure, load default map
-
 		if (WorldContext.PendingNetGame)
 		{
 			CancelPending(WorldContext);
 		}
-		// Handle failure URL.
-		UE_LOG(LogNet, Log, TEXT("%s"), TEXT("Failed; returning to Entry") );
 		if (WorldContext.World() != NULL)
 		{
 			ResetLoaders( WorldContext.World()->GetOuter() );
 		}
 
+		if (WorldContext.WorldType == EWorldType::GameRPC)
+		{
+			UE_LOG(LogNet, Log, TEXT("RPC connection failed; retrying..."));
+
+			// Clean up all current net-drivers before we create a new one
+			// Need to copy the array as DestroyNamedNetDriver_Local mutates it
+			{
+				TArray<FNamedNetDriver> ActiveNetDrivers = WorldContext.ActiveNetDrivers;
+				for (const FNamedNetDriver& ActiveNetDriver : ActiveNetDrivers)
+				{
+					if (ActiveNetDriver.NetDriver)
+					{
+						DestroyNamedNetDriver_Local(WorldContext, ActiveNetDriver.NetDriver->NetDriverName);
+					}
+				}
+				CollectGarbage(GARBAGE_COLLECTION_KEEPFLAGS);
+			}
+
+			// Just reload the RPC world (as we have no real map to load)
+			WorldContext.PendingNetGame = NewObject<UPendingNetGame>();
+			WorldContext.PendingNetGame->Initialize(WorldContext.LastURL);
+			WorldContext.PendingNetGame->InitNetDriver();
+			if (!WorldContext.PendingNetGame->NetDriver)
+			{
+				// UPendingNetGame will set the appropriate error code and connection lost type, so
+				// we just have to propagate that message to the game.
+				BroadcastTravelFailure(WorldContext.World(), ETravelFailure::PendingNetGameCreateFailure, WorldContext.PendingNetGame->ConnectionError);
+				WorldContext.PendingNetGame = NULL;
+				return EBrowseReturnVal::Failure;
+			}
+
+			return EBrowseReturnVal::Pending;
+		}
+
+		// Browsing after a failure, load default map
+		UE_LOG(LogNet, Log, TEXT("Connection failed; returning to Entry"));
+		
 		const UGameMapsSettings* GameMapsSettings = GetDefault<UGameMapsSettings>();
-		const FString TextURL = GameMapsSettings->GetGameDefaultMap() + GameMapsSettings->LocalMapOptions;
-		if (!LoadMap(WorldContext, FURL(&URL, *TextURL, TRAVEL_Partial), NULL, Error))
+		const FURL DefaultURL = FURL(&URL, *(GameMapsSettings->GetGameDefaultMap() + GameMapsSettings->LocalMapOptions), TRAVEL_Partial);
+
+		if (!LoadMap(WorldContext, DefaultURL, NULL, Error))
 		{
-			HandleBrowseToDefaultMapFailure(WorldContext, TextURL, Error);
+			HandleBrowseToDefaultMapFailure(WorldContext, DefaultURL.ToString(), Error);
 			return EBrowseReturnVal::Failure;
 		}
 
@@ -11512,6 +11590,12 @@ bool UEngine::LoadMap( FWorldContext& WorldContext, FURL URL, class UPendingNetG
 		}
 	}
 
+	// Is this a minimal net RPC world?
+	if (WorldContext.WorldType == EWorldType::GameRPC)
+	{
+		UGameInstance::CreateMinimalNetRPCWorld(*URL.Map, WorldPackage, NewWorld);
+	}
+
 	const FString URLTrueMapName = URL.Map;
 
 	// Normal map loading
diff --git a/Engine/Source/Runtime/Engine/Private/UnrealExporter.cpp b/Engine/Source/Runtime/Engine/Private/UnrealExporter.cpp
--- a/Engine/Source/Runtime/Engine/Private/UnrealExporter.cpp
+++ b/Engine/Source/Runtime/Engine/Private/UnrealExporter.cpp
@@ -23,6 +23,8 @@
 #include "GameFramework/Actor.h"
 #include "Model.h"
 #include "Misc/FeedbackContext.h"
+#include "AssetExportTask.h"
+#include "UObject/GCObjectScopeGuard.h"
 
 DEFINE_LOG_CATEGORY_STATIC(LogExporter, Log, All);
 
@@ -203,150 +205,38 @@ void UExporter::ExportToOutputDevice(const FExportObjectInnerContext* Context, U
 int32 UExporter::ExportToFile( UObject* Object, UExporter* InExporter, const TCHAR* Filename, bool InSelectedOnly, bool NoReplaceIdentical, bool Prompt )
 {
 #if WITH_EDITOR
-	check(Object);
-
-	CurrentFilename = Filename;
-
-	UExporter*	Exporter	= InExporter;
-	int32		Result		= 0;
-	FString		Extension;
-
-	if (!Exporter)
-	{
-		// look for an exporter with all possible extensions, so an exporter can have something like *.xxx.yyy as an extension
-		int32 SearchStart = 0;
-		int32 DotLocation;
-		while (!Exporter && (DotLocation = CurrentFilename.Find(TEXT("."), ESearchCase::CaseSensitive, ESearchDir::FromStart, SearchStart)) != INDEX_NONE)
-		{
-			// get everything after the current .
-			Extension = CurrentFilename.Mid(DotLocation + 1);
-
-			// try to find an exporter with it
-			Exporter = FindExporter( Object, *Extension );
-
-			// skip past the dot in case we look again
-			SearchStart = DotLocation + 1;
-		}
-	}
-
-	if( !Exporter )
-	{
-		UE_LOG(LogExporter, Warning, TEXT("No %s exporter found for %s"), *Extension, *Object->GetFullName() );
-		CurrentFilename = TEXT("");
-		return 0;
-	}
-
-	Exporter->bSelectedOnly = InSelectedOnly;
-
-	if( Exporter->bText )
-	{
-		FStringOutputDevice Buffer;
-		const FExportObjectInnerContext Context;
-		ExportToOutputDevice( &Context, Object, Exporter, Buffer, *Extension, 0, PPF_ExportsNotFullyQualified, InSelectedOnly );
-		if ( Buffer.Len() == 0 )
-		{
-			Result = -1;
-		}
-		else
-		{
-			if( NoReplaceIdentical )
-			{
-				FString FileBytes;
-				if ( FFileHelper::LoadFileToString(FileBytes,Filename) && FCString::Strcmp(*Buffer,*FileBytes) == 0 )
-				{
-					UE_LOG(LogExporter, Log,  TEXT("Not replacing %s because identical"), Filename );
-					Result = 1;
-					goto Done;
-				}
-
-				if( Prompt )
-				{
-					if( !GWarn->YesNof( FText::Format( NSLOCTEXT("Core", "Overwrite", "The file '{0}' needs to be updated.  Do you want to overwrite the existing version?"), FText::FromString( Filename ) ) ) )
-					{
-						Result = 1;
-						goto Done;
-					}
-				}
-			}
-			if( !FFileHelper::SaveStringToFile( Buffer, Filename ) )
-			{
-	#if 0
-				if( GWarn->YesNof( FText::Format( NSLOCTEXT("Core", "OverwriteReadOnly", "'{0}' is marked read-only.  Would you like to try to force overwriting it?"), FText::FromString( Filename ) ) ) )
-				{
-					IFileManager::Get().Delete( Filename, 0, 1 );
-					if( FFileHelper::SaveStringToFile( Buffer, Filename ) )
-					{
-						Result = 1;
-						goto Done;
-					}
-				}
-	#endif
-				UE_LOG(LogExporter, Error, TEXT("Error exporting %s: couldn't open file '%s'"), *Object->GetFullName(), Filename);
-				goto Done;
-			}
-			Result = 1;
-		}
-	}
-	else
-	{
-		for( int32 i = 0; i < Exporter->GetFileCount(); i++ )
-		{
-			FBufferArchive Buffer;
-			if( ExportToArchive( Object, Exporter, Buffer, *Extension, i ) )
-			{
-				FString UniqueFilename = Exporter->GetUniqueFilename( Filename, i );
-
-				if( NoReplaceIdentical )
-				{
-					TArray<uint8> FileBytes;
-
-					if(	FFileHelper::LoadFileToArray( FileBytes, *UniqueFilename )
-					&&	FileBytes.Num() == Buffer.Num()
-					&&	FMemory::Memcmp( &FileBytes[ 0 ], &Buffer[ 0 ], Buffer.Num() ) == 0 )
-					{
-						UE_LOG(LogExporter, Log,  TEXT( "Not replacing %s because identical" ), *UniqueFilename );
-						Result = 1;
-						goto Done;
-					}
-					if( Prompt )
-					{
-						if( !GWarn->YesNof( FText::Format( NSLOCTEXT("Core", "Overwrite", "The file '{0}' needs to be updated.  Do you want to overwrite the existing version?"), FText::FromString( UniqueFilename ) ) ) )
-						{
-							Result = 1;
-							goto Done;
-						}
-					}
-				}
-
-				if( !FFileHelper::SaveArrayToFile( Buffer, *UniqueFilename ) )
-				{
-					UE_LOG(LogExporter, Error, TEXT("Error exporting %s: couldn't open file '%s'"), *Object->GetFullName(), *UniqueFilename);
-					goto Done;
-				}
-			}
-		}
-		Result = 1;
-	}
-Done:
-	CurrentFilename = TEXT("");
-
-	return Result;
+	UAssetExportTask* ExportTask = NewObject<UAssetExportTask>();
+	FGCObjectScopeGuard ExportTaskGuard(ExportTask);
+	ExportTask->Object = Object;
+	ExportTask->Exporter = InExporter;
+	ExportTask->Filename = Filename;
+	ExportTask->bSelected = InSelectedOnly;
+	ExportTask->bReplaceIdentical = !NoReplaceIdentical;
+	ExportTask->bPrompt = Prompt;
+	ExportTask->bUseFileArchive = false;
+	ExportTask->bWriteEmptyFiles = false;
+	ExportTask->bAutomated = false;
+	return RunAssetExportTask(ExportTask) ? 1 : 0;
 #else
 	return 0;
 #endif
 }
 
 
-int32 UExporter::ExportToFileEx( FExportToFileParams& ExportParams )
+bool UExporter::RunAssetExportTask(class UAssetExportTask* Task)
 {
 #if WITH_EDITOR
-	check(ExportParams.Object);
+	check(Task);
 
-	CurrentFilename = ExportParams.Filename;
+	CurrentFilename = Task->Filename;
+	struct FScopedCleanup
+	{
+		~FScopedCleanup() { UExporter::CurrentFilename = TEXT(""); }
+	};
+	FScopedCleanup ScopedCleanup;
 
-	UExporter*	Exporter	= ExportParams.Exporter;
-	FString		Extension	= FPaths::GetExtension(ExportParams.Filename);
-	int32		Result		= 0;
+	UExporter*	Exporter	= Task->Exporter;
+	FString		Extension	= FPaths::GetExtension(Task->Filename);
 
 	if (!Exporter)
 	{
@@ -359,7 +249,7 @@ int32 UExporter::ExportToFileEx( FExportToFileParams& ExportParams )
 			Extension = CurrentFilename.Mid(DotLocation + 1);
 
 			// try to find an exporter with it
-			Exporter = FindExporter( ExportParams.Object, *Extension );
+			Exporter = FindExporter(Task->Object, *Extension );
 
 			// skip past the dot in case we look again
 			SearchStart = DotLocation + 1;
@@ -368,25 +258,26 @@ int32 UExporter::ExportToFileEx( FExportToFileParams& ExportParams )
 
 	if( !Exporter )
 	{
-		UE_LOG(LogExporter, Warning, TEXT("No %s exporter found for %s"), *Extension, *(ExportParams.Object->GetFullName()) );
-		CurrentFilename = TEXT("");
-		return 0;
+		Task->Errors.Add(FString::Printf(TEXT("No %s exporter found for %s"), *Extension, *(Task->Object->GetFullName())));
+		UE_LOG(LogExporter, Warning, TEXT("%s"), *Task->Errors.Last());
+		return false;
 	}
 
-	Exporter->bSelectedOnly = ExportParams.InSelectedOnly;
+	Exporter->ExportTask = Task;
+	Exporter->bSelectedOnly = Task->bSelected;
 
-	FOutputDevice* TextBuffer = NULL;
+	TSharedPtr<FOutputDevice> TextBuffer;
 	if( Exporter->bText )
 	{
 		bool bIsFileDevice = false;
-		FString TempFile = FPaths::GetPath(ExportParams.Filename);
-		if (Exporter->bForceFileOperations || ExportParams.bUseFileArchive)
+		FString TempFile = FPaths::GetPath(Task->Filename);
+		if (Exporter->bForceFileOperations || Task->bUseFileArchive)
 		{
 			IFileManager::Get().MakeDirectory(*TempFile);
 
 			TempFile += TEXT("/UnrealExportFile.tmp");
-			TextBuffer = new FOutputDeviceFile(*TempFile);
-			if (TextBuffer)
+			TextBuffer = MakeShareable(new FOutputDeviceFile(*TempFile));
+			if (TextBuffer.IsValid())
 			{
 				TextBuffer->SetSuppressEventTag(true);
 				TextBuffer->SetAutoEmitLineTerminator(false);
@@ -394,67 +285,65 @@ int32 UExporter::ExportToFileEx( FExportToFileParams& ExportParams )
 			}
 		}
 
-		if (TextBuffer == NULL)
+		if (!TextBuffer.IsValid())
 		{
-			if (ExportParams.bUseFileArchive)
+			if (Task->bUseFileArchive)
 			{
 				UE_LOG(LogExporter, Warning, TEXT("Failed to create file output device... defaulting to string buffer"));
 			}
-			TextBuffer = new FStringOutputDevice();
+			TextBuffer = MakeShareable(new FStringOutputDevice());
 		}
-		const FExportObjectInnerContext Context(ExportParams.IgnoreObjectList);
-		ExportToOutputDevice( &Context, ExportParams.Object, Exporter, *TextBuffer, *Extension, 0, PPF_ExportsNotFullyQualified, ExportParams.InSelectedOnly );
+		const FExportObjectInnerContext Context(Task->IgnoreObjectList);
+		ExportToOutputDevice(&Context, Task->Object, Exporter, *TextBuffer, *Extension, 0, PPF_ExportsNotFullyQualified, Task->bSelected);
 		if (bIsFileDevice)
 		{
 			TextBuffer->TearDown();
-			IFileManager::Get().Move(ExportParams.Filename, *TempFile, 1, 1);
+			IFileManager::Get().Move(*Task->Filename, *TempFile, 1, 1);
 		}
 		else
 		{
-			FStringOutputDevice& StringBuffer = *((FStringOutputDevice*)TextBuffer);
+			FStringOutputDevice& StringBuffer = *((FStringOutputDevice*)TextBuffer.Get());
 			if ( StringBuffer.Len() == 0 )
 			{
-				Result = -1;
+				// non-fatal
+				return true;
 			}
 			else
 			{
-				if( ExportParams.NoReplaceIdentical )
+				if ( !Task->bReplaceIdentical )
 				{
 					FString FileBytes;
 					if
-						(	FFileHelper::LoadFileToString(FileBytes,ExportParams.Filename)
+						(	FFileHelper::LoadFileToString(FileBytes, *Task->Filename)
 						&&	FCString::Strcmp(*StringBuffer,*FileBytes)==0 )
 					{
-						UE_LOG(LogExporter, Log,  TEXT("Not replacing %s because identical"), ExportParams.Filename );
-						Result = 1;
-						goto Done;
+						UE_LOG(LogExporter, Log, TEXT("Not replacing %s because identical"), *Task->Filename);
+						return true;
 					}
-					if( ExportParams.Prompt )
+					if( Task->bPrompt )
 					{
-						if( !GWarn->YesNof( FText::Format( NSLOCTEXT("Core", "Overwrite", "The file '{0}' needs to be updated.  Do you want to overwrite the existing version?"), FText::FromString(  ExportParams.Filename ) ) ) )
+						if( !GWarn->YesNof(FText::Format(NSLOCTEXT("Core", "Overwrite", "The file '{0}' needs to be updated.  Do you want to overwrite the existing version?"), FText::FromString(  Task->Filename ) ) ) )
 						{
-							Result = 1;
-							goto Done;
+							return true;
 						}
 					}
 				}
-				if( !FFileHelper::SaveStringToFile( StringBuffer, ExportParams.Filename ) )
+				if(!FFileHelper::SaveStringToFile( StringBuffer, *Task->Filename ) )
 				{
 #if 0
-					if( GWarn->YesNof( FText::Format( NSLOCTEXT("Core", "OverwriteReadOnly", "'{0}' is marked read-only.  Would you like to try to force overwriting it?"), FText::FromString( ExportParams.Filename ) ) ) )
+					if(GWarn->YesNof(FText::Format(NSLOCTEXT("Core", "OverwriteReadOnly", "'{0}' is marked read-only.  Would you like to try to force overwriting it?"), FText::FromString(Task->Filename))))
 					{
-						IFileManager::Get().Delete( ExportParams.Filename, 0, 1 );
-						if( FFileHelper::SaveStringToFile( StringBuffer, ExportParams.Filename ) )
+						IFileManager::Get().Delete( Task->Filename, 0, 1 );
+						if(FFileHelper::SaveStringToFile( StringBuffer, *Task->Filename ) )
 						{
-							Result = 1;
-							goto Done;
+							return true;
 						}
 					}
 #endif
-					UE_LOG(LogExporter, Error, TEXT("%s"), *FString::Printf( TEXT("Error exporting %s: couldn't open file '%s'"), *(ExportParams.Object->GetFullName()), ExportParams.Filename));
-					goto Done;
+					UE_LOG(LogExporter, Error, TEXT("%s"), *FString::Printf(TEXT("Error exporting %s: couldn't open file '%s'"), *(Task->Object->GetFullName()), *Task->Filename));
+					return false;
 				}
-				Result = 1;
+				return true;
 			}
 		}
 	}
@@ -463,11 +352,11 @@ int32 UExporter::ExportToFileEx( FExportToFileParams& ExportParams )
 		for( int32 i = 0; i < Exporter->GetFileCount(); i++ )
 		{
 			FBufferArchive Buffer;
-			if( ExportToArchive( ExportParams.Object, Exporter, Buffer, *Extension, i ) )
+			if(ExportToArchive(Task->Object, Exporter, Buffer, *Extension, i))
 			{
-				FString UniqueFilename = Exporter->GetUniqueFilename( ExportParams.Filename, i );
+				FString UniqueFilename = Exporter->GetUniqueFilename(*Task->Filename, i);
 
-				if( ExportParams.NoReplaceIdentical )
+				if(!Task->bReplaceIdentical)
 				{
 					TArray<uint8> FileBytes;
 
@@ -476,43 +365,72 @@ int32 UExporter::ExportToFileEx( FExportToFileParams& ExportParams )
 					&&	FMemory::Memcmp( &FileBytes[ 0 ], &Buffer[ 0 ], Buffer.Num() ) == 0 )
 					{
 						UE_LOG(LogExporter, Log,  TEXT( "Not replacing %s because identical" ), *UniqueFilename );
-						Result = 1;
-						goto Done;
+						return true;
 					}
-					if( ExportParams.Prompt )
+					if(Task->bPrompt)
 					{
 						if( !GWarn->YesNof( FText::Format( NSLOCTEXT("Core", "Overwrite", "The file '{0}' needs to be updated.  Do you want to overwrite the existing version?"), FText::FromString( UniqueFilename ) ) ) )
 						{
-							Result = 1;
-							goto Done;
+							return true;
 						}
 					}
 				}
 
-				if ( !ExportParams.WriteEmptyFiles && !Buffer.Num() )
+				if (!Task->bWriteEmptyFiles && !Buffer.Num())
 				{
-					Result = 1;
-					goto Done;
+					return true;
 				}
 
 				if( !FFileHelper::SaveArrayToFile( Buffer, *UniqueFilename ) )
 				{
-					UE_LOG(LogExporter, Error, TEXT("Error exporting %s: couldn't open file '%s'"), *(ExportParams.Object->GetFullName()), *UniqueFilename);
-					goto Done;
+					Task->Errors.Add(FString::Printf(TEXT("Error exporting %s: couldn't open file '%s'"), *(Task->Object->GetFullName()), *UniqueFilename));
+					UE_LOG(LogExporter, Error, TEXT("%s"), *Task->Errors.Last());
+					return false;
 				}
 			}
 		}
-		Result = 1;
+		return true;
 	}
-Done:
-	if ( TextBuffer != NULL )
+	return false;
+#else
+	return false;
+#endif
+}
+
+bool UExporter::RunAssetExportTasks(const TArray<UAssetExportTask*>& ExportTasks)
+{
+#if WITH_EDITOR
+	bool bSuccess = true;
+	for (UAssetExportTask* Task : ExportTasks)
 	{
-		delete TextBuffer;
-		TextBuffer = NULL;
+		if (!RunAssetExportTask(Task))
+		{
+			bSuccess = false;
+		}
 	}
-	CurrentFilename = TEXT("");
+	return bSuccess;
+#else
+	return false;
+#endif
+}
 
-	return Result;
+int32 UExporter::ExportToFileEx( FExportToFileParams& ExportParams )
+{
+#if WITH_EDITOR
+	check(ExportParams.Object);
+	UAssetExportTask* ExportTask = NewObject<UAssetExportTask>();
+	FGCObjectScopeGuard ExportTaskGuard(ExportTask);
+	ExportTask->Object = ExportParams.Object;
+	ExportTask->Exporter = ExportParams.Exporter;
+	ExportTask->Filename = ExportParams.Filename;
+	ExportTask->bSelected = ExportParams.InSelectedOnly;
+	ExportTask->bReplaceIdentical = !ExportParams.NoReplaceIdentical;
+	ExportTask->bPrompt = ExportParams.Prompt;
+	ExportTask->bUseFileArchive = ExportParams.bUseFileArchive;
+	ExportTask->bWriteEmptyFiles = ExportParams.WriteEmptyFiles;
+	ExportTask->IgnoreObjectList = ExportParams.IgnoreObjectList;
+	ExportTask->bAutomated = false;
+	return RunAssetExportTask(ExportTask) ? 1 : 0;
 #else
 	return 0;
 #endif
@@ -595,7 +513,7 @@ FExportObjectInnerContext::FExportObjectInnerContext()
 }
 
 
-FExportObjectInnerContext::FExportObjectInnerContext(TArray<UObject*>& ObjsToIgnore)
+FExportObjectInnerContext::FExportObjectInnerContext(const TArray<UObject*>& ObjsToIgnore)
 {
 	// For each object . . .
 	for (UObject* InnerObj : TObjectRange<UObject>(RF_ClassDefaultObject, true, EInternalObjectFlags::PendingKill))
diff --git a/Engine/Source/Runtime/Engine/Private/World.cpp b/Engine/Source/Runtime/Engine/Private/World.cpp
--- a/Engine/Source/Runtime/Engine/Private/World.cpp
+++ b/Engine/Source/Runtime/Engine/Private/World.cpp
@@ -4812,8 +4812,9 @@ bool UWorld::Listen( FURL& InURL )
 		}
 		return false;
 	}
-	static bool LanPlay = FParse::Param(FCommandLine::Get(),TEXT("lanplay"));
-	if ( !LanPlay && (NetDriver->MaxInternetClientRate < NetDriver->MaxClientRate) && (NetDriver->MaxInternetClientRate > 2500) )
+	static const bool bLanPlay = FParse::Param(FCommandLine::Get(),TEXT("lanplay"));
+	const bool bLanSpeed = bLanPlay || InURL.HasOption(TEXT("LAN"));
+	if ( !bLanSpeed && (NetDriver->MaxInternetClientRate < NetDriver->MaxClientRate) && (NetDriver->MaxInternetClientRate > 2500) )
 	{
 		NetDriver->MaxClientRate = NetDriver->MaxInternetClientRate;
 	}
@@ -5932,8 +5933,8 @@ FConstLevelIterator	UWorld::GetLevelIterator() const
 
 ULevel* UWorld::GetLevel( int32 InLevelIndex ) const
 {
-	check( InLevelIndex < Levels.Num() );
-		check(Levels[InLevelIndex]);
+	check(InLevelIndex < Levels.Num());
+	check(Levels[InLevelIndex]);
 	return Levels[ InLevelIndex ];
 }
 
@@ -6013,7 +6014,7 @@ bool UWorld::IsPlayInVulkanPreview() const
 
 bool UWorld::IsGameWorld() const
 {
-	return WorldType == EWorldType::Game || WorldType == EWorldType::PIE || WorldType == EWorldType::GamePreview;
+	return WorldType == EWorldType::Game || WorldType == EWorldType::PIE || WorldType == EWorldType::GamePreview || WorldType == EWorldType::GameRPC;
 }
 
 bool UWorld::IsEditorWorld() const
diff --git a/Engine/Source/Runtime/Engine/Public/Model.h b/Engine/Source/Runtime/Engine/Public/Model.h
--- a/Engine/Source/Runtime/Engine/Public/Model.h
+++ b/Engine/Source/Runtime/Engine/Public/Model.h
@@ -489,10 +489,10 @@ public:
 	ENGINE_API void EmptyModel( int32 EmptySurfInfo, int32 EmptyPolys );
 
 	/** Begins releasing the model's resources. */
-	void BeginReleaseResources();
+	ENGINE_API void BeginReleaseResources();
 
 	/** Begins initializing the model's VertexBuffer. */
-	void UpdateVertices();
+	ENGINE_API void UpdateVertices();
 
 	/** Compute the "center" location of all the verts */
 	FVector GetCenter();
diff --git a/Engine/Source/Runtime/Engine/Public/Net/RepLayout.h b/Engine/Source/Runtime/Engine/Public/Net/RepLayout.h
--- a/Engine/Source/Runtime/Engine/Public/Net/RepLayout.h
+++ b/Engine/Source/Runtime/Engine/Public/Net/RepLayout.h
@@ -528,8 +528,8 @@ public:
 	void ClearSharedSerializationForRPC();
 
 	// Struct support
-	void SerializePropertiesForStruct( UStruct * Struct, FBitArchive & Ar, UPackageMap	* Map, void* Data, bool & bHasUnmapped ) const;	
-	void InitFromStruct( UStruct * InStruct, const UNetConnection* ServerConnection = nullptr );
+	ENGINE_API void SerializePropertiesForStruct( UStruct * Struct, FBitArchive & Ar, UPackageMap	* Map, void* Data, bool & bHasUnmapped ) const;	
+	ENGINE_API void InitFromStruct( UStruct * InStruct, const UNetConnection* ServerConnection = nullptr );
 
 	// Serializes all replicated properties of a UObject in or out of an archive (depending on what type of archive it is)
 	ENGINE_API void SerializeObjectReplicatedProperties(UObject* Object, FBitArchive & Ar) const;
diff --git a/Engine/Source/Runtime/Engine/Public/UnrealExporter.h b/Engine/Source/Runtime/Engine/Public/UnrealExporter.h
--- a/Engine/Source/Runtime/Engine/Public/UnrealExporter.h
+++ b/Engine/Source/Runtime/Engine/Public/UnrealExporter.h
@@ -30,7 +30,7 @@ public:
 	 * Creates the map from objects to their direct inners.
 	 * @param	ObjsToIgnore	An array of objects that should NOT be put in the list
 	 */
-	FExportObjectInnerContext(TArray<UObject*>& ObjsToIgnore);
+	FExportObjectInnerContext(const TArray<UObject*>& ObjsToIgnore);
 
 	/**Empty Constructor for derived export contexts */
 	FExportObjectInnerContext(const bool bIgnoredValue) {};
diff --git a/Engine/Source/Runtime/Media/Private/MediaModule.cpp b/Engine/Source/Runtime/Media/Private/MediaModule.cpp
--- a/Engine/Source/Runtime/Media/Private/MediaModule.cpp
+++ b/Engine/Source/Runtime/Media/Private/MediaModule.cpp
@@ -57,6 +57,11 @@ public:
 		return Ticker;
 	}
 
+	virtual FSimpleMulticastDelegate& GetOnTickPreEngineCompleted() override
+	{
+		return OnTickPreEngineCompleted;
+	}
+
 	virtual void LockToTimecode(bool Locked) override
 	{
 		TimecodeLocked = Locked;
@@ -95,6 +100,8 @@ public:
 		}
 
 		Clock.TickInput();
+
+		OnTickPreEngineCompleted.Broadcast();
 	}
 
 	virtual void TickPreSlate() override
@@ -159,6 +166,9 @@ private:
 	/** High-frequency ticker thread. */
 	FRunnableThread* TickerThread;
 
+	/** Delegate to receive TickPreEngine */
+	FSimpleMulticastDelegate OnTickPreEngineCompleted;
+
 	/** Whether media objects should lock to the media clock's time code. */
 	bool TimecodeLocked;
 
diff --git a/Engine/Source/Runtime/Media/Public/IMediaCache.h b/Engine/Source/Runtime/Media/Public/IMediaCache.h
--- a/Engine/Source/Runtime/Media/Public/IMediaCache.h
+++ b/Engine/Source/Runtime/Media/Public/IMediaCache.h
@@ -53,6 +53,17 @@ public:
 	{
 		return false; // override in child classes, if supported
 	}
+	
+	/**
+	 * Query the number of samples for the specified caching state.
+	 *
+	 * @param State The sample state we're interested in.
+	 * @return Sample count.
+	 */
+	virtual int32 GetSampleCount(EMediaCacheState State) const
+	{
+		return 0; // override in child classes, if supported
+	}
 
 public:
 
diff --git a/Engine/Source/Runtime/Media/Public/IMediaModule.h b/Engine/Source/Runtime/Media/Public/IMediaModule.h
--- a/Engine/Source/Runtime/Media/Public/IMediaModule.h
+++ b/Engine/Source/Runtime/Media/Public/IMediaModule.h
@@ -3,6 +3,7 @@
 #pragma once
 
 #include "Containers/Array.h"
+#include "Delegates/Delegate.h"
 #include "Misc/Timespan.h"
 #include "Modules/ModuleInterface.h"
 #include "Templates/SharedPointer.h"
@@ -117,6 +118,13 @@ public:
 	 */
 	virtual IMediaTicker& GetTicker() = 0;
 
+	/**
+	 * Get a Delegate that is trigger once all MediaClockSink are TickInput
+	 *
+	 * @return the OnTickPreEngineCompleted
+	 */
+	virtual FSimpleMulticastDelegate& GetOnTickPreEngineCompleted() = 0;
+
 	/**
 	 * Whether media objects should lock to the media clock's time code.
 	 *
diff --git a/Engine/Source/Runtime/Media/Public/IMediaPlayer.h b/Engine/Source/Runtime/Media/Public/IMediaPlayer.h
--- a/Engine/Source/Runtime/Media/Public/IMediaPlayer.h
+++ b/Engine/Source/Runtime/Media/Public/IMediaPlayer.h
@@ -4,6 +4,7 @@
 
 #include "Containers/UnrealString.h"
 #include "Internationalization/Text.h"
+#include "Math/Interval.h"
 #include "Misc/Paths.h"
 #include "Misc/Timespan.h"
 
diff --git a/Engine/Source/Runtime/MediaAssets/Private/Assets/MediaPlayer.cpp b/Engine/Source/Runtime/MediaAssets/Private/Assets/MediaPlayer.cpp
--- a/Engine/Source/Runtime/MediaAssets/Private/Assets/MediaPlayer.cpp
+++ b/Engine/Source/Runtime/MediaAssets/Private/Assets/MediaPlayer.cpp
@@ -35,18 +35,22 @@ UMediaPlayer::UMediaPlayer(const FObjectInitializer& ObjectInitializer)
 	, Shuffle(false)
 	, Loop(false)
 	, PlaylistIndex(INDEX_NONE)
+	, TimeDelay(FTimespan::Zero())
 	, HorizontalFieldOfView(90.0f)
 	, VerticalFieldOfView(60.0f)
 	, ViewRotation(FRotator::ZeroRotator)
-	, PlayerFacade(MakeShareable(new FMediaPlayerFacade))
 	, PlayerGuid(FGuid::NewGuid())
 	, PlayOnNext(false)
 #if WITH_EDITOR
 	, WasPlayingInPIE(false)
 #endif
 {
-	PlayerFacade->OnMediaEvent().AddUObject(this, &UMediaPlayer::HandlePlayerMediaEvent);
-	Playlist = NewObject<UMediaPlaylist>(GetTransientPackage(), NAME_None, RF_Transactional | RF_Transient);
+	if (!HasAnyFlags(RF_ClassDefaultObject))
+	{
+		PlayerFacade = MakeShared<FMediaPlayerFacade, ESPMode::ThreadSafe>();
+		PlayerFacade->OnMediaEvent().AddUObject(this, &UMediaPlayer::HandlePlayerMediaEvent);
+		Playlist = NewObject<UMediaPlaylist>(GetTransientPackage(), NAME_None, RF_Transactional | RF_Transient);
+	}
 }
 
 
@@ -277,6 +281,12 @@ FRotator UMediaPlayer::GetViewRotation() const
 }
 
 
+FTimespan UMediaPlayer::GetTimeDelay() const
+{
+	return PlayerFacade->TimeDelay;
+}
+
+
 bool UMediaPlayer::HasError() const
 {
 	return PlayerFacade->HasError();
@@ -577,6 +587,13 @@ bool UMediaPlayer::SetViewRotation(const FRotator& Rotation, bool Absolute)
 }
 
 
+void UMediaPlayer::SetTimeDelay(FTimespan InTimeDelay)
+{
+	UE_LOG(LogMediaAssets, Verbose, TEXT("%s.SetTimeDelay %s"), *GetFName().ToString(), *InTimeDelay.ToString());
+	PlayerFacade->TimeDelay = InTimeDelay;
+}
+
+
 bool UMediaPlayer::SupportsRate(float Rate, bool Unthinned) const
 {
 	return PlayerFacade->SupportsRate(Rate, Unthinned);
@@ -628,15 +645,18 @@ void UMediaPlayer::ResumePIE()
 
 void UMediaPlayer::BeginDestroy()
 {
-	IMediaModule* MediaModule = FModuleManager::LoadModulePtr<IMediaModule>("Media");
-
-	if (MediaModule != nullptr)
+	if (!HasAnyFlags(RF_ClassDefaultObject))
 	{
-		MediaModule->GetClock().RemoveSink(PlayerFacade.ToSharedRef());
-		MediaModule->GetTicker().RemoveTickable(PlayerFacade.ToSharedRef());
-	}
+		IMediaModule* MediaModule = FModuleManager::LoadModulePtr<IMediaModule>("Media");
 
-	PlayerFacade->Close();
+		if (MediaModule != nullptr)
+		{
+			MediaModule->GetClock().RemoveSink(PlayerFacade.ToSharedRef());
+			MediaModule->GetTicker().RemoveTickable(PlayerFacade.ToSharedRef());
+		}
+
+		PlayerFacade->Close();
+	}
 
 	Super::BeginDestroy();
 }
@@ -652,8 +672,11 @@ void UMediaPlayer::PostDuplicate(bool bDuplicateForPIE)
 {
 	Super::PostDuplicate(bDuplicateForPIE);
 
-	PlayerGuid = FGuid::NewGuid();
-	PlayerFacade->SetGuid(PlayerGuid);
+	if (!HasAnyFlags(RF_ClassDefaultObject))
+	{
+		PlayerGuid = FGuid::NewGuid();
+		PlayerFacade->SetGuid(PlayerGuid);
+	}
 }
 
 
@@ -661,15 +684,18 @@ void UMediaPlayer::PostInitProperties()
 {
 	Super::PostInitProperties();
 
-	// Set the player GUID - required for UMediaPlayers dynamically allocated at runtime
-	PlayerFacade->SetGuid(PlayerGuid);
+	if (!HasAnyFlags(RF_ClassDefaultObject))
+	{
+		// Set the player GUID - required for UMediaPlayers dynamically allocated at runtime
+		PlayerFacade->SetGuid(PlayerGuid);
 
-	IMediaModule* MediaModule = FModuleManager::LoadModulePtr<IMediaModule>("Media");
+		IMediaModule* MediaModule = FModuleManager::LoadModulePtr<IMediaModule>("Media");
 
-	if (MediaModule != nullptr)
-	{
-		MediaModule->GetClock().AddSink(PlayerFacade.ToSharedRef());
-		MediaModule->GetTicker().AddTickable(PlayerFacade.ToSharedRef());
+		if (MediaModule != nullptr)
+		{
+			MediaModule->GetClock().AddSink(PlayerFacade.ToSharedRef());
+			MediaModule->GetTicker().AddTickable(PlayerFacade.ToSharedRef());
+		}
 	}
 }
 
@@ -677,8 +703,11 @@ void UMediaPlayer::PostLoad()
 {
 	Super::PostLoad();
 
-	// Set the player GUID - required for UMediaPlayer assets
-	PlayerFacade->SetGuid(PlayerGuid);
+	if (!HasAnyFlags(RF_ClassDefaultObject))
+	{
+		// Set the player GUID - required for UMediaPlayer assets
+		PlayerFacade->SetGuid(PlayerGuid);
+	}
 }
 
 
@@ -694,6 +723,10 @@ void UMediaPlayer::PostEditChangeProperty(FPropertyChangedEvent& PropertyChanged
 	{
 		SetLooping(Loop);
 	}
+	else if (PropertyName == GET_MEMBER_NAME_CHECKED(UMediaPlayer, TimeDelay))
+	{
+		SetTimeDelay(TimeDelay);
+	}
 
 	Super::PostEditChangeProperty(PropertyChangedEvent);
 }
@@ -719,6 +752,7 @@ void UMediaPlayer::HandlePlayerMediaEvent(EMediaEvent Event)
 		PlayerFacade->SetLooping(Loop && (Playlist->Num() == 1));
 		PlayerFacade->SetViewField(HorizontalFieldOfView, VerticalFieldOfView, true);
 		PlayerFacade->SetViewOrientation(FQuat(ViewRotation), true);
+		PlayerFacade->TimeDelay = TimeDelay;
 
 		OnMediaOpened.Broadcast(PlayerFacade->GetUrl());
 
diff --git a/Engine/Source/Runtime/MediaAssets/Private/Assets/MediaTexture.cpp b/Engine/Source/Runtime/MediaAssets/Private/Assets/MediaTexture.cpp
--- a/Engine/Source/Runtime/MediaAssets/Private/Assets/MediaTexture.cpp
+++ b/Engine/Source/Runtime/MediaAssets/Private/Assets/MediaTexture.cpp
@@ -15,6 +15,7 @@
 
 #include "MediaPlayer.h"
 #include "Misc/MediaTextureResource.h"
+#include "IMediaTextureSample.h"
 
 
 /* Local helpers
@@ -350,3 +351,23 @@ void UMediaTexture::UpdateQueue()
 		SampleQueue.Reset();
 	}
 }
+
+FTimespan UMediaTexture::GetNextSampleTime() const
+{
+	FTimespan SampleTime;
+
+	TSharedPtr<IMediaTextureSample, ESPMode::ThreadSafe> Sample;
+	const bool bHasSucceed = SampleQueue->Peek(Sample);
+	if (bHasSucceed)
+	{
+		SampleTime = Sample->GetTime();
+	}
+
+	return SampleTime;
+}
+
+int32 UMediaTexture::GetAvailableSampleCount() const
+{
+	return SampleQueue->Num();
+}
+
diff --git a/Engine/Source/Runtime/MediaAssets/Public/MediaPlayer.h b/Engine/Source/Runtime/MediaAssets/Public/MediaPlayer.h
--- a/Engine/Source/Runtime/MediaAssets/Public/MediaPlayer.h
+++ b/Engine/Source/Runtime/MediaAssets/Public/MediaPlayer.h
@@ -398,6 +398,15 @@ public:
 	UFUNCTION(BlueprintCallable, Category="Media|MediaPlayer")
 	FRotator GetViewRotation() const;
 
+	/**
+	 * Delay of the player's time.
+	 *
+	 * @return Delay added to the player's time used to manually sync multiple sources.
+	 * @see SetTimeDelay
+	 */
+	UFUNCTION(BlueprintCallable, Category="Media|MediaPlayer")
+	FTimespan GetTimeDelay() const;
+
 	/**
 	 * Check whether the player is in an error state.
 	 *
@@ -717,6 +726,18 @@ public:
 	UFUNCTION(BlueprintCallable, Category = "Media|MediaPlayer")
 	bool SetViewRotation(const FRotator& Rotation, bool Absolute);
 
+	/**
+	 * Delay of the player's time.
+	 *
+	 * This setting can be used to manually sync multiple sources.
+	 * Set to 1 seconds, if you would like that Player to play 1 second behind his current time.
+	 * If the value is too big, it is possible that the player would not hold that frame for that long.
+	 * @return true on success, false otherwise.
+	 * @see GetTimeDelay
+	 */
+	UFUNCTION(BlueprintCallable, Category = "Media|MediaPlayer")
+	void SetTimeDelay(FTimespan TimeDelay);
+
 	/**
 	 * Check whether the specified playback rate is supported.
 	 *
@@ -978,6 +999,14 @@ protected:
 	UPROPERTY(BlueprintReadOnly, Category=Playback)
 	int32 PlaylistIndex;
 
+	/**
+	 * Delay of the player's time.
+	 *
+	 * @see SetTimeDelay, GetTimeDelay
+	 */
+	UPROPERTY(BlueprintReadOnly, Category=Playback)
+	FTimespan TimeDelay;
+
 protected:
 
 	/**
diff --git a/Engine/Source/Runtime/MediaAssets/Public/MediaTexture.h b/Engine/Source/Runtime/MediaAssets/Public/MediaTexture.h
--- a/Engine/Source/Runtime/MediaAssets/Public/MediaTexture.h
+++ b/Engine/Source/Runtime/MediaAssets/Public/MediaTexture.h
@@ -93,6 +93,19 @@ public:
 	UFUNCTION(BlueprintCallable, Category="Media|MediaTexture")
 	void SetMediaPlayer(UMediaPlayer* NewMediaPlayer);
 
+	/**
+	 * Gets the next sample Time. Only valid if GetAvailableSampleCount is greater than 0
+	 * @return FTimespan of the next sample
+	 * @see GetAvailableSampleCount
+	 */
+	FTimespan GetNextSampleTime() const;
+
+	/**
+	 * Gets the number of samples that are available
+	 * @return The number of samples in the queue
+	 */
+	int32 GetAvailableSampleCount() const;
+
 #if WITH_EDITOR
 
 	/**
diff --git a/Engine/Source/Runtime/MediaAssets/Public/TimeSynchronizableMediaSource.h b/Engine/Source/Runtime/MediaAssets/Public/TimeSynchronizableMediaSource.h
new file mode 100644
--- /dev/null
+++ b/Engine/Source/Runtime/MediaAssets/Public/TimeSynchronizableMediaSource.h
@@ -0,0 +1,40 @@
+// Copyright 1998-2018 Epic Games, Inc. All Rights Reserved.
+
+#pragma once
+
+#include "BaseMediaSource.h"
+
+#include "TimeSynchronizableMediaSource.generated.h"
+
+
+namespace TimeSynchronizableMedia
+{
+	/** Name of bUseTimeSynchronization media option. */
+	static FName UseTimeSynchronizatioOption("UseTimeSynchronization");
+}
+
+/**
+ * Base class for media sources that can be used with TimeSynchronizationManager.
+ */
+UCLASS(Abstract)
+class MEDIAASSETS_API UTimeSynchronizableMediaSource : public UBaseMediaSource
+{
+	GENERATED_BODY()
+	
+public:
+	/** Default constructor. */
+	UTimeSynchronizableMediaSource();
+
+public:
+
+	/** Option to enable using Time from the TimeSynchronizationManager. */
+	UPROPERTY(BlueprintReadWrite, EditAnywhere, Category=Synchronization)
+	bool bUseTimeSynchronization;
+
+public:
+	//~ IMediaOptions interface
+	using Super::GetMediaOption;
+	virtual bool GetMediaOption(const FName& Key, bool DefaultValue) const override;
+	virtual bool HasMediaOption(const FName& Key) const override;
+
+};
diff --git a/Engine/Source/Runtime/MediaUtils/Private/MediaPlayerFacade.cpp b/Engine/Source/Runtime/MediaUtils/Private/MediaPlayerFacade.cpp
--- a/Engine/Source/Runtime/MediaUtils/Private/MediaPlayerFacade.cpp
+++ b/Engine/Source/Runtime/MediaUtils/Private/MediaPlayerFacade.cpp
@@ -70,7 +70,8 @@ namespace MediaPlayerFacade
 *****************************************************************************/
 
 FMediaPlayerFacade::FMediaPlayerFacade()
-	: Cache(new FMediaSampleCache)
+	: TimeDelay(FTimespan::Zero())
+	, Cache(new FMediaSampleCache)
 	, LastRate(0.0f)
 { }
 
@@ -287,7 +288,17 @@ TRangeSet<float> FMediaPlayerFacade::GetSupportedRates(bool Unthinned) const
 
 FTimespan FMediaPlayerFacade::GetTime() const
 {
-	return Player.IsValid() ? Player->GetControls().GetTime() : FTimespan::Zero();
+	if (!Player.IsValid())
+	{
+		return FTimespan::Zero(); // no media opened
+	}
+
+	FTimespan Result = Player->GetControls().GetTime() - TimeDelay;
+	if (Result.GetTicks() < 0)
+	{
+		Result = FTimespan::Zero();
+	}
+	return Result;
 }
 
 
@@ -838,7 +849,7 @@ void FMediaPlayerFacade::TickFetch(FTimespan DeltaTime, FTimespan Timecode)
 	// determine range of valid samples
 	TRange<FTimespan> TimeRange;
 
-	const FTimespan Time = Player->GetControls().GetTime();
+	const FTimespan Time = GetTime();
 	
 	if (Rate > 0.0f)
 	{
@@ -883,7 +894,7 @@ void FMediaPlayerFacade::TickOutput(FTimespan DeltaTime, FTimespan /*Timecode*/)
 	}
 
 	IMediaControls& Controls = Player->GetControls();
-	Cache->Tick(DeltaTime, Controls.GetRate(), Controls.GetTime());
+	Cache->Tick(DeltaTime, Controls.GetRate(), GetTime());
 }
 
 
@@ -912,7 +923,7 @@ void FMediaPlayerFacade::TickTickable()
 	TRange<FTimespan> AudioTimeRange;
 	TRange<FTimespan> MetadataTimeRange;
 
-	const FTimespan Time = Player->GetControls().GetTime();
+	const FTimespan Time = GetTime();
 
 	if (LastRate > 0.0f)
 	{
diff --git a/Engine/Source/Runtime/MediaUtils/Public/MediaPlayerFacade.h b/Engine/Source/Runtime/MediaUtils/Public/MediaPlayerFacade.h
--- a/Engine/Source/Runtime/MediaUtils/Public/MediaPlayerFacade.h
+++ b/Engine/Source/Runtime/MediaUtils/Public/MediaPlayerFacade.h
@@ -55,6 +55,9 @@ public:
 	/** Name of the desired native player, if any. */
 	FName DesiredPlayerName;
 
+	/** Extra time to reduce from current player's time. */
+	FTimespan TimeDelay;
+
 public:
 
 	/** Default constructor. */
diff --git a/Engine/Source/Runtime/MediaUtils/Public/MediaSampleQueue.h b/Engine/Source/Runtime/MediaUtils/Public/MediaSampleQueue.h
--- a/Engine/Source/Runtime/MediaUtils/Public/MediaSampleQueue.h
+++ b/Engine/Source/Runtime/MediaUtils/Public/MediaSampleQueue.h
@@ -5,13 +5,17 @@
 #include "Containers/Queue.h"
 #include "HAL/PlatformAtomics.h"
 #include "IMediaSamples.h"
+#include "HAL/CriticalSection.h"
+#include "Math/Interval.h"
 #include "Misc/App.h"
+#include "Misc/ScopeLock.h"
 #include "Misc/Timespan.h"
 #include "Templates/SharedPointer.h"
 
 #include "MediaSampleSink.h"
 #include "MediaSampleSource.h"
 
+#include "IMediaTextureSample.h"
 
 /**
  * Template for media sample queues.
@@ -167,7 +171,7 @@ protected:
 		}
 	}
 
-private:
+protected:
 
 	/** Number of samples in the queue. */
 	int32 NumSamples;
@@ -191,3 +195,4 @@ typedef TMediaSampleQueue<class IMediaOverlaySample> FMediaOverlaySampleQueue;
 
 /** Type definition for texture sample queue. */
 typedef TMediaSampleQueue<class IMediaTextureSample> FMediaTextureSampleQueue;
+
diff --git a/Engine/Source/Runtime/MediaUtils/Public/MediaSamples.h b/Engine/Source/Runtime/MediaUtils/Public/MediaSamples.h
--- a/Engine/Source/Runtime/MediaUtils/Public/MediaSamples.h
+++ b/Engine/Source/Runtime/MediaUtils/Public/MediaSamples.h
@@ -146,17 +146,17 @@ public:
 private:
 
 	/** Audio sample queue. */
-	TMediaSampleQueue<IMediaAudioSample> AudioSampleQueue;
+	FMediaAudioSampleQueue AudioSampleQueue;
 
 	/** Caption sample queue. */
-	TMediaSampleQueue<IMediaOverlaySample> CaptionSampleQueue;
+	FMediaOverlaySampleQueue CaptionSampleQueue;
 
 	/** Metadata sample queue. */
-	TMediaSampleQueue<IMediaBinarySample> MetadataSampleQueue;
+	FMediaBinarySampleQueue MetadataSampleQueue;
 
 	/** Subtitle sample queue. */
-	TMediaSampleQueue<IMediaOverlaySample> SubtitleSampleQueue;
+	FMediaOverlaySampleQueue SubtitleSampleQueue;
 
 	/** Video sample queue. */
-	TMediaSampleQueue<IMediaTextureSample> VideoSampleQueue;
+	FMediaTextureSampleQueue VideoSampleQueue;
 };
diff --git a/Engine/Source/Runtime/MovieSceneCapture/Private/FrameGrabber.cpp b/Engine/Source/Runtime/MovieSceneCapture/Private/FrameGrabber.cpp
--- a/Engine/Source/Runtime/MovieSceneCapture/Private/FrameGrabber.cpp
+++ b/Engine/Source/Runtime/MovieSceneCapture/Private/FrameGrabber.cpp
@@ -153,13 +153,16 @@ void FViewportSurfaceReader::ResolveRenderTarget(const FViewportRHIRef& Viewport
 
 		FTexture2DRHIRef SourceBackBuffer = RHICmdList.GetViewportBackBuffer(ViewportRHI);
 
-		if (TargetSize.X != SourceBackBuffer->GetSizeX() || TargetSize.Y != SourceBackBuffer->GetSizeY())
+		const bool bIsSourceBackBufferSameAsWindowSize = SourceBackBuffer->GetSizeX() == WindowSize.X && SourceBackBuffer->GetSizeY() == WindowSize.Y;
+		const bool bIsSourceBackBufferSameAsTargetSize = TargetSize.X == SourceBackBuffer->GetSizeX() && TargetSize.Y == SourceBackBuffer->GetSizeY();
+
+		if (bIsSourceBackBufferSameAsWindowSize || bIsSourceBackBufferSameAsTargetSize)
 		{
-			PixelShader->SetParameters(RHICmdList, TStaticSamplerState<SF_Bilinear>::GetRHI(), SourceBackBuffer);
+			PixelShader->SetParameters(RHICmdList, TStaticSamplerState<SF_Point>::GetRHI(), SourceBackBuffer);
 		}
 		else
 		{
-			PixelShader->SetParameters(RHICmdList, TStaticSamplerState<SF_Point>::GetRHI(), SourceBackBuffer);
+			PixelShader->SetParameters(RHICmdList, TStaticSamplerState<SF_Bilinear>::GetRHI(), SourceBackBuffer);
 		}
 
 		float U = float(CaptureRect.Min.X) / float(SourceBackBuffer->GetSizeX());
@@ -173,7 +176,7 @@ void FViewportSurfaceReader::ResolveRenderTarget(const FViewportRHIRef& Viewport
 			TargetSize.X,							// Dest Width
 			TargetSize.Y,							// Dest Height
 			U, V,									// Source U, V
-			1, 1,									// Source USize, VSize
+			SizeU, SizeV,							// Source USize, VSize
 			CaptureRect.Max - CaptureRect.Min,		// Target buffer size
 			FIntPoint(1, 1),						// Source texture size
 			*VertexShader,
@@ -206,23 +209,26 @@ void FViewportSurfaceReader::ResolveRenderTarget(const FViewportRHIRef& Viewport
 	});
 }
 
-FFrameGrabber::FFrameGrabber(TSharedRef<FSceneViewport> Viewport, FIntPoint DesiredBufferSize, EPixelFormat InPixelFormat, uint32 NumSurfaces)
+FFrameGrabber::FFrameGrabber(TSharedRef<FSceneViewport> Viewport, FIntPoint DesiredBufferSize, EPixelFormat InPixelFormat, uint32 NumSurfaces, bool bAlwaysFlushOnDraw)
 {
 	State = EFrameGrabberState::Inactive;
 
-	// cause the viewport to always flush on draw
-	Viewport->IncrementFlushOnDraw();
-
+	if (bAlwaysFlushOnDraw)
 	{
-		// Setup a functor to decrement the flag on destruction (this class isn't necessarily tied to scene viewports)
-		TWeakPtr<FSceneViewport> WeakViewport = Viewport;
-		OnShutdown = [WeakViewport]{
-			TSharedPtr<FSceneViewport> PinnedViewport = WeakViewport.Pin();
-			if (PinnedViewport.IsValid())
-			{
-				PinnedViewport->DecrementFlushOnDraw();
-			}
-		};
+		// cause the viewport to always flush on draw
+		Viewport->IncrementFlushOnDraw();
+
+		{
+			// Setup a functor to decrement the flag on destruction (this class isn't necessarily tied to scene viewports)
+			TWeakPtr<FSceneViewport> WeakViewport = Viewport;
+			OnShutdown = [WeakViewport] {
+				TSharedPtr<FSceneViewport> PinnedViewport = WeakViewport.Pin();
+				if (PinnedViewport.IsValid())
+				{
+					PinnedViewport->DecrementFlushOnDraw();
+				}
+			};
+		}
 	}
 
 	TargetSize = DesiredBufferSize;
@@ -232,6 +238,7 @@ FFrameGrabber::FFrameGrabber(TSharedRef<FSceneViewport> Viewport, FIntPoint Desi
 	check(NumSurfaces != 0);
 
 	FIntRect CaptureRect(0,0,Viewport->GetSize().X, Viewport->GetSize().Y);
+	FIntPoint WindowSize(0, 0);
 
 	// Set up the capture rectangle
 	TSharedPtr<SViewport> ViewportWidget = Viewport->GetViewportWidget().Pin();
@@ -252,14 +259,18 @@ FFrameGrabber::FFrameGrabber(TSharedRef<FSceneViewport> Viewport, FIntPoint Desi
 			{
 				FArrangedWidget ArrangedWidget = WidgetPath.FindArrangedWidget(ViewportWidget.ToSharedRef()).Get(FArrangedWidget::NullWidget);
 
-				FVector2D Position = ArrangedWidget.Geometry.AbsolutePosition;
-				FVector2D Size = ArrangedWidget.Geometry.GetDrawSize();
+				FVector2D Position = ArrangedWidget.Geometry.GetAbsolutePosition();
+				FVector2D Size = ArrangedWidget.Geometry.GetAbsoluteSize();
 
 				CaptureRect = FIntRect(
 					Position.X,
 					Position.Y,
 					Position.X + Size.X,
 					Position.Y + Size.Y);
+
+				FVector2D AbsoluteSize = InnerWindowGeometry.GetAbsoluteSize();
+				WindowSize = FIntPoint(AbsoluteSize.X, AbsoluteSize.Y);
+
 			}
 		}
 	}
@@ -270,6 +281,7 @@ FFrameGrabber::FFrameGrabber(TSharedRef<FSceneViewport> Viewport, FIntPoint Desi
 	{
 		Surfaces.Emplace(InPixelFormat, DesiredBufferSize);
 		Surfaces.Last().Surface.SetCaptureRect(CaptureRect);
+		Surfaces.Last().Surface.SetWindowSize(WindowSize);
 	}
 
 	// Ensure textures are setup
@@ -422,22 +434,32 @@ void FFrameGrabber::OnFrameReady(int32 BufferIndex, FColor* ColorBuffer, int32 W
 	
 	const FResolveSurface& Surface = Surfaces[BufferIndex];
 
-	FCapturedFrameData ResolvedFrameData(TargetSize, Surface.Payload);
-
-	ResolvedFrameData.ColorBuffer.InsertUninitialized(0, TargetSize.X * TargetSize.Y);
-	FColor* Dest = &ResolvedFrameData.ColorBuffer[0];
-
-	const int32 MaxWidth = FMath::Min(TargetSize.X, Width);
-	for (int32 Row = 0; Row < FMath::Min(Height, TargetSize.Y); ++Row)
+	bool bExecuteDefaultGrabber = true;
+	if (Surface.Payload.IsValid())
 	{
-		FMemory::Memcpy(Dest, ColorBuffer, sizeof(FColor)*MaxWidth);
-		ColorBuffer += Width;
-		Dest += MaxWidth;
+		bExecuteDefaultGrabber = Surface.Payload->OnFrameReady_RenderThread(ColorBuffer, FIntPoint(Width, Height), TargetSize);
 	}
 
+	if (bExecuteDefaultGrabber)
 	{
-		FScopeLock Lock(&CapturedFramesMutex);
-		CapturedFrames.Add(MoveTemp(ResolvedFrameData));
+		FCapturedFrameData ResolvedFrameData(TargetSize, Surface.Payload);
+
+		ResolvedFrameData.ColorBuffer.InsertUninitialized(0, TargetSize.X * TargetSize.Y);
+		FColor* Dest = &ResolvedFrameData.ColorBuffer[0];
+
+		const int32 MaxWidth = FMath::Min(TargetSize.X, Width);
+		for (int32 Row = 0; Row < FMath::Min(Height, TargetSize.Y); ++Row)
+		{
+			FMemory::Memcpy(Dest, ColorBuffer, sizeof(FColor)*MaxWidth);
+			ColorBuffer += Width;
+			Dest += MaxWidth;
+		}
+
+		{
+			FScopeLock Lock(&CapturedFramesMutex);
+			CapturedFrames.Add(MoveTemp(ResolvedFrameData));
+		}
 	}
+
 	OutstandingFrameCount.Decrement();
 }
diff --git a/Engine/Source/Runtime/MovieSceneCapture/Public/FrameGrabber.h b/Engine/Source/Runtime/MovieSceneCapture/Public/FrameGrabber.h
--- a/Engine/Source/Runtime/MovieSceneCapture/Public/FrameGrabber.h
+++ b/Engine/Source/Runtime/MovieSceneCapture/Public/FrameGrabber.h
@@ -38,6 +38,9 @@ struct MOVIESCENECAPTURE_API FViewportSurfaceReader
 	/** Set the rectangle within which to read pixels */
 	void SetCaptureRect(FIntRect InCaptureRect) { CaptureRect = InCaptureRect; }
 
+	/** Set the window size that we expect from the BackBuffer */
+	void SetWindowSize(FIntPoint InWindowSize) { WindowSize = InWindowSize; }
+
 protected:
 
 	/** Set up this surface to the specified width/height */
@@ -55,6 +58,9 @@ protected:
 	/** The rectangle to read from the surface */
 	FIntRect CaptureRect;
 
+	/** In windows mode, the size of the widget with the border */
+	FIntPoint WindowSize;
+
 	/** The desired pixel format of the resolved textures */
 	EPixelFormat PixelFormat;
 
@@ -65,6 +71,12 @@ protected:
 struct IFramePayload
 {
 	virtual ~IFramePayload() {}
+
+	/**
+	 * Called when the buffer is now available in CPU ram
+	 * Return true if you would like to execute the default behavior. (If you return false, GetCapturedFrames will be empty).
+	 */
+	virtual bool OnFrameReady_RenderThread(FColor* ColorBuffer, FIntPoint BufferSize, FIntPoint TargetSize) const { return true; }
 };
 
 typedef TSharedPtr<IFramePayload, ESPMode::ThreadSafe> FFramePayloadPtr;
@@ -106,12 +118,13 @@ public:
 	/**
 	 * Construct this frame grabber
 	 *
-	 * @param InViewport		The viewport we are to grab frames for
-	 * @param DesiredBufferSize	The desired size of captured frames
-	 * @param InPixelFormat		The desired pixel format to store captured frames as
-	 * @param InNumSurfaces		The number of destination surfaces contained in our buffer 
+	 * @param InViewport			The viewport we are to grab frames for
+	 * @param DesiredBufferSize		The desired size of captured frames
+	 * @param InPixelFormat			The desired pixel format to store captured frames as
+	 * @param InNumSurfaces			The number of destination surfaces contained in our buffer 
+	 * @param bAlwaysFlushOnDraw	The viewport will flush at everydraw (increase synchronization, reduce performance)
 	 */
-	FFrameGrabber(TSharedRef<FSceneViewport> Viewport, FIntPoint DesiredBufferSize, EPixelFormat InPixelFormat = PF_B8G8R8A8, uint32 NumSurfaces = 2);
+	FFrameGrabber(TSharedRef<FSceneViewport> Viewport, FIntPoint DesiredBufferSize, EPixelFormat InPixelFormat = PF_B8G8R8A8, uint32 NumSurfaces = 2, bool bAlwaysFlushOnDraw = true);
 
 	/** Destructor */
 	~FFrameGrabber();
diff --git a/Engine/Source/Runtime/MovieSceneTracks/Private/Tracks/MovieSceneCameraCutTrack.cpp b/Engine/Source/Runtime/MovieSceneTracks/Private/Tracks/MovieSceneCameraCutTrack.cpp
--- a/Engine/Source/Runtime/MovieSceneTracks/Private/Tracks/MovieSceneCameraCutTrack.cpp
+++ b/Engine/Source/Runtime/MovieSceneTracks/Private/Tracks/MovieSceneCameraCutTrack.cpp
@@ -124,7 +124,7 @@ FFrameNumber UMovieSceneCameraCutTrack::FindEndTimeForCameraCut( FFrameNumber St
 
 	for( UMovieSceneSection* Section : Sections )
 	{
-		if( Section->HasStartFrame() && Section->GetInclusiveStartFrame() >= StartTime )
+		if( Section->HasStartFrame() && Section->GetInclusiveStartFrame() > StartTime )
 		{
 			ExclusiveEndTime = Section->GetInclusiveStartFrame();
 			break;
diff --git a/Engine/Source/Runtime/Projects/Private/ModuleDescriptor.cpp b/Engine/Source/Runtime/Projects/Private/ModuleDescriptor.cpp
--- a/Engine/Source/Runtime/Projects/Private/ModuleDescriptor.cpp
+++ b/Engine/Source/Runtime/Projects/Private/ModuleDescriptor.cpp
@@ -1,6 +1,7 @@
 // Copyright 1998-2018 Epic Games, Inc. All Rights Reserved.
 
 #include "ModuleDescriptor.h"
+#include "Misc/App.h"
 #include "Misc/ScopedSlowTask.h"
 #include "Dom/JsonObject.h"
 #include "Modules/ModuleManager.h"
@@ -194,6 +195,28 @@ bool FModuleDescriptor::Read(const FJsonObject& Object, FText& OutFailReason)
 		}
 	}
 
+	// Read the whitelisted target configurations
+	TSharedPtr<FJsonValue> WhitelistTargetConfigurationsValue = Object.TryGetField(TEXT("WhitelistTargetConfigurations"));
+	if (WhitelistTargetConfigurationsValue.IsValid() && WhitelistTargetConfigurationsValue->Type == EJson::Array)
+	{
+		const TArray< TSharedPtr< FJsonValue > >& ConfigsArray = WhitelistTargetConfigurationsValue->AsArray();
+		for (int Idx = 0; Idx < ConfigsArray.Num(); Idx++)
+		{
+			WhitelistTargetConfigurations.Add(ConfigsArray[Idx]->AsString());
+		}
+	}
+
+	// Read the blacklisted target configurations
+	TSharedPtr<FJsonValue> BlacklistTargetConfigurationsValue = Object.TryGetField(TEXT("BlacklistTargetConfigurations"));
+	if (BlacklistTargetConfigurationsValue.IsValid() && BlacklistTargetConfigurationsValue->Type == EJson::Array)
+	{
+		const TArray< TSharedPtr< FJsonValue > >& ConfigsArray = BlacklistTargetConfigurationsValue->AsArray();
+		for (int Idx = 0; Idx < ConfigsArray.Num(); Idx++)
+		{
+			BlacklistTargetConfigurations.Add(ConfigsArray[Idx]->AsString());
+		}
+	}
+
 	// Read the additional dependencies
 	TSharedPtr<FJsonValue> AdditionalDependenciesValue = Object.TryGetField(TEXT("AdditionalDependencies"));
 	if (AdditionalDependenciesValue.IsValid() && AdditionalDependenciesValue->Type == EJson::Array)
@@ -284,6 +307,24 @@ void FModuleDescriptor::Write(TJsonWriter<>& Writer) const
 		}
 		Writer.WriteArrayEnd();
 	}
+	if (WhitelistTargetConfigurations.Num() > 0)
+	{
+		Writer.WriteArrayStart(TEXT("WhitelistTargetConfigurations"));
+		for (int Idx = 0; Idx < WhitelistTargetConfigurations.Num(); Idx++)
+		{
+			Writer.WriteValue(WhitelistTargetConfigurations[Idx]);
+		}
+		Writer.WriteArrayEnd();
+	}
+	if (BlacklistTargetConfigurations.Num() > 0)
+	{
+		Writer.WriteArrayStart(TEXT("BlacklistTargetConfigurations"));
+		for (int Idx = 0; Idx < BlacklistTargetConfigurations.Num(); Idx++)
+		{
+			Writer.WriteValue(BlacklistTargetConfigurations[Idx]);
+		}
+		Writer.WriteArrayEnd();
+	}
 	if (AdditionalDependencies.Num() > 0)
 	{
 		Writer.WriteArrayStart(TEXT("AdditionalDependencies"));
@@ -312,7 +353,9 @@ void FModuleDescriptor::WriteArray(TJsonWriter<>& Writer, const TCHAR* Name, con
 bool FModuleDescriptor::IsCompiledInCurrentConfiguration() const
 {
 	// Cache the string for the current platform
-	static FString UBTPlatform(FPlatformMisc::GetUBTPlatform());
+	static const FString UBTPlatform = FPlatformMisc::GetUBTPlatform();
+	static const FString UBTTarget = FPlatformMisc::GetUBTTarget();
+	static const FString UBTTargetConfiguration = EBuildConfigurations::ToString(FApp::GetBuildConfiguration());
 
 	// Check the platform is whitelisted
 	if(WhitelistPlatforms.Num() > 0 && !WhitelistPlatforms.Contains(UBTPlatform))
@@ -326,8 +369,6 @@ bool FModuleDescriptor::IsCompiledInCurrentConfiguration() const
 		return false;
 	}
 
-	static FString UBTTarget(FPlatformMisc::GetUBTTarget());
-
 	// Check the target is whitelisted
 	if (WhitelistTargets.Num() > 0 && !WhitelistTargets.Contains(UBTTarget))
 	{
@@ -340,6 +381,18 @@ bool FModuleDescriptor::IsCompiledInCurrentConfiguration() const
 		return false;
 	}
 
+	// Check the target configuration is whitelisted
+	if (WhitelistTargetConfigurations.Num() > 0 && !WhitelistTargetConfigurations.Contains(UBTTargetConfiguration))
+	{
+		return false;
+	}
+
+	// Check the target configuration is not blacklisted
+	if (BlacklistTargetConfigurations.Num() > 0 && BlacklistTargetConfigurations.Contains(UBTTargetConfiguration))
+	{
+		return false;
+	}
+
 	// Check the module is compatible with this target. This should match ModuleDescriptor.IsCompiledInConfiguration in UBT
 	switch (Type)
 	{
diff --git a/Engine/Source/Runtime/Projects/Private/PluginManager.cpp b/Engine/Source/Runtime/Projects/Private/PluginManager.cpp
--- a/Engine/Source/Runtime/Projects/Private/PluginManager.cpp
+++ b/Engine/Source/Runtime/Projects/Private/PluginManager.cpp
@@ -645,7 +645,7 @@ bool FPluginManager::ConfigureEnabledPlugin(const FPluginReferenceDescriptor& Fi
 			const FPluginReferenceDescriptor& Reference = *NewPluginReferences[Idx];
 
 			// Check if the plugin is required for this platform
-			if(!Reference.IsEnabledForPlatform(FPlatformMisc::GetUBTPlatform()) || !Reference.IsEnabledForTarget(FPlatformMisc::GetUBTTarget()))
+			if(!Reference.IsEnabledForPlatform(FPlatformMisc::GetUBTPlatform()) || !Reference.IsEnabledForTargetConfiguration(EBuildConfigurations::ToString(FApp::GetBuildConfiguration())) || !Reference.IsEnabledForTarget(FPlatformMisc::GetUBTTarget()))
 			{
 				UE_LOG(LogPluginManager, Verbose, TEXT("Ignoring plugin '%s' for platform/configuration"), *Reference.Name);
 				continue;
diff --git a/Engine/Source/Runtime/Projects/Private/PluginReferenceDescriptor.cpp b/Engine/Source/Runtime/Projects/Private/PluginReferenceDescriptor.cpp
--- a/Engine/Source/Runtime/Projects/Private/PluginReferenceDescriptor.cpp
+++ b/Engine/Source/Runtime/Projects/Private/PluginReferenceDescriptor.cpp
@@ -61,6 +61,29 @@ bool FPluginReferenceDescriptor::IsEnabledForTarget(const FString& Target) const
     return true;
 }
 
+bool FPluginReferenceDescriptor::IsEnabledForTargetConfiguration(const FString& TargetConfiguration) const
+{
+	// If it's not enabled at all, return false
+	if (!bEnabled)
+	{
+		return false;
+	}
+
+	// If there is a list of whitelisted target configurations, and this isn't one of them, return false
+	if (WhitelistTargetConfigurations.Num() > 0 && !WhitelistTargetConfigurations.Contains(TargetConfiguration))
+	{
+		return false;
+	}
+
+	// If this target configuration is blacklisted, also return false
+	if (BlacklistTargetConfigurations.Contains(TargetConfiguration))
+	{
+		return false;
+	}
+
+	return true;
+}
+
 bool FPluginReferenceDescriptor::IsSupportedTargetPlatform(const FString& Platform) const
 {
 	return SupportedTargetPlatforms.Num() == 0 || SupportedTargetPlatforms.Contains(Platform);
@@ -93,6 +116,10 @@ bool FPluginReferenceDescriptor::Read( const FJsonObject& Object, FText& OutFail
 	Object.TryGetStringArrayField(TEXT("WhitelistPlatforms"), WhitelistPlatforms);
 	Object.TryGetStringArrayField(TEXT("BlacklistPlatforms"), BlacklistPlatforms);
 
+	// Get the target configuration lists
+	Object.TryGetStringArrayField(TEXT("WhitelistTargetConfigurations"), WhitelistTargetConfigurations);
+	Object.TryGetStringArrayField(TEXT("BlacklistTargetConfigurations"), BlacklistTargetConfigurations);
+
 	// Get the target lists
 	Object.TryGetStringArrayField(TEXT("WhitelistTargets"), WhitelistTargets);
 	Object.TryGetStringArrayField(TEXT("BlacklistTargets"), BlacklistTargets);
@@ -177,6 +204,30 @@ void FPluginReferenceDescriptor::Write( TJsonWriter<>& Writer ) const
 		Writer.WriteArrayEnd();
 	}
 
+	if (WhitelistTargetConfigurations.Num() > 0)
+	{
+		Writer.WriteArrayStart(TEXT("WhitelistTargetConfigurations"));
+
+		for (int Idx = 0; Idx < WhitelistTargetConfigurations.Num(); Idx++)
+		{
+			Writer.WriteValue(WhitelistTargetConfigurations[Idx]);
+		}
+
+		Writer.WriteArrayEnd();
+	}
+
+	if (BlacklistTargetConfigurations.Num() > 0)
+	{
+		Writer.WriteArrayStart(TEXT("BlacklistTargetConfigurations"));
+
+		for (int Idx = 0; Idx < BlacklistTargetConfigurations.Num(); Idx++)
+		{
+			Writer.WriteValue(BlacklistTargetConfigurations[Idx]);
+		}
+
+		Writer.WriteArrayEnd();
+	}
+
 	if (WhitelistTargets.Num() > 0)
 	{
 		Writer.WriteArrayStart(TEXT("WhitelistTargets"));
diff --git a/Engine/Source/Runtime/Projects/Public/ModuleDescriptor.h b/Engine/Source/Runtime/Projects/Public/ModuleDescriptor.h
--- a/Engine/Source/Runtime/Projects/Public/ModuleDescriptor.h
+++ b/Engine/Source/Runtime/Projects/Public/ModuleDescriptor.h
@@ -122,6 +122,12 @@ struct PROJECTS_API FModuleDescriptor
 	/** List of disallowed targets */
 	TArray<FString> BlacklistTargets;
 
+	/** List of allowed target configurations */
+	TArray<FString> WhitelistTargetConfigurations;
+
+	/** List of disallowed target configurations */
+	TArray<FString> BlacklistTargetConfigurations;
+
 	/** List of additional dependencies for building this module. */
 	TArray<FString> AdditionalDependencies;
 
diff --git a/Engine/Source/Runtime/Projects/Public/PluginReferenceDescriptor.h b/Engine/Source/Runtime/Projects/Public/PluginReferenceDescriptor.h
--- a/Engine/Source/Runtime/Projects/Public/PluginReferenceDescriptor.h
+++ b/Engine/Source/Runtime/Projects/Public/PluginReferenceDescriptor.h
@@ -33,6 +33,12 @@ struct PROJECTS_API FPluginReferenceDescriptor
 	/** If enabled, list of platforms for which the plugin should be disabled. */
 	TArray<FString> BlacklistPlatforms;
  
+	/** If enabled, list of target configurations for which the plugin should be enabled (or all target configurations if blank). */
+	TArray<FString> WhitelistTargetConfigurations;
+
+	/** If enabled, list of target configurations for which the plugin should be disabled. */
+	TArray<FString> BlacklistTargetConfigurations;
+
 	/** If enabled, list of targets for which the plugin should be enabled (or all targets if blank). */
 	TArray<FString> WhitelistTargets;
 
@@ -48,6 +54,9 @@ struct PROJECTS_API FPluginReferenceDescriptor
 	/** Determines whether the plugin is enabled for the given platform */
 	bool IsEnabledForPlatform(const FString& Platform) const;
 
+	/** Determines whether the plugin is enabled for the given target configuration */
+	bool IsEnabledForTargetConfiguration(const FString& TargetConfiguration) const;
+
 	/** Determines whether the plugin is enabled for the given target */
 	bool IsEnabledForTarget(const FString& Target) const;
 
diff --git a/Engine/Source/Runtime/Renderer/Private/PostProcess/PostProcessing.cpp b/Engine/Source/Runtime/Renderer/Private/PostProcess/PostProcessing.cpp
--- a/Engine/Source/Runtime/Renderer/Private/PostProcess/PostProcessing.cpp
+++ b/Engine/Source/Runtime/Renderer/Private/PostProcess/PostProcessing.cpp
@@ -169,8 +169,9 @@ static TAutoConsoleVariable<int32> CVarAlphaChannel(
 	TEXT("r.PostProcessing.PropagateAlpha"),
 	0,
 	TEXT("0 to disable scene alpha channel support in the post processing.\n")
-	TEXT(" 0: disabled (default)\n")
-	TEXT(" 1: enabled"),
+	TEXT(" 0: disabled (default);\n")
+	TEXT(" 1: enabled in linear color space;\n")
+	TEXT(" 2: same as 1, but also enable it through the tonemapper. Compositing after the tonemapper is incorrect, as their is no meaning to tonemap the alpha channel. This is only meant to be use exclusively for broadcasting hardware that does not support linear color space compositing and tonemapping."),
 	ECVF_ReadOnly);
 	
 static TAutoConsoleVariable<int32> CVarPostProcessingPreferCompute(
diff --git a/Engine/Source/Runtime/Renderer/Private/PostProcess/SceneRenderTargets.cpp b/Engine/Source/Runtime/Renderer/Private/PostProcess/SceneRenderTargets.cpp
--- a/Engine/Source/Runtime/Renderer/Private/PostProcess/SceneRenderTargets.cpp
+++ b/Engine/Source/Runtime/Renderer/Private/PostProcess/SceneRenderTargets.cpp
@@ -100,6 +100,17 @@ static TAutoConsoleVariable<int32> CVarGBufferFormat(
 	TEXT(" 5: high precision"),
 	ECVF_RenderThreadSafe);
 
+static TAutoConsoleVariable<int32> CVarDefaultBackBufferPixelFormat(
+	TEXT("r.DefaultBackBufferPixelFormat"),
+	4,
+	TEXT("Defines the default back buffer pixel format.\n")
+	TEXT(" 0: 8bit RGBA\n")
+	TEXT(" 1: 16bit RGBA\n")
+	TEXT(" 2: Float RGB\n")
+	TEXT(" 3: Float RGBA\n")
+	TEXT(" 4: 10bit RGB, 2bit Alpha\n"),
+	ECVF_ReadOnly);
+
 int32 GAllowCustomMSAAResolves = 1;
 static FAutoConsoleVariableRef CVarAllowCustomResolves(
    TEXT("r.MSAA.AllowCustomResolves"),
diff --git a/Engine/Source/Runtime/ShaderCore/Private/Shader.cpp b/Engine/Source/Runtime/ShaderCore/Private/Shader.cpp
--- a/Engine/Source/Runtime/ShaderCore/Private/Shader.cpp
+++ b/Engine/Source/Runtime/ShaderCore/Private/Shader.cpp
@@ -2016,9 +2016,16 @@ void ShaderMapAppendKeyString(EShaderPlatform Platform, FString& KeyString)
 		static const auto CVar = IConsoleManager::Get().FindTConsoleVariableDataInt(TEXT("r.PostProcessing.PropagateAlpha"));
 		if (CVar && CVar->GetValueOnAnyThread() > 0)
 		{
+			if (CVar->GetValueOnAnyThread() == 2)
+			{
+				KeyString += TEXT("_SA2");
+			}
+			else
+			{
 			KeyString += TEXT("_SA");
 		}
 	}
+	}
 
 	{
 		static const auto CVar = IConsoleManager::Get().FindTConsoleVariableDataInt(TEXT("r.VertexFoggingForOpaque"));
diff --git a/Engine/Source/Runtime/Slate/Private/Framework/Application/SlateApplication.cpp b/Engine/Source/Runtime/Slate/Private/Framework/Application/SlateApplication.cpp
--- a/Engine/Source/Runtime/Slate/Private/Framework/Application/SlateApplication.cpp
+++ b/Engine/Source/Runtime/Slate/Private/Framework/Application/SlateApplication.cpp
@@ -2707,7 +2707,7 @@ void FSlateApplication::UnregisterOnWindowActionNotification(FDelegateHandle Han
 	}
 }
 
-TSharedPtr<SWindow> FSlateApplication::FindBestParentWindowForDialogs(const TSharedPtr<SWidget>& InWidget)
+TSharedPtr<SWindow> FSlateApplication::FindBestParentWindowForDialogs(const TSharedPtr<SWidget>& InWidget, const ESlateParentWindowSearchMethod InParentWindowSearchMethod)
 {
 	TSharedPtr<SWindow> ParentWindow = ( InWidget.IsValid() ) ? FindWidgetWindow(InWidget.ToSharedRef()) : TSharedPtr<SWindow>();
 
@@ -2715,7 +2715,7 @@ TSharedPtr<SWindow> FSlateApplication::FindBestParentWindowForDialogs(const TSha
 	{
 		// First check the active top level window.
 		TSharedPtr<SWindow> ActiveTopWindow = GetActiveTopLevelWindow();
-		if ( ActiveTopWindow.IsValid() && ActiveTopWindow->IsRegularWindow() )
+		if ( ActiveTopWindow.IsValid() && ActiveTopWindow->IsRegularWindow() && InParentWindowSearchMethod == ESlateParentWindowSearchMethod::ActiveWindow )
 		{
 			ParentWindow = ActiveTopWindow;
 		}
@@ -2737,9 +2737,9 @@ TSharedPtr<SWindow> FSlateApplication::FindBestParentWindowForDialogs(const TSha
 	return ParentWindow;
 }
 
-const void* FSlateApplication::FindBestParentWindowHandleForDialogs(const TSharedPtr<SWidget>& InWidget)
+const void* FSlateApplication::FindBestParentWindowHandleForDialogs(const TSharedPtr<SWidget>& InWidget, const ESlateParentWindowSearchMethod InParentWindowSearchMethod)
 {
-	TSharedPtr<SWindow> ParentWindow = FindBestParentWindowForDialogs(InWidget);
+	TSharedPtr<SWindow> ParentWindow = FindBestParentWindowForDialogs(InWidget, InParentWindowSearchMethod);
 
 	const void* ParentWindowWindowHandle = nullptr;
 	if ( ParentWindow.IsValid() && ParentWindow->GetNativeWindow().IsValid() )
diff --git a/Engine/Source/Runtime/Slate/Public/Framework/Application/SlateApplication.h b/Engine/Source/Runtime/Slate/Public/Framework/Application/SlateApplication.h
--- a/Engine/Source/Runtime/Slate/Public/Framework/Application/SlateApplication.h
+++ b/Engine/Source/Runtime/Slate/Public/Framework/Application/SlateApplication.h
@@ -1553,13 +1553,13 @@ public:
 	* Given an optional widget, try and get the most suitable parent window to use with dialogs (such as file and directory pickers).
 	* This will first try and get the window that owns the widget (if provided), before falling back to using the MainFrame window.
 	*/
-	TSharedPtr<SWindow> FindBestParentWindowForDialogs(const TSharedPtr<SWidget>& InWidget);
+	TSharedPtr<SWindow> FindBestParentWindowForDialogs(const TSharedPtr<SWidget>& InWidget, const ESlateParentWindowSearchMethod InParentWindowSearchMethod = ESlateParentWindowSearchMethod::ActiveWindow);
 
 	/**
 	* Given an optional widget, try and get the most suitable parent window handle to use with dialogs (such as file and directory pickers).
 	* This will first try and get the window that owns the widget (if provided), before falling back to using the MainFrame window.
 	*/
-	const void* FindBestParentWindowHandleForDialogs(const TSharedPtr<SWidget>& InWidget);
+	const void* FindBestParentWindowHandleForDialogs(const TSharedPtr<SWidget>& InWidget, const ESlateParentWindowSearchMethod InParentWindowSearchMethod = ESlateParentWindowSearchMethod::ActiveWindow);
 
 public:
 #if WITH_EDITORONLY_DATA
