Fix non-unity compile error.

diff --git a/Engine/Source/Runtime/Engine/Private/GameInstance.cpp b/Engine/Source/Runtime/Engine/Private/GameInstance.cpp
--- a/Engine/Source/Runtime/Engine/Private/GameInstance.cpp
+++ b/Engine/Source/Runtime/Engine/Private/GameInstance.cpp
@@ -27,6 +27,7 @@
 #include "Kismet/GameplayStatics.h"
 #include "Framework/Application/SlateApplication.h"
 #include "GenericPlatform/GenericApplication.h"
+#include "Misc/PackageName.h"
 
 #if WITH_EDITOR
 #include "Settings/LevelEditorPlaySettings.h"
