Fixed crash when selecting Asset Manager in Project Settings.

diff --git a/Engine/Source/Runtime/Engine/Private/AssetManagerTypes.cpp b/Engine/Source/Runtime/Engine/Private/AssetManagerTypes.cpp
--- a/Engine/Source/Runtime/Engine/Private/AssetManagerTypes.cpp
+++ b/Engine/Source/Runtime/Engine/Private/AssetManagerTypes.cpp
@@ -135,10 +135,13 @@ void UAssetManagerSettings::PreEditChange(UProperty* PropertyAboutToChange)
 {
 	Super::PreEditChange(PropertyAboutToChange);
 
-	FName PropertyName = PropertyAboutToChange->GetFName();
-	if (PropertyName == GET_MEMBER_NAME_CHECKED(UAssetManagerSettings, MetaDataTagsForAssetRegistry))
+	if (PropertyAboutToChange != nullptr)
 	{
-		ClearMetaDataTagsSettings();
+		FName PropertyName = PropertyAboutToChange->GetFName();
+		if (PropertyName == GET_MEMBER_NAME_CHECKED(UAssetManagerSettings, MetaDataTagsForAssetRegistry))
+		{
+			ClearMetaDataTagsSettings();
+		}
 	}
 }
 
