Make trimesh cooking failures logs instead of warnings, as sometimes they are unavoidable (degenerate mesh) and don't want to give CIS cook warnings

diff --git a/Engine/Source/Runtime/Engine/Private/PhysicsEngine/BodySetup.cpp b/Engine/Source/Runtime/Engine/Private/PhysicsEngine/BodySetup.cpp
--- a/Engine/Source/Runtime/Engine/Private/PhysicsEngine/BodySetup.cpp
+++ b/Engine/Source/Runtime/Engine/Private/PhysicsEngine/BodySetup.cpp
@@ -1517,7 +1517,7 @@ FByteBulkData* UBodySetup::GetCookedData(FName Format, bool bRuntimeOnlyOptimize
 		}
 		else if(!bIsRuntime)	//only want to warn if DDC cooking failed - if it's really trying to use runtime and we can't, the runtime cooker code will catch it
 		{
-			UE_LOG(LogPhysics, Warning, TEXT("Attempt to build physics data for %s when we are unable to."), *GetPathName());
+			UE_LOG(LogPhysics, Log, TEXT("Attempt to build physics data for %s when we are unable to."), *GetPathName());
 		}
 	}
 #endif // WITH_PHYSX && WITH_EDITOR
diff --git a/Engine/Source/Runtime/Engine/Private/PhysicsEngine/PhysDerivedData.cpp b/Engine/Source/Runtime/Engine/Private/PhysicsEngine/PhysDerivedData.cpp
--- a/Engine/Source/Runtime/Engine/Private/PhysicsEngine/PhysDerivedData.cpp
+++ b/Engine/Source/Runtime/Engine/Private/PhysicsEngine/PhysDerivedData.cpp
@@ -189,7 +189,7 @@ bool FDerivedDataPhysXCooker::BuildTriMesh( TArray<uint8>& OutData,  const FTriM
 	if( !bResult )
 	{
 		bError = true;
-		UE_LOG(LogPhysics, Warning, TEXT("Failed to cook TriMesh: %s."), *CollisionDataProvider->GetPathName());
+		UE_LOG(LogPhysics, Log, TEXT("Failed to cook TriMesh: %s."), *CollisionDataProvider->GetPathName());
 	}
 
 	// If we want UV info, copy that now
diff --git a/Engine/Source/Runtime/Engine/Private/PhysicsEngine/PhysXCookHelper.cpp b/Engine/Source/Runtime/Engine/Private/PhysicsEngine/PhysXCookHelper.cpp
--- a/Engine/Source/Runtime/Engine/Private/PhysicsEngine/PhysXCookHelper.cpp
+++ b/Engine/Source/Runtime/Engine/Private/PhysicsEngine/PhysXCookHelper.cpp
@@ -25,7 +25,7 @@ void FPhysXCookHelper::CreatePhysicsMeshes_Concurrent()
 		const bool bError = !PhysXCookingModule->GetPhysXCooking()->CreateTriMesh(FPlatformProperties::GetPhysicsFormat(), CookInfo.TriMeshCookFlags, CookInfo.TriangleMeshDesc.Vertices, CookInfo.TriangleMeshDesc.Indices, CookInfo.TriangleMeshDesc.MaterialIndices, CookInfo.TriangleMeshDesc.bFlipNormals, OutTriangleMeshes[0]);
 		if (bError)
 		{
-			UE_LOG(LogPhysics, Warning, TEXT("Failed to cook TriMesh: %s."), *CookInfo.OuterDebugName);
+			UE_LOG(LogPhysics, Log, TEXT("Failed to cook TriMesh: %s."), *CookInfo.OuterDebugName);
 		}
 		else if (CookInfo.bSupportUVFromHitResults)
 		{
