Copying //UE4/Dev-Framework to //UE4/Dev-Main (Source: //UE4/Dev-Framework @ 4006870)

diff --git a/Engine/Source/Runtime/CoreUObject/Public/UObject/Object.h b/Engine/Source/Runtime/CoreUObject/Public/UObject/Object.h
--- a/Engine/Source/Runtime/CoreUObject/Public/UObject/Object.h
+++ b/Engine/Source/Runtime/CoreUObject/Public/UObject/Object.h
@@ -759,9 +759,7 @@ public:
 
 #if WITH_EDITOR
 	/**
-	 * Serializes all objects which have this object as their archetype into GMemoryArchive, then recursively calls this function
-	 * on each of those objects until the full list has been processed.
-	 * Called when a property value is about to be modified in an archetype object. 
+	 * Calls PreEditChange on all instances based on an archetype in AffectedObjects. Recurses on any instances.
 	 *
 	 * @param	AffectedObjects		the array of objects which have this object in their ObjectArchetype chain and will be affected by the change.
 	 *								Objects which have this object as their direct ObjectArchetype are removed from the list once they're processed.
@@ -769,8 +767,7 @@ public:
 	void PropagatePreEditChange( TArray<UObject*>& AffectedObjects, FEditPropertyChain& PropertyAboutToChange );
 
 	/**
-	 * De-serializes all objects which have this object as their archetype from the GMemoryArchive, then recursively calls this function
-	 * on each of those objects until the full list has been processed.
+	 * Calls PostEditChange on all instances based on an archetype in AffectedObjects. Recurses on any instances.
 	 *
 	 * @param	AffectedObjects		the array of objects which have this object in their ObjectArchetype chain and will be affected by the change.
 	 *								Objects which have this object as their direct ObjectArchetype are removed from the list once they're processed.
diff --git a/Engine/Source/Runtime/Engine/Classes/Components/WindDirectionalSourceComponent.h b/Engine/Source/Runtime/Engine/Classes/Components/WindDirectionalSourceComponent.h
--- a/Engine/Source/Runtime/Engine/Classes/Components/WindDirectionalSourceComponent.h
+++ b/Engine/Source/Runtime/Engine/Classes/Components/WindDirectionalSourceComponent.h
@@ -19,7 +19,7 @@ class FWindData;
 class FWindSourceSceneProxy;
 
 /** Component that provides a directional wind source. Only affects SpeedTree assets. */
-UCLASS(collapsecategories, hidecategories=(Object, Mobility), editinlinenew)
+UCLASS(MinimalAPI,collapsecategories, hidecategories=(Object, Mobility), editinlinenew)
 class UWindDirectionalSourceComponent : public USceneComponent
 {
 	GENERATED_UCLASS_BODY()
diff --git a/Engine/Source/Runtime/Engine/Classes/Engine/AssetManagerTypes.h b/Engine/Source/Runtime/Engine/Classes/Engine/AssetManagerTypes.h
--- a/Engine/Source/Runtime/Engine/Classes/Engine/AssetManagerTypes.h
+++ b/Engine/Source/Runtime/Engine/Classes/Engine/AssetManagerTypes.h
@@ -19,9 +19,12 @@ enum class EPrimaryAssetCookRule : uint8
 	/** Asset should never be cooked/shipped in any situation. An error will be generated if something depends on it. */
 	NeverCook,
 
-	/** Asset can be cooked for development and testing, but should never be shipped in a production build. */
+	/** Asset will be cooked in development if something else depends on it, but will never be cooked in a production build. */
 	DevelopmentCook,
 
+	/** Asset will always be cooked in development, but should never be cooked in a production build. */
+	DevelopmentAlwaysCook,
+
 	/** Asset will always be cooked, in both production and development. */
 	AlwaysCook,
 };
diff --git a/Engine/Source/Runtime/Engine/Classes/Engine/Blueprint.h b/Engine/Source/Runtime/Engine/Classes/Engine/Blueprint.h
--- a/Engine/Source/Runtime/Engine/Classes/Engine/Blueprint.h
+++ b/Engine/Source/Runtime/Engine/Classes/Engine/Blueprint.h
@@ -12,6 +12,7 @@
 #include "EdGraph/EdGraphPin.h"
 #include "Engine/BlueprintCore.h"
 #include "UObject/SoftObjectPath.h"
+#include "Blueprint/BlueprintSupport.h"
 #include "Blueprint.generated.h"
 
 class FCompilerResultsLog;
diff --git a/Engine/Source/Runtime/Engine/Classes/Engine/BlueprintCore.h b/Engine/Source/Runtime/Engine/Classes/Engine/BlueprintCore.h
--- a/Engine/Source/Runtime/Engine/Classes/Engine/BlueprintCore.h
+++ b/Engine/Source/Runtime/Engine/Classes/Engine/BlueprintCore.h
@@ -52,11 +52,6 @@ public:
 
 	virtual void Serialize( FArchive& Ar ) override;
 
-#if WITH_EDITORONLY_DATA
-	/** Gets asset registry tags */
-	virtual void GetAssetRegistryTags(TArray<FAssetRegistryTag>& OutTags) const override;
-#endif
-
 	/** Generates a new blueprint Guid, used when creating new blueprints */
 	void GenerateNewGuid()
 	{
diff --git a/Engine/Source/Runtime/Engine/Classes/Engine/World.h b/Engine/Source/Runtime/Engine/Classes/Engine/World.h
--- a/Engine/Source/Runtime/Engine/Classes/Engine/World.h
+++ b/Engine/Source/Runtime/Engine/Classes/Engine/World.h
@@ -2194,6 +2194,7 @@ public:
 
 	//~ Begin UObject Interface
 	virtual void Serialize( FArchive& Ar ) override;
+	virtual void BeginDestroy() override;
 	virtual void FinishDestroy() override;
 	virtual void PostLoad() override;
 	virtual bool PreSaveRoot(const TCHAR* Filename) override;
diff --git a/Engine/Source/Runtime/Engine/Classes/GameFramework/Actor.h b/Engine/Source/Runtime/Engine/Classes/GameFramework/Actor.h
--- a/Engine/Source/Runtime/Engine/Classes/GameFramework/Actor.h
+++ b/Engine/Source/Runtime/Engine/Classes/GameFramework/Actor.h
@@ -297,6 +297,9 @@ private:
 	UPROPERTY(Transient, DuplicateTransient)
 	uint8 bActorIsBeingDestroyed:1;
 
+	/** Set if an Actor tries to be destroyed while it is beginning play so that once BeginPlay ends we can issue the destroy call. */
+	uint8 bActorWantsDestroyDuringBeginPlay : 1;
+
 	enum class EActorBeginPlayState : uint8
 	{
 		HasNotBegunPlay,
@@ -2923,6 +2926,7 @@ private:
 
 	friend struct FMarkActorIsBeingDestroyed;
 	friend struct FActorParentComponentSetter;
+	friend struct FSetActorWantsDestroyDuringBeginPlay;
 
 	// Static helpers for accessing functions on SceneComponent.
 	// These are templates for no other reason than to delay compilation until USceneComponent is defined.
@@ -3012,6 +3016,20 @@ private:
 	friend UWorld;
 };
 
+/** This should only be used by UWorld::DestroyActor when the actor is in the process of beginning play so it can't be destroyed yet */
+struct FSetActorWantsDestroyDuringBeginPlay
+{
+private:
+	FSetActorWantsDestroyDuringBeginPlay(AActor* InActor)
+	{
+		ensure(InActor->IsActorBeginningPlay()); // Doesn't make sense to call this under any other circumstances
+		InActor->bActorWantsDestroyDuringBeginPlay = true;
+	}
+
+	friend UWorld;
+};
+
+
 /**
  * TInlineComponentArray is simply a TArray that reserves a fixed amount of space on the stack
  * to try to avoid heap allocation when there are fewer than a specified number of elements expected in the result.
diff --git a/Engine/Source/Runtime/Engine/Private/Actor.cpp b/Engine/Source/Runtime/Engine/Private/Actor.cpp
--- a/Engine/Source/Runtime/Engine/Private/Actor.cpp
+++ b/Engine/Source/Runtime/Engine/Private/Actor.cpp
@@ -674,7 +674,7 @@ void AActor::PostLoadSubobjects(FObjectInstancingGraph* OuterInstanceGraph)
 
 	ResetOwnedComponents();
 
-	if (RootComponent && bHadRoot && OldRoot != RootComponent)
+	if (RootComponent && bHadRoot && OldRoot != RootComponent && OldRoot->IsIn(this))
 	{
 		UE_LOG(LogActor, Log, TEXT("Root component has changed, relocating new root component to old position %s->%s"), *OldRoot->GetFullName(), *GetRootComponent()->GetFullName());
 		GetRootComponent()->RelativeRotation = OldRotation;
@@ -1972,8 +1972,7 @@ void AActor::RouteEndPlay(const EEndPlayReason::Type EndPlayReason)
 {
 	if (bActorInitialized)
 	{
-		UWorld* World = GetWorld();
-		if (World && World->HasBegunPlay())
+		if (ActorHasBegunPlay == EActorBeginPlayState::HasBegunPlay)
 		{
 			EndPlay(EndPlayReason);
 		}
@@ -1984,7 +1983,7 @@ void AActor::RouteEndPlay(const EEndPlayReason::Type EndPlayReason)
 			ClearComponentOverlaps();
 
 			bActorInitialized = false;
-			if (World)
+			if (UWorld* World = GetWorld())
 			{
 				World->RemoveNetworkActor(this);
 			}
@@ -2808,10 +2807,26 @@ void AActor::PostSpawnInitialize(FTransform const& UserSpawnTransform, AActor* I
 	USceneComponent* const SceneRootComponent = FixupNativeActorComponents(this);
 	if (SceneRootComponent != nullptr)
 	{
+		check(SceneRootComponent->GetOwner() == this);
+		
+		FVector RootComponentRelativeLocation = SceneRootComponent->RelativeLocation;
+		FRotator RootComponentRelativeRotation = SceneRootComponent->RelativeRotation;
+
+		// For converted BP root components, we must zero out RelativeLocation/RelativeRotation here. The reason is that in the non-nativized
+		// case, we ignore them when we instance a scene component that will also become the root (see USCS_Node::ExecuteNodeOnActor). Once
+		// a Blueprint class is nativized, we no longer run through that path, but we need to keep the same rotation/translation as before.
+		// We used to ignore them at nativization time, but that doesn't work because existing placements of the Blueprint component may rely
+		// on the value that's stored in the CDO, because it won't have been serialized out to the instance as a result of delta serialization.
+		if (Cast<UDynamicClass>(SceneRootComponent->GetArchetype()->GetOuter()->GetClass()))
+		{
+			RootComponentRelativeLocation = FVector::ZeroVector;
+			RootComponentRelativeRotation = FRotator::ZeroRotator;
+		}
+
 		// Set the actor's location and rotation since it has a native rootcomponent
 		// Note that we respect any initial transformation the root component may have from the CDO, so the final transform
 		// might necessarily be exactly the passed-in UserSpawnTransform.
- 		const FTransform RootTransform(SceneRootComponent->RelativeRotation, SceneRootComponent->RelativeLocation, SceneRootComponent->RelativeScale3D);
+ 		const FTransform RootTransform(RootComponentRelativeRotation, RootComponentRelativeLocation, SceneRootComponent->RelativeScale3D);
  		const FTransform FinalRootComponentTransform = RootTransform * UserSpawnTransform;
 		SceneRootComponent->SetWorldTransform(FinalRootComponentTransform);
 	}
@@ -3162,6 +3177,13 @@ void AActor::DispatchBeginPlay()
 
 		ensure(BeginPlayCallDepth - 1 == CurrentCallDepth);
 		BeginPlayCallDepth = CurrentCallDepth;
+
+		if (bActorWantsDestroyDuringBeginPlay)
+		{
+			// Pass true for bNetForce as either it doesn't matter or it was true the first time to even 
+			// get to the point we set bActorWantsDestroyDuringBeginPlay to true
+			World->DestroyActor(this, true); 
+		}
 	}
 }
 
diff --git a/Engine/Source/Runtime/Engine/Private/ActorConstruction.cpp b/Engine/Source/Runtime/Engine/Private/ActorConstruction.cpp
--- a/Engine/Source/Runtime/Engine/Private/ActorConstruction.cpp
+++ b/Engine/Source/Runtime/Engine/Private/ActorConstruction.cpp
@@ -884,7 +884,50 @@ static TAutoConsoleVariable<int32> CVarLogBlueprintComponentInstanceCalls(
 	TEXT("Log Blueprint Component instance calls; debugging."));
 #endif
 
-UActorComponent* AActor::CreateComponentFromTemplate(UActorComponent* Template, const FName InName)
+static FName FindFirstFreeName(UObject* Outer, FName BaseName)
+{
+	int32 Lower = 0;
+	FName Ret = FName(BaseName, Lower);
+
+	// Binary search if we appear to have used this name a lot, else
+	// just linear search for the first free index:
+	if (FindObjectFast<UObject>(Outer, FName(BaseName, 100)))
+	{
+		// could be replaced by Algo::LowerBound if TRange or some other
+		// integer range type could be made compatible with the algo's
+		// implementation:
+		int32 Upper = INT_MAX;
+		while (true)
+		{
+			int32 Next = (Upper - Lower) / 2 + Lower;
+			if (FindObjectFast<UObject>(Outer, FName(BaseName, Next)))
+			{
+				Lower = Next + 1;
+			}
+			else
+			{
+				Upper = Next;
+			}
+
+			if (Upper == Lower)
+			{
+				Ret = FName(BaseName, Lower);
+				break;
+			}
+		}
+	}
+	else
+	{
+		while (FindObjectFast<UObject>(Outer, Ret))
+		{
+			Ret = FName(BaseName, ++Lower);
+		}
+	}
+
+	return Ret;
+}
+
+UActorComponent* AActor::CreateComponentFromTemplate(UActorComponent* Template, FName InName)
 {
 	SCOPE_CYCLE_COUNTER(STAT_InstanceActorComponent);
 
@@ -894,15 +937,20 @@ UActorComponent* AActor::CreateComponentFromTemplate(UActorComponent* Template,
 #if !UE_BUILD_SHIPPING
 		const double StartTime = FPlatformTime::Seconds();
 #endif
-		// Make sure, that the name of the instance is different than the name of the template. This ensures that archetypes will not be recycled as instances in the nativized case.
-		const FName NewComponentName = (InName != NAME_None) ? InName : MakeUniqueObjectName(this, Template->GetClass(), Template->GetFName());
-		ensure(NewComponentName != Template->GetFName());
-
-		// Resolve any name conflicts.
-		CheckComponentInstanceName(NewComponentName);
+		if (InName == NAME_None)
+		{
+			// Search for a unique name based on our template. Our template is going to be our archetype
+			// thanks to logic in UBPGC::FindArchetype:
+			InName = FindFirstFreeName(this, Template->GetFName());
+		}
+		else
+		{
+			// Resolve any name conflicts.
+			CheckComponentInstanceName(InName);
+		}
 
 		// Note we aren't copying the the RF_ArchetypeObject flag. Also note the result is non-transactional by default.
-		NewActorComp = (UActorComponent*)StaticDuplicateObject(Template, this, NewComponentName, RF_AllFlags & ~(RF_ArchetypeObject | RF_Transactional | RF_WasLoaded | RF_Public | RF_InheritableComponentTemplate));
+		NewActorComp = (UActorComponent*)StaticDuplicateObject(Template, this, InName, RF_AllFlags & ~(RF_ArchetypeObject | RF_Transactional | RF_WasLoaded | RF_Public | RF_InheritableComponentTemplate));
 
 		// Handle post-creation tasks.
 		PostCreateBlueprintComponent(NewActorComp);
@@ -910,14 +958,14 @@ UActorComponent* AActor::CreateComponentFromTemplate(UActorComponent* Template,
 #if !UE_BUILD_SHIPPING
 		if (CVarLogBlueprintComponentInstanceCalls.GetValueOnGameThread())
 		{
-			UE_LOG(LogBlueprint, Log, TEXT("%s: CreateComponentFromTemplate() - %s \'%s\' completed in %.02g ms"), *GetName(), *Template->GetClass()->GetName(), *NewComponentName.ToString(), (FPlatformTime::Seconds() - StartTime) * 1000.0);
+			UE_LOG(LogBlueprint, Log, TEXT("%s: CreateComponentFromTemplate() - %s \'%s\' completed in %.02g ms"), *GetName(), *Template->GetClass()->GetName(), *InName.ToString(), (FPlatformTime::Seconds() - StartTime) * 1000.0);
 		}
 #endif
 	}
 	return NewActorComp;
 }
 
-UActorComponent* AActor::CreateComponentFromTemplateData(const FBlueprintCookedComponentInstancingData* TemplateData, const FName InName)
+UActorComponent* AActor::CreateComponentFromTemplateData(const FBlueprintCookedComponentInstancingData* TemplateData, FName InName)
 {
 	SCOPE_CYCLE_COUNTER(STAT_InstanceActorComponent);
 
@@ -943,18 +991,24 @@ UActorComponent* AActor::CreateComponentFromTemplateData(const FBlueprintCookedC
 #if !UE_BUILD_SHIPPING
 		const double StartTime = FPlatformTime::Seconds();
 #endif
-		// Make sure, that the name of the instance is different than the name of the template. This ensures that archetypes will not be recycled as instances in the nativized case.
-		const FName NewComponentName = (InName != NAME_None) ? InName : MakeUniqueObjectName(this, TemplateData->ComponentTemplateClass, TemplateData->ComponentTemplateName);
-		ensure(NewComponentName != TemplateData->ComponentTemplateName);
-
 		// Resolve any name conflicts.
-		CheckComponentInstanceName(NewComponentName);
+		if (InName == NAME_None)
+		{
+			// Search for a unique name based on our template. Our template is going to be our archetype
+			// thanks to logic in UBPGC::FindArchetype:
+			InName = FindFirstFreeName(this, TemplateData->ComponentTemplateName);
+		}
+		else
+		{
+			// Resolve any name conflicts.
+			CheckComponentInstanceName(InName);
+		}
 
 		// Note we aren't copying the the RF_ArchetypeObject flag. Also note the result is non-transactional by default.
 		NewActorComp = NewObject<UActorComponent>(
 			this,
 			TemplateData->ComponentTemplateClass,
-			NewComponentName,
+			InName,
 			EObjectFlags(TemplateData->ComponentTemplateFlags) & ~(RF_ArchetypeObject | RF_Transactional | RF_WasLoaded | RF_Public | RF_InheritableComponentTemplate)
 		);
 
@@ -978,7 +1032,7 @@ UActorComponent* AActor::CreateComponentFromTemplateData(const FBlueprintCookedC
 #if !UE_BUILD_SHIPPING
 		if (CVarLogBlueprintComponentInstanceCalls.GetValueOnGameThread())
 		{
-			UE_LOG(LogBlueprint, Log, TEXT("%s: CreateComponentFromTemplateData() - %s \'%s\' completed in %.02g ms"), *GetName(), *TemplateData->ComponentTemplateClass->GetName(), *NewComponentName.ToString(), (FPlatformTime::Seconds() - StartTime) * 1000.0);
+			UE_LOG(LogBlueprint, Log, TEXT("%s: CreateComponentFromTemplateData() - %s \'%s\' completed in %.02g ms"), *GetName(), *TemplateData->ComponentTemplateClass->GetName(), *InName.ToString(), (FPlatformTime::Seconds() - StartTime) * 1000.0);
 		}
 #endif
 	}
diff --git a/Engine/Source/Runtime/Engine/Private/AssetManager.cpp b/Engine/Source/Runtime/Engine/Private/AssetManager.cpp
--- a/Engine/Source/Runtime/Engine/Private/AssetManager.cpp
+++ b/Engine/Source/Runtime/Engine/Private/AssetManager.cpp
@@ -296,7 +296,7 @@ FPrimaryAssetId UAssetManager::DeterminePrimaryAssetIdForObject(const UObject* O
 
 bool UAssetManager::IsAssetDataBlueprintOfClassSet(const FAssetData& AssetData, const TSet<FName>& ClassNameSet)
 {
-	const FString ParentClassFromData = AssetData.GetTagValueRef<FString>("ParentClass");
+	const FString ParentClassFromData = AssetData.GetTagValueRef<FString>(FBlueprintTags::ParentClassPath);
 	if (!ParentClassFromData.IsEmpty())
 	{
 		const FString ClassObjectPath = FPackageName::ExportTextPathToObjectPath(ParentClassFromData);
@@ -2845,12 +2845,16 @@ void UAssetManager::ModifyCook(TArray<FName>& PackagesToCook, TArray<FName>& Pac
 		{
 			EPrimaryAssetCookRule CookRule = GetPackageCookRule(AssetData.PackageName);
 
-			if (CookRule == EPrimaryAssetCookRule::AlwaysCook && !TypeInfo.bIsEditorOnly)
+			// Treat DevAlwaysCook as AlwaysCook, may get excluded in VerifyCanCookPackage
+			bool bAlwaysCook = (CookRule == EPrimaryAssetCookRule::AlwaysCook || CookRule == EPrimaryAssetCookRule::DevelopmentAlwaysCook);
+			bool bCanCook = VerifyCanCookPackage(AssetData.PackageName, false);
+			
+			if (bAlwaysCook && bCanCook && !TypeInfo.bIsEditorOnly)
 			{
-				// If this is always cook and not editor only, cook it
+				// If this is always cook, not excluded, and not editor only, cook it
 				PackagesToCook.AddUnique(AssetData.PackageName);
 			}
-			else if (!VerifyCanCookPackage(AssetData.PackageName, false))
+			else if (!bCanCook)
 			{
 				// If this package cannot be cooked, add to exclusion list
 				PackagesToNeverCook.AddUnique(AssetData.PackageName);
@@ -2903,7 +2907,7 @@ bool UAssetManager::VerifyCanCookPackage(FName PackageName, bool bLogError) cons
 		
 		return false;
 	}
-	else if (CookRule == EPrimaryAssetCookRule::DevelopmentCook && bOnlyCookProductionAssets)
+	else if ((CookRule == EPrimaryAssetCookRule::DevelopmentCook || CookRule == EPrimaryAssetCookRule::DevelopmentAlwaysCook) && bOnlyCookProductionAssets)
 	{
 		if (bLogError)
 		{
diff --git a/Engine/Source/Runtime/Engine/Private/AssetManagerTypes.cpp b/Engine/Source/Runtime/Engine/Private/AssetManagerTypes.cpp
--- a/Engine/Source/Runtime/Engine/Private/AssetManagerTypes.cpp
+++ b/Engine/Source/Runtime/Engine/Private/AssetManagerTypes.cpp
@@ -6,11 +6,22 @@
 
 bool FPrimaryAssetTypeInfo::FillRuntimeData()
 {
+	if (PrimaryAssetType == NAME_None)
+	{
+		// Invalid type
+		return false;
+	}
+
+	if (!ensureMsgf(!AssetBaseClass.IsNull(), TEXT("Primary Asset Type %s must have a class set!"), *PrimaryAssetType.ToString()))
+	{
+		return false;
+	}
+
 	// Hot reload may have messed up asset pointer
 	AssetBaseClass.ResetWeakPtr();
 	AssetBaseClassLoaded = AssetBaseClass.LoadSynchronous();
 
-	if (!ensureMsgf(AssetBaseClassLoaded, TEXT("Failed to load Primary Asset Type class %s!"), *AssetBaseClass.ToString()))
+	if (!ensureMsgf(AssetBaseClassLoaded, TEXT("Failed to load class %s for Primary Asset Type %s!"), *AssetBaseClass.ToString(), *PrimaryAssetType.ToString()))
 	{
 		return false;
 	}
@@ -37,11 +48,7 @@ bool FPrimaryAssetTypeInfo::FillRuntimeData()
 		return false;
 	}
 
-	if (PrimaryAssetType == NAME_None)
-	{
-		// Invalid type
-		return false;
-	}
+	
 
 	return true;
 }
diff --git a/Engine/Source/Runtime/Engine/Private/Blueprint.cpp b/Engine/Source/Runtime/Engine/Private/Blueprint.cpp
--- a/Engine/Source/Runtime/Engine/Private/Blueprint.cpp
+++ b/Engine/Source/Runtime/Engine/Private/Blueprint.cpp
@@ -18,9 +18,7 @@
 #include "Modules/ModuleManager.h"
 
 #if WITH_EDITOR
-#include "Blueprint/BlueprintSupport.h"
 #include "BlueprintCompilationManager.h"
-#include "Blueprint/BlueprintSupport.h"
 #include "Editor/UnrealEd/Classes/Settings/ProjectPackagingSettings.h"
 #include "Engine/SimpleConstructionScript.h"
 #include "Engine/SCS_Node.h"
@@ -339,25 +337,6 @@ void UBlueprintCore::Serialize(FArchive& Ar)
 	}
 }
 
-#if WITH_EDITORONLY_DATA
-void UBlueprintCore::GetAssetRegistryTags(TArray<FAssetRegistryTag>& OutTags) const
-{
-	Super::GetAssetRegistryTags(OutTags);
-
-	FString GeneratedClassVal;
-	if ( GeneratedClass != NULL )
-	{
-		GeneratedClassVal = FString::Printf(TEXT("%s'%s'"), *GeneratedClass->GetClass()->GetName(), *GeneratedClass->GetPathName());
-	}
-	else
-	{
-		GeneratedClassVal = TEXT("None");
-	}
-
-	OutTags.Add( FAssetRegistryTag("GeneratedClass", GeneratedClassVal, FAssetRegistryTag::TT_Hidden) );
-}
-#endif
-
 void UBlueprintCore::GenerateDeterministicGuid()
 {
 	FString HashString = GetPathName();
@@ -921,12 +900,20 @@ void UBlueprint::GetAssetRegistryTags(TArray<FAssetRegistryTag>& OutTags) const
 
 	Super::GetAssetRegistryTags(OutTags);
 
-	FString ParentClassPackageName;
+	// Output native parent class and generated class as if they were AssetRegistrySearchable
+	FString GeneratedClassVal;
+	if (GeneratedClass)
+	{
+		GeneratedClassVal = FString::Printf(TEXT("%s'%s'"), *GeneratedClass->GetClass()->GetName(), *GeneratedClass->GetPathName());
+	}
+	else
+	{
+		GeneratedClassVal = TEXT("None");
+	}
+
 	FString NativeParentClassName;
 	if ( ParentClass )
 	{
-		ParentClassPackageName = ParentClass->GetOutermost()->GetName();
-
 		// Walk up until we find a native class (ie 'while they are BP classes')
 		UClass* NativeParentClass = ParentClass;
 		while (Cast<UBlueprintGeneratedClass>(NativeParentClass) != nullptr) // can't use IsA on UClass
@@ -937,22 +924,23 @@ void UBlueprint::GetAssetRegistryTags(TArray<FAssetRegistryTag>& OutTags) const
 	}
 	else
 	{
-		ParentClassPackageName = TEXT("None");
 		NativeParentClassName = TEXT("None");
 	}
 
-	//NumReplicatedProperties
+	OutTags.Add(FAssetRegistryTag(FBlueprintTags::GeneratedClassPath, GeneratedClassVal, FAssetRegistryTag::TT_Hidden));
+	OutTags.Add(FAssetRegistryTag(FBlueprintTags::NativeParentClassPath, NativeParentClassName, FAssetRegistryTag::TT_Alphabetical));
+
+	// BlueprintGeneratedClass is not automatically traversed so we have to manually add NumReplicatedProperties
 	int32 NumReplicatedProperties = 0;
 	UBlueprintGeneratedClass* BlueprintClass = Cast<UBlueprintGeneratedClass>(SkeletonGeneratedClass);
 	if (BlueprintClass)
 	{
 		NumReplicatedProperties = BlueprintClass->NumReplicatedProperties;
 	}
+	OutTags.Add(FAssetRegistryTag(FBlueprintTags::NumReplicatedProperties, FString::FromInt(NumReplicatedProperties), FAssetRegistryTag::TT_Numerical));
 
-	OutTags.Add(FAssetRegistryTag("NumReplicatedProperties", FString::FromInt(NumReplicatedProperties), FAssetRegistryTag::TT_Numerical));
-	OutTags.Add(FAssetRegistryTag("ParentClassPackage", ParentClassPackageName, FAssetRegistryTag::TT_Hidden));
-	OutTags.Add(FAssetRegistryTag("NativeParentClass", NativeParentClassName, FAssetRegistryTag::TT_Alphabetical));
-	OutTags.Add(FAssetRegistryTag(GET_MEMBER_NAME_CHECKED(UBlueprint, BlueprintDescription), BlueprintDescription, FAssetRegistryTag::TT_Hidden));
+	// This is explicit so it can be added as hidden
+	OutTags.Add(FAssetRegistryTag(FBlueprintTags::BlueprintDescription, BlueprintDescription, FAssetRegistryTag::TT_Hidden));
 
 	uint32 ClassFlagsTagged = 0;
 	if (BlueprintClass)
@@ -963,16 +951,16 @@ void UBlueprint::GetAssetRegistryTags(TArray<FAssetRegistryTag>& OutTags) const
 	{
 		ClassFlagsTagged = GetClass()->GetClassFlags();
 	}
-	OutTags.Add( FAssetRegistryTag("ClassFlags", FString::FromInt(ClassFlagsTagged), FAssetRegistryTag::TT_Hidden) );
+	OutTags.Add( FAssetRegistryTag(FBlueprintTags::ClassFlags, FString::FromInt(ClassFlagsTagged), FAssetRegistryTag::TT_Hidden) );
 
-	OutTags.Add( FAssetRegistryTag( "IsDataOnly",
+	OutTags.Add( FAssetRegistryTag(FBlueprintTags::IsDataOnly,
 			FBlueprintEditorUtils::IsDataOnlyBlueprint(this) ? TEXT("True") : TEXT("False"),
 			FAssetRegistryTag::TT_Alphabetical ) );
 
 	// Only add the FiB tags in the editor, this now gets run for standalone uncooked games
 	if ( ParentClass && GIsEditor)
 	{
-		OutTags.Add( FAssetRegistryTag("FiBData", FFindInBlueprintSearchManager::Get().QuerySingleBlueprint((UBlueprint*)this, false), FAssetRegistryTag::TT_Hidden) );
+		OutTags.Add( FAssetRegistryTag(FBlueprintTags::FindInBlueprintsData, FFindInBlueprintSearchManager::Get().QuerySingleBlueprint((UBlueprint*)this, false), FAssetRegistryTag::TT_Hidden) );
 	}
 
 	// Only show for strict blueprints (not animation or widget blueprints)
@@ -994,7 +982,7 @@ void UBlueprint::GetAssetRegistryTags(TArray<FAssetRegistryTag>& OutTags) const
 				}
 			}
 		}
-		OutTags.Add(FAssetRegistryTag("NativeComponents", FString::FromInt(NumNativeComponents), UObject::FAssetRegistryTag::TT_Numerical));
+		OutTags.Add(FAssetRegistryTag(FBlueprintTags::NumNativeComponents, FString::FromInt(NumNativeComponents), UObject::FAssetRegistryTag::TT_Numerical));
 
 		// Determine how many components are added via a SimpleConstructionScript (both newly introduced and inherited from parent BPs)
 		int32 NumAddedComponents = 0;
@@ -1006,7 +994,7 @@ void UBlueprint::GetAssetRegistryTags(TArray<FAssetRegistryTag>& OutTags) const
 				NumAddedComponents += AssociatedBP->SimpleConstructionScript->GetAllNodesConst().Num();
 			}
 		}
-		OutTags.Add(FAssetRegistryTag("BlueprintComponents", FString::FromInt(NumAddedComponents), UObject::FAssetRegistryTag::TT_Numerical));
+		OutTags.Add(FAssetRegistryTag(FBlueprintTags::NumBlueprintComponents, FString::FromInt(NumAddedComponents), UObject::FAssetRegistryTag::TT_Numerical));
 	}
 }
 
diff --git a/Engine/Source/Runtime/Engine/Private/ComponentInstanceDataCache.cpp b/Engine/Source/Runtime/Engine/Private/ComponentInstanceDataCache.cpp
--- a/Engine/Source/Runtime/Engine/Private/ComponentInstanceDataCache.cpp
+++ b/Engine/Source/Runtime/Engine/Private/ComponentInstanceDataCache.cpp
@@ -102,7 +102,7 @@ public:
 	virtual FArchive& operator<<(UObject*& Object) override
 	{
 		UObject* SerializedObject = Object;
-		if (Object && Object->IsIn(Component))
+		if (Object && Component && Object->IsIn(Component))
 		{
 			SerializedObject = GetDuplicatedObject(Object);
 		}
@@ -197,8 +197,9 @@ FActorComponentInstanceData::FActorComponentInstanceData(const UActorComponent*
 		FComponentPropertyWriter ComponentPropertyWriter(SourceComponent, SavedProperties, InstancedObjects);
 
 		// Cache off the length of an array that will come from SerializeTaggedProperties that had no properties saved in to it.
-		auto GetSizeOfEmptyArchive = [](const UActorComponent* DummyComponent) -> int32
+		auto GetSizeOfEmptyArchive = []() -> int32
 		{
+			const UActorComponent* DummyComponent = GetDefault<UActorComponent>();
 			TArray<uint8> NoWrittenPropertyReference;
 			TArray<UObject*> NoInstances;
 			FComponentPropertyWriter NullWriter(nullptr, NoWrittenPropertyReference, NoInstances);
@@ -211,7 +212,7 @@ FActorComponentInstanceData::FActorComponentInstanceData(const UActorComponent*
 			return NoWrittenPropertyReference.Num();
 		};
 
-		static const int32 SizeOfEmptyArchive = GetSizeOfEmptyArchive(SourceComponent);
+		static const int32 SizeOfEmptyArchive = GetSizeOfEmptyArchive();
 
 		// SerializeTaggedProperties will always put a sentinel NAME_None at the end of the Archive. 
 		// If that is the only thing in the buffer then empty it because we want to know that we haven't stored anything.
@@ -268,7 +269,7 @@ void FActorComponentInstanceData::ApplyToComponent(UActorComponent* Component, c
 
 		for (UObject* InstancedObject : InstancedObjects)
 		{
-			InstancedObject->Rename(nullptr, Component);
+			InstancedObject->Rename(nullptr, Component, REN_DontCreateRedirectors | REN_ForceNoResetLoaders);
 		}
 
 		FComponentPropertyReader ComponentPropertyReader(Component, SavedProperties);
diff --git a/Engine/Source/Runtime/Engine/Private/Components/SceneComponent.cpp b/Engine/Source/Runtime/Engine/Private/Components/SceneComponent.cpp
--- a/Engine/Source/Runtime/Engine/Private/Components/SceneComponent.cpp
+++ b/Engine/Source/Runtime/Engine/Private/Components/SceneComponent.cpp
@@ -935,8 +935,10 @@ void USceneComponent::DestroyComponent(bool bPromoteChildren/*= false*/)
 						else
 						{
 							// Default to first child node
-							check(AttachedChildren[0] != nullptr);
-							ChildToPromote = AttachedChildren[0];
+							if(ensureMsgf(AttachedChildren[0] != nullptr, TEXT("Deleting a non-root scene component with no promotable AttachChildren: %s"), *GetFullName()))
+							{
+								ChildToPromote = AttachedChildren[0];
+							}
 						}
 					}
 
@@ -1680,7 +1682,7 @@ void USceneComponent::SetRelativeRotationCache(const FRotationConversionCache& I
 
 void USceneComponent::SetupAttachment(class USceneComponent* InParent, FName InSocketName)
 {
-	if (ensureMsgf(!bRegistered, TEXT("SetupAttachment should only be used to initialize AttachParent and AttachSocketName for a future AttachTo. Once a component is registered you must use AttachTo.")))
+	if (ensureMsgf(!bRegistered, TEXT("SetupAttachment should only be used to initialize AttachParent and AttachSocketName for a future AttachToComponent. Once a component is registered you must use AttachToComponent.")))
 	{
 		if (ensureMsgf(InParent != this, TEXT("Cannot attach a component to itself.")))
 		{
diff --git a/Engine/Source/Runtime/Engine/Private/LevelActor.cpp b/Engine/Source/Runtime/Engine/Private/LevelActor.cpp
--- a/Engine/Source/Runtime/Engine/Private/LevelActor.cpp
+++ b/Engine/Source/Runtime/Engine/Private/LevelActor.cpp
@@ -567,6 +567,12 @@ bool UWorld::DestroyActor( AActor* ThisActor, bool bNetForce, bool bShouldModify
 			// Don't destroy PlayerControllers and BeaconClients
 			return false;
 		}
+
+		if (ThisActor->IsActorBeginningPlay())
+		{
+			FSetActorWantsDestroyDuringBeginPlay SetActorWantsDestroyDuringBeginPlay(ThisActor);
+			return true; // while we didn't actually destroy it now, we are going to, so tell the calling code it succeeded
+		}
 	}
 	else
 	{
diff --git a/Engine/Source/Runtime/Engine/Private/ObjectLibrary.cpp b/Engine/Source/Runtime/Engine/Private/ObjectLibrary.cpp
--- a/Engine/Source/Runtime/Engine/Private/ObjectLibrary.cpp
+++ b/Engine/Source/Runtime/Engine/Private/ObjectLibrary.cpp
@@ -455,7 +455,7 @@ int32 UObjectLibrary::LoadBlueprintAssetDataFromPaths(const TArray<FString>& Pat
 			FAssetData& Data = AssetDataList[AssetIdx];
 
 			bool bShouldRemove = true;
-			const FString ParentClassFromData = Data.GetTagValueRef<FString>("ParentClass");
+			const FString ParentClassFromData = Data.GetTagValueRef<FString>(FBlueprintTags::ParentClassPath);
 			if (!ParentClassFromData.IsEmpty())
 			{
 				const FString ClassObjectPath = FPackageName::ExportTextPathToObjectPath(ParentClassFromData);
diff --git a/Engine/Source/Runtime/Engine/Private/PlayerController.cpp b/Engine/Source/Runtime/Engine/Private/PlayerController.cpp
--- a/Engine/Source/Runtime/Engine/Private/PlayerController.cpp
+++ b/Engine/Source/Runtime/Engine/Private/PlayerController.cpp
@@ -64,6 +64,7 @@
 #include "Slate/SceneViewport.h"
 #include "Engine/NetworkObjectList.h"
 #include "GameFramework/GameSession.h"
+#include "GameMapsSettings.h"
 
 DEFINE_LOG_CATEGORY(LogPlayerController);
 
@@ -3948,10 +3949,10 @@ void APlayerController::ProcessForceFeedbackAndHaptics(const float DeltaTime, co
 	}	
 #endif
 
+	UWorld* World = GetWorld();
+
 	if (bProcessFeedback)
 	{
-		UWorld* World = GetWorld();
-
 		// --- Force Feedback --------------------------
 		for (int32 Index = ActiveForceFeedbackEffects.Num() - 1; Index >= 0; --Index)
 		{
@@ -4015,29 +4016,39 @@ void APlayerController::ProcessForceFeedbackAndHaptics(const float DeltaTime, co
 
 	if (FSlateApplication::IsInitialized())
 	{
-		const int32 ControllerId = GetInputIndex();
+		int32 ControllerId = GetInputIndex();
 
-		IInputInterface* InputInterface = FSlateApplication::Get().GetInputInterface();
-		if (InputInterface)
+		if (ControllerId != INVALID_CONTROLLERID)
 		{
-			InputInterface->SetForceFeedbackChannelValues(ControllerId, (bForceFeedbackEnabled ? ForceFeedbackValues : FForceFeedbackValues()));
-
-			bool bAreHapticsDisabled = (CVarDisableHaptics.GetValueOnGameThread() > 0) || bDisableHaptics;
-			if (!bAreHapticsDisabled)
+			IInputInterface* InputInterface = FSlateApplication::Get().GetInputInterface();
+			if (InputInterface)
 			{
-				// Haptic Updates
-				if (bLeftHapticsNeedUpdate)
-				{
-					InputInterface->SetHapticFeedbackValues(ControllerId, (int32)EControllerHand::Left, LeftHaptics);
-				}
-				if (bRightHapticsNeedUpdate)
+				// Adjust the ControllerId to account for the controller ID offset applied in UGameViewportClient::InputKey/Axis
+				// to play the force feedback on the correct controller if the offset player gamepad IDs feature is in use. 
+				const int32 NumLocalPlayers = World->GetGameInstance()->GetNumLocalPlayers();
+				if (NumLocalPlayers > 1 && GetDefault<UGameMapsSettings>()->bOffsetPlayerGamepadIds)
 				{
-					InputInterface->SetHapticFeedbackValues(ControllerId, (int32)EControllerHand::Right, RightHaptics);
+					--ControllerId;
 				}
 
-				if (bGunHapticsNeedUpdate)
+				InputInterface->SetForceFeedbackChannelValues(ControllerId, (bForceFeedbackEnabled ? ForceFeedbackValues : FForceFeedbackValues()));
+
+				const bool bAreHapticsDisabled = (CVarDisableHaptics.GetValueOnGameThread() > 0) || bDisableHaptics;
+				if (!bAreHapticsDisabled)
 				{
-					InputInterface->SetHapticFeedbackValues(ControllerId, (int32)EControllerHand::Gun, GunHaptics);
+					// Haptic Updates
+					if (bLeftHapticsNeedUpdate)
+					{
+						InputInterface->SetHapticFeedbackValues(ControllerId, (int32)EControllerHand::Left, LeftHaptics);
+					}
+					if (bRightHapticsNeedUpdate)
+					{
+						InputInterface->SetHapticFeedbackValues(ControllerId, (int32)EControllerHand::Right, RightHaptics);
+					}
+					if (bGunHapticsNeedUpdate)
+					{
+						InputInterface->SetHapticFeedbackValues(ControllerId, (int32)EControllerHand::Gun, GunHaptics);
+					}
 				}
 			}
 		}
diff --git a/Engine/Source/Runtime/Engine/Private/UnrealEngine.cpp b/Engine/Source/Runtime/Engine/Private/UnrealEngine.cpp
--- a/Engine/Source/Runtime/Engine/Private/UnrealEngine.cpp
+++ b/Engine/Source/Runtime/Engine/Private/UnrealEngine.cpp
@@ -12586,14 +12586,13 @@ bool UEngine::PrepareMapChange(FWorldContext &Context, const TArray<FName>& Leve
 	// Make sure we don't interrupt a pending map change in progress.
 	if( !IsPreparingMapChange(Context) )
 	{
-		Context.LevelsToLoadForPendingMapChange.Empty();
-		Context.LevelsToLoadForPendingMapChange += LevelNames;
+		Context.LevelsToLoadForPendingMapChange.Reset();
+		Context.LevelsToLoadForPendingMapChange.Append(LevelNames);
 
 #if !(UE_BUILD_SHIPPING || UE_BUILD_TEST)
 		// Verify that all levels specified are in the package file cache.
-		for( int32 LevelIndex=0; LevelIndex < Context.LevelsToLoadForPendingMapChange.Num(); LevelIndex++ )
+		for (const FName LevelName : Context.LevelsToLoadForPendingMapChange)
 		{
-			const FName LevelName = Context.LevelsToLoadForPendingMapChange[LevelIndex];
 			if( !FPackageName::DoesPackageExist( LevelName.ToString() ) )
 			{
 				Context.LevelsToLoadForPendingMapChange.Empty();
@@ -12618,10 +12617,8 @@ bool UEngine::PrepareMapChange(FWorldContext &Context, const TArray<FName>& Leve
 		}
 
 		// Kick off async loading of packages.
-		for( int32 LevelIndex=0; LevelIndex < Context.LevelsToLoadForPendingMapChange.Num(); LevelIndex++ )
+		for (const FName LevelName : Context.LevelsToLoadForPendingMapChange)
 		{
-			const FName LevelName = Context.LevelsToLoadForPendingMapChange[LevelIndex];
-
 			STAT_ADD_CUSTOMMESSAGE_NAME( STAT_NamedMarker, *(FString( TEXT( "PrepareMapChange - " ) + LevelName.ToString() )) );
 			LoadPackageAsync(LevelName.ToString(),
 				FLoadPackageAsyncDelegate::CreateStatic(&AsyncMapChangeLevelLoadCompletionCallback, Context.ContextHandle)
@@ -12800,11 +12797,12 @@ bool UEngine::CommitMapChange( FWorldContext &Context )
 			Context.World(),
 			*FString::Printf(TEXT("LevelStreamingPersistent_%s"), *FakePersistentLevel->GetOutermost()->GetName()) );
 
+		// Add it to the world info's list of levels.
+		Context.World()->AddStreamingLevel(LevelStreamingPersistent);
+
 		// Propagate level and name to streaming object.
 		FStreamingLevelPrivateAccessor::SetLoadedLevel(LevelStreamingPersistent, FakePersistentLevel);
 		LevelStreamingPersistent->SetWorldAssetByPackageName(FakePersistentLevelName);
-		// And add it to the world info's list of levels.
-		Context.World()->AddStreamingLevel(LevelStreamingPersistent);
 
 		UWorld* FakeWorld = CastChecked<UWorld>(FakePersistentLevel->GetOuter());
 
@@ -12816,7 +12814,7 @@ bool UEngine::CommitMapChange( FWorldContext &Context )
 		}
 
 		// Move the secondary levels to the world info levels array.
-		Context.World()->SetStreamingLevels(FakeWorld->GetStreamingLevels());
+		Context.World()->AddStreamingLevels(FakeWorld->GetStreamingLevels());
 		FakeWorld->ClearStreamingLevels();
 
 		// fixup up any kismet streaming objects to force them to be loaded if they were preloaded, this
diff --git a/Engine/Source/Runtime/Engine/Private/World.cpp b/Engine/Source/Runtime/Engine/Private/World.cpp
--- a/Engine/Source/Runtime/Engine/Private/World.cpp
+++ b/Engine/Source/Runtime/Engine/Private/World.cpp
@@ -759,6 +759,20 @@ void UWorld::PostDuplicate(bool bDuplicateForPIE)
 #endif // WITH_EDITOR
 }
 
+void UWorld::BeginDestroy()
+{
+	Super::BeginDestroy();
+
+	for (FLevelCollection& LevelCollection : LevelCollections)
+	{
+		TSet<ULevel*> CollectionLevels = LevelCollection.GetLevels();
+		for (ULevel* CollectionLevel : CollectionLevels)
+		{
+			LevelCollection.RemoveLevel(CollectionLevel);
+		}
+	}
+}
+
 void UWorld::FinishDestroy()
 {
 	// Avoid cleanup if the world hasn't been initialized, e.g., the default object or a world object that got loaded
diff --git a/Engine/Source/Runtime/Engine/Public/VisualLogger/VisualLogger.h b/Engine/Source/Runtime/Engine/Public/VisualLogger/VisualLogger.h
--- a/Engine/Source/Runtime/Engine/Public/VisualLogger/VisualLogger.h
+++ b/Engine/Source/Runtime/Engine/Public/VisualLogger/VisualLogger.h
@@ -37,13 +37,13 @@
 #define UE_CVLOG_LOCATION(Condition, LogOwner, CategoryName, Verbosity, Location, Radius, Color, Format, ...)  if(FVisualLogger::IsRecording() && Condition) {UE_VLOG_LOCATION(LogOwner, CategoryName, Verbosity, Location, Radius, Color, Format, ##__VA_ARGS__);} 
 // Box shape
 #define UE_VLOG_BOX(LogOwner, CategoryName, Verbosity, Box, Color, Format, ...) if(FVisualLogger::IsRecording()) FVisualLogger::GeometryShapeLogf(LogOwner, CategoryName, ELogVerbosity::Verbosity, Box, FMatrix::Identity, Color, Format, ##__VA_ARGS__)
-#define UE_CVLOG_BOX(Condition, LogOwner, CategoryName, Verbosity, Box, Color, Format, ...)  if(FVisualLogger::IsRecording() && Condition) {UE_VLOG_BOX(Actor, CategoryName, Verbosity, Box, FMatrix::Identity, Color, Format, ##__VA_ARGS__);} 
+#define UE_CVLOG_BOX(Condition, LogOwner, CategoryName, Verbosity, Box, Color, Format, ...)  if(FVisualLogger::IsRecording() && Condition) {UE_VLOG_BOX(LogOwner, CategoryName, Verbosity, Box, FMatrix::Identity, Color, Format, ##__VA_ARGS__);} 
 // Oriented box shape
 #define UE_VLOG_OBOX(LogOwner, CategoryName, Verbosity, Box, Matrix, Color, Format, ...) if(FVisualLogger::IsRecording()) FVisualLogger::GeometryShapeLogf(LogOwner, CategoryName, ELogVerbosity::Verbosity, Box, Matrix, Color, Format, ##__VA_ARGS__)
-#define UE_CVLOG_OBOX(Condition, LogOwner, CategoryName, Verbosity, Box, Matrix, Color, Format, ...)  if(FVisualLogger::IsRecording() && Condition) {UE_VLOG_BOX(Actor, CategoryName, Verbosity, Box, Matrix, Color, Format, ##__VA_ARGS__);} 
+#define UE_CVLOG_OBOX(Condition, LogOwner, CategoryName, Verbosity, Box, Matrix, Color, Format, ...)  if(FVisualLogger::IsRecording() && Condition) {UE_VLOG_OBOX(LogOwner, CategoryName, Verbosity, Box, Matrix, Color, Format, ##__VA_ARGS__);} 
 // Cone shape
 #define UE_VLOG_CONE(LogOwner, CategoryName, Verbosity, Orgin, Direction, Length, Angle, Color, Format, ...) if(FVisualLogger::IsRecording()) FVisualLogger::GeometryShapeLogf(LogOwner, CategoryName, ELogVerbosity::Verbosity, Orgin, Direction, Length, Angle, Color, Format, ##__VA_ARGS__)
-#define UE_CVLOG_CONE(Condition, LogOwner, CategoryName, Verbosity, Orgin, Direction, Length, Angle, Color, Format, ...)  if(FVisualLogger::IsRecording() && Condition) {UE_VLOG_CONE(Actor, CategoryName, Verbosity, Orgin, Direction, Length, Angle, Color, Format, ##__VA_ARGS__);} 
+#define UE_CVLOG_CONE(Condition, LogOwner, CategoryName, Verbosity, Orgin, Direction, Length, Angle, Color, Format, ...)  if(FVisualLogger::IsRecording() && Condition) {UE_VLOG_CONE(LogOwner, CategoryName, Verbosity, Orgin, Direction, Length, Angle, Color, Format, ##__VA_ARGS__);} 
 // Cylinder shape
 #define UE_VLOG_CYLINDER(LogOwner, CategoryName, Verbosity, Start, End, Radius, Color, Format, ...) if(FVisualLogger::IsRecording()) FVisualLogger::GeometryShapeLogf(LogOwner, CategoryName, ELogVerbosity::Verbosity, Start, End, Radius, Color, Format, ##__VA_ARGS__)
 #define UE_CVLOG_CYLINDER(Condition, LogOwner, CategoryName, Verbosity, Start, End, Radius, Color, Format, ...)  if(FVisualLogger::IsRecording() && Condition) {UE_VLOG_CYLINDER(LogOwner, CategoryName, Verbosity, Start, End, Radius, Color, Format, ##__VA_ARGS__);} 
diff --git a/Engine/Source/Runtime/NavigationSystem/Private/NavigationSystem.cpp b/Engine/Source/Runtime/NavigationSystem/Private/NavigationSystem.cpp
--- a/Engine/Source/Runtime/NavigationSystem/Private/NavigationSystem.cpp
+++ b/Engine/Source/Runtime/NavigationSystem/Private/NavigationSystem.cpp
@@ -3166,6 +3166,11 @@ void UNavigationSystemV1::Build()
 		}
 	}
 
+#if !UE_BUILD_SHIPPING
+	// no longer report that navmesh needs to be rebuild
+	bDirtyAreasReportedWhileAccumulationLocked = false;
+#endif // !UE_BUILD_SHIPPING
+
 	UE_LOG(LogNavigation, Display, TEXT("UNavigationSystemV1::Build total execution time: %.5f"), float(FPlatformTime::Seconds() - BuildStartTime));
 }
 
diff --git a/Engine/Source/Runtime/Slate/Private/Framework/Application/SlateApplication.cpp b/Engine/Source/Runtime/Slate/Private/Framework/Application/SlateApplication.cpp
--- a/Engine/Source/Runtime/Slate/Private/Framework/Application/SlateApplication.cpp
+++ b/Engine/Source/Runtime/Slate/Private/Framework/Application/SlateApplication.cpp
@@ -2467,6 +2467,9 @@ void FSlateApplication::UnregisterGameViewport()
 {
 	ResetToDefaultPointerInputSettings();
 
+	bIsFakingTouched = false;
+	bIsGameFakingTouch = false;
+
 	if (GameViewportWidget.IsValid())
 	{
 		GameViewportWidget.Pin()->SetActive(false);
@@ -5165,13 +5168,13 @@ void FSlateApplication::SetGameIsFakingTouchEvents(const bool bIsFaking, FVector
 {
 	if ( bIsGameFakingTouch != bIsFaking )
 	{
-	if (bIsFakingTouched && !bIsFaking && bIsGameFakingTouch && !bIsFakingTouch)
-	{
-		OnTouchEnded((CursorLocation ? *CursorLocation : PlatformApplication->Cursor->GetPosition()), 0, 0);
-	}
+		if (bIsFakingTouched && !bIsFaking && bIsGameFakingTouch && !bIsFakingTouch)
+		{
+			OnTouchEnded((CursorLocation ? *CursorLocation : PlatformApplication->Cursor->GetPosition()), 0, 0);
+		}
 
-	bIsGameFakingTouch = bIsFaking;
-}
+		bIsGameFakingTouch = bIsFaking;
+	}
 }
 
 #endif
