Fixed UCanvas safe zone offset. Change by [REDACTED_PERSON].

diff --git a/Engine/Source/Runtime/Engine/Private/UserInterface/Canvas.cpp b/Engine/Source/Runtime/Engine/Private/UserInterface/Canvas.cpp
--- a/Engine/Source/Runtime/Engine/Private/UserInterface/Canvas.cpp
+++ b/Engine/Source/Runtime/Engine/Private/UserInterface/Canvas.cpp
@@ -1209,6 +1209,8 @@ UCanvas::UCanvas(const FObjectInitializer& ObjectInitializer)
 	UnsafeSizeY = 0;
 	SafeZonePadX = 0;
 	SafeZonePadY = 0;
+	SafeZonePadEX = 0;
+	SafeZonePadEY = 0;
 	CachedDisplayWidth = 0;
 	CachedDisplayHeight = 0;
 
@@ -1347,6 +1349,8 @@ void UCanvas::UpdateSafeZoneData()
 
 		SafeZonePadX = (CachedDisplayWidth - (CachedDisplayWidth * SafeRegionPercentage.X))/2.f;
 		SafeZonePadY = (CachedDisplayHeight - (CachedDisplayHeight * SafeRegionPercentage.Y))/2.f;
+		SafeZonePadX = SafeZonePadEX;
+		SafeZonePadY = SafeZonePadEY;
 	}
 	else if(FSlateApplication::IsInitialized())
 	{
@@ -1356,6 +1360,8 @@ void UCanvas::UpdateSafeZoneData()
 
 		SafeZonePadX = FMath::CeilToInt(DisplayMetrics.TitleSafePaddingSize.X);
 		SafeZonePadY = FMath::CeilToInt(DisplayMetrics.TitleSafePaddingSize.Y);
+		SafeZonePadEX = FMath::CeilToInt(DisplayMetrics.TitleSafePaddingSize.Z);
+		SafeZonePadEY = FMath::CeilToInt(DisplayMetrics.TitleSafePaddingSize.W);
 
 		CachedDisplayWidth = DisplayMetrics.PrimaryDisplayWidth;
 		CachedDisplayHeight = DisplayMetrics.PrimaryDisplayHeight;
