Branch snapshot for CL 3327560

diff --git a/Engine/Source/Runtime/Core/Public/Containers/SparseArray.h b/Engine/Source/Runtime/Core/Public/Containers/SparseArray.h
--- a/Engine/Source/Runtime/Core/Public/Containers/SparseArray.h
+++ b/Engine/Source/Runtime/Core/Public/Containers/SparseArray.h
@@ -884,6 +884,22 @@ struct TContainerTraits<TSparseArray<ElementType, Allocator> > : public TContain
 		TContainerTraits<typename TSparseArray<ElementType, Allocator>::AllocationBitArrayType>::MoveWillEmptyContainer };
 };
 
+//
+// TSparseArray operator news.
+//
+template <typename T, typename Allocator> void* operator new( size_t Size, TSparseArray<T, Allocator>& Array )
+{
+	check( Size == sizeof( T ) );
+	const int32 Index = Array.AddUninitialized().Index;
+	return &Array[ Index ];
+}
+template <typename T, typename Allocator> void* operator new( size_t Size, TSparseArray<T, Allocator>& Array, int32 Index )
+{
+	check( Size == sizeof( T ) );
+	Array.InsertUninitialized( Index );
+	return &Array[ Index ];
+}
+
 struct FScriptSparseArrayLayout
 {
 	int32 ElementOffset;
diff --git a/Engine/Source/Runtime/Core/Public/Misc/Change.h b/Engine/Source/Runtime/Core/Public/Misc/Change.h
new file mode 100644
--- /dev/null
+++ b/Engine/Source/Runtime/Core/Public/Misc/Change.h
@@ -0,0 +1,95 @@
+// Copyright 1998-2015 Epic Games, Inc. All Rights Reserved.
+
+#pragma once
+
+#include "UniquePtr.h"
+#include "Array.h"
+
+// @todo mesheditor: Comment these classes and enums!
+
+class CORE_API FChange
+{
+
+public:
+
+	/** Makes the change to the object, returning a new change that can be used to perfectly roll back this change */
+	virtual TUniquePtr<FChange> Execute( UObject* Object ) = 0;
+
+	/** Describes this change (for debugging) */
+	virtual FString ToString() const = 0;
+
+	/** Prints this change to the log, including sub-changes if there are any.  For compound changes, there might be multiple lines.  You should not need to override this function. */
+	virtual void PrintToLog( class FFeedbackContext& FeedbackContext, const int32 IndentLevel = 0 );
+
+	/** Virtual destructor */
+	virtual ~FChange()
+	{
+	}
+
+protected:
+
+	/** Protected default constructor */
+	FChange()
+	{
+	}
+
+private:
+	// Non-copyable
+	FChange( const FChange& ) = delete;
+	FChange& operator=( const FChange& ) = delete;
+
+};
+
+
+
+struct CORE_API FCompoundChangeInput
+{
+	FCompoundChangeInput()
+	{
+	}
+
+	explicit FCompoundChangeInput( FCompoundChangeInput&& RHS )
+		: Subchanges( MoveTemp( RHS.Subchanges ) )
+	{
+	}
+
+	/** Ordered list of changes that comprise everything needed to describe this change */
+	TArray<TUniquePtr<FChange>> Subchanges;
+
+private:
+	// Non-copyable
+	FCompoundChangeInput( const FCompoundChangeInput& ) = delete;
+	FCompoundChangeInput& operator=( const FCompoundChangeInput& ) = delete;
+};
+
+
+class CORE_API FCompoundChange : public FChange
+{
+
+public:
+
+	/** Constructor */
+	explicit FCompoundChange( FCompoundChangeInput&& InitInput )
+		: Input( MoveTemp( InitInput ) )
+	{
+	}
+
+	// Parent class overrides
+	virtual TUniquePtr<FChange> Execute( UObject* Object ) override;
+	virtual FString ToString() const override;
+	virtual void PrintToLog( class FFeedbackContext& FeedbackContext, const int32 IndentLevel = 0 ) override;
+
+
+private:
+	
+	/** The data we need to make this change */
+	FCompoundChangeInput Input;
+
+private:
+	// Non-copyable
+	FCompoundChange( const FCompoundChange& ) = delete;
+	FCompoundChange& operator=( const FCompoundChange& ) = delete;
+
+};
+
+
diff --git a/Engine/Source/Runtime/Core/Public/Misc/ITransaction.h b/Engine/Source/Runtime/Core/Public/Misc/ITransaction.h
--- a/Engine/Source/Runtime/Core/Public/Misc/ITransaction.h
+++ b/Engine/Source/Runtime/Core/Public/Misc/ITransaction.h
@@ -4,6 +4,7 @@
 
 #include "CoreTypes.h"
 #include "UObject/UObjectHierarchyFwd.h"
+#include "Change.h"
 
 // Class for handling undo/redo transactions among objects.
 typedef void(*STRUCT_DC)( void* TPtr );						// default construct
@@ -59,10 +60,20 @@ public:
 	 * Saves an UObject to the transaction.
 	 *
 	 * @param Object The object to save.
+	 *
 	 * @see SaveArray
 	 */
 	virtual void SaveObject( UObject* Object ) = 0;
 
+	/**
+	 * Stores a command that can be used to undo a change to the specified object.  This may be called multiple times in the
+	 * same transaction to stack up changes that will be rolled back in reverse order.  No copy of the object itself is stored.
+	 *
+	 * @param Object The object the undo change will apply to
+	 * @param CustomChange The change that can be used to undo the changes to this object.
+	 */
+	virtual void StoreUndo( UObject* Object, TUniquePtr<FChange> CustomChange ) = 0;
+
 	/**
 	 * Sets the transaction's primary object.
 	 *
diff --git a/Engine/Source/Runtime/Core/Public/ProfilingDebugging/ABTesting.h b/Engine/Source/Runtime/Core/Public/ProfilingDebugging/ABTesting.h
--- a/Engine/Source/Runtime/Core/Public/ProfilingDebugging/ABTesting.h
+++ b/Engine/Source/Runtime/Core/Public/ProfilingDebugging/ABTesting.h
@@ -122,7 +122,7 @@ public:
 	~FScopedABTimer()
 	{
 		Stop();
-		FABTest::Get().ReportScopeTime(GetAccumulatedTime());
+		FABTest::Get().ReportScopeTime(Accumulator);
 	}
 private:
 	double TimerData;
diff --git a/Engine/Source/Runtime/Core/Public/ProfilingDebugging/ScopedTimers.h b/Engine/Source/Runtime/Core/Public/ProfilingDebugging/ScopedTimers.h
--- a/Engine/Source/Runtime/Core/Public/ProfilingDebugging/ScopedTimers.h
+++ b/Engine/Source/Runtime/Core/Public/ProfilingDebugging/ScopedTimers.h
@@ -37,12 +37,7 @@ public:
 		return StopTime;
 	}
 
-	double GetAccumulatedTime() const
-	{
-		return Accumulator;
-	}
-
-private:
+protected:
 	/** Start time, captured in ctor. */
 	double StartTime;
 	/** Time variable to update. */
@@ -53,7 +48,7 @@ private:
  * Utility class for tracking the duration of a scoped action (the user 
  * doesn't have to call Start() and Stop() manually).
  */
-class FScopedDurationTimer : private FDurationTimer
+class FScopedDurationTimer : public FDurationTimer
 {
 public:
 	explicit FScopedDurationTimer(double& AccumulatorIn)
@@ -68,6 +63,28 @@ public:
 	}
 };
 
+/**
+ * Same as FScopedDurationTimer, except that it tracks the time value internally so you don't have to
+ * pass in a double to accumulate.  Call GetTime() to get the total time since starting.
+ */
+class FAutoScopedDurationTimer : public FScopedDurationTimer
+{
+public:
+	FAutoScopedDurationTimer()
+		: FScopedDurationTimer(AccumulatorValue)
+	{
+	}
+
+	double GetTime()
+	{
+		Stop();
+		return AccumulatorValue;
+	}
+
+private:
+	double AccumulatorValue;
+};
+
 /**
  * Utility class for tracking the duration of a scoped action to an accumulator in a thread-safe fashion.
  * Can accumulate into a 32bit or 64bit counter.
diff --git a/Engine/Source/Runtime/Engine/Classes/Components/StaticMeshComponent.h b/Engine/Source/Runtime/Engine/Classes/Components/StaticMeshComponent.h
--- a/Engine/Source/Runtime/Engine/Classes/Components/StaticMeshComponent.h
+++ b/Engine/Source/Runtime/Engine/Classes/Components/StaticMeshComponent.h
@@ -537,6 +537,9 @@ public:
 	/** Whether or not the component supports default collision from its static mesh asset */
 	virtual bool SupportsDefaultCollision();
 
+	/** Whether we can support dithered LOD transitions (default behavior checks all materials). Used for HISMC LOD. */
+	virtual bool SupportsDitheredLODTransitions();
+
 private:
 	/** Initializes the resources used by the static mesh component. */
 	void InitResources();
diff --git a/Engine/Source/Runtime/Engine/Classes/Engine/StaticMesh.h b/Engine/Source/Runtime/Engine/Classes/Engine/StaticMesh.h
--- a/Engine/Source/Runtime/Engine/Classes/Engine/StaticMesh.h
+++ b/Engine/Source/Runtime/Engine/Classes/Engine/StaticMesh.h
@@ -238,6 +238,7 @@ struct FMeshSectionInfoMap
 	GENERATED_USTRUCT_BODY()
 
 	/** Maps an LOD+Section to the material it should render with. */
+	UPROPERTY()
 	TMap<uint32,FMeshSectionInfo> Map;
 
 	/** Serialize. */
@@ -426,7 +427,7 @@ class UStaticMesh : public UObject, public IInterface_CollisionDataProvider, pub
 	FMeshSectionInfoMap SectionInfoMap;
 
 	/** The LOD group to which this mesh belongs. */
-	UPROPERTY(AssetRegistrySearchable)
+	UPROPERTY(EditAnywhere, AssetRegistrySearchable, Category=StaticMesh)
 	FName LODGroup;
 
 	/** If true, the screen sizees at which LODs swap are computed automatically. */
@@ -590,6 +591,11 @@ protected:
 	TArray<UAssetUserData*> AssetUserData;
 
 public:
+	/** The editable mesh representation of this static mesh */
+	// @todo: Maybe we don't want this visible in the details panel in the end; for now, this might aid debugging.
+	UPROPERTY(Instanced, VisibleAnywhere, Category = EditableMesh)
+	class UEditableMesh* EditableMesh;
+
 	/** Pre-build navigation collision */
 	UPROPERTY(VisibleAnywhere, transient, duplicatetransient, Instanced, Category = Navigation)
 	class UNavCollision* NavCollision;
@@ -604,7 +610,7 @@ public:
 	ENGINE_API virtual void PreEditChange(UProperty* PropertyAboutToChange) override;
 	ENGINE_API virtual void PostEditChangeProperty(FPropertyChangedEvent& PropertyChangedEvent) override;
 	ENGINE_API virtual void GetAssetRegistryTagMetadata(TMap<FName, FAssetRegistryTagMetadata>& OutMetadata) const override;
-	ENGINE_API void SetLODGroup(FName NewGroup);
+	ENGINE_API void SetLODGroup(FName NewGroup, bool bRebuildImmediately = true);
 	ENGINE_API void BroadcastNavCollisionChange();
 #endif // WITH_EDITOR
 	ENGINE_API virtual void Serialize(FArchive& Ar) override;
@@ -845,7 +851,6 @@ private:
 	 */
 	void CacheDerivedData();
 
-
 	FOnPreMeshBuild PreMeshBuild;
 	FOnPostMeshBuild PostMeshBuild;
 
diff --git a/Engine/Source/Runtime/Engine/Classes/Materials/Material.h b/Engine/Source/Runtime/Engine/Classes/Materials/Material.h
--- a/Engine/Source/Runtime/Engine/Classes/Materials/Material.h
+++ b/Engine/Source/Runtime/Engine/Classes/Materials/Material.h
@@ -447,9 +447,9 @@ public:
 	UPROPERTY(EditAnywhere, Category=Material)
 	uint32 TwoSided:1;
 
-	/** Whether the material should support a dithered LOD transition. */
-	UPROPERTY(EditAnywhere, Category=Material, AdvancedDisplay)
-	uint32 DitheredLODTransition:1;
+	/** Whether meshes rendered with the material should support dithered LOD transitions. */
+	UPROPERTY(EditAnywhere, Category = Material, AdvancedDisplay, meta = (DisplayName = "Dithered LOD Transition"))
+	uint32 DitheredLODTransition : 1;
 
 	/** Dither opacity mask. When combined with Temporal AA this can be used as a form of limited translucency which supports all lighting features. */
 	UPROPERTY(EditAnywhere, Category=Material, AdvancedDisplay)
diff --git a/Engine/Source/Runtime/Engine/Engine.Build.cs b/Engine/Source/Runtime/Engine/Engine.Build.cs
--- a/Engine/Source/Runtime/Engine/Engine.Build.cs
+++ b/Engine/Source/Runtime/Engine/Engine.Build.cs
@@ -103,6 +103,7 @@ public class Engine : ModuleRules
                 "CinematicCamera",
 				"Analytics",
 				"AnalyticsET",
+                "MeshEditingRuntime"
             }
 		);
 
@@ -157,6 +158,8 @@ public class Engine : ModuleRules
         CircularlyReferencedDependentModules.Add("UMG");
         CircularlyReferencedDependentModules.Add("MaterialShaderQualitySettings");
         CircularlyReferencedDependentModules.Add("CinematicCamera");
+        // @todo mesheditor: See if we can find a way to break this cycle later on
+        CircularlyReferencedDependentModules.Add("MeshEditingRuntime");
 
         // The AnimGraphRuntime module is not needed by Engine proper, but it is loaded in LaunchEngineLoop.cpp,
         // and needs to be listed in an always-included module in order to be compiled into standalone games
diff --git a/Engine/Source/Runtime/Engine/Private/Components/StaticMeshComponent.cpp b/Engine/Source/Runtime/Engine/Private/Components/StaticMeshComponent.cpp
--- a/Engine/Source/Runtime/Engine/Private/Components/StaticMeshComponent.cpp
+++ b/Engine/Source/Runtime/Engine/Private/Components/StaticMeshComponent.cpp
@@ -1580,6 +1580,20 @@ bool UStaticMeshComponent::SupportsDefaultCollision()
 	return GetStaticMesh() && GetBodySetup() == GetStaticMesh()->BodySetup;
 }
 
+bool UStaticMeshComponent::SupportsDitheredLODTransitions()
+{
+	// Only support dithered transitions if all materials do.
+	TArray<class UMaterialInterface*> Materials = GetMaterials();
+	for (UMaterialInterface* Material : Materials)
+	{
+		if (Material && !Material->IsDitheredLODTransition())
+		{
+			return false;
+		}
+	}
+	return true;
+}
+
 void UStaticMeshComponent::UpdateCollisionFromStaticMesh()
 {
 	if(bUseDefaultCollision && SupportsDefaultCollision())
diff --git a/Engine/Source/Runtime/Engine/Private/ConvexVolume.cpp b/Engine/Source/Runtime/Engine/Private/ConvexVolume.cpp
--- a/Engine/Source/Runtime/Engine/Private/ConvexVolume.cpp
+++ b/Engine/Source/Runtime/Engine/Private/ConvexVolume.cpp
@@ -413,6 +413,51 @@ bool FConvexVolume::IntersectSphere(const FVector& Origin,const float& Radius, b
 	return Result;
 }
 
+bool FConvexVolume::IntersectLineSegment(const FVector& InStart, const FVector& InEnd) const
+{
+	// @todo: not optimized
+	// Not sure if there's a better algorithm for this; in any case, there's scope for vectorizing some of this stuff
+	// using the permuted planes array.
+
+	// Take copies of the line segment start/end points so they can be modified
+	FVector Start(InStart);
+	FVector End(InEnd);
+
+	// Iterate through all planes, successively clipping the line segment against each one,
+	// until it is either completely contained within the convex volume (intersects), or
+	// it is completely outside (doesn't intersect)
+	for (const FPlane& Plane : Planes)
+	{
+		const float DistanceFromStart = Plane.PlaneDot(Start);
+		const float DistanceFromEnd = Plane.PlaneDot(End);
+
+		if (DistanceFromStart > 0.0f && DistanceFromEnd > 0.0f)
+		{
+			// Both points are outside one of the frustum planes, so cannot intersect
+			return false;
+		}
+
+		if (DistanceFromStart < 0.0f && DistanceFromEnd < 0.0f)
+		{
+			// Both points are inside this frustum plane, no need to clip it against the plane
+			continue;
+		}
+
+		// Clip the line segment against the plane
+		const FVector IntersectionPoint = FMath::LinePlaneIntersection(Start, End, Plane);
+		if (DistanceFromStart > 0.0f)
+		{
+			Start = IntersectionPoint;
+		}
+		else
+		{
+			End = IntersectionPoint;
+		}
+	}
+
+	return true;
+}
+
 void GetViewFrustumBounds(FConvexVolume& OutResult, const FMatrix& ViewProjectionMatrix, bool UseNearPlane)
 {
 	GetViewFrustumBounds(OutResult, ViewProjectionMatrix, FPlane(), false, UseNearPlane);
diff --git a/Engine/Source/Runtime/Engine/Private/GeomTools.cpp b/Engine/Source/Runtime/Engine/Private/GeomTools.cpp
--- a/Engine/Source/Runtime/Engine/Private/GeomTools.cpp
+++ b/Engine/Source/Runtime/Engine/Private/GeomTools.cpp
@@ -388,21 +388,22 @@ void FGeomTools::Buid2DPolysFromEdges(TArray<FUtilPoly2D>& OutPolys, const TArra
 }
 
 /** Given three direction vectors, indicates if A and B are on the same 'side' of Vec. */
-bool FGeomTools::VectorsOnSameSide(const FVector& Vec, const FVector& A, const FVector& B)
+bool FGeomTools::VectorsOnSameSide(const FVector& Vec, const FVector& A, const FVector& B, const float SameSideDotProductEpsilon)
 {
 	const FVector CrossA = Vec ^ A;
 	const FVector CrossB = Vec ^ B;
-	return !FMath::IsNegativeFloat(CrossA | CrossB);
+	float DotWithEpsilon = SameSideDotProductEpsilon + ( CrossA | CrossB );
+	return !FMath::IsNegativeFloat(DotWithEpsilon);
 }
 
 /** Util to see if P lies within triangle created by A, B and C. */
-bool FGeomTools::PointInTriangle(const FVector& A, const FVector& B, const FVector& C, const FVector& P)
+bool FGeomTools::PointInTriangle(const FVector& A, const FVector& B, const FVector& C, const FVector& P, const float InsideTriangleDotProductEpsilon)
 {
 	// Cross product indicates which 'side' of the vector the point is on
 	// If its on the same side as the remaining vert for all edges, then its inside.	
-	if( VectorsOnSameSide(B-A, P-A, C-A) &&
-		VectorsOnSameSide(C-B, P-B, A-B) &&
-		VectorsOnSameSide(A-C, P-C, B-C) )
+	if( VectorsOnSameSide(B-A, P-A, C-A, InsideTriangleDotProductEpsilon) &&
+		VectorsOnSameSide(C-B, P-B, A-B, InsideTriangleDotProductEpsilon) &&
+		VectorsOnSameSide(A-C, P-C, B-C, InsideTriangleDotProductEpsilon) )
 	{
 		return true;
 	}
diff --git a/Engine/Source/Runtime/Engine/Private/HierarchicalInstancedStaticMesh.cpp b/Engine/Source/Runtime/Engine/Private/HierarchicalInstancedStaticMesh.cpp
--- a/Engine/Source/Runtime/Engine/Private/HierarchicalInstancedStaticMesh.cpp
+++ b/Engine/Source/Runtime/Engine/Private/HierarchicalInstancedStaticMesh.cpp
@@ -27,6 +27,7 @@
 #include "StaticMeshResources.h"
 #include "Components/HierarchicalInstancedStaticMeshComponent.h"
 #include "InstancedStaticMesh.h"
+#include "SceneManagement.h"
 
 static TAutoConsoleVariable<int32> CVarFoliageSplitFactor(
 	TEXT("foliage.SplitFactor"),
@@ -747,6 +748,7 @@ class FHierarchicalStaticMeshSceneProxy : public FInstancedStaticMeshSceneProxy
 	TMap<uint32, FFoliageOcclusionResults> OcclusionResults;
 	bool bIsGrass;
 	uint32 SceneProxyCreatedFrameNumberRenderThread;
+	bool bDitheredLODTransitions;
 
 #if !(UE_BUILD_SHIPPING || UE_BUILD_TEST)
 	mutable TArray<uint32> SingleDebugRuns[MAX_STATIC_MESH_LODS];
@@ -767,6 +769,7 @@ public:
 		, LastUnbuiltIndex(InComponent->GetNumRenderInstances()-1)
 		, bIsGrass(bInIsGrass)
 		, SceneProxyCreatedFrameNumberRenderThread(UINT32_MAX)
+		, bDitheredLODTransitions(InComponent->SupportsDitheredLODTransitions())
 #if !(UE_BUILD_SHIPPING || UE_BUILD_TEST)
 		, CaptureTag(0)
 #endif
@@ -783,6 +786,8 @@ public:
 		, FirstUnbuiltIndex(InComponent->NumBuiltRenderInstances)
 		, LastUnbuiltIndex(InComponent->GetNumRenderInstances()-1)
 		, bIsGrass(bInIsGrass)
+		, SceneProxyCreatedFrameNumberRenderThread(UINT32_MAX)
+		, bDitheredLODTransitions(InComponent->SupportsDitheredLODTransitions())
 #if !(UE_BUILD_SHIPPING || UE_BUILD_TEST)
 		, CaptureTag(0)
 #endif
@@ -1315,7 +1320,7 @@ void FHierarchicalStaticMeshSceneProxy::GetDynamicMeshElements(const TArray<cons
 
 	QUICK_SCOPE_CYCLE_COUNTER(STAT_HierarchicalInstancedStaticMeshSceneProxy_GetMeshElements);
 
-	bool bMultipleSections = ALLOW_DITHERED_LOD_FOR_INSTANCED_STATIC_MESHES && CVarDitheredLOD.GetValueOnRenderThread() > 0;
+	bool bMultipleSections = ALLOW_DITHERED_LOD_FOR_INSTANCED_STATIC_MESHES && bDitheredLODTransitions && CVarDitheredLOD.GetValueOnRenderThread() > 0;
 	bool bSingleSections = !bMultipleSections;
 	bool bOverestimate = CVarOverestimateLOD.GetValueOnRenderThread() > 0;
 
@@ -1327,6 +1332,9 @@ void FHierarchicalStaticMeshSceneProxy::GetDynamicMeshElements(const TArray<cons
 		{
 			const FSceneView* View = Views[ViewIndex];
 
+			// With stereo rendering views may have asymmetric FOVs, so always use the left eye's view for LOD calculation
+			const FSceneView* ViewForLOD = ( View->Family != nullptr && View->StereoPass == eSSP_RIGHT_EYE ) ? View->Family->Views[ 0 ] : View;
+
 			FFoliageElementParams ElementParams;
 			ElementParams.bSelectionRenderEnabled = GIsEditor && ViewFamily.EngineShowFlags.Selection;
 			ElementParams.NumSelectionGroups = (ElementParams.bSelectionRenderEnabled && bHasSelectedInstances) ? 2 : 1;
@@ -1395,7 +1403,7 @@ void FHierarchicalStaticMeshSceneProxy::GetDynamicMeshElements(const TArray<cons
 					}
 					else
 					{
-						const FMatrix LocalViewProjForCulling = GetLocalToWorld() * View->ViewMatrices.GetViewProjectionMatrix();
+						const FMatrix LocalViewProjForCulling = GetLocalToWorld() * ViewForLOD->ViewMatrices.GetViewProjectionMatrix();
 						GetViewFrustumBounds(InstanceParams.ViewFrustumLocal, LocalViewProjForCulling, false);
 					}
 
@@ -1448,39 +1456,25 @@ void FHierarchicalStaticMeshSceneProxy::GetDynamicMeshElements(const TArray<cons
 				float MaxDrawDistanceScale = GetCachedScalabilityCVars().ViewDistanceScale;
 				float SphereRadius = RenderData->Bounds.SphereRadius;
 
-				for (int32 LODIndex = 0; LODIndex < InstanceParams.LODs; LODIndex++)
+				float FinalCull = MAX_flt;
+				if (MinSize > 0.0)
 				{
-					InstanceParams.LODPlanesMax[LODIndex] = MIN_flt;
-					InstanceParams.LODPlanesMin[LODIndex] = MAX_flt;
+					FinalCull = ComputeBoundsDrawDistance(MinSize, SphereRadius, ViewForLOD->ViewMatrices.GetProjectionMatrix()) * LODScale;
 				}
-				ElementParams.FinalCullDistance = MIN_flt;
-				for (int32 SampleIndex = 0; SampleIndex < 2; SampleIndex++)
+				if (UserData_AllInstances.EndCullDistance > 0.0f)
 				{
-					// we aren't going to accurately figure out LOD planes for both samples and try all 4 possibilities, rather we just consider the worst case LOD range; the shaders will render correctly
-					float Fac = View->GetTemporalLODDistanceFactor(SampleIndex, bMultipleSections) * SphereRadius * LODScale;
-
-					float FinalCull = MAX_flt;
-
-					if (MinSize > 0.0)
-					{
-						FinalCull = FMath::Sqrt(Fac / MinSize);
-					}
-					if (UserData_AllInstances.EndCullDistance > 0.0f && UserData_AllInstances.EndCullDistance < FinalCull)
-					{
-						FinalCull = UserData_AllInstances.EndCullDistance;
-					}
-					FinalCull *= MaxDrawDistanceScale;
-					ElementParams.FinalCullDistance = FMath::Max(ElementParams.FinalCullDistance, FinalCull);
+					FinalCull = FMath::Min(FinalCull, UserData_AllInstances.EndCullDistance * MaxDrawDistanceScale);
+				}
+				ElementParams.FinalCullDistance = FinalCull;
 
-					for (int32 LODIndex = 1; LODIndex < InstanceParams.LODs; LODIndex++)
-					{
-						float Val = FMath::Min(FMath::Sqrt(Fac / RenderData->ScreenSize[LODIndex]), FinalCull);
-						InstanceParams.LODPlanesMin[LODIndex - 1] = FMath::Min(InstanceParams.LODPlanesMin[LODIndex - 1], Val - LODRandom);
-						InstanceParams.LODPlanesMax[LODIndex - 1] = FMath::Max(InstanceParams.LODPlanesMax[LODIndex - 1], Val);
-					}
-					InstanceParams.LODPlanesMin[InstanceParams.LODs - 1] = FMath::Min(InstanceParams.LODPlanesMin[InstanceParams.LODs - 1], FinalCull - LODRandom);
-					InstanceParams.LODPlanesMax[InstanceParams.LODs - 1] = FMath::Max(InstanceParams.LODPlanesMax[InstanceParams.LODs - 1], FinalCull);
+				for (int32 LODIndex = 1; LODIndex < InstanceParams.LODs; LODIndex++)
+				{
+					float Distance = ComputeBoundsDrawDistance(RenderData->ScreenSize[LODIndex], SphereRadius, ViewForLOD->ViewMatrices.GetProjectionMatrix()) * LODScale;
+					InstanceParams.LODPlanesMin[LODIndex - 1] = Distance - LODRandom;
+					InstanceParams.LODPlanesMax[LODIndex - 1] = Distance;
 				}
+				InstanceParams.LODPlanesMin[InstanceParams.LODs - 1] = FinalCull - LODRandom;
+				InstanceParams.LODPlanesMax[InstanceParams.LODs - 1] = FinalCull;
 
 				for (int32 LODIndex = 0; LODIndex < InstanceParams.LODs; LODIndex++)
 				{
@@ -1628,38 +1622,26 @@ void FHierarchicalStaticMeshSceneProxy::GetDynamicMeshElements(const TArray<cons
 						const float SphereRadius = RenderData->Bounds.SphereRadius;
 
 						checkSlow(NumLODs > 0);
-						for (int32 LODIndex = 0; LODIndex < NumLODs; LODIndex++)
+						
+						float FinalCull = MAX_flt;
+						if (MinSize > 0.0)
 						{
-							LODPlanesMax[LODIndex] = MIN_flt;
-							LODPlanesMin[LODIndex] = MAX_flt;
+							FinalCull = ComputeBoundsDrawDistance(MinSize, SphereRadius, ViewForLOD->ViewMatrices.GetProjectionMatrix()) * LODScale;
 						}
-						ElementParams.FinalCullDistance = MIN_flt;
-						for (int32 SampleIndex = 0; SampleIndex < 2; SampleIndex++)
+						if (UserData_AllInstances.EndCullDistance > 0.0f)
 						{
-							// we aren't going to accurately figure out LOD planes for both samples and try all 4 possibilities, rather we just consider the worst case LOD range; the shaders will render correctly
-							float Fac = View->GetTemporalLODDistanceFactor(SampleIndex, bMultipleSections) * SphereRadius * LODScale;
-							float FinalCull = MAX_flt;
-
-							if (MinSize > 0.0)
-							{
-								FinalCull = FMath::Sqrt(Fac / MinSize);
-							}
-							if (UserData_AllInstances.EndCullDistance > 0.0f && UserData_AllInstances.EndCullDistance < FinalCull)
-							{
-								FinalCull = UserData_AllInstances.EndCullDistance;
-							}
-							FinalCull *= MaxDrawDistanceScale;
-							ElementParams.FinalCullDistance = FMath::Max(ElementParams.FinalCullDistance, FinalCull);
+							FinalCull = FMath::Min(FinalCull, UserData_AllInstances.EndCullDistance * MaxDrawDistanceScale);
+						}
+						ElementParams.FinalCullDistance = FinalCull;
 
-							for (int32 LODIndex = 1; LODIndex < NumLODs; LODIndex++)
-							{
-								float Val = FMath::Min(FMath::Sqrt(Fac / RenderData->ScreenSize[LODIndex]), FinalCull);
-								LODPlanesMin[LODIndex - 1] = FMath::Min(LODPlanesMin[LODIndex - 1], Val - LODRandom);
-								LODPlanesMax[LODIndex - 1] = FMath::Max(LODPlanesMax[LODIndex - 1], Val);
-							}
-							LODPlanesMin[NumLODs - 1] = FMath::Min(LODPlanesMin[NumLODs - 1], FinalCull - LODRandom);
-							LODPlanesMax[NumLODs - 1] = FMath::Max(LODPlanesMax[NumLODs - 1], FinalCull);
+						for (int32 LODIndex = 1; LODIndex < NumLODs; LODIndex++)
+						{
+							float Distance = ComputeBoundsDrawDistance(RenderData->ScreenSize[LODIndex], SphereRadius, ViewForLOD->ViewMatrices.GetProjectionMatrix()) * LODScale;
+							LODPlanesMin[LODIndex - 1] = Distance - LODRandom;
+							LODPlanesMax[LODIndex - 1] = Distance;
 						}
+						LODPlanesMin[NumLODs - 1] = FinalCull - LODRandom;
+						LODPlanesMax[NumLODs - 1] = FinalCull;
 
 						// calculate runs
 						int32 MinLOD = 0;
diff --git a/Engine/Source/Runtime/Engine/Private/InstancedStaticMesh.cpp b/Engine/Source/Runtime/Engine/Private/InstancedStaticMesh.cpp
--- a/Engine/Source/Runtime/Engine/Private/InstancedStaticMesh.cpp
+++ b/Engine/Source/Runtime/Engine/Private/InstancedStaticMesh.cpp
@@ -22,6 +22,7 @@
 #include "PhysicsEngine/BodySetup.h"
 #include "GameFramework/WorldSettings.h"
 #include "ComponentRecreateRenderStateContext.h"
+#include "SceneManagement.h"
 
 const int32 InstancedStaticMeshMaxTexCoord = 8;
 
@@ -2013,6 +2014,9 @@ void FInstancedStaticMeshVertexFactoryShaderParameters::SetMesh( FRHICommandList
 		InstancingWorldViewOriginOne.W = 1.0f;
 		if (InstancingUserData && BatchElement.InstancedLODRange)
 		{
+			// With stereo rendering views may have asymmetric FOVs, so always use the left eye's view for LOD calculation
+			const FSceneView& ViewForLOD = ( View.Family != nullptr && View.StereoPass == eSSP_RIGHT_EYE ) ? *View.Family->Views[ 0 ] : View;
+
 			int32 FirstLOD = InstancingUserData->MinLOD;
 
 			int32 DebugMin = FMath::Min(CVarMinLOD.GetValueOnRenderThread(), InstancingUserData->MeshRenderData->LODResources.Num() - 1);
@@ -2048,34 +2052,31 @@ void FInstancedStaticMeshVertexFactoryShaderParameters::SetMesh( FRHICommandList
 			for (int32 SampleIndex = 0; SampleIndex < 2; SampleIndex++)
 			{
 				FVector4& InstancingViewZCompare(SampleIndex ? InstancingViewZCompareOne : InstancingViewZCompareZero);
-				float Fac = View.GetTemporalLODDistanceFactor(SampleIndex) * SphereRadius * LODScale;
 
 				float FinalCull = MAX_flt;
 				if (MinSize > 0.0)
 				{
-					FinalCull = FMath::Sqrt(Fac / MinSize);
+					FinalCull = ComputeBoundsDrawDistance(MinSize, SphereRadius, ViewForLOD.ViewMatrices.GetProjectionMatrix()) * LODScale;
 				}
-				if (InstancingUserData->EndCullDistance > 0.0f && InstancingUserData->EndCullDistance < FinalCull)
+				if (InstancingUserData->EndCullDistance > 0.0f)
 				{
-					FinalCull = InstancingUserData->EndCullDistance;
+					FinalCull = FMath::Min(FinalCull, InstancingUserData->EndCullDistance * MaxDrawDistanceScale);
 				}
-				FinalCull *= MaxDrawDistanceScale;
 
 				InstancingViewZCompare.Z = FinalCull;
 				if (BatchElement.InstancedLODIndex < InstancingUserData->MeshRenderData->LODResources.Num() - 1)
 				{
-					float NextCut = InstancingUserData->MeshRenderData->ScreenSize[BatchElement.InstancedLODIndex + 1];
-					InstancingViewZCompare.Z = FMath::Min(FMath::Sqrt(Fac / NextCut), FinalCull);
+					float NextCut = ComputeBoundsDrawDistance(InstancingUserData->MeshRenderData->ScreenSize[BatchElement.InstancedLODIndex + 1], SphereRadius, ViewForLOD.ViewMatrices.GetProjectionMatrix()) * LODScale;
+					InstancingViewZCompare.Z = FMath::Min(NextCut, FinalCull);
 				}
 
 				InstancingViewZCompare.X = MIN_flt;
 				if (BatchElement.InstancedLODIndex > FirstLOD)
 				{
-					float CurCut = FMath::Sqrt(Fac / InstancingUserData->MeshRenderData->ScreenSize[BatchElement.InstancedLODIndex]);
+					float CurCut = ComputeBoundsDrawDistance(InstancingUserData->MeshRenderData->ScreenSize[BatchElement.InstancedLODIndex], SphereRadius, ViewForLOD.ViewMatrices.GetProjectionMatrix()) * LODScale;
 					if (CurCut < FinalCull)
 					{
 						InstancingViewZCompare.Y = CurCut;
-
 					}
 					else
 					{
@@ -2095,7 +2096,7 @@ void FInstancedStaticMeshVertexFactoryShaderParameters::SetMesh( FRHICommandList
 			InstancingWorldViewOriginZero = View.GetTemporalLODOrigin(0);
 			InstancingWorldViewOriginOne = View.GetTemporalLODOrigin(1);
 
-			float Alpha = View.GetTemporalLODTransition();
+			float Alpha = ViewForLOD.GetTemporalLODTransition();
 			InstancingWorldViewOriginZero.W = 1.0f - Alpha;
 			InstancingWorldViewOriginOne.W = Alpha;
 
diff --git a/Engine/Source/Runtime/Engine/Private/PositionVertexBuffer.cpp b/Engine/Source/Runtime/Engine/Private/PositionVertexBuffer.cpp
--- a/Engine/Source/Runtime/Engine/Private/PositionVertexBuffer.cpp
+++ b/Engine/Source/Runtime/Engine/Private/PositionVertexBuffer.cpp
@@ -58,14 +58,17 @@ void FPositionVertexBuffer::Init(const TArray<FStaticMeshBuildVertex>& InVertice
 
 	// Allocate the vertex data buffer.
 	VertexData->ResizeBuffer(NumVertices);
-	Data = VertexData->GetDataPointer();
-
-	// Copy the vertices into the buffer.
-	for(int32 VertexIndex = 0;VertexIndex < InVertices.Num();VertexIndex++)
+	if( NumVertices > 0 )
 	{
-		const FStaticMeshBuildVertex& SourceVertex = InVertices[VertexIndex];
-		const uint32 DestVertexIndex = VertexIndex;
-		VertexPosition(DestVertexIndex) = SourceVertex.Position;
+		Data = VertexData->GetDataPointer();
+
+		// Copy the vertices into the buffer.
+		for(int32 VertexIndex = 0;VertexIndex < InVertices.Num();VertexIndex++)
+		{
+			const FStaticMeshBuildVertex& SourceVertex = InVertices[VertexIndex];
+			const uint32 DestVertexIndex = VertexIndex;
+			VertexPosition(DestVertexIndex) = SourceVertex.Position;
+		}
 	}
 }
 
@@ -81,9 +84,12 @@ void FPositionVertexBuffer::Init(const FPositionVertexBuffer& InVertexBuffer)
 		AllocateData();
 		check( Stride == InVertexBuffer.GetStride() );
 		VertexData->ResizeBuffer(NumVertices);
-		Data = VertexData->GetDataPointer();
-		const uint8* InData = InVertexBuffer.Data;
-		FMemory::Memcpy( Data, InData, Stride * NumVertices );
+		if( NumVertices > 0 )
+		{
+			Data = VertexData->GetDataPointer();
+			const uint8* InData = InVertexBuffer.Data;
+			FMemory::Memcpy( Data, InData, Stride * NumVertices );
+		}
 	}
 }
 
@@ -95,23 +101,45 @@ void FPositionVertexBuffer::Init(const TArray<FVector>& InPositions)
 		AllocateData();
 		check( Stride == InPositions.GetTypeSize() );
 		VertexData->ResizeBuffer(NumVertices);
-		Data = VertexData->GetDataPointer();
-		FMemory::Memcpy( Data, InPositions.GetData(), Stride * NumVertices );
+		if( NumVertices > 0 )
+		{
+			Data = VertexData->GetDataPointer();
+			FMemory::Memcpy( Data, InPositions.GetData(), Stride * NumVertices );
+		}
 	}
 }
 
-/**
-* Removes the cloned vertices used for extruding shadow volumes.
-* @param NumVertices - The real number of static mesh vertices which should remain in the buffer upon return.
-*/
-void FPositionVertexBuffer::RemoveLegacyShadowVolumeVertices(uint32 InNumVertices)
+void FPositionVertexBuffer::AppendVertices( const FStaticMeshBuildVertex* Vertices, const uint32 NumVerticesToAppend )
 {
-	check(VertexData);
-	VertexData->ResizeBuffer(InNumVertices);
-	NumVertices = InNumVertices;
+	if (VertexData == nullptr && NumVerticesToAppend > 0)
+	{
+		// Allocate the vertex data storage type if the buffer was never allocated before
+		AllocateData();
+	}
+
 
-	// Make a copy of the vertex data pointer.
-	Data = VertexData->GetDataPointer();
+	check( VertexData != nullptr );	// Must only be called after Init() has already initialized the buffer!
+	if( NumVerticesToAppend > 0 )
+	{
+		check( Vertices != nullptr );
+
+		const uint32 FirstDestVertexIndex = NumVertices;
+		NumVertices += NumVerticesToAppend;
+		VertexData->ResizeBuffer( NumVertices );
+		if( NumVertices > 0 )
+		{
+			Data = VertexData->GetDataPointer();
+
+			// Copy the vertices into the buffer.
+			for( uint32 VertexIter = 0; VertexIter < NumVerticesToAppend; ++VertexIter )
+			{
+				const FStaticMeshBuildVertex& SourceVertex = Vertices[ VertexIter ];
+
+				const uint32 DestVertexIndex = FirstDestVertexIndex + VertexIter;
+				VertexPosition( DestVertexIndex ) = SourceVertex.Position;
+			}
+		}
+	}
 }
 
 /**
@@ -135,8 +163,11 @@ void FPositionVertexBuffer::Serialize( FArchive& Ar, bool bNeedsCPUAccess )
 		// Serialize the vertex data.
 		VertexData->Serialize(Ar);
 
-		// Make a copy of the vertex data pointer.
-		Data = VertexData->GetDataPointer();
+		if( NumVertices > 0 )
+		{
+			// Make a copy of the vertex data pointer.
+			Data = VertexData->GetDataPointer();
+		}
 	}
 }
 
diff --git a/Engine/Source/Runtime/Engine/Private/RawIndexBuffer.cpp b/Engine/Source/Runtime/Engine/Private/RawIndexBuffer.cpp
--- a/Engine/Source/Runtime/Engine/Private/RawIndexBuffer.cpp
+++ b/Engine/Source/Runtime/Engine/Private/RawIndexBuffer.cpp
@@ -195,6 +195,54 @@ void FRawStaticIndexBuffer::SetIndices(const TArray<uint32>& InIndices, EIndexBu
 	}
 }
 
+void FRawStaticIndexBuffer::InsertIndices( const uint32 At, const uint32* IndicesToAppend, const uint32 NumIndicesToAppend )
+{
+	if( NumIndicesToAppend > 0 )
+	{
+		const uint32 IndexStride = b32Bit ? sizeof( uint32 ) : sizeof( uint16 );
+
+		IndexStorage.InsertUninitialized( At * IndexStride, NumIndicesToAppend * IndexStride );
+		uint8* const DestIndices = &IndexStorage[ At * IndexStride ];
+
+		if( IndicesToAppend )
+		{
+			if( b32Bit )
+			{
+				// If the indices are 32 bit we can just do a memcpy.
+				FMemory::Memcpy( DestIndices, IndicesToAppend, NumIndicesToAppend * IndexStride );
+			}
+			else
+			{
+				// Copy element by element demoting 32-bit integers to 16-bit.
+				uint16* DestIndices16Bit = (uint16*)DestIndices;
+				for( uint32 Index = 0; Index < NumIndicesToAppend; ++Index )
+				{
+					DestIndices16Bit[ Index ] = IndicesToAppend[ Index ];
+				}
+			}
+		}
+		else
+		{
+			// If no indices to insert were supplied, just clear the buffer
+			FMemory::Memset( DestIndices, 0, NumIndicesToAppend * IndexStride );
+		}
+	}
+}
+
+void FRawStaticIndexBuffer::AppendIndices( const uint32* IndicesToAppend, const uint32 NumIndicesToAppend )
+{
+	InsertIndices( b32Bit ? IndexStorage.Num() / 4 : IndexStorage.Num() / 2, IndicesToAppend, NumIndicesToAppend );
+}
+
+void FRawStaticIndexBuffer::RemoveIndicesAt( const uint32 At, const uint32 NumIndicesToRemove )
+{
+	if( NumIndicesToRemove > 0 )
+	{
+		const int32 IndexStride = b32Bit ? sizeof( uint32 ) : sizeof( uint16 );
+		IndexStorage.RemoveAt( At * IndexStride, NumIndicesToRemove * IndexStride );
+	}
+}
+
 void FRawStaticIndexBuffer::GetCopy(TArray<uint32>& OutIndices) const
 {
 	int32 NumIndices = b32Bit ? (IndexStorage.Num() / 4) : (IndexStorage.Num() / 2);
diff --git a/Engine/Source/Runtime/Engine/Private/Rendering/ColorVertexBuffer.cpp b/Engine/Source/Runtime/Engine/Private/Rendering/ColorVertexBuffer.cpp
--- a/Engine/Source/Runtime/Engine/Private/Rendering/ColorVertexBuffer.cpp
+++ b/Engine/Source/Runtime/Engine/Private/Rendering/ColorVertexBuffer.cpp
@@ -119,14 +119,17 @@ void FColorVertexBuffer::Init(const TArray<FStaticMeshBuildVertex>& InVertices)
 
 		// Allocate the vertex data buffer.
 		VertexData->ResizeBuffer(NumVertices);
-		Data = VertexData->GetDataPointer();
-
-		// Copy the vertices into the buffer.
-		for(int32 VertexIndex = 0;VertexIndex < InVertices.Num();VertexIndex++)
+		if( NumVertices > 0 )
 		{
-			const FStaticMeshBuildVertex& SourceVertex = InVertices[VertexIndex];
-			const uint32 DestVertexIndex = VertexIndex;
-			VertexColor(DestVertexIndex) = SourceVertex.Color;
+			Data = VertexData->GetDataPointer();
+
+			// Copy the vertices into the buffer.
+			for(int32 VertexIndex = 0;VertexIndex < InVertices.Num();VertexIndex++)
+			{
+				const FStaticMeshBuildVertex& SourceVertex = InVertices[VertexIndex];
+				const uint32 DestVertexIndex = VertexIndex;
+				VertexColor(DestVertexIndex) = SourceVertex.Color;
+			}
 		}
 	}
 }
@@ -143,25 +146,38 @@ void FColorVertexBuffer::Init(const FColorVertexBuffer& InVertexBuffer)
 		AllocateData();
 		check( Stride == InVertexBuffer.GetStride() );
 		VertexData->ResizeBuffer(NumVertices);
-		Data = VertexData->GetDataPointer();
-		const uint8* InData = InVertexBuffer.Data;
-		FMemory::Memcpy( Data, InData, Stride * NumVertices );
+		if( NumVertices > 0 )
+		{
+			Data = VertexData->GetDataPointer();
+			const uint8* InData = InVertexBuffer.Data;
+			FMemory::Memcpy( Data, InData, Stride * NumVertices );
+		}
 	}
 }
 
-/**
-* Removes the cloned vertices used for extruding shadow volumes.
-* @param NumVertices - The real number of static mesh vertices which should remain in the buffer upon return.
-*/
-void FColorVertexBuffer::RemoveLegacyShadowVolumeVertices(uint32 InNumVertices)
+void FColorVertexBuffer::AppendVertices( const FStaticMeshBuildVertex* Vertices, const uint32 NumVerticesToAppend )
 {
-	if( VertexData != NULL )
+	check( VertexData != nullptr );	// Must only be called after Init() has already initialized the buffer!
+	if( NumVerticesToAppend > 0 )
 	{
-		VertexData->ResizeBuffer(InNumVertices);
-		NumVertices = InNumVertices;
+		check( Vertices != nullptr );
 
-		// Make a copy of the vertex data pointer.
-		Data = VertexData->GetDataPointer();
+		const uint32 FirstDestVertexIndex = NumVertices;
+		NumVertices += NumVerticesToAppend;
+		VertexData->ResizeBuffer( NumVertices );
+		if( NumVertices > 0 )
+		{
+			Data = VertexData->GetDataPointer();
+
+			// Copy the vertices into the buffer.
+			for( uint32 VertexIter = 0; VertexIter < NumVerticesToAppend; ++VertexIter )
+			{
+				const FStaticMeshBuildVertex& SourceVertex = Vertices[ VertexIter ];
+
+				const uint32 DestVertexIndex = FirstDestVertexIndex + VertexIter;
+				VertexColor( DestVertexIndex ) = SourceVertex.Color;
+			}
+		}
 	}
 }
 
@@ -337,7 +353,10 @@ void FColorVertexBuffer::InitFromColorArray( const FColor *InColors, const uint3
 	}
 
 	// Make a copy of the vertex data pointer.
-	Data = VertexData->GetDataPointer();
+	if( NumVertices > 0 )
+	{
+		Data = VertexData->GetDataPointer();
+	}
 }
 
 uint32 FColorVertexBuffer::GetAllocatedSize() const
diff --git a/Engine/Source/Runtime/Engine/Private/SceneManagement.cpp b/Engine/Source/Runtime/Engine/Private/SceneManagement.cpp
--- a/Engine/Source/Runtime/Engine/Private/SceneManagement.cpp
+++ b/Engine/Source/Runtime/Engine/Private/SceneManagement.cpp
@@ -375,7 +375,9 @@ int8 ComputeTemporalStaticMeshLOD( const FStaticMeshRenderData* RenderData, cons
 {
 	const int32 NumLODs = MAX_STATIC_MESH_LODS;
 
-	const float ScreenRadiusSquared = ComputeTemporalLODBoundsScreenRadiusSquared(Origin, SphereRadius, View, SampleIndex) * FactorScale * FactorScale;
+	// With stereo rendering views may have asymmetric FOVs, so always use the left eye's view for LOD calculation
+	const FSceneView& ViewForLOD = ( View.Family != nullptr && View.StereoPass == eSSP_RIGHT_EYE ) ? *View.Family->Views[ 0 ] : View;
+	const float ScreenRadiusSquared = ComputeTemporalLODBoundsScreenRadiusSquared(Origin, SphereRadius, ViewForLOD, SampleIndex) * FactorScale * FactorScale;
 
 	// Walk backwards and return the first matching LOD
 	for(int32 LODIndex = NumLODs - 1 ; LODIndex >= 0 ; --LODIndex)
@@ -395,7 +397,9 @@ int8 ComputeStaticMeshLOD( const FStaticMeshRenderData* RenderData, const FVecto
 	{
 		const int32 NumLODs = MAX_STATIC_MESH_LODS;
 
-		const float ScreenRadiusSquared = ComputeBoundsScreenRadiusSquared(Origin, SphereRadius, View) * FactorScale * FactorScale;
+		// With stereo rendering views may have asymmetric FOVs, so always use the left eye's view for LOD calculation
+		const FSceneView& ViewForLOD = ( View.Family != nullptr && View.StereoPass == eSSP_RIGHT_EYE ) ? *View.Family->Views[ 0 ] : View;
+		const float ScreenRadiusSquared = ComputeBoundsScreenRadiusSquared(Origin, SphereRadius, ViewForLOD) * FactorScale * FactorScale;
 
 		// Walk backwards and return the first matching LOD
 		for (int32 LODIndex = NumLODs - 1; LODIndex >= 0; --LODIndex)
@@ -430,13 +434,15 @@ FLODMask ComputeLODForMeshes( const TIndirectArray<class FStaticMesh>& StaticMes
 	{
 		int32 NumMeshes = StaticMeshes.Num();
 
+		// With stereo rendering views may have asymmetric FOVs, so always use the left eye's view for LOD calculation
+		const FSceneView& ViewForLOD = ( View.Family != nullptr && View.StereoPass == eSSP_RIGHT_EYE ) ? *View.Family->Views[ 0 ] : View;
 		if (NumMeshes && StaticMeshes[0].bDitheredLODTransition)
 		{
 			for (int32 SampleIndex = 0; SampleIndex < 2; SampleIndex++)
 			{
 				int32 MinLODFound = INT_MAX;
 				bool bFoundLOD = false;
-				const float ScreenRadiusSquared = ComputeTemporalLODBoundsScreenRadiusSquared(Origin, SphereRadius, View, SampleIndex);
+				const float ScreenRadiusSquared = ComputeTemporalLODBoundsScreenRadiusSquared(Origin, SphereRadius, ViewForLOD, SampleIndex);
 
 				for(int32 MeshIndex = NumMeshes-1 ; MeshIndex >= 0 ; --MeshIndex)
 				{
@@ -464,7 +470,7 @@ FLODMask ComputeLODForMeshes( const TIndirectArray<class FStaticMesh>& StaticMes
 		{
 			int32 MinLODFound = INT_MAX;
 			bool bFoundLOD = false;
-			const float ScreenRadiusSquared = ComputeBoundsScreenRadiusSquared(Origin, SphereRadius, View);
+			const float ScreenRadiusSquared = ComputeBoundsScreenRadiusSquared(Origin, SphereRadius, ViewForLOD);
 
 			for(int32 MeshIndex = NumMeshes-1 ; MeshIndex >= 0 ; --MeshIndex)
 			{
diff --git a/Engine/Source/Runtime/Engine/Private/SceneView.cpp b/Engine/Source/Runtime/Engine/Private/SceneView.cpp
--- a/Engine/Source/Runtime/Engine/Private/SceneView.cpp
+++ b/Engine/Source/Runtime/Engine/Private/SceneView.cpp
@@ -797,13 +797,16 @@ static TAutoConsoleVariable<int32> CVarCompensateForFOV(
 
 float FSceneView::GetLODDistanceFactor() const
 {
+	// With stereo rendering views may have asymmetric FOVs, so always use the left eye's view for LOD calculation
+	const FSceneView* ViewForLOD = ( this->Family != nullptr && this->StereoPass == eSSP_RIGHT_EYE ) ? Family->Views[0] : this;
+
 	bool bCompensateForFOV = bUseFieldOfViewForLOD && CVarCompensateForFOV.GetValueOnAnyThread() != 0;
-	float ScreenScaleX = bCompensateForFOV ? ViewMatrices.GetProjectionMatrix().M[0][0] : 1.0f;
-	float ScreenScaleY = bCompensateForFOV ? ViewMatrices.GetProjectionMatrix().M[1][1] : (static_cast<float>(ViewRect.Width()) / ViewRect.Height());
+	float ScreenScaleX = bCompensateForFOV ? ViewForLOD->ViewMatrices.GetProjectionMatrix().M[0][0] : 1.0f;
+	float ScreenScaleY = bCompensateForFOV ? ViewForLOD->ViewMatrices.GetProjectionMatrix().M[1][1] : (static_cast<float>(ViewForLOD->ViewRect.Width()) / ViewForLOD->ViewRect.Height());
 
-	const float ScreenMultiple = FMath::Max(ViewRect.Width() / 2.0f * ScreenScaleX,
-		ViewRect.Height() / 2.0f * ScreenScaleY);
-	float Fac = PI * ScreenMultiple * ScreenMultiple / ViewRect.Area();
+	const float ScreenMultiple = FMath::Max(ViewForLOD->ViewRect.Width() / 2.0f * ScreenScaleX,
+		ViewForLOD->ViewRect.Height() / 2.0f * ScreenScaleY);
+	float Fac = PI * ScreenMultiple * ScreenMultiple / ViewForLOD->ViewRect.Area();
 	return Fac;
 }
 
diff --git a/Engine/Source/Runtime/Engine/Private/SkeletalRender.cpp b/Engine/Source/Runtime/Engine/Private/SkeletalRender.cpp
--- a/Engine/Source/Runtime/Engine/Private/SkeletalRender.cpp
+++ b/Engine/Source/Runtime/Engine/Private/SkeletalRender.cpp
@@ -68,7 +68,9 @@ void FSkeletalMeshObject::UpdateMinDesiredLODLevel(const FSceneView* View, const
 	static const auto* SkeletalMeshLODRadiusScale = IConsoleManager::Get().FindTConsoleVariableDataFloat(TEXT("r.SkeletalMeshLODRadiusScale"));
 	float LODScale = FMath::Clamp(SkeletalMeshLODRadiusScale->GetValueOnRenderThread(), 0.25f, 1.0f);
 
-	const float ScreenRadiusSquared = ComputeBoundsScreenRadiusSquared(Bounds.Origin, Bounds.SphereRadius, *View) * LODScale * LODScale;
+	// With stereo rendering views may have asymmetric FOVs, so always use the left eye's view for LOD calculation
+	const FSceneView* ViewForLOD = ( View->Family != nullptr && View->StereoPass == eSSP_RIGHT_EYE ) ? View->Family->Views[ 0 ] : View;
+	const float ScreenRadiusSquared = ComputeBoundsScreenRadiusSquared(Bounds.Origin, Bounds.SphereRadius, *ViewForLOD) * LODScale * LODScale;
 
 	check( SkeletalMeshLODInfo.Num() == SkeletalMeshResource->LODModels.Num() );
 
diff --git a/Engine/Source/Runtime/Engine/Private/StaticMesh.cpp b/Engine/Source/Runtime/Engine/Private/StaticMesh.cpp
--- a/Engine/Source/Runtime/Engine/Private/StaticMesh.cpp
+++ b/Engine/Source/Runtime/Engine/Private/StaticMesh.cpp
@@ -35,6 +35,7 @@
 #include "Interfaces/ITargetPlatformManagerModule.h"
 #include "Engine/Engine.h"
 #include "EngineGlobals.h"
+#include "EditableMesh.h"
 
 #if WITH_EDITOR
 #include "RawMesh.h"
@@ -124,19 +125,22 @@ void FStaticMeshVertexBuffer::Init(const TArray<FStaticMeshBuildVertex>& InVerti
 
 	// Allocate the vertex data buffer.
 	VertexData->ResizeBuffer(NumVertices);
-	Data = VertexData->GetDataPointer();
-
-	// Copy the vertices into the buffer.
-	for(int32 VertexIndex = 0;VertexIndex < InVertices.Num();VertexIndex++)
+	if( NumVertices > 0 )
 	{
-		const FStaticMeshBuildVertex& SourceVertex = InVertices[VertexIndex];
-		const uint32 DestVertexIndex = VertexIndex;
-		SetVertexTangents(DestVertexIndex, SourceVertex.TangentX, SourceVertex.TangentY, SourceVertex.TangentZ);
-
-		for(uint32 UVIndex = 0; UVIndex < NumTexCoords; UVIndex++)
-		{
-			SetVertexUV(DestVertexIndex,UVIndex,SourceVertex.UVs[UVIndex]);
-		}
+	    Data = VertexData->GetDataPointer();
+    
+	    // Copy the vertices into the buffer.
+	    for(int32 VertexIndex = 0;VertexIndex < InVertices.Num();VertexIndex++)
+	    {
+		    const FStaticMeshBuildVertex& SourceVertex = InVertices[VertexIndex];
+		    const uint32 DestVertexIndex = VertexIndex;
+		    SetVertexTangents(DestVertexIndex, SourceVertex.TangentX, SourceVertex.TangentY, SourceVertex.TangentZ);
+    
+		    for(uint32 UVIndex = 0; UVIndex < NumTexCoords; UVIndex++)
+		    {
+			    SetVertexUV(DestVertexIndex,UVIndex,SourceVertex.UVs[UVIndex]);
+		    }
+	    }
 	}
 }
 
@@ -156,24 +160,52 @@ void FStaticMeshVertexBuffer::Init(const FStaticMeshVertexBuffer& InVertexBuffer
 		AllocateData();
 		check( GetStride() == InVertexBuffer.GetStride() );
 		VertexData->ResizeBuffer(NumVertices);
-		Data = VertexData->GetDataPointer();
-		const uint8* InData = InVertexBuffer.GetRawVertexData();
-		FMemory::Memcpy( Data, InData, Stride * NumVertices );
+		if( NumVertices > 0 )
+		{
+			Data = VertexData->GetDataPointer();
+			const uint8* InData = InVertexBuffer.GetRawVertexData();
+			FMemory::Memcpy( Data, InData, Stride * NumVertices );
+		}
 	}
 }
 
-/**
-* Removes the cloned vertices used for extruding shadow volumes.
-* @param NumVertices - The real number of static mesh vertices which should remain in the buffer upon return.
-*/
-void FStaticMeshVertexBuffer::RemoveLegacyShadowVolumeVertices(uint32 InNumVertices)
+void FStaticMeshVertexBuffer::AppendVertices( const FStaticMeshBuildVertex* Vertices, const uint32 NumVerticesToAppend )
 {
-	check(VertexData);
-	VertexData->ResizeBuffer(InNumVertices);
-	NumVertices = InNumVertices;
+	if (VertexData == nullptr && NumVerticesToAppend > 0)
+	{
+		NumTexCoords = 1;
 
-	// Make a copy of the vertex data pointer.
-	Data = VertexData->GetDataPointer();
+		// Allocate the vertex data storage type if it has never been allocated before
+		AllocateData();
+	}
+
+	check( VertexData != nullptr );	// Must only be called after Init() has already initialized the buffer!
+	if( NumVerticesToAppend > 0 )
+	{
+		check( Vertices != nullptr );
+
+		const uint32 FirstDestVertexIndex = NumVertices;
+		NumVertices += NumVerticesToAppend;
+		VertexData->ResizeBuffer( NumVertices );
+		if( NumVertices > 0 )
+		{
+			Data = VertexData->GetDataPointer();
+
+			// Copy the vertices into the buffer.
+			for( uint32 VertexIter = 0; VertexIter < NumVerticesToAppend; ++VertexIter )
+			{
+				const FStaticMeshBuildVertex& SourceVertex = Vertices[ VertexIter ];
+
+				const uint32 DestVertexIndex = FirstDestVertexIndex + VertexIter;
+
+				SetVertexTangents( DestVertexIndex, SourceVertex.TangentX, SourceVertex.TangentY, SourceVertex.TangentZ );
+				for( uint32 UVIndex = 0; UVIndex < NumTexCoords; UVIndex++ )
+				{
+					SetVertexUV( DestVertexIndex, UVIndex, SourceVertex.UVs[ UVIndex ] );
+				}
+			}
+		}
+	}
 }
 
 template<typename SrcVertexTypeT, typename DstVertexTypeT>
@@ -248,7 +280,10 @@ void FStaticMeshVertexBuffer::Serialize( FArchive& Ar, bool bNeedsCPUAccess )
 			VertexData->Serialize(Ar);
 
 			// Make a copy of the vertex data pointer.
-			Data = VertexData->GetDataPointer();
+			if( NumVertices > 0 )
+			{
+				Data = VertexData->GetDataPointer();
+			}
 		}
 	}
 }
@@ -1876,6 +1911,14 @@ void UStaticMesh::PostEditChangeProperty(FPropertyChangedEvent& PropertyChangedE
 	UProperty* PropertyThatChanged = PropertyChangedEvent.Property;
 	const FName PropertyName = PropertyThatChanged ? PropertyThatChanged->GetFName() : NAME_None;
 	
+	if (PropertyName == GET_MEMBER_NAME_CHECKED(UStaticMesh, LODGroup))
+	{
+		// Force an update of LOD group settings
+
+		// Dont rebuild inside here.  We're doing that below.
+		bool bRebuild = false;
+		SetLODGroup(LODGroup, bRebuild);
+	}
 	LightMapResolution = FMath::Max(LightMapResolution, 0);
 
 	if (PropertyChangedEvent.MemberProperty 
@@ -1919,7 +1962,7 @@ void UStaticMesh::PostEditChangeProperty(FPropertyChangedEvent& PropertyChangedE
 	Super::PostEditChangeProperty(PropertyChangedEvent);
 }
 
-void UStaticMesh::SetLODGroup(FName NewGroup)
+void UStaticMesh::SetLODGroup(FName NewGroup, bool bRebuildImmediately)
 {
 #if WITH_EDITORONLY_DATA
 	const bool bBeforeDerivedDataCached = (RenderData == nullptr);
@@ -1959,6 +2002,11 @@ void UStaticMesh::SetLODGroup(FName NewGroup)
 	if (!bBeforeDerivedDataCached)
 	{
 		bAutoComputeLODScreenSize = true;
+		
+	}
+
+	if (bRebuildImmediately && !bBeforeDerivedDataCached)
+	{
 		PostEditChange();
 	}
 #endif
@@ -2226,7 +2274,12 @@ FArchive& operator<<(FArchive& Ar, FMeshSectionInfo& Info)
 
 void FMeshSectionInfoMap::Serialize(FArchive& Ar)
 {
-	Ar << Map;
+	Ar.UsingCustomVersion(FReleaseObjectVersion::GUID);
+
+	if (Ar.CustomVer(FReleaseObjectVersion::GUID) < FReleaseObjectVersion::UPropertryForMeshSectionSerialize)
+	{
+		Ar << Map;
+	}
 }
 
 #endif // #if WITH_EDITORONLY_DATA
@@ -2591,7 +2644,7 @@ void UStaticMesh::PostLoad()
 		}
 
 		CacheDerivedData();
-		
+
 		// Only required in an editor build as other builds process this in a different place
 		if (bRequiresLODDistanceConversion)
 		{
@@ -2656,6 +2709,28 @@ void UStaticMesh::PostLoad()
 	{
 		CalculateExtendedBounds();
 	}
+
+	if (SectionInfoMap.Map.Num() == 0)
+	{
+		UE_LOG(LogStaticMesh, Warning, TEXT("StaticMesh '%s' has no SectionInfoMap. Please re-export from the original source data."), *GetName());
+
+		// Before this serialization issue was fixed, some assets were resaved and permanently lost their section info map.
+		// This attempts to recreate it based on the render data.
+		SectionInfoMap.Clear();
+		for (int32 LODResourceIndex = 0; LODResourceIndex < RenderData->LODResources.Num(); ++LODResourceIndex)
+		{
+			FStaticMeshLODResources& LOD = RenderData->LODResources[LODResourceIndex];
+			const int32 NumSections = LOD.Sections.Num();
+			for (int32 SectionIndex = 0; SectionIndex < NumSections; ++SectionIndex)
+			{
+				const int32 MaterialIndex = LOD.Sections[SectionIndex].MaterialIndex;
+				if (StaticMaterials.IsValidIndex(MaterialIndex))
+				{
+					SectionInfoMap.Set(LODResourceIndex, SectionIndex, FMeshSectionInfo(MaterialIndex));
+				}
+			}
+		}
+	}
 #endif // #if WITH_EDITOR
 
 	// We want to always have a BodySetup, its used for per-poly collision as well
@@ -2664,6 +2739,14 @@ void UStaticMesh::PostLoad()
 		CreateBodySetup();
 	}
 
+
+	// @todo mesheditor: think about how to make this better
+	// We have an editable mesh, so supersede the vertex/index buffers from the DDC with those generated by the editable mesh
+	if (EditableMesh && !EditableMesh->IsPreviewingSubdivisions())
+	{
+		EditableMesh->RebuildRenderMesh();
+	}
+
 	CreateNavCollision();
 }
 
diff --git a/Engine/Source/Runtime/Engine/Private/SubDSurfaceComponent.cpp b/Engine/Source/Runtime/Engine/Private/SubDSurfaceComponent.cpp
--- a/Engine/Source/Runtime/Engine/Private/SubDSurfaceComponent.cpp
+++ b/Engine/Source/Runtime/Engine/Private/SubDSurfaceComponent.cpp
@@ -455,7 +455,7 @@ static float const g_creaseweights[4] = { 3.0f, 3.0f, 3.0f, 3.0f };
 //   InDeltaPos[0]=InPos[1]-InPos[0], InDeltaPos[1]=InPos[2]-InPos[0]
 //   InDeltaUV[0]=InUV[1]-InUV[0], InDeltaUV[1]=InUV[2]-InUV[0]
 // degenerated cases are not handled yet
-void ComputeTangents(FVector OutTangentXYZ[3], FVector InDeltaPos[2], FVector2D InDeltaUV[2])
+ENGINE_API void ComputeTangents(FVector OutTangentXYZ[3], FVector InDeltaPos[2], FVector2D InDeltaUV[2])
 {
 	// originally from: Lengyel, Eric. "Computing Tangent Space Basis Vectors for an Arbitrary Mesh". Terathon Software 3D Graphics Library, 2001. http://www.terathon.com/code/tangent.html
 
diff --git a/Engine/Source/Runtime/Engine/Private/UnrealClient.cpp b/Engine/Source/Runtime/Engine/Private/UnrealClient.cpp
--- a/Engine/Source/Runtime/Engine/Private/UnrealClient.cpp
+++ b/Engine/Source/Runtime/Engine/Private/UnrealClient.cpp
@@ -19,7 +19,7 @@
 #include "Matinee/MatineeActor.h"
 #include "EditorSupportDelegates.h"
 #include "HighResScreenshot.h"
-
+#include "HModel.h"
 #include "Framework/Notifications/NotificationManager.h"
 #include "Widgets/Notifications/SNotificationList.h"
 #include "Engine/PostProcessVolume.h"
@@ -1430,6 +1430,51 @@ HHitProxy* FViewport::GetHitProxy(int32 X,int32 Y)
 	return HitProxy;
 }
 
+void FViewport::GetActorsAndModelsInHitProxy(FIntRect InRect, TSet<AActor*>& OutActors, TSet<UModel*>& OutModels)
+{
+	const TArray<FColor>& RawHitProxyData = GetRawHitProxyData(InRect);
+
+	OutActors.Empty();
+	OutModels.Empty();
+
+	// Lower the resolution with massive box selects
+	const int32 Step = (InRect.Width() > 500 && InRect.Height() > 500) ? 4 : 1;
+
+	for (int32 Y = InRect.Min.Y; Y < InRect.Max.Y; Y = Y < InRect.Max.Y - 1 ? FMath::Min(InRect.Max.Y-1, Y+Step) : ++Y )
+	{
+		const FColor* SourceData = &RawHitProxyData[Y * SizeX];
+		for (int32 X = InRect.Min.X; X < InRect.Max.X; X = X < InRect.Max.X-1 ? FMath::Min(InRect.Max.X-1, X + Step) : ++X )
+		{
+			FHitProxyId HitProxyId(SourceData[X]);
+			HHitProxy* HitProxy = GetHitProxyById(HitProxyId);
+
+			if (HitProxy)
+			{
+				if( HitProxy->IsA(HActor::StaticGetType()) )
+				{
+					AActor* Actor = ((HActor*)HitProxy)->Actor;
+					if (Actor)
+					{
+						OutActors.Add(Actor);
+					}
+				}
+				else if( HitProxy->IsA(HModel::StaticGetType()) )
+				{
+					OutModels.Add( ((HModel*)HitProxy)->GetModel() );
+				}
+				else if( HitProxy->IsA(HBSPBrushVert::StaticGetType()) )
+				{
+					HBSPBrushVert* HitBSPBrushVert = ((HBSPBrushVert*)HitProxy);
+					if( HitBSPBrushVert->Brush.IsValid() )
+					{
+						OutActors.Add( HitBSPBrushVert->Brush.Get() );
+					}
+				}
+			}
+		}
+	}
+}
+
 void FViewport::UpdateViewportRHI(bool bDestroyed, uint32 NewSizeX, uint32 NewSizeY, EWindowMode::Type NewWindowMode, EPixelFormat PreferredPixelFormat)
 {
 	// Make sure we're not in the middle of streaming textures.
diff --git a/Engine/Source/Runtime/Engine/Private/UnrealEngine.cpp b/Engine/Source/Runtime/Engine/Private/UnrealEngine.cpp
--- a/Engine/Source/Runtime/Engine/Private/UnrealEngine.cpp
+++ b/Engine/Source/Runtime/Engine/Private/UnrealEngine.cpp
@@ -1011,6 +1011,7 @@ void UEngine::Init(IEngineLoop* InEngineLoop)
 	// Dynamically load engine runtime modules
 	{
 		FModuleManager::Get().LoadModuleChecked(TEXT("StreamingPauseRendering"));
+		FModuleManager::Get().LoadModuleChecked(TEXT("MeshEditingRuntime"));
 		FModuleManager::Get().LoadModuleChecked(TEXT("Niagara"));
 		FModuleManager::Get().LoadModuleChecked(TEXT("GeometryCache"));
 	}
diff --git a/Engine/Source/Runtime/Engine/Public/ConvexVolume.h b/Engine/Source/Runtime/Engine/Public/ConvexVolume.h
--- a/Engine/Source/Runtime/Engine/Public/ConvexVolume.h
+++ b/Engine/Source/Runtime/Engine/Public/ConvexVolume.h
@@ -91,6 +91,7 @@ public:
 	bool IntersectBox(const FVector& Origin,const FVector& Extent, bool& bOutFullyContained) const;
 	bool IntersectSphere(const FVector& Origin,const float& Radius) const;
 	bool IntersectSphere(const FVector& Origin,const float& Radius, bool& bOutFullyContained) const;
+	bool IntersectLineSegment(const FVector& Start, const FVector& End) const;
 
 	/**
 	 * Intersection test with a translated axis-aligned box.
@@ -101,8 +102,14 @@ public:
 	 */
 	bool IntersectBox(const FVector& Origin,const FVector& Translation,const FVector& Extent) const;
 
+	/** Determines whether the given point lies inside the convex volume */
+	bool IntersectPoint(const FVector& Point) const
+	{
+		return IntersectSphere(Point, 0.0f);
+	}
+
 	/**
-	 * Serializor
+	 * Serializer
 	 *
 	 * @param	Ar				Archive to serialize data to
 	 * @param	ConvexVolume	Convex volumes to serialize to archive
diff --git a/Engine/Source/Runtime/Engine/Public/GeomTools.h b/Engine/Source/Runtime/Engine/Public/GeomTools.h
--- a/Engine/Source/Runtime/Engine/Public/GeomTools.h
+++ b/Engine/Source/Runtime/Engine/Public/GeomTools.h
@@ -156,8 +156,8 @@ namespace FGeomTools
 	ENGINE_API void Split2DPolysWithPlane(FUtilPoly2DSet& PolySet, const FPlane& Plane, const FColor& ExteriorVertColor, const FColor& InteriorVertColor);
 
 	/** Given three direction vectors, indicates if A and B are on the same 'side' of Vec. */
-	ENGINE_API bool VectorsOnSameSide(const FVector& Vec, const FVector& A, const FVector& B);
+	ENGINE_API bool VectorsOnSameSide(const FVector& Vec, const FVector& A, const FVector& B, const float SameSideDotProductEpsilon = 0.0f );
 
 	/** Util to see if P lies within triangle created by A, B and C. */
-	ENGINE_API bool PointInTriangle(const FVector& A, const FVector& B, const FVector& C, const FVector& P);
+	ENGINE_API bool PointInTriangle(const FVector& A, const FVector& B, const FVector& C, const FVector& P, const float InsideTriangleDotProductEpsilon = 0.0f);
 };
diff --git a/Engine/Source/Runtime/Engine/Public/RawIndexBuffer.h b/Engine/Source/Runtime/Engine/Public/RawIndexBuffer.h
--- a/Engine/Source/Runtime/Engine/Public/RawIndexBuffer.h
+++ b/Engine/Source/Runtime/Engine/Public/RawIndexBuffer.h
@@ -138,6 +138,27 @@ public:
 	 */
 	FRawStaticIndexBuffer(bool InNeedsCPUAccess=false);
 
+	/**
+	 * Sets a single index value.  Consider using SetIndices() instead if you're setting a lot of indices.
+	 * @param	At	The index of the index to set
+	 * @param	NewIndexValue	The index value
+	 */
+	ENGINE_API inline void SetIndex( const uint32 At, const uint32 NewIndexValue )
+	{
+		check( At >= 0 && At < (uint32)IndexStorage.Num() );
+
+		if( b32Bit )
+		{
+			uint32* Indices32Bit = (uint32*)IndexStorage.GetData();
+			Indices32Bit[ At ] = NewIndexValue;
+		}
+		else
+		{
+			uint16* Indices16Bit = (uint16*)IndexStorage.GetData();
+			Indices16Bit[ At ] = (uint16)NewIndexValue;
+		}
+	}
+
 	/**
 	 * Set the indices stored within this buffer.
 	 * @param InIndices		The new indices to copy in to the buffer.
@@ -145,6 +166,49 @@ public:
 	 */
 	ENGINE_API void SetIndices(const TArray<uint32>& InIndices, EIndexBufferStride::Type DesiredStride);
 
+	/**
+	 * Insert indices at the given position in the buffer
+	 * @param	At					Index to insert at
+	 * @param	IndicesToAppend		Pointer to the array of indices to insert
+	 * @param	NumIndicesToAppend	How many indices are in the IndicesToAppend array
+	 */
+	ENGINE_API void InsertIndices( const uint32 At, const uint32* IndicesToAppend, const uint32 NumIndicesToAppend );
+
+	/**
+	 * Append indices to the end of the buffer
+	 * @param	IndicesToAppend		Pointer to the array of indices to add to the end
+	 * @param	NumIndicesToAppend	How many indices are in the IndicesToAppend array
+	 */
+	ENGINE_API void AppendIndices( const uint32* IndicesToAppend, const uint32 NumIndicesToAppend );
+
+	/** @return Gets a specific index value */
+	ENGINE_API inline uint32 GetIndex( const uint32 At ) const
+	{
+		check( At >= 0 && At < (uint32)IndexStorage.Num() );
+		uint32 IndexValue;
+		if( b32Bit )
+		{
+			const uint32* SrcIndices32Bit = (const uint32*)IndexStorage.GetData();
+			IndexValue = SrcIndices32Bit[ At ];
+		}
+		else
+		{
+			const uint16* SrcIndices16Bit = (const uint16*)IndexStorage.GetData();
+			IndexValue = SrcIndices16Bit[ At ];
+		}
+
+		return IndexValue;
+	}
+
+
+	/**
+	 * Removes indices from the buffer
+	 *
+	 * @param	At	The index of the first index to remove
+	 * @param	NumIndicesToRemove	How many indices to remove
+	 */
+	ENGINE_API void RemoveIndicesAt( const uint32 At, const uint32 NumIndicesToRemove );
+
 	/**
 	 * Retrieve a copy of the indices in this buffer. Only valid if created with
 	 * NeedsCPUAccess set to true or the resource has not yet been initialized.
diff --git a/Engine/Source/Runtime/Engine/Public/Rendering/ColorVertexBuffer.h b/Engine/Source/Runtime/Engine/Public/Rendering/ColorVertexBuffer.h
--- a/Engine/Source/Runtime/Engine/Public/Rendering/ColorVertexBuffer.h
+++ b/Engine/Source/Runtime/Engine/Public/Rendering/ColorVertexBuffer.h
@@ -32,13 +32,15 @@ public:
 	* @param InVertexBuffer - The vertex buffer to initialize from.
 	*/
 	void Init(const FColorVertexBuffer& InVertexBuffer);
-
+	
 	/**
-	* Removes the cloned vertices used for extruding shadow volumes.
-	* @param NumVertices - The real number of static mesh vertices which should remain in the buffer upon return.
-	*/
-	void RemoveLegacyShadowVolumeVertices(uint32 InNumVertices);
-
+	 * Appends the specified vertices to the end of the buffer
+	 *
+	 * @param	Vertices	The vertex data to be appended.  Must not be nullptr.
+	 * @param	NumVerticesToAppend		How many vertices should be added
+	 */
+	ENGINE_API void AppendVertices( const FStaticMeshBuildVertex* Vertices, const uint32 NumVerticesToAppend );
+	
 	/**
 	* Serializer
 	*
diff --git a/Engine/Source/Runtime/Engine/Public/StaticMeshResources.h b/Engine/Source/Runtime/Engine/Public/StaticMeshResources.h
--- a/Engine/Source/Runtime/Engine/Public/StaticMeshResources.h
+++ b/Engine/Source/Runtime/Engine/Public/StaticMeshResources.h
@@ -420,10 +420,12 @@ public:
 	ENGINE_API void Init(const TArray<FVector>& InPositions);
 
 	/**
-	* Removes the cloned vertices used for extruding shadow volumes.
-	* @param NumVertices - The real number of static mesh vertices which should remain in the buffer upon return.
-	*/
-	void RemoveLegacyShadowVolumeVertices(uint32 InNumVertices);
+	 * Appends the specified vertices to the end of the buffer
+	 *
+	 * @param	Vertices	The vertex data to be appended.  Must not be nullptr.
+	 * @param	NumVerticesToAppend		How many vertices should be added
+	 */
+	ENGINE_API void AppendVertices( const FStaticMeshBuildVertex* Vertices, const uint32 NumVerticesToAppend );
 
 	/**
 	* Serializer
@@ -509,10 +511,12 @@ public:
 	void Init(const FStaticMeshVertexBuffer& InVertexBuffer);
 
 	/**
-	 * Removes the cloned vertices used for extruding shadow volumes.
-	 * @param NumVertices - The real number of static mesh vertices which should remain in the buffer upon return.
+	 * Appends the specified vertices to the end of the buffer
+	 *
+	 * @param	Vertices	The vertex data to be appended.  Must not be nullptr.
+	 * @param	NumVerticesToAppend		How many vertices should be added
 	 */
-	void RemoveLegacyShadowVolumeVertices(uint32 NumVertices);
+	ENGINE_API void AppendVertices( const FStaticMeshBuildVertex* Vertices, const uint32 NumVerticesToAppend );
 
 	/**
 	* Serializer
@@ -799,7 +803,7 @@ struct FStaticMeshLODResources
 	/** Default constructor. */
 	FStaticMeshLODResources();
 
-	~FStaticMeshLODResources();
+	ENGINE_API ~FStaticMeshLODResources();
 
 	/** Initializes all rendering resources. */
 	void InitResources(UStaticMesh* Parent);
@@ -835,7 +839,7 @@ class FStaticMeshRenderData
 {
 public:
 	/** Default constructor. */
-	FStaticMeshRenderData();
+	ENGINE_API FStaticMeshRenderData();
 
 	/** Per-LOD resources. */
 	TIndirectArray<FStaticMeshLODResources> LODResources;
diff --git a/Engine/Source/Runtime/Engine/Public/UnrealClient.h b/Engine/Source/Runtime/Engine/Public/UnrealClient.h
--- a/Engine/Source/Runtime/Engine/Public/UnrealClient.h
+++ b/Engine/Source/Runtime/Engine/Public/UnrealClient.h
@@ -19,6 +19,7 @@
 class FCanvas;
 class FViewport;
 class FViewportClient;
+class UModel;
 
 /**
  * A render target.
@@ -424,6 +425,13 @@ public:
 	 */
 	ENGINE_API HHitProxy* GetHitProxy(int32 X,int32 Y);
 
+	/**
+	 * Returns all actors and models found in the hit proxy within a specified region.
+	 * InRect must be entirely within the viewport's client area.
+	 * If the hit proxies are not cached, this will call ViewportClient->Draw with a hit-testing canvas.
+	 */
+	ENGINE_API void GetActorsAndModelsInHitProxy(FIntRect InRect, TSet<AActor*>& OutActors, TSet<UModel*>& OutModels);
+
 	/**
 	 * Retrieves the interface to the viewport's frame, if it has one.
 	 * @return The viewport's frame interface.
diff --git a/Engine/Source/Runtime/GeometryCache/Classes/GeometryCacheMeshData.h b/Engine/Source/Runtime/GeometryCache/Classes/GeometryCacheMeshData.h
--- a/Engine/Source/Runtime/GeometryCache/Classes/GeometryCacheMeshData.h
+++ b/Engine/Source/Runtime/GeometryCache/Classes/GeometryCacheMeshData.h
@@ -5,9 +5,10 @@
 #include "UObject/ObjectMacros.h"
 #include "ProfilingDebugging/ResourceSize.h"
 #include "DynamicMeshBuilder.h"
-
+#include "GeometryObjectVersion.h"
 #include "GeometryCacheMeshData.generated.h"
 
+
 /** Stores per-batch data used for rendering */
 USTRUCT()
 struct FGeometryCacheMeshBatchInfo
@@ -56,6 +57,8 @@ struct FGeometryCacheMeshData
 	/** Serialization for FVertexAnimationSample. */
 	friend FArchive& operator<<(FArchive& Ar, FGeometryCacheMeshData& Mesh)
 	{
+		Ar.UsingCustomVersion(FGeometryObjectVersion::GUID);
+
 		int32 NumVertices = 0;
 		
 		if (Ar.IsSaving())
@@ -69,19 +72,26 @@ struct FGeometryCacheMeshData
 			Mesh.Vertices.AddUninitialized(NumVertices);
 		}
 
-		for (int32 VertexIndex = 0; VertexIndex < NumVertices; ++VertexIndex)
+		if (Ar.CustomVer(FGeometryObjectVersion::GUID) < FGeometryObjectVersion::CompressGeometryCache)
 		{
-			FDynamicMeshVertex& Vertex = Mesh.Vertices[VertexIndex];
-			Ar << Vertex.Position;
-			Ar << Vertex.TextureCoordinate;
-			Ar << Vertex.TangentX;
-			Ar << Vertex.TangentZ;
-			Ar << Vertex.Color;
+			for (int32 VertexIndex = 0; VertexIndex < NumVertices; ++VertexIndex)
+			{
+				FDynamicMeshVertex& Vertex = Mesh.Vertices[VertexIndex];
+				Ar << Vertex.Position;
+				Ar << Vertex.TextureCoordinate;
+				Ar << Vertex.TangentX;
+				Ar << Vertex.TangentZ;
+				Ar << Vertex.Color;
+			}
+		}
+		else if(NumVertices)
+		{
+			Ar.SerializeCompressed(&Mesh.Vertices[0], Mesh.Vertices.Num()*Mesh.Vertices.GetTypeSize(), COMPRESS_ZLIB);
 		}
 
+	
 		Ar << Mesh.BoundingBox;
 		Ar << Mesh.BatchesInfo;
-
 		Ar << Mesh.Indices;	
 
 		return Ar;
