Fixed bidirectional dependency between Engine and MeshEditingRuntime modules
- Changed StaticMesh to just have a blind UObject pointer to the EditableMesh

diff --git a/Engine/Source/Runtime/Engine/Classes/Engine/StaticMesh.h b/Engine/Source/Runtime/Engine/Classes/Engine/StaticMesh.h
--- a/Engine/Source/Runtime/Engine/Classes/Engine/StaticMesh.h
+++ b/Engine/Source/Runtime/Engine/Classes/Engine/StaticMesh.h
@@ -601,7 +601,7 @@ public:
 	/** The editable mesh representation of this static mesh */
 	// @todo: Maybe we don't want this visible in the details panel in the end; for now, this might aid debugging.
 	UPROPERTY(Instanced, VisibleAnywhere, Category = EditableMesh)
-	class UEditableMesh* EditableMesh;
+	class UObject* EditableMesh;
 
 	/** Pre-build navigation collision */
 	UPROPERTY(VisibleAnywhere, transient, duplicatetransient, Instanced, Category = Navigation)
diff --git a/Engine/Source/Runtime/Engine/Engine.Build.cs b/Engine/Source/Runtime/Engine/Engine.Build.cs
--- a/Engine/Source/Runtime/Engine/Engine.Build.cs
+++ b/Engine/Source/Runtime/Engine/Engine.Build.cs
@@ -102,7 +102,6 @@ public class Engine : ModuleRules
                 "CinematicCamera",
 				"Analytics",
 				"AnalyticsET",
-                "MeshEditingRuntime"
             }
 		);
 
@@ -163,8 +162,6 @@ public class Engine : ModuleRules
         CircularlyReferencedDependentModules.Add("UMG");
         CircularlyReferencedDependentModules.Add("MaterialShaderQualitySettings");
         CircularlyReferencedDependentModules.Add("CinematicCamera");
-        // @todo mesheditor: See if we can find a way to break this cycle later on
-        CircularlyReferencedDependentModules.Add("MeshEditingRuntime");
 
         // The AnimGraphRuntime module is not needed by Engine proper, but it is loaded in LaunchEngineLoop.cpp,
         // and needs to be listed in an always-included module in order to be compiled into standalone games
@@ -180,6 +177,7 @@ public class Engine : ModuleRules
 				"MovieSceneCapture",
 				"MovieSceneTracks",
 				"HeadMountedDisplay",
+				"MeshEditingRuntime",
 				"StreamingPauseRendering",
                 "Niagara",
 			}
diff --git a/Engine/Source/Runtime/MeshEditingRuntime/StaticMeshEditableMeshFormat.cpp b/Engine/Source/Runtime/MeshEditingRuntime/StaticMeshEditableMeshFormat.cpp
--- a/Engine/Source/Runtime/MeshEditingRuntime/StaticMeshEditableMeshFormat.cpp
+++ b/Engine/Source/Runtime/MeshEditingRuntime/StaticMeshEditableMeshFormat.cpp
@@ -35,7 +35,7 @@ UEditableMesh* FStaticMeshEditableMeshFormat::MakeEditableMesh( UPrimitiveCompon
 		const UStaticMeshComponent& StaticMeshComponent = *StaticMeshComponentPtr;
 		if( StaticMeshComponent.GetStaticMesh() != nullptr )
 		{
-			UEditableMesh* EditableMesh = StaticMeshComponent.GetStaticMesh()->EditableMesh;
+			UEditableMesh* EditableMesh = Cast<UEditableMesh>( StaticMeshComponent.GetStaticMesh()->EditableMesh );
 			if( EditableMesh )
 			{
 				EditableMesh->SetSubMeshAddress( SubMeshAddress );
