Merging //UE4/Dev-Main to Dev-Geometry (//UE4/Dev-Geometry) (up to CL 3517493)
- Missing files only

diff --git a/Engine/Source/Runtime/Engine/Classes/Kismet/HeadMountedDisplayFunctionLibrary.h b/Engine/Source/Runtime/Engine/Classes/Kismet/HeadMountedDisplayFunctionLibrary.h
new file mode 100644
--- /dev/null
+++ b/Engine/Source/Runtime/Engine/Classes/Kismet/HeadMountedDisplayFunctionLibrary.h
@@ -0,0 +1,10 @@
+// Copyright 1998-2017 Epic Games, Inc. All Rights Reserved.
+
+#pragma once
+
+#include "Platform.h"
+
+COMPILE_WARNING("This file has been moved to HeadMountedDisplay/Public/HeadMountedDisplayFunctionLibrary.h. Some types were moved to HeadMountedDisplay/Public/HeadMountedDisplayTypes.h.  Replace all #include \"Kismet\\HeadMountedDisplayFunctionLibrary.h\" with: #include \"HeadMountedDisplayFunctionLibrary.h\" or #include \"HeadMountedDisplayTypes.h\".")
+
+// Needs highly explicit path, or we just end up including this header for an infinite loop.
+#include "../../../HeadMountedDisplay/Public/HeadMountedDisplayFunctionLibrary.h"
\ No newline at end of file
