- Remove FSoftObjectPath attribute
- Add DDC for to avoid converting many time the same data

diff --git a/Engine/Plugins/Editor/MeshEditor/Source/MeshEditor/MeshEditorMode.cpp b/Engine/Plugins/Editor/MeshEditor/Source/MeshEditor/MeshEditorMode.cpp
--- a/Engine/Plugins/Editor/MeshEditor/Source/MeshEditor/MeshEditorMode.cpp
+++ b/Engine/Plugins/Editor/MeshEditor/Source/MeshEditor/MeshEditorMode.cpp
@@ -2279,10 +2279,15 @@ bool FMeshEditorMode::AssignMaterialToSelectedPolygons( UMaterialInterface* Sele
 
 		// Refresh selection (committing may have created a new mesh instance)
 		GetSelectedMeshesAndPolygons( /* Out */ MeshesAndPolygons );
-
 		for( const auto& MeshAndPolygons : MeshesAndPolygons )
 		{
 			UEditableMesh* EditableMesh = MeshAndPolygons.Key;
+
+			int32 MaterialIndex = INDEX_NONE;
+			FName ImportedMaterialName = NAME_None;
+			//We should be able to add a material at least one adapter must succeed
+			check(EditableMesh->FindOrAddMaterial(SelectedMaterial, MaterialIndex, ImportedMaterialName));
+
 			const UMeshDescription* MeshDescription = EditableMesh->GetMeshDescription();
 
 			UPrimitiveComponent* Component = nullptr;
@@ -2301,17 +2306,17 @@ bool FMeshEditorMode::AssignMaterialToSelectedPolygons( UMaterialInterface* Sele
 				// further details if there is more than one polygon group which matches the material.
 				FPolygonGroupID PolygonGroupToAssign = FPolygonGroupID::Invalid;
 
-				const TPolygonGroupAttributeArray<FSoftObjectPath>& PolygonGroupMaterialAssets = MeshDescription->PolygonGroupAttributes().GetAttributes<FSoftObjectPath>( MeshAttribute::PolygonGroup::MaterialAsset );
+				const TPolygonGroupAttributeArray<int>& PolygonGroupMaterialIndexes = MeshDescription->PolygonGroupAttributes().GetAttributes<int>(MeshAttribute::PolygonGroup::MaterialIndex);
+				const TPolygonGroupAttributeArray<FName>& PolygonGroupImportedMaterialSlotName = MeshDescription->PolygonGroupAttributes().GetAttributes<FName>(MeshAttribute::PolygonGroup::ImportedMaterialSlotName);
 
 				for( const FPolygonGroupID PolygonGroupID : MeshDescription->PolygonGroups().GetElementIDs() )
 				{
-					UMaterialInterface* Material = Cast<UMaterialInterface>( PolygonGroupMaterialAssets[ PolygonGroupID ].TryLoad() );
-					if( Material == SelectedMaterial )
+					int32 PolygonGroupMaterialIndex = PolygonGroupMaterialIndexes[PolygonGroupID];
+					if(MaterialIndex == PolygonGroupMaterialIndex)
 					{
 						// We only expect to find one polygon group containing this material at the moment.
 						// We need to provide a way of distinguishing different polygon groups with the same material.
 						ensure( PolygonGroupToAssign == FPolygonGroupID::Invalid );
-
 						PolygonGroupToAssign = PolygonGroupID;
 					}
 				}
@@ -2319,36 +2324,16 @@ bool FMeshEditorMode::AssignMaterialToSelectedPolygons( UMaterialInterface* Sele
 				// If we didn't find the material being used anywhere, create a new polygon group
 				if( PolygonGroupToAssign == FPolygonGroupID::Invalid )
 				{
-					// Helper function which returns a unique FName for the material slot name, based on the material's asset name,
-					// and adding a unique suffix if there are other polygon groups with the same material slot name.
-					auto MakeUniqueSlotName = []( const UMeshDescription* MeshDescription, FName Name ) -> FName
-					{
-						const TPolygonGroupAttributeArray<FName>& MaterialSlotNames = MeshDescription->PolygonGroupAttributes().GetAttributes<FName>( MeshAttribute::PolygonGroup::MaterialSlotName );
-						for( const FPolygonGroupID PolygonGroupID : MeshDescription->PolygonGroups().GetElementIDs() )
-						{
-							const FName ExistingName = MaterialSlotNames[ PolygonGroupID ];
-							if( ExistingName.GetComparisonIndex() == Name.GetComparisonIndex() )
-							{
-								Name = FName( Name, FMath::Max( Name.GetNumber(), ExistingName.GetNumber() + 1 ) );
-							}
-						}
-						return Name;
-					};
-
-					const FName UniqueSlotName = MakeUniqueSlotName( MeshDescription, SelectedMaterial->GetFName() );
 
 					static TArray<FPolygonGroupToCreate> PolygonGroupsToCreate;
 					PolygonGroupsToCreate.Reset( 1 );
 					PolygonGroupsToCreate.Emplace();
 					
 					FPolygonGroupToCreate& PolygonGroupToCreate = PolygonGroupsToCreate.Last();
-					PolygonGroupToCreate.PolygonGroupAttributes.Attributes.Emplace( MeshAttribute::PolygonGroup::MaterialAsset, 0, FMeshElementAttributeValue( FSoftObjectPath(SelectedMaterial) ) );
-					PolygonGroupToCreate.PolygonGroupAttributes.Attributes.Emplace( MeshAttribute::PolygonGroup::MaterialSlotName, 0, FMeshElementAttributeValue( UniqueSlotName ) );
-					PolygonGroupToCreate.PolygonGroupAttributes.Attributes.Emplace( MeshAttribute::PolygonGroup::CastShadow, 0, FMeshElementAttributeValue( true ) );
-					PolygonGroupToCreate.PolygonGroupAttributes.Attributes.Emplace( MeshAttribute::PolygonGroup::EnableCollision, 0, FMeshElementAttributeValue( true ) );
-
+					PolygonGroupToCreate.PolygonGroupAttributes.Attributes.Emplace(MeshAttribute::PolygonGroup::ImportedMaterialSlotName, 0, FMeshElementAttributeValue(ImportedMaterialName));
+					PolygonGroupToCreate.PolygonGroupAttributes.Attributes.Emplace(MeshAttribute::PolygonGroup::MaterialIndex, 0, FMeshElementAttributeValue(MaterialIndex));
 					static TArray<FPolygonGroupID> NewPolygonGroupIDs;
-					EditableMesh->CreatePolygonGroups( PolygonGroupsToCreate, NewPolygonGroupIDs );
+					EditableMesh->CreatePolygonGroups( PolygonGroupsToCreate, NewPolygonGroupIDs);
 					PolygonGroupToAssign = NewPolygonGroupIDs[ 0 ];
 				}
 
diff --git a/Engine/Plugins/Editor/MeshEditor/Source/MeshEditor/MeshEditorStaticMeshAdapter.h b/Engine/Plugins/Editor/MeshEditor/Source/MeshEditor/MeshEditorStaticMeshAdapter.h
--- a/Engine/Plugins/Editor/MeshEditor/Source/MeshEditor/MeshEditorStaticMeshAdapter.h
+++ b/Engine/Plugins/Editor/MeshEditor/Source/MeshEditor/MeshEditorStaticMeshAdapter.h
@@ -22,6 +22,8 @@ public:
 	/** Creates a editable static mesh from the specified component and sub-mesh address */
 	void Initialize( UEditableMesh* EditableMesh, class UWireframeMesh* WireframeMesh );
 
+	virtual bool FindOrAddMaterial(class UMaterialInterface* Material, int32& MaterialIndex, FName& ImportedMaterialName) override { return false; }
+
 	virtual void InitializeFromEditableMesh( const UEditableMesh* EditableMesh ) override {}
 	virtual void OnRebuildRenderMeshStart( const UEditableMesh* EditableMesh, const bool bInvalidateLighting ) override;
 	virtual void OnRebuildRenderMesh( const UEditableMesh* EditableMesh ) override;
diff --git a/Engine/Plugins/Editor/MeshEditor/Source/MeshEditor/MeshEditorSubdividedStaticMeshAdapter.h b/Engine/Plugins/Editor/MeshEditor/Source/MeshEditor/MeshEditorSubdividedStaticMeshAdapter.h
--- a/Engine/Plugins/Editor/MeshEditor/Source/MeshEditor/MeshEditorSubdividedStaticMeshAdapter.h
+++ b/Engine/Plugins/Editor/MeshEditor/Source/MeshEditor/MeshEditorSubdividedStaticMeshAdapter.h
@@ -22,6 +22,8 @@ public:
 	/** Creates a editable static mesh from the specified component and sub-mesh address */
 	void Initialize( UEditableMesh* EditableMesh, class UWireframeMesh* WireframeMesh );
 
+	virtual bool FindOrAddMaterial(class UMaterialInterface* Material, int32& MaterialIndex, FName& ImportedMaterialName) override { return false; }
+
 	virtual void InitializeFromEditableMesh( const UEditableMesh* EditableMesh ) override {}
 	virtual void OnRebuildRenderMeshStart( const UEditableMesh* EditableMesh, const bool bInvalidateLighting ) override;
 	virtual void OnRebuildRenderMesh( const UEditableMesh* EditableMesh ) override;
diff --git a/Engine/Plugins/Runtime/EditableMesh/Source/EditableMesh/EditableMesh.cpp b/Engine/Plugins/Runtime/EditableMesh/Source/EditableMesh/EditableMesh.cpp
--- a/Engine/Plugins/Runtime/EditableMesh/Source/EditableMesh/EditableMesh.cpp
+++ b/Engine/Plugins/Runtime/EditableMesh/Source/EditableMesh/EditableMesh.cpp
@@ -567,6 +567,22 @@ void UEditableMesh::RebuildRenderMesh()
 	}
 }
 
+bool UEditableMesh::FindOrAddMaterial(class UMaterialInterface* Material, int32& MaterialIndex, FName& ImportedMaterialName)
+{
+	//Only one adapter should add a material
+	//TODO: Do we want to support editing a SM and a SK sharing a MeshDescription?
+	for (UEditableMeshAdapter* Adapter : Adapters)
+	{
+		if (Adapter->FindOrAddMaterial(Material, MaterialIndex, ImportedMaterialName))
+		{
+			if (MaterialIndex != INDEX_NONE)
+			{
+				return true;
+			}
+		}
+	}
+	return false;
+}
 
 void UEditableMesh::StartModification( const EMeshModificationType MeshModificationType, const EMeshTopologyChange MeshTopologyChange )
 {
@@ -4545,22 +4561,28 @@ void UEditableMesh::CreatePolygonGroups( const TArray<FPolygonGroupToCreate>& Po
 			{
 				PolygonGroupID = GetMeshDescription()->CreatePolygonGroup();
 			}
-
+			//Apply the PolygonGroup attribute here
+			for (const FMeshElementAttributeData& EdgeAttribute : PolygonGroupToCreate.PolygonGroupAttributes.Attributes)
+			{
+				ApplyAttribute(GetMeshDescription()->PolygonGroupAttributes(), EdgeAttribute, PolygonGroupID);
+			}
 			OutNewPolygonGroupIDs.Add( PolygonGroupID );
 		}
 	}
 
-	for( UEditableMeshAdapter* Adapter : Adapters )
+	//Notify adapters that we just create some PolygonGroup
+	//This must be call before sending the modified attributes to the adapter
+	for (UEditableMeshAdapter* Adapter : Adapters)
 	{
-		Adapter->OnCreatePolygonGroups( this, OutNewPolygonGroupIDs );
+		Adapter->OnCreatePolygonGroups(this, OutNewPolygonGroupIDs);
 	}
 
-	// Apply attributes
-	for( int32 Index = 0; Index < OutNewPolygonGroupIDs.Num(); ++Index )
+	// Apply attributes and notify the adapter for each modified attributes
+	for (int32 Index = 0; Index < OutNewPolygonGroupIDs.Num(); ++Index)
 	{
-		for( const FMeshElementAttributeData& EdgeAttribute : PolygonGroupsToCreate[ Index ].PolygonGroupAttributes.Attributes )
+		for (const FMeshElementAttributeData& EdgeAttribute : PolygonGroupsToCreate[Index].PolygonGroupAttributes.Attributes)
 		{
-			SetPolygonGroupAttribute( OutNewPolygonGroupIDs[ Index ], EdgeAttribute );
+			SetPolygonGroupAttribute(OutNewPolygonGroupIDs[Index], EdgeAttribute);
 		}
 	}
 
diff --git a/Engine/Plugins/Runtime/EditableMesh/Source/EditableMesh/EditableStaticMeshAdapter.cpp b/Engine/Plugins/Runtime/EditableMesh/Source/EditableMesh/EditableStaticMeshAdapter.cpp
--- a/Engine/Plugins/Runtime/EditableMesh/Source/EditableMesh/EditableStaticMeshAdapter.cpp
+++ b/Engine/Plugins/Runtime/EditableMesh/Source/EditableMesh/EditableStaticMeshAdapter.cpp
@@ -104,7 +104,7 @@ void UEditableStaticMeshAdapter::InitEditableStaticMesh( UEditableMesh* Editable
 
 				TEdgeAttributeArray<bool>& EdgeHardnesses = MeshDescription->EdgeAttributes().GetAttributes<bool>( MeshAttribute::Edge::IsHard );
 
-				TPolygonGroupAttributeArray<FSoftObjectPath>& PolygonGroupMaterialAssets = MeshDescription->PolygonGroupAttributes().GetAttributes<FSoftObjectPath>( MeshAttribute::PolygonGroup::MaterialAsset );
+				TPolygonGroupAttributeArray<int>& PolygonGroupMaterialIndexes = MeshDescription->PolygonGroupAttributes().GetAttributes<int>(MeshAttribute::PolygonGroup::MaterialIndex);
 				TPolygonGroupAttributeArray<FName>& PolygonGroupMaterialSlotNames = MeshDescription->PolygonGroupAttributes().GetAttributes<FName>( MeshAttribute::PolygonGroup::MaterialSlotName );
 				TPolygonGroupAttributeArray<FName>& PolygonGroupImportedMaterialSlotNames = MeshDescription->PolygonGroupAttributes().GetAttributes<FName>( MeshAttribute::PolygonGroup::ImportedMaterialSlotName );
 				TPolygonGroupAttributeArray<bool>& PolygonGroupCollision = MeshDescription->PolygonGroupAttributes().GetAttributes<bool>( MeshAttribute::PolygonGroup::EnableCollision );
@@ -276,7 +276,7 @@ void UEditableStaticMeshAdapter::InitEditableStaticMesh( UEditableMesh* Editable
 
 					// Create a new polygon group
 					const FPolygonGroupID NewPolygonGroupID = MeshDescription->CreatePolygonGroup();
-					PolygonGroupMaterialAssets[ NewPolygonGroupID ] = FSoftObjectPath(MaterialInterface);
+					PolygonGroupMaterialIndexes[NewPolygonGroupID] = RenderingSection.MaterialIndex;
 					PolygonGroupMaterialSlotNames[ NewPolygonGroupID ] = StaticMaterial.MaterialSlotName;
 					PolygonGroupImportedMaterialSlotNames[ NewPolygonGroupID ] = StaticMaterial.ImportedMaterialSlotName;
 					PolygonGroupCollision[ NewPolygonGroupID ] = RenderingSection.bEnableCollision;
@@ -450,6 +450,47 @@ void UEditableStaticMeshAdapter::InitFromBlankStaticMesh( UEditableMesh* Editabl
 	StaticMesh = &InStaticMesh;
 }
 
+bool UEditableStaticMeshAdapter::FindOrAddMaterial(UMaterialInterface* Material, int32& MaterialIndex, FName& ImportedMaterialName)
+{
+	if (!StaticMesh)
+	{
+		return false;
+	}
+	//Look if the material already exist
+	for (int32 CurrentMaterialIndex = 0; CurrentMaterialIndex < StaticMesh->StaticMaterials.Num(); ++CurrentMaterialIndex)
+	{
+		if (StaticMesh->StaticMaterials[CurrentMaterialIndex].MaterialInterface == Material)
+		{
+			MaterialIndex = CurrentMaterialIndex;
+			ImportedMaterialName = StaticMesh->StaticMaterials[CurrentMaterialIndex].ImportedMaterialSlotName;
+			return true;
+		}
+	}
+	//If we did not found any existing material, create an entry in the static mesh material list and return the slot index
+	FName MaterialName = Material->GetFName();
+	bool bFoundMatch = true;
+	int32 Counter = 1;
+	while (bFoundMatch)
+	{
+		bFoundMatch = false;
+		for (const FStaticMaterial& StaticMaterial : StaticMesh->StaticMaterials)
+		{
+			if (StaticMaterial.ImportedMaterialSlotName == MaterialName)
+			{
+				bFoundMatch = true;
+				break;
+			}
+		}
+		if (bFoundMatch)
+		{
+			MaterialName = *FString::Printf(TEXT("%s_%d"), *Material->GetName(), Counter);
+			Counter++;
+		}
+	}
+	MaterialIndex  = StaticMesh->StaticMaterials.Add(FStaticMaterial(Material, MaterialName, MaterialName));
+	ImportedMaterialName = MaterialName;
+	return true;
+}
 
 void UEditableStaticMeshAdapter::InitializeFromEditableMesh( const UEditableMesh* EditableMesh )
 {
@@ -1651,12 +1692,12 @@ FStaticMeshLODResources& UEditableStaticMeshAdapter::GetStaticMeshLOD()
 }
 
 
-void UEditableStaticMeshAdapter::OnCreatePolygonGroups( const UEditableMesh* EditableMesh, const TArray<FPolygonGroupID>& PolygonGroupIDs )
+void UEditableStaticMeshAdapter::OnCreatePolygonGroups( const UEditableMesh* EditableMesh, const TArray<FPolygonGroupID>& PolygonGroupIDs)
 {
-	const UMeshDescription* MeshDescription = EditableMesh->GetMeshDescription();
+	UMeshDescription* MeshDescription = EditableMesh->GetMeshDescription();
 	check( MeshDescription );
 
-	const TPolygonGroupAttributeArray<FSoftObjectPath>& PolygonGroupMaterialAssets = MeshDescription->PolygonGroupAttributes().GetAttributes<FSoftObjectPath>( MeshAttribute::PolygonGroup::MaterialAsset );
+	TPolygonGroupAttributeArray<int>& PolygonGroupMaterialIndexes = MeshDescription->PolygonGroupAttributes().GetAttributes<int>(MeshAttribute::PolygonGroup::MaterialIndex);
 	const TPolygonGroupAttributeArray<FName>& PolygonGroupMaterialSlotNames = MeshDescription->PolygonGroupAttributes().GetAttributes<FName>( MeshAttribute::PolygonGroup::MaterialSlotName );
 	const TPolygonGroupAttributeArray<FName>& PolygonGroupImportedMaterialSlotNames = MeshDescription->PolygonGroupAttributes().GetAttributes<FName>( MeshAttribute::PolygonGroup::ImportedMaterialSlotName );
 	const TPolygonGroupAttributeArray<bool>& PolygonGroupCollision = MeshDescription->PolygonGroupAttributes().GetAttributes<bool>( MeshAttribute::PolygonGroup::EnableCollision );
@@ -1667,6 +1708,19 @@ void UEditableStaticMeshAdapter::OnCreatePolygonGroups( const UEditableMesh* Edi
 		uint32 LODSectionIndex = 0;
 		if( !EditableMesh->IsPreviewingSubdivisions() )
 		{
+			//Verify the meshdescription new polygon group data is correct, The attribute must match the static mesh material slot array
+			//The staticmesh material slot entry was create before executing this function
+			bool bFoundMaterial = false;
+			for (int32 MaterialIndex = 0; MaterialIndex < StaticMesh->StaticMaterials.Num(); ++MaterialIndex)
+			{
+				if (PolygonGroupImportedMaterialSlotNames[PolygonGroupID] == StaticMesh->StaticMaterials[MaterialIndex].ImportedMaterialSlotName)
+				{
+					bFoundMaterial = true;
+					check(PolygonGroupMaterialIndexes[PolygonGroupID] == MaterialIndex);
+				}
+			}
+			check(bFoundMaterial);
+
 			// Need to create a new rendering section. This is added to the end of the array.
 			FStaticMeshLODResources& StaticMeshLOD = GetStaticMeshLOD();
 
@@ -1698,14 +1752,7 @@ void UEditableStaticMeshAdapter::OnCreatePolygonGroups( const UEditableMesh* Edi
 			StaticMeshSection.MaxVertexIndex = 0;
 			StaticMeshSection.bEnableCollision = PolygonGroupCollision[ PolygonGroupID ];
 			StaticMeshSection.bCastShadow = PolygonGroupCastShadow[ PolygonGroupID ];
-
-			UMaterialInterface* Material = Cast<UMaterialInterface>( PolygonGroupMaterialAssets[ PolygonGroupID ].TryLoad() );
-#if WITH_EDITORONLY_DATA
-			StaticMeshSection.MaterialIndex = StaticMesh->StaticMaterials.Emplace(Material, PolygonGroupMaterialSlotNames[PolygonGroupID], PolygonGroupImportedMaterialSlotNames[PolygonGroupID]);
-#else
-			StaticMeshSection.MaterialIndex = StaticMesh->StaticMaterials.Emplace(Material, PolygonGroupMaterialSlotNames[PolygonGroupID]);
-#endif
-			
+			StaticMeshSection.MaterialIndex = PolygonGroupMaterialIndexes[ PolygonGroupID ];
 		}
 
 		// Insert the rendering polygon group for keeping track of these index buffer properties
@@ -1726,20 +1773,23 @@ void UEditableStaticMeshAdapter::OnSetPolygonGroupAttribute( const UEditableMesh
 	FStaticMeshLODResources& StaticMeshLOD = GetStaticMeshLOD();
 	FStaticMeshSection& StaticMeshSection = StaticMeshLOD.Sections[ RenderingPolygonGroup.RenderingSectionIndex ];
 
-	const TPolygonGroupAttributeArray<FSoftObjectPath>& PolygonGroupMaterialAssets = MeshDescription->PolygonGroupAttributes().GetAttributes<FSoftObjectPath>( MeshAttribute::PolygonGroup::MaterialAsset );
+	
+	const TPolygonGroupAttributeArray<int>& PolygonGroupMaterialIndexes = MeshDescription->PolygonGroupAttributes().GetAttributes<int>(MeshAttribute::PolygonGroup::MaterialIndex);
 	const TPolygonGroupAttributeArray<FName>& PolygonGroupMaterialSlotNames = MeshDescription->PolygonGroupAttributes().GetAttributes<FName>( MeshAttribute::PolygonGroup::MaterialSlotName );
 	const TPolygonGroupAttributeArray<FName>& PolygonGroupImportedMaterialSlotNames = MeshDescription->PolygonGroupAttributes().GetAttributes<FName>( MeshAttribute::PolygonGroup::ImportedMaterialSlotName );
+	
+	const int32 MaterialIndex = PolygonGroupMaterialIndexes[PolygonGroupID];
 
-	if( Attribute.AttributeName == MeshAttribute::PolygonGroup::MaterialAsset ||
-		Attribute.AttributeName == MeshAttribute::PolygonGroup::MaterialSlotName ||
+	if (Attribute.AttributeName == MeshAttribute::PolygonGroup::MaterialIndex)
+	{
+		StaticMeshSection.MaterialIndex = MaterialIndex;
+	}
+	else if ( Attribute.AttributeName == MeshAttribute::PolygonGroup::MaterialSlotName ||
 		Attribute.AttributeName == MeshAttribute::PolygonGroup::ImportedMaterialSlotName )
 	{
-		UMaterialInterface* Material = Cast<UMaterialInterface>( PolygonGroupMaterialAssets[ PolygonGroupID ].TryLoad() );
-#if WITH_EDITORONLY_DATA
-		StaticMesh->StaticMaterials[ StaticMeshSection.MaterialIndex ] = FStaticMaterial( Material, PolygonGroupMaterialSlotNames[ PolygonGroupID ], PolygonGroupImportedMaterialSlotNames[ PolygonGroupID ] );
-#else
-		StaticMesh->StaticMaterials[StaticMeshSection.MaterialIndex] = FStaticMaterial(Material, PolygonGroupMaterialSlotNames[PolygonGroupID]);
-#endif
+		
+		StaticMesh->StaticMaterials[MaterialIndex].MaterialSlotName = PolygonGroupMaterialSlotNames[PolygonGroupID];
+		StaticMesh->StaticMaterials[MaterialIndex].ImportedMaterialSlotName = PolygonGroupImportedMaterialSlotNames[PolygonGroupID];
 	}
 	else if( Attribute.AttributeName == MeshAttribute::PolygonGroup::CastShadow )
 	{
diff --git a/Engine/Plugins/Runtime/EditableMesh/Source/EditableMesh/EditableStaticMeshAdapter.h b/Engine/Plugins/Runtime/EditableMesh/Source/EditableMesh/EditableStaticMeshAdapter.h
--- a/Engine/Plugins/Runtime/EditableMesh/Source/EditableMesh/EditableStaticMeshAdapter.h
+++ b/Engine/Plugins/Runtime/EditableMesh/Source/EditableMesh/EditableStaticMeshAdapter.h
@@ -121,6 +121,8 @@ public:
 	void InitEditableStaticMesh( UEditableMesh* EditableMesh, class UPrimitiveComponent& Component, const FEditableMeshSubMeshAddress& InitSubMeshAddress );
 	EDITABLEMESH_API void InitFromBlankStaticMesh( UEditableMesh* EditableMesh, UStaticMesh& InStaticMesh );
 
+	virtual bool FindOrAddMaterial(class UMaterialInterface* Material, int32& MaterialIndex, FName& ImportedMaterialName) override;
+
 	virtual void InitializeFromEditableMesh( const UEditableMesh* EditableMesh ) override;
 	virtual void OnRebuildRenderMeshStart( const UEditableMesh* EditableMesh, const bool bInvalidateLighting ) override;
 	virtual void OnRebuildRenderMesh( const UEditableMesh* EditableMesh ) override;
diff --git a/Engine/Plugins/Runtime/EditableMesh/Source/EditableMesh/Public/EditableMesh.h b/Engine/Plugins/Runtime/EditableMesh/Source/EditableMesh/Public/EditableMesh.h
--- a/Engine/Plugins/Runtime/EditableMesh/Source/EditableMesh/Public/EditableMesh.h
+++ b/Engine/Plugins/Runtime/EditableMesh/Source/EditableMesh/Public/EditableMesh.h
@@ -210,6 +210,8 @@ public:
 	UFUNCTION( BlueprintPure, Category="Editable Mesh" )
 	FMeshTriangle GetPolygonTriangulatedTriangle( const FPolygonID PolygonID, int32 PolygonTriangleNumber ) const;
 
+	bool FindOrAddMaterial(class UMaterialInterface* Material, int32& MaterialIndex, FName& ImportedMaterialName);
+
 protected:
 
 	/** Given a set of index remappings, fixes up references in the octree */
diff --git a/Engine/Plugins/Runtime/EditableMesh/Source/EditableMesh/Public/EditableMeshAdapter.h b/Engine/Plugins/Runtime/EditableMesh/Source/EditableMesh/Public/EditableMeshAdapter.h
--- a/Engine/Plugins/Runtime/EditableMesh/Source/EditableMesh/Public/EditableMeshAdapter.h
+++ b/Engine/Plugins/Runtime/EditableMesh/Source/EditableMesh/Public/EditableMeshAdapter.h
@@ -16,6 +16,8 @@ class EDITABLEMESH_API UEditableMeshAdapter : public UObject
 
 public:
 
+	virtual bool FindOrAddMaterial(class UMaterialInterface* Material, int32& MaterialIndex, FName& ImportedMaterialName) PURE_VIRTUAL(, return false;);
+
 	virtual void InitializeFromEditableMesh( const UEditableMesh* EditableMesh ) PURE_VIRTUAL(,);
 	virtual void OnRebuildRenderMeshStart( const UEditableMesh* EditableMesh, const bool bInvalidateLighting ) PURE_VIRTUAL(,);
 	virtual void OnRebuildRenderMesh( const UEditableMesh* EditableMesh ) PURE_VIRTUAL(,);
diff --git a/Engine/Plugins/Runtime/EditableMesh/Source/EditableMesh/Public/EditableMeshTypes.h b/Engine/Plugins/Runtime/EditableMesh/Source/EditableMesh/Public/EditableMeshTypes.h
--- a/Engine/Plugins/Runtime/EditableMesh/Source/EditableMesh/Public/EditableMeshTypes.h
+++ b/Engine/Plugins/Runtime/EditableMesh/Source/EditableMesh/Public/EditableMeshTypes.h
@@ -334,8 +334,7 @@ enum class EMeshElementAttributeType : uint8
 	Float,
 	Int,
 	Bool,
-	FName,
-	FSoftObjectPath
+	FName
 };
 
 
@@ -356,7 +355,6 @@ struct FMeshElementAttributeValue
 	explicit FMeshElementAttributeValue( const int Value ) : Type( EMeshElementAttributeType::Int ), Value_Int( Value ) {}
 	explicit FMeshElementAttributeValue( const bool Value ) : Type( EMeshElementAttributeType::Bool ), Value_Bool( Value ) {}
 	explicit FMeshElementAttributeValue( const FName Value ) : Type( EMeshElementAttributeType::FName ), Value_FName( Value ) {}
-	explicit FMeshElementAttributeValue(FSoftObjectPath const Value ) : Type( EMeshElementAttributeType::FSoftObjectPath), Value_FSoftObjectPath( Value ) {}
 
 	/** Returns the type of this attribute value */
 	EMeshElementAttributeType GetType() const
@@ -373,7 +371,6 @@ struct FMeshElementAttributeValue
 	template <> int GetValue<int>() const { check( Type == EMeshElementAttributeType::Int ); return Value_Int; }
 	template <> bool GetValue<bool>() const { check( Type == EMeshElementAttributeType::Bool ); return Value_Bool; }
 	template <> FName GetValue<FName>() const { check( Type == EMeshElementAttributeType::FName ); return Value_FName; }
-	template <> FSoftObjectPath GetValue<FSoftObjectPath>() const { check( Type == EMeshElementAttributeType::FSoftObjectPath); return Value_FSoftObjectPath; }
 
 	/**
 	 * Calls the specified polymorphic lambda, with this attribute value as its passed-in parameter.
@@ -393,7 +390,6 @@ struct FMeshElementAttributeValue
 			case EMeshElementAttributeType::Int: Func( Value_Int ); break;
 			case EMeshElementAttributeType::Bool: Func( Value_Bool ); break;
 			case EMeshElementAttributeType::FName: Func( Value_FName ); break;
-			case EMeshElementAttributeType::FSoftObjectPath: Func(Value_FSoftObjectPath); break;
 		}
 	}
 
@@ -410,7 +406,6 @@ struct FMeshElementAttributeValue
 			case EMeshElementAttributeType::Int: return Lex::ToString( Value_Int );
 			case EMeshElementAttributeType::Bool: return Lex::ToString( Value_Bool );
 			case EMeshElementAttributeType::FName: return Value_FName.ToString();
-			case EMeshElementAttributeType::FSoftObjectPath: return GetNameSafe(Value_FSoftObjectPath.TryLoad());
 		}
 
 		return FString();
@@ -430,7 +425,6 @@ struct FMeshElementAttributeValue
 			case EMeshElementAttributeType::Int: Ar << AttributeValue.Value_Int; break;
 			case EMeshElementAttributeType::Bool: Ar << AttributeValue.Value_Bool; break;
 			case EMeshElementAttributeType::FName: Ar << AttributeValue.Value_FName; break;
-			case EMeshElementAttributeType::FSoftObjectPath: Ar << AttributeValue.Value_FSoftObjectPath; break;
 		}
 
 		return Ar;
@@ -454,7 +448,6 @@ struct FMeshElementAttributeValue
 			case EMeshElementAttributeType::Int: return Value1.Value_Int == Value2.Value_Int;
 			case EMeshElementAttributeType::Bool: return Value1.Value_Bool == Value2.Value_Bool;
 			case EMeshElementAttributeType::FName: return Value1.Value_FName == Value2.Value_FName;
-			case EMeshElementAttributeType::FSoftObjectPath: return Value1.Value_FSoftObjectPath == Value2.Value_FSoftObjectPath;
 		}
 
 		return false;
@@ -475,7 +468,6 @@ private:
 	int Value_Int;
 	bool Value_Bool;
 	FName Value_FName;
-	FSoftObjectPath Value_FSoftObjectPath;
 };
 
 
diff --git a/Engine/Source/Developer/RawMesh/Public/RawMesh.h b/Engine/Source/Developer/RawMesh/Public/RawMesh.h
--- a/Engine/Source/Developer/RawMesh/Public/RawMesh.h
+++ b/Engine/Source/Developer/RawMesh/Public/RawMesh.h
@@ -84,8 +84,16 @@ struct FRawMesh
 	 * Also updates the material index map.
 	 */
 	RAWMESH_API void CompactMaterialIndices();
+
 };
 
+/**
+* Add the raw mesh data to an archive
+*
+*/
+RAWMESH_API FArchive& operator<<(FArchive& Ar, FRawMesh& RawMesh);
+
+
 /**
  * Bulk data storage for raw meshes.
  */
diff --git a/Engine/Source/Editor/UnrealEd/Private/Fbx/FbxStaticMeshImport.cpp b/Engine/Source/Editor/UnrealEd/Private/Fbx/FbxStaticMeshImport.cpp
--- a/Engine/Source/Editor/UnrealEd/Private/Fbx/FbxStaticMeshImport.cpp
+++ b/Engine/Source/Editor/UnrealEd/Private/Fbx/FbxStaticMeshImport.cpp
@@ -1661,13 +1661,17 @@ UStaticMesh* UnFbx::FFbxImporter::ImportStaticMeshAsSingle(UObject* InParent, TA
 		}
 	}
 
-	if (!GetDefault<UEditorExperimentalSettings>()->bUseMeshDescription)
+	if (GetDefault<UEditorExperimentalSettings>()->bUseMeshDescription)
+	{
+		// SetOriginalMeshDescription will generate the rawmesh data
+		StaticMesh->SetOriginalMeshDescription(LODIndex, MeshDescription);
+	}
+	else
 	{
 		// Store the new raw mesh.
-		SrcModel.RawMeshBulkData->SaveRawMesh(NewRawMesh);
+		SrcModel.SaveRawMesh(NewRawMesh);
 	}
 
-
 	if (bBuildStatus)
 	{
 		int32 MaxMaterialIndex = 0;
@@ -1834,7 +1838,7 @@ UStaticMesh* UnFbx::FFbxImporter::ImportStaticMeshAsSingle(UObject* InParent, TA
 					}
 				}
 
-				SrcModel.RawMeshBulkData->SaveRawMesh(LocalRawMesh);
+				SrcModel.SaveRawMesh(LocalRawMesh);
 			}
 
 			// Setup per-section info and the materials array.
diff --git a/Engine/Source/Runtime/Engine/Classes/Engine/StaticMesh.h b/Engine/Source/Runtime/Engine/Classes/Engine/StaticMesh.h
--- a/Engine/Source/Runtime/Engine/Classes/Engine/StaticMesh.h
+++ b/Engine/Source/Runtime/Engine/Classes/Engine/StaticMesh.h
@@ -933,6 +933,16 @@ private:
 	 */
 	void FixupZeroTriangleSections();
 
+	/**
+	* Return mesh data key. The key is the ddc filename for the mesh data
+	*/
+	bool GetMeshDataKey(FString& OutKey);
+
+	/**
+	* Caches mesh data.
+	*/
+	void CacheMeshData();
+
 public:
 	/**
 	 * Caches derived renderable data.
@@ -949,5 +959,18 @@ private:
 	 */
 	bool CleanUpRedondantMaterialPostLoad;
 
+	/*
+	 * Mesh data key DDC version enum
+	 */
+	enum EMeshDataKeyVersion
+	{
+		VER_MDK_INITIAL = 0,
+
+		// -----<new versions can be added before this line>-------------------------------------------------
+		// - this needs to be the last line (see note below)
+		VER_MDK_AUTOMATIC_VERSION_PLUS_ONE,
+		VER_MDK_AUTOMATIC_VERSION = VER_MDK_AUTOMATIC_VERSION_PLUS_ONE - 1
+	};
+
 #endif // #if WITH_EDITOR
 };
diff --git a/Engine/Source/Runtime/Engine/Private/StaticMesh.cpp b/Engine/Source/Runtime/Engine/Private/StaticMesh.cpp
--- a/Engine/Source/Runtime/Engine/Private/StaticMesh.cpp
+++ b/Engine/Source/Runtime/Engine/Private/StaticMesh.cpp
@@ -2084,23 +2084,14 @@ bool FStaticMeshSourceModel::IsRawMeshEmpty() const
 
 void FStaticMeshSourceModel::LoadRawMesh(FRawMesh& OutRawMesh) const
 {
-	bool bUseMeshDescription = GetDefault<UEditorExperimentalSettings>()->bUseMeshDescription;
-	if (bUseMeshDescription && OriginalMeshDescription != nullptr)
+	if (RawMeshBulkData->IsEmpty() && OriginalMeshDescription != nullptr)
 	{
 		FMeshDescriptionOperations::ConverToRawMesh(OriginalMeshDescription, OutRawMesh);
 	}
 	else
 	{
-		if (RawMeshBulkData->IsEmpty() && OriginalMeshDescription != nullptr)
-		{
-			FMeshDescriptionOperations::ConverToRawMesh(OriginalMeshDescription, OutRawMesh);
-		}
-		else
-		{
-			RawMeshBulkData->LoadRawMesh(OutRawMesh);
-		}
+		RawMeshBulkData->LoadRawMesh(OutRawMesh);
 	}
-	
 }
 
 void FStaticMeshSourceModel::SaveRawMesh(FRawMesh& InRawMesh)
@@ -2109,17 +2100,9 @@ void FStaticMeshSourceModel::SaveRawMesh(FRawMesh& InRawMesh)
 	{
 		return;
 	}
-	bool bUseMeshDescription = GetDefault<UEditorExperimentalSettings>()->bUseMeshDescription;
-	if (bUseMeshDescription && OriginalMeshDescription != nullptr)
-	{
-		FMeshDescriptionOperations::ConverFromRawMesh(InRawMesh, OriginalMeshDescription);
-		RawMeshBulkData->Empty();
-	}
-	else
-	{
-		RawMeshBulkData->SaveRawMesh(InRawMesh);
-		OriginalMeshDescription = nullptr;
-	}
+	//Save both format
+	RawMeshBulkData->SaveRawMesh(InRawMesh);
+	FMeshDescriptionOperations::ConverFromRawMesh(InRawMesh, OriginalMeshDescription);
 }
 
 void FStaticMeshSourceModel::SerializeBulkData(FArchive& Ar, UObject* Owner)
@@ -2310,35 +2293,7 @@ UMeshDescription* UStaticMesh::GetOriginalMeshDescription(int32 LodIndex)
 {
 	if (SourceModels.IsValidIndex(LodIndex))
 	{
-		if (GetDefault<UEditorExperimentalSettings>()->bUseMeshDescription)
-		{
-			if (SourceModels[LodIndex].OriginalMeshDescription == nullptr && !SourceModels[LodIndex].RawMeshBulkData->IsEmpty())
-			{
-				FRawMesh TempRawMesh;
-				SourceModels[LodIndex].RawMeshBulkData->LoadRawMesh(TempRawMesh);
-				//The original mesh description is null we must create one
-				SourceModels[LodIndex].OriginalMeshDescription = NewObject<UMeshDescription>(this, NAME_None);
-				//Convert RawMesh to meshdescription
-				FMeshDescriptionOperations::ConverFromRawMesh(TempRawMesh, SourceModels[LodIndex].OriginalMeshDescription);
-			}
-			return SourceModels[LodIndex].OriginalMeshDescription;
-		}
-		else
-		{
-			if (!SourceModels[LodIndex].RawMeshBulkData->IsEmpty())
-			{
-				FRawMesh TempRawMesh;
-				SourceModels[LodIndex].RawMeshBulkData->LoadRawMesh(TempRawMesh);
-				if (SourceModels[LodIndex].OriginalMeshDescription == nullptr)
-				{
-					//The original mesh description is null we must create one
-					SourceModels[LodIndex].OriginalMeshDescription = NewObject<UMeshDescription>(this, NAME_None);
-				}
-				//Convert RawMesh to meshdescription
-				FMeshDescriptionOperations::ConverFromRawMesh(TempRawMesh, SourceModels[LodIndex].OriginalMeshDescription);
-			}
-			return SourceModels[LodIndex].OriginalMeshDescription;
-		}
+		return SourceModels[LodIndex].OriginalMeshDescription;
 	}
 	return nullptr;
 }
@@ -2347,27 +2302,19 @@ void UStaticMesh::SetOriginalMeshDescription(int32 LodIndex, class UMeshDescript
 {
 	if (SourceModels.IsValidIndex(LodIndex))
 	{
-		if (!GetDefault<UEditorExperimentalSettings>()->bUseMeshDescription)
+		if (MeshDescription != nullptr)
 		{
-			if (MeshDescription != nullptr)
-			{
-				FRawMesh TempRawMesh;
-				//Convert RawMesh to meshdescription
-				FMeshDescriptionOperations::ConverToRawMesh(MeshDescription, TempRawMesh);
-				SourceModels[LodIndex].RawMeshBulkData->SaveRawMesh(TempRawMesh);
-			}
-			else
-			{
-				//We have a null mesh description so empty the raw mesh
-				SourceModels[LodIndex].RawMeshBulkData->Empty();
-			}
-			SourceModels[LodIndex].OriginalMeshDescription = nullptr;
+			FRawMesh TempRawMesh;
+			//Convert RawMesh to meshdescription
+			FMeshDescriptionOperations::ConverToRawMesh(MeshDescription, TempRawMesh);
+			SourceModels[LodIndex].RawMeshBulkData->SaveRawMesh(TempRawMesh);
 		}
 		else
 		{
-			SourceModels[LodIndex].OriginalMeshDescription = MeshDescription;
+			//Mesh description is null, remove the rawmesh data
 			SourceModels[LodIndex].RawMeshBulkData->Empty();
 		}
+		SourceModels[LodIndex].OriginalMeshDescription = MeshDescription;
 	}
 }
 
@@ -2419,10 +2366,309 @@ void UStaticMesh::FixupMaterialSlotName()
 	}
 }
 
+bool UStaticMesh::GetMeshDataKey(FString& OutKey)
+{
+	if (SourceModels.Num() < 1)
+	{
+		return false;
+	}
+	OutKey += TEXT("MESHDATAKEY_STATICMESH___Path_");
+	OutKey += this->GetFullGroupName(false);
+	OutKey += TEXT("___Version_");
+	OutKey += FString::FromInt(VER_MDK_AUTOMATIC_VERSION+1);
+	OutKey += TEXT("___LOD_COUNT_");
+	OutKey += FString::FromInt(SourceModels.Num());
+	FSHA1 Sha;
+	for (int32 LodIndex = 0; LodIndex < SourceModels.Num(); ++LodIndex)
+	{
+		FString LodIndexString = FString::Printf(TEXT("%d_"), LodIndex);
+		FStaticMeshSourceModel& SourceModel = SourceModels[LodIndex];
+		bool bEmptyRawData = true;
+		if (SourceModel.OriginalMeshDescription != nullptr)
+		{
+			LodIndexString += SourceModel.OriginalMeshDescription->GetIdString();
+			bEmptyRawData = false;
+		}
+		
+		if(!SourceModel.RawMeshBulkData->IsEmpty())
+		{
+			LodIndexString += SourceModel.RawMeshBulkData->GetIdString();
+			bEmptyRawData = false;
+		}
+		
+		if(bEmptyRawData)
+		{
+			LodIndexString += TEXT("REDUCELOD");
+		}
+		TArray<TCHAR> LodIndexArray = LodIndexString.GetCharArray();
+		Sha.Update((uint8*)LodIndexArray.GetData(), LodIndexArray.Num() * LodIndexArray.GetTypeSize());
+	}
+	OutKey += TEXT("___DATAHASH_");
+	Sha.Final();
+	// Retrieve the hash and use it to construct a pseudo-GUID.
+	uint32 Hash[5];
+	Sha.GetHash((uint8*)Hash);
+	FGuid Guid = FGuid(Hash[0] ^ Hash[4], Hash[1], Hash[2], Hash[3]);
+	OutKey += Guid.ToString(EGuidFormats::Digits);
+	return true;
+}
+
+struct FMeshDataSerialize
+{
+private:
+	int32 LodNumber;
+	TArray<bool> RawMeshTable;
+	TArray<FRawMesh> RawMeshContent;
+	TArray<bool> MeshDescriptionTable;
+	TArray<UMeshDescription*> MeshDescriptionContent;
+
+public:
+	FMeshDataSerialize()
+	{
+		LodNumber = 0;
+	}
+
+	~FMeshDataSerialize()
+	{
+		for (int32 LodIndex = 0; LodIndex < LodNumber; ++LodIndex)
+		{
+			if (MeshDescriptionTable[LodIndex] == true)
+			{
+				MeshDescriptionContent[LodIndex]->ClearFlags(RF_AllFlags);
+				MeshDescriptionContent[LodIndex]->MarkPendingKill();
+			}
+		}
+		LodNumber = 0;
+		RawMeshTable.Empty();
+		RawMeshContent.Empty();
+		MeshDescriptionTable.Empty();
+		MeshDescriptionContent.Empty();
+	}
+
+	//Fill the structure data using the static mesh
+	void ReadDataFromStaticMesh(UStaticMesh* StaticMesh)
+	{
+		LodNumber = StaticMesh->SourceModels.Num();
+		RawMeshTable.AddZeroed(LodNumber);
+		RawMeshContent.AddDefaulted(LodNumber);
+		MeshDescriptionTable.AddZeroed(LodNumber);
+		MeshDescriptionContent.AddZeroed(LodNumber);
+		for (int32 LodIndex = 0; LodIndex < LodNumber; ++LodIndex)
+		{
+			StaticMesh->SourceModels[LodIndex];
+			FStaticMeshSourceModel& SourceModel = StaticMesh->SourceModels[LodIndex];
+			RawMeshTable[LodIndex] = !SourceModel.RawMeshBulkData->IsEmpty();
+			if (RawMeshTable[LodIndex])
+			{
+				SourceModel.RawMeshBulkData->LoadRawMesh(RawMeshContent[LodIndex]);
+			}
+			MeshDescriptionTable[LodIndex] = SourceModel.OriginalMeshDescription != nullptr;
+			if (MeshDescriptionTable[LodIndex])
+			{
+				MeshDescriptionContent[LodIndex] = Cast<UMeshDescription>(StaticDuplicateObject(SourceModel.OriginalMeshDescription, GetTransientPackage(), NAME_None));
+				MeshDescriptionContent[LodIndex]->SetFlags(RF_Standalone);
+			}
+		}
+	}
+
+	//Apply the structure data to a static mesh
+	bool ApplyDataToStaticMesh(UStaticMesh* StaticMesh, int32 LodIndex, bool bRawMesh)
+	{
+		if (LodIndex >= LodNumber || LodIndex < 0 || !StaticMesh->SourceModels.IsValidIndex(LodIndex))
+		{
+			return false;
+		}
+
+		if (bRawMesh && RawMeshTable.IsValidIndex(LodIndex) && RawMeshTable[LodIndex] && RawMeshContent[LodIndex].IsValid())
+		{
+			StaticMesh->SourceModels[LodIndex].RawMeshBulkData->SaveRawMesh(RawMeshContent[LodIndex]);
+			return true;
+		}
+
+		if (!bRawMesh && MeshDescriptionTable.IsValidIndex(LodIndex) && MeshDescriptionTable[LodIndex] && MeshDescriptionContent[LodIndex] != nullptr)
+		{
+			StaticMesh->SourceModels[LodIndex].OriginalMeshDescription = Cast<UMeshDescription>(StaticDuplicateObject(MeshDescriptionContent[LodIndex], StaticMesh, NAME_None, RF_AllFlags & ~RF_Standalone));
+			return true;
+		}
+
+		return false;
+	}
+
+	void ReadArchive(TArray<uint8>& DerivedData)
+	{
+		FMemoryReader Ar(DerivedData, /*bIsPersistent=*/ true);
+		Ar.Serialize((void*)(&LodNumber), sizeof(int32));
+		RawMeshTable.AddZeroed(LodNumber);
+		RawMeshContent.AddDefaulted(LodNumber);
+		MeshDescriptionTable.AddZeroed(LodNumber);
+		MeshDescriptionContent.AddZeroed(LodNumber);
+		for (int32 LodIndex = 0; LodIndex < LodNumber; ++LodIndex)
+		{
+			int32 HasRawMesh;
+			Ar.Serialize((void*)(&HasRawMesh), sizeof(int32));
+			RawMeshTable[LodIndex] = HasRawMesh == 0 ? false : true;
+			
+			if (RawMeshTable[LodIndex])
+			{
+				//Read the FRawMesh
+				Ar << RawMeshContent[LodIndex];
+			}
+
+			int32 HasMeshDescription;
+			Ar.Serialize((void*)(&HasMeshDescription), sizeof(int32));
+			MeshDescriptionTable[LodIndex] = HasMeshDescription == 0 ? false : true;
+			if (MeshDescriptionTable[LodIndex])
+			{
+				//Write the UMeshDescription
+				MeshDescriptionContent[LodIndex] = NewObject<UMeshDescription>(GetTransientPackage(), NAME_None, RF_Standalone);
+				MeshDescriptionContent[LodIndex]->Serialize(Ar);
+			}
+		}
+	}
+
+	void WriteArchive(TArray<uint8>& DerivedData)
+	{
+		FMemoryWriter Ar(DerivedData, /*bIsPersistent=*/ true);
+		Ar.Serialize((void*)(&LodNumber), sizeof(int32));
+		for (int32 LodIndex = 0; LodIndex < LodNumber; ++LodIndex)
+		{
+			int32 HasRawMesh = RawMeshTable[LodIndex] ? 1 : 0;
+			Ar.Serialize((void*)(&HasRawMesh), sizeof(int32));
+			if (HasRawMesh != 0)
+			{
+				//Write the FRawMeshBulkData
+				Ar << RawMeshContent[LodIndex];
+			}
+
+			int32 HasMeshDescription = MeshDescriptionTable[LodIndex] ? 1 : 0;
+			Ar.Serialize((void*)(&HasMeshDescription), sizeof(int32));
+			if (HasMeshDescription)
+			{
+				//Write the UMeshDescription
+				MeshDescriptionContent[LodIndex]->Serialize(Ar);
+			}
+		}
+	}
+};
+
+void UStaticMesh::CacheMeshData()
+{
+	bool bMeshDescriptionMode = GetDefault<UEditorExperimentalSettings>()->bUseMeshDescription;
+	//TODO DDC
+	//Look if a ddc exist for this mesh data
+	bool IsValidCache = false;
+	FMeshDataSerialize MeshDataSerialize;
+	FString MeshDataKey;
+	if (GetMeshDataKey(MeshDataKey))
+	{
+		//Get the ddc if it exist
+		TArray<uint8> DerivedData;
+		if (GetDerivedDataCacheRef().GetSynchronous(*MeshDataKey, DerivedData))
+		{
+			MeshDataSerialize.ReadArchive(DerivedData);
+			IsValidCache = true;
+		}
+	}
+
+	//We have two lambda, one to create a raw mesh if the mesh description exist and another lambda to create a mesh description if the raw mesh exist
+	//We use lambda because we want to look for them in a order decide by the mesh description mode
+
+	//Create a raw mesh from a valid mesh description
+	auto MeshDescriptionLambda = [this, &IsValidCache, &MeshDataSerialize](int32 LodIndex)
+	{
+		FStaticMeshSourceModel& SourceModel = SourceModels[LodIndex];
+		if (SourceModel.OriginalMeshDescription != nullptr)
+		{
+			if (SourceModel.RawMeshBulkData->IsEmpty())
+			{
+				//Convert the mesh description to raw mesh and mark the ddc for save
+				FRawMesh TempRawMesh;
+				//Use the ddc if it exist
+				if (!MeshDataSerialize.ApplyDataToStaticMesh(this, LodIndex, true))
+				{
+					IsValidCache = false;
+					//Convert RawMesh to meshdescription
+					FMeshDescriptionOperations::ConverToRawMesh(SourceModel.OriginalMeshDescription, TempRawMesh);
+				}
+				SourceModel.RawMeshBulkData->SaveRawMesh(TempRawMesh);
+			}
+			return true;
+		}
+		return false;
+	};
+
+	//Create a mesh description from a valid raw mesh
+	auto RawMeshLambda = [this, &IsValidCache, &MeshDataSerialize](int32 LodIndex)
+	{
+		FStaticMeshSourceModel& SourceModel = SourceModels[LodIndex];
+		if (!SourceModel.RawMeshBulkData->IsEmpty())
+		{
+			if (SourceModel.OriginalMeshDescription == nullptr)
+			{
+				//Use the ddc if it exist
+				if (!MeshDataSerialize.ApplyDataToStaticMesh(this, LodIndex, false))
+				{
+					IsValidCache = false;
+					FRawMesh TempRawMesh;
+					//Convert the rawmesh to mesh description mark the ddc for save
+					SourceModel.RawMeshBulkData->LoadRawMesh(TempRawMesh);
+					//The original mesh description is null we must create one
+					SourceModel.OriginalMeshDescription = NewObject<UMeshDescription>(this, NAME_None);
+					//Convert the raw mesh to mesh description and save the mesh description in the ddc
+					FMeshDescriptionOperations::ConverFromRawMesh(TempRawMesh, SourceModel.OriginalMeshDescription);
+				}
+			}
+			return true;
+		}
+		return false;
+	};
+
+	for (int32 LodIndex = 0; LodIndex < SourceModels.Num(); ++LodIndex)
+	{
+		if (bMeshDescriptionMode)
+		{
+			if (!MeshDescriptionLambda(LodIndex))
+			{
+				RawMeshLambda(LodIndex);
+			}
+		}
+		else
+		{
+			if (!RawMeshLambda(LodIndex))
+			{
+				MeshDescriptionLambda(LodIndex);
+			}
+		}
+
+		//We cannot check because there is some case where the MeshMerge as generate some mesh LOD with no rawmesh and no reduction
+		//if (LodIndex == 0 || (SourceModels[LodIndex].ReductionSettings.PercentTriangles == 1.0f && SourceModels[LodIndex].ReductionSettings.MaxDeviation == 0.0f))
+		//{
+			//Check if we have all the data, we need both data until we deprecate FRawMesh
+			//check(SourceModels[LodIndex].OriginalMeshDescription != nullptr);
+			//check(!SourceModels[LodIndex].RawMeshBulkData->IsEmpty());
+		//}
+	}
+
+	//Write a cache that will be useful next time we load the asset (if not saved)
+	if (!IsValidCache && !MeshDataKey.IsEmpty())
+	{
+		//We need to write the cache
+		MeshDataSerialize.ReadDataFromStaticMesh(this);
+		//Write the ddc cache
+		TArray<uint8> DerivedData;
+		MeshDataSerialize.WriteArchive(DerivedData);
+		int32 LodNumber = SourceModels.Num();
+		GetDerivedDataCacheRef().Put(*MeshDataKey, DerivedData);
+	}
+}
+
 #endif
 
 void UStaticMesh::CacheDerivedData()
 {
+#if WITH_EDITORONLY_DATA
+	CacheMeshData();
+#endif
 	// Cache derived data for the running platform.
 	ITargetPlatformManagerModule& TargetPlatformManager = GetTargetPlatformManagerRef();
 	ITargetPlatform* RunningPlatform = TargetPlatformManager.GetRunningTargetPlatform();
@@ -3259,29 +3505,7 @@ void UStaticMesh::GetVertexColorData(TMap<FVector, FColor>& VertexColorData)
 	const uint32 PaintingMeshLODIndex = 0;
 	if (SourceModels.IsValidIndex(PaintingMeshLODIndex))
 	{
-		if (GetDefault<UEditorExperimentalSettings>()->bUseMeshDescription)
-		{
-			UMeshDescription* PaintLODMeshDescription = GetOriginalMeshDescription(PaintingMeshLODIndex);
-			if (PaintLODMeshDescription != nullptr)
-			{
-				const TVertexAttributeArray<FVector>& VertexPositions = PaintLODMeshDescription->VertexAttributes().GetAttributes<FVector>(MeshAttribute::Vertex::Position);
-				const TVertexInstanceAttributeArray<FVector4>& VertexInstanceColors = PaintLODMeshDescription->VertexInstanceAttributes().GetAttributes<FVector4>(MeshAttribute::VertexInstance::Color);
-				if (VertexInstanceColors.Num() > 0)
-				{
-					// Build a mapping of vertex positions to vertex colors.
-					for (const FVertexInstanceID& VertexInstanceID : PaintLODMeshDescription->VertexInstances().GetElementIDs())
-					{
-						FVector Position = VertexPositions[PaintLODMeshDescription->GetVertexInstanceVertex(VertexInstanceID)];
-						FColor Color = FLinearColor(VertexInstanceColors[VertexInstanceID]).ToFColor(true);
-						if (!VertexColorData.Contains(Position))
-						{
-							VertexColorData.Add(Position, Color);
-						}
-					}
-				}
-			}
-		}
-		else if (SourceModels[PaintingMeshLODIndex].IsRawMeshEmpty() == false)
+		if (SourceModels[PaintingMeshLODIndex].IsRawMeshEmpty() == false)
 		{
 			// Extract the raw mesh.
 			FRawMesh Mesh;
@@ -3320,36 +3544,7 @@ void UStaticMesh::SetVertexColorData(const TMap<FVector, FColor>& VertexColorDat
 	const uint32 PaintingMeshLODIndex = 0;
 	if (SourceModels.IsValidIndex(PaintingMeshLODIndex))
 	{
-		if (GetDefault<UEditorExperimentalSettings>()->bUseMeshDescription)
-		{
-			UMeshDescription* PaintLODMeshDescription = GetOriginalMeshDescription(PaintingMeshLODIndex);
-			if (PaintLODMeshDescription != nullptr)
-			{
-				const TVertexAttributeArray<FVector>& VertexPositions = PaintLODMeshDescription->VertexAttributes().GetAttributes<FVector>(MeshAttribute::Vertex::Position);
-				//Empty vertex color attributes
-				PaintLODMeshDescription->VertexInstanceAttributes().GetAttributesSet<FVector4>(MeshAttribute::VertexInstance::Color).SetNumIndices(0);
-				PaintLODMeshDescription->VertexInstanceAttributes().GetAttributesSet<FVector4>(MeshAttribute::VertexInstance::Color).SetNumIndices(1);
-				TVertexInstanceAttributeArray<FVector4>& VertexInstanceColors = PaintLODMeshDescription->VertexInstanceAttributes().GetAttributes<FVector4>(MeshAttribute::VertexInstance::Color);
-				if (PaintLODMeshDescription->VertexInstances().Num() > 0)
-				{
-					// Build a mapping of vertex positions to vertex colors.
-					for (const FVertexInstanceID& VertexInstanceID : PaintLODMeshDescription->VertexInstances().GetElementIDs())
-					{
-						FVector Position = VertexPositions[PaintLODMeshDescription->GetVertexInstanceVertex(VertexInstanceID)];
-						const FColor* Color = VertexColorData.Find(Position);
-						if (Color)
-						{
-							VertexInstanceColors[VertexInstanceID] = FLinearColor(*Color);
-						}
-						else
-						{
-							VertexInstanceColors[VertexInstanceID] = FLinearColor::White;
-						}
-					}
-				}
-			}
-		}
-		else if (SourceModels[PaintingMeshLODIndex].IsRawMeshEmpty() == false)
+		if (SourceModels[PaintingMeshLODIndex].IsRawMeshEmpty() == false)
 		{
 			// Extract the raw mesh.
 			FRawMesh Mesh;
@@ -3392,18 +3587,7 @@ ENGINE_API void UStaticMesh::RemoveVertexColors()
 
 	for (FStaticMeshSourceModel& SourceModel : SourceModels)
 	{
-		if (GetDefault<UEditorExperimentalSettings>()->bUseMeshDescription)
-		{
-			//Empty the attribute indices
-			if (SourceModel.OriginalMeshDescription != nullptr)
-			{
-				//Remove attributes content
-				SourceModel.OriginalMeshDescription->VertexInstanceAttributes().GetAttributesSet<FVector4>(MeshAttribute::VertexInstance::Color).SetNumIndices(0);
-				//Add back one indice so we are ready to add again some paint color attribute
-				SourceModel.OriginalMeshDescription->VertexInstanceAttributes().GetAttributesSet<FVector4>(MeshAttribute::VertexInstance::Color).SetNumIndices(1);
-			}
-		}
-		else if (!SourceModel.IsRawMeshEmpty())
+		if (!SourceModel.IsRawMeshEmpty())
 		{
 			FRawMesh RawMesh;
 			SourceModel.LoadRawMesh(RawMesh);
diff --git a/Engine/Source/Runtime/MeshBuilder/Private/BuildStatisticManager.h b/Engine/Source/Runtime/MeshBuilder/Private/BuildStatisticManager.h
--- a/Engine/Source/Runtime/MeshBuilder/Private/BuildStatisticManager.h
+++ b/Engine/Source/Runtime/MeshBuilder/Private/BuildStatisticManager.h
@@ -4,25 +4,40 @@
 #include "CoreMinimal.h"
 #include "CoreTypes.h"
 
-DECLARE_LOG_CATEGORY_EXTERN(LogBuildStatistic, Log, All);
+DECLARE_LOG_CATEGORY_EXTERN(LogMeshDescriptionBuildStatistic, Log, All);
 
 namespace BuildStatisticManager
 {
+	struct FStatisticData
+	{
+	public:
+		FStatisticData()
+			: TotalTime(0.0)
+			, Counter(0)
+		{ }
+		double TotalTime;
+		uint32 Counter;
+	};
+
 	class FBuildStatisticScope
 	{
 	public:
-		FBuildStatisticScope(const FString& InTimerDescription)
+		FBuildStatisticScope(const FString& InTimerDescription, FStatisticData& InStatisticData)
+			: StatisticData(InStatisticData)
+			, TimerDescription(InTimerDescription)
 		{
-			TimerDescription = InTimerDescription;
 			StartTime = FPlatformTime::Cycles();
 		}
 		~FBuildStatisticScope()
 		{
-			float ScopeTime = FPlatformTime::ToSeconds(FPlatformTime::Cycles() - StartTime);
-			UE_LOG(LogBuildStatistic, Log, TEXT("%s: %f seconds"), *TimerDescription, ScopeTime);
+			double ScopeTime = FPlatformTime::ToSeconds64(FPlatformTime::Cycles() - StartTime);
+			StatisticData.TotalTime += ScopeTime;
+			StatisticData.Counter += 1;
+			UE_LOG(LogMeshDescriptionBuildStatistic, Log, TEXT("%s: %f seconds - Counter: %d - Total Time: %f"), *TimerDescription, ScopeTime, StatisticData.Counter, StatisticData.TotalTime);
 		}
 	private:
 		int32 StartTime;
 		FString TimerDescription;
+		FStatisticData& StatisticData;
 	};
 }
\ No newline at end of file
diff --git a/Engine/Source/Runtime/MeshBuilder/Private/MeshDescriptionHelper.cpp b/Engine/Source/Runtime/MeshBuilder/Private/MeshDescriptionHelper.cpp
--- a/Engine/Source/Runtime/MeshBuilder/Private/MeshDescriptionHelper.cpp
+++ b/Engine/Source/Runtime/MeshBuilder/Private/MeshDescriptionHelper.cpp
@@ -20,10 +20,13 @@
 #include "IMeshReductionManagerModule.h"
 #include "Materials/Material.h"
 #include "RawMesh.h"
+#include "BuildStatisticManager.h"
 
 //Enable all check
 //#define ENABLE_NTB_CHECK
 
+DEFINE_LOG_CATEGORY(LogMeshDescriptionBuildStatistic);
+
 namespace MeshDescriptionMikktSpaceInterface
 {
 	//Mikk t spce static function
@@ -54,7 +57,7 @@ UMeshDescription* FMeshDescriptionHelper::GetRenderMeshDescription(UObject* Owne
 		return RenderMeshDescription;
 	}
 	//Copy The Original Mesh Description in the render mesh description
-	RenderMeshDescription = Cast<UMeshDescription>(StaticDuplicateObject(OriginalMeshDescription, Owner, NAME_None, RF_NoFlags));
+	RenderMeshDescription = Cast<UMeshDescription>(StaticDuplicateObject(OriginalMeshDescription, Owner, NAME_None));
 	float ComparisonThreshold = BuildSettings->bRemoveDegenerates ? THRESH_POINTS_ARE_SAME : 0.0f;
 
 	FVertexInstanceArray& VertexInstanceArray = RenderMeshDescription->VertexInstances();
@@ -925,6 +928,8 @@ void FMeshDescriptionHelper::ConvertSmoothGroupToHardEdges(const struct FRawMesh
 
 void FMeshDescriptionHelper::ConverToRawMesh(const UMeshDescription* SourceMeshDescription, struct FRawMesh &DestinationRawMesh)
 {
+	static BuildStatisticManager::FStatisticData StatisticData;
+	BuildStatisticManager::FBuildStatisticScope(TEXT("ConvertToRawMesh"), StatisticData);
 	DestinationRawMesh.Empty();
 
 	//Gather all array data
@@ -995,6 +1000,8 @@ void FMeshDescriptionHelper::ConverToRawMesh(const UMeshDescription* SourceMeshD
 
 void FMeshDescriptionHelper::ConverFromRawMesh(const struct FRawMesh &SourceRawMesh, UMeshDescription* DestinationMeshDescription)
 {
+	static BuildStatisticManager::FStatisticData StatisticData;
+	BuildStatisticManager::FBuildStatisticScope(TEXT("ConverFromRawMesh"), StatisticData);
 	DestinationMeshDescription->Empty();
 	//Gather all array data
 	TVertexAttributeArray<FVector>& VertexPositions = DestinationMeshDescription->VertexAttributes().GetAttributes<FVector>(MeshAttribute::Vertex::Position);
diff --git a/Engine/Source/Runtime/MeshDescription/Private/MeshAttributes.cpp b/Engine/Source/Runtime/MeshDescription/Private/MeshAttributes.cpp
--- a/Engine/Source/Runtime/MeshDescription/Private/MeshAttributes.cpp
+++ b/Engine/Source/Runtime/MeshDescription/Private/MeshAttributes.cpp
@@ -21,7 +21,6 @@ namespace MeshAttribute
 	const FName Polygon::Binormal( "Binormal" );
 	const FName Polygon::Center( "Center" );
 
-	const FName PolygonGroup::MaterialAsset( "MaterialAsset" );
 	const FName PolygonGroup::MaterialSlotName( "MaterialSlotName" );
 	const FName PolygonGroup::ImportedMaterialSlotName( "ImportedMaterialSlotName" );
 	const FName PolygonGroup::EnableCollision( "EnableCollision" );
diff --git a/Engine/Source/Runtime/MeshDescription/Private/MeshDescription.cpp b/Engine/Source/Runtime/MeshDescription/Private/MeshDescription.cpp
--- a/Engine/Source/Runtime/MeshDescription/Private/MeshDescription.cpp
+++ b/Engine/Source/Runtime/MeshDescription/Private/MeshDescription.cpp
@@ -35,12 +35,11 @@ void UMeshDescription::RegisterBaseAttributes()
 	PolygonAttributes().RegisterAttribute<FVector>(MeshAttribute::Polygon::Center, 1, FVector::ZeroVector, EMeshAttributeFlags::AutoGenerated);
 
 	// Add basic polygon group attributes
-	PolygonGroupAttributes().RegisterAttribute<FSoftObjectPath>(MeshAttribute::PolygonGroup::MaterialAsset); //Deprecated
+	PolygonGroupAttributes().RegisterAttribute<int>(MeshAttribute::PolygonGroup::MaterialIndex); //This attribute represent the section material index
 	PolygonGroupAttributes().RegisterAttribute<FName>(MeshAttribute::PolygonGroup::MaterialSlotName); //Deprecated
 	PolygonGroupAttributes().RegisterAttribute<FName>(MeshAttribute::PolygonGroup::ImportedMaterialSlotName); //The unique key to match the mesh material slot
 	PolygonGroupAttributes().RegisterAttribute<bool>(MeshAttribute::PolygonGroup::EnableCollision); //Deprecated
 	PolygonGroupAttributes().RegisterAttribute<bool>(MeshAttribute::PolygonGroup::CastShadow); //Deprecated
-	PolygonGroupAttributes().RegisterAttribute<int>(MeshAttribute::PolygonGroup::MaterialIndex); //This attribute is used for FRawMesh conversion only and is temporary, use the ImportedMaterialSlotName to match the mesh material slot
 }
 
 void UMeshDescription::UnRegisterBaseAttributes()
@@ -67,12 +66,11 @@ void UMeshDescription::UnRegisterBaseAttributes()
 	PolygonAttributes().UnregisterAttribute<FVector>(MeshAttribute::Polygon::Center);
 
 	// Remove basic polygon group attributes
-	PolygonGroupAttributes().UnregisterAttribute<FSoftObjectPath>(MeshAttribute::PolygonGroup::MaterialAsset);
+	PolygonGroupAttributes().UnregisterAttribute<int>(MeshAttribute::PolygonGroup::MaterialIndex);
 	PolygonGroupAttributes().UnregisterAttribute<FName>(MeshAttribute::PolygonGroup::MaterialSlotName);
 	PolygonGroupAttributes().UnregisterAttribute<FName>(MeshAttribute::PolygonGroup::ImportedMaterialSlotName);
 	PolygonGroupAttributes().UnregisterAttribute<bool>(MeshAttribute::PolygonGroup::EnableCollision);
 	PolygonGroupAttributes().UnregisterAttribute<bool>(MeshAttribute::PolygonGroup::CastShadow);
-	PolygonGroupAttributes().UnregisterAttribute<int>(MeshAttribute::PolygonGroup::MaterialIndex);
 }
 
 void UMeshDescription::Serialize( FArchive& Ar )
@@ -112,8 +110,6 @@ FString UMeshDescription::GetIdString()
 	FObjectWriter Ar(this, TempBytes);
 	Serialize(Ar);
 	FSHA1 Sha;
-	TArray<TCHAR> OwnerName = GetPathName().GetCharArray();
-	Sha.Update((uint8*)OwnerName.GetData(), OwnerName.Num() * OwnerName.GetTypeSize());
 	if (TempBytes.Num() > 0)
 	{
 		uint8* Buffer = TempBytes.GetData();
diff --git a/Engine/Source/Runtime/MeshDescription/Public/MeshAttributeArray.h b/Engine/Source/Runtime/MeshDescription/Public/MeshAttributeArray.h
--- a/Engine/Source/Runtime/MeshDescription/Public/MeshAttributeArray.h
+++ b/Engine/Source/Runtime/MeshDescription/Public/MeshAttributeArray.h
@@ -19,8 +19,7 @@ using AttributeTypes = TTuple
 	float,
 	int,
 	bool,
-	FName,
-	FSoftObjectPath
+	FName
 >;
 
 
diff --git a/Engine/Source/Runtime/MeshDescription/Public/MeshAttributes.h b/Engine/Source/Runtime/MeshDescription/Public/MeshAttributes.h
--- a/Engine/Source/Runtime/MeshDescription/Public/MeshAttributes.h
+++ b/Engine/Source/Runtime/MeshDescription/Public/MeshAttributes.h
@@ -40,7 +40,6 @@ namespace MeshAttribute
 
 	namespace PolygonGroup
 	{
-		extern MESHDESCRIPTION_API const FName MaterialAsset;
 		extern MESHDESCRIPTION_API const FName MaterialSlotName;
 		extern MESHDESCRIPTION_API const FName ImportedMaterialSlotName;
 		extern MESHDESCRIPTION_API const FName EnableCollision;
