Bump meshdescription build GUID to fire up CIS

diff --git a/Engine/Source/Runtime/Engine/Private/StaticMesh.cpp b/Engine/Source/Runtime/Engine/Private/StaticMesh.cpp
--- a/Engine/Source/Runtime/Engine/Private/StaticMesh.cpp
+++ b/Engine/Source/Runtime/Engine/Private/StaticMesh.cpp
@@ -2458,7 +2458,7 @@ void UStaticMesh::FixupMaterialSlotName()
 // differences, etc.) replace the version GUID below with a new one.
 // In case of merge conflicts with DDC versions, you MUST generate a new GUID
 // and set this new GUID as the version.                                       
-#define MESHDATAKEY_STATICMESH_DERIVEDDATA_VER TEXT("3F6401D9624A4E2AB07998C997BC52D5")
+#define MESHDATAKEY_STATICMESH_DERIVEDDATA_VER TEXT("F9C67C97D5854B2DB45CC1AF78135EE0")
 
 static const FString& GetMeshDataKeyStaticMeshDerivedDataVersion()
 {
