-Fix fbx automation test when not in meshdescription mode (crash)
-Change Warning to be Log when the serialization add  some texture coordinnate

diff --git a/Engine/Source/Runtime/Engine/Private/StaticMesh.cpp b/Engine/Source/Runtime/Engine/Private/StaticMesh.cpp
--- a/Engine/Source/Runtime/Engine/Private/StaticMesh.cpp
+++ b/Engine/Source/Runtime/Engine/Private/StaticMesh.cpp
@@ -2190,7 +2190,10 @@ void FStaticMeshSourceModel::SaveRawMesh(FRawMesh& InRawMesh)
 	}
 	//Save both format
 	RawMeshBulkData->SaveRawMesh(InRawMesh);
-	FMeshDescriptionOperations::ConverFromRawMesh(InRawMesh, OriginalMeshDescription);
+	if (OriginalMeshDescription != nullptr)
+	{
+		FMeshDescriptionOperations::ConverFromRawMesh(InRawMesh, OriginalMeshDescription);
+	}
 }
 
 void FStaticMeshSourceModel::SerializeBulkData(FArchive& Ar, UObject* Owner)
@@ -2363,7 +2366,7 @@ UMeshDescription* UStaticMesh::GetMeshDescription(int32 LodIndex) const
 
 void UStaticMesh::SetMeshDescription(int32 LodIndex, class UMeshDescription* InMeshDescription)
 {
-	if (!MeshDescriptions.IsValidIndex(LodIndex))
+	if (!MeshDescriptions.IsValidIndex(LodIndex) && LodIndex >= 0)
 	{
 		//Add nullptr missing entries
 		MeshDescriptions.AddZeroed(LodIndex - MeshDescriptions.Num() + 1);
@@ -2381,6 +2384,14 @@ UMeshDescription* UStaticMesh::GetOriginalMeshDescription(int32 LodIndex)
 {
 	if (SourceModels.IsValidIndex(LodIndex))
 	{
+		if (SourceModels[LodIndex].OriginalMeshDescription == nullptr && !SourceModels[LodIndex].IsRawMeshEmpty())
+		{
+			//Fill the MeshDescription on demand
+			SourceModels[LodIndex].OriginalMeshDescription = NewObject<UMeshDescription>(this, NAME_None);
+			FRawMesh LodRawMesh;
+			SourceModels[LodIndex].LoadRawMesh(LodRawMesh);
+			FMeshDescriptionOperations::ConverFromRawMesh(LodRawMesh, SourceModels[LodIndex].OriginalMeshDescription);
+		}
 		return SourceModels[LodIndex].OriginalMeshDescription;
 	}
 	return nullptr;
@@ -2388,22 +2399,22 @@ UMeshDescription* UStaticMesh::GetOriginalMeshDescription(int32 LodIndex)
 
 void UStaticMesh::SetOriginalMeshDescription(int32 LodIndex, class UMeshDescription* MeshDescription)
 {
-	if (SourceModels.IsValidIndex(LodIndex))
+	//The source model must be create before calling this function
+	check(SourceModels.IsValidIndex(LodIndex));
+
+	if (MeshDescription != nullptr)
 	{
-		if (MeshDescription != nullptr)
-		{
-			FRawMesh TempRawMesh;
-			//Convert RawMesh to meshdescription
-			FMeshDescriptionOperations::ConverToRawMesh(MeshDescription, TempRawMesh);
-			SourceModels[LodIndex].RawMeshBulkData->SaveRawMesh(TempRawMesh);
-		}
-		else
-		{
-			//Mesh description is null, remove the rawmesh data
-			SourceModels[LodIndex].RawMeshBulkData->Empty();
-		}
-		SourceModels[LodIndex].OriginalMeshDescription = MeshDescription;
+		FRawMesh TempRawMesh;
+		//Convert RawMesh to meshdescription
+		FMeshDescriptionOperations::ConverToRawMesh(MeshDescription, TempRawMesh);
+		SourceModels[LodIndex].RawMeshBulkData->SaveRawMesh(TempRawMesh);
+	}
+	else
+	{
+		//Mesh description is null, remove the rawmesh data
+		SourceModels[LodIndex].RawMeshBulkData->Empty();
 	}
+	SourceModels[LodIndex].OriginalMeshDescription = MeshDescription;
 }
 
 void UStaticMesh::ClearOriginalMeshDescription(int32 LodIndex)
diff --git a/Engine/Source/Runtime/MeshBuilder/Private/MeshDescriptionHelper.cpp b/Engine/Source/Runtime/MeshBuilder/Private/MeshDescriptionHelper.cpp
--- a/Engine/Source/Runtime/MeshBuilder/Private/MeshDescriptionHelper.cpp
+++ b/Engine/Source/Runtime/MeshBuilder/Private/MeshDescriptionHelper.cpp
@@ -1000,6 +1000,7 @@ void FMeshDescriptionHelper::ConverToRawMesh(const UMeshDescription* SourceMeshD
 
 void FMeshDescriptionHelper::ConverFromRawMesh(const struct FRawMesh &SourceRawMesh, UMeshDescription* DestinationMeshDescription)
 {
+	check(DestinationMeshDescription != nullptr);
 	static BuildStatisticManager::FStatisticData StatisticData;
 	BuildStatisticManager::FBuildStatisticScope StatScope(TEXT("ConverFromRawMesh"), StatisticData);
 	DestinationMeshDescription->Empty();
diff --git a/Engine/Source/Runtime/MeshDescription/Public/MeshAttributeArray.h b/Engine/Source/Runtime/MeshDescription/Public/MeshAttributeArray.h
--- a/Engine/Source/Runtime/MeshDescription/Public/MeshAttributeArray.h
+++ b/Engine/Source/Runtime/MeshDescription/Public/MeshAttributeArray.h
@@ -530,13 +530,13 @@ private:
 					// Re-register any attributes that aren't there
 					AttributeIndicesArrayType NewAttributeIndicesArray( RegisteredAttribute.NumIndices, RegisteredAttribute.DefaultValue, RegisteredAttribute.Flags, AttributesMap.NumElements );
 
-					UE_LOG( LogMeshDescription, Warning, TEXT( "Didn't find attribute '%s' - adding" ), *RegisteredAttribute.Name.ToString() );
+					UE_LOG( LogMeshDescription, Log, TEXT( "Didn't find attribute '%s' - adding" ), *RegisteredAttribute.Name.ToString() );
 					AttributesMap.Map.Add( RegisteredAttribute.Name, MoveTemp( NewAttributeIndicesArray ) );
 				}
 				else if( AttributeIndicesArray->GetNumIndices() != RegisteredAttribute.NumIndices )
 				{
 					// Amend the number of attribute indices for any which are not correct
-					UE_LOG( LogMeshDescription, Warning, TEXT( "Found attribute '%s' with the wrong number of indices - amending" ), *RegisteredAttribute.Name.ToString() );
+					UE_LOG( LogMeshDescription, Log, TEXT( "Found attribute '%s' with the wrong number of indices - amending" ), *RegisteredAttribute.Name.ToString() );
 					AttributeIndicesArray->SetNumIndices( RegisteredAttribute.NumIndices );
 				}
 				else
