Fix non-editor build error. UStaticMesh::RegisterMeshAttributes is now defined unconditionally.

diff --git a/Engine/Source/Runtime/Engine/Classes/Engine/StaticMesh.h b/Engine/Source/Runtime/Engine/Classes/Engine/StaticMesh.h
--- a/Engine/Source/Runtime/Engine/Classes/Engine/StaticMesh.h
+++ b/Engine/Source/Runtime/Engine/Classes/Engine/StaticMesh.h
@@ -24,6 +24,7 @@
 class FSpeedTreeWind;
 class UAssetUserData;
 class UMaterialInterface;
+class UMeshDescription;
 struct FStaticMeshLODResources;
 
 /*-----------------------------------------------------------------------------
@@ -175,7 +176,7 @@ struct FStaticMeshSourceModel
 #if WITH_EDITORONLY_DATA
 	/* Original Imported mesh description data. Optional for all but the first LOD. Autogenerate LOD do not have original mesh description*/
 	UPROPERTY(Transient)
-	class UMeshDescription* OriginalMeshDescription;
+	UMeshDescription* OriginalMeshDescription;
 #endif
 
 	/** Settings applied when building the mesh. */
@@ -671,6 +672,11 @@ public:
 	UPROPERTY(Instanced, VisibleAnywhere, Category = EditableMesh)
 	class UObject* EditableMesh;
 
+	/**
+	 * Registers the mesh attributes required by the mesh description for a static mesh.
+	 */
+	ENGINE_API static void RegisterMeshAttributes( UMeshDescription* MeshDescription );
+
 #if WITH_EDITORONLY_DATA
 	/**
 	 * The MeshDescription use to build the render data.
@@ -678,26 +684,21 @@ public:
 	 * respect the build options (tangent, UVs, reduce, ...) we then keep the copy in this array in case we need it
 	 */
 	UPROPERTY(VisibleAnywhere, Category = MeshDescription)
-	TArray<class UMeshDescription*> MeshDescriptions;
-
-	/**
-	 * Registers the mesh attributes required by the mesh description for a static mesh.
-	 */
-	ENGINE_API static void RegisterMeshAttributes( UMeshDescription* MeshDescription );
+	TArray<UMeshDescription*> MeshDescriptions;
 
 	/**
 	 * Accessors for the mesh description use to build the render data
 	 */
-	ENGINE_API class UMeshDescription* GetMeshDescription(int32 LodIndex=0) const;
-	ENGINE_API void SetMeshDescription(int32 LodIndex, class UMeshDescription* InMeshDescription);
+	ENGINE_API UMeshDescription* GetMeshDescription(int32 LodIndex=0) const;
+	ENGINE_API void SetMeshDescription(int32 LodIndex, UMeshDescription* InMeshDescription);
 	ENGINE_API int32 GetMeshDescriptionCount() const;
 
 	/**
 	 * Accessors for the original mesh description imported data
 	 * The original import data is necessary to start from the full data when applying build options.
 	 */
-	ENGINE_API class UMeshDescription* GetOriginalMeshDescription(int32 LodIndex = 0);
-	ENGINE_API void SetOriginalMeshDescription(int32 LodIndex, class UMeshDescription* MeshDescription);
+	ENGINE_API UMeshDescription* GetOriginalMeshDescription(int32 LodIndex = 0);
+	ENGINE_API void SetOriginalMeshDescription(int32 LodIndex, UMeshDescription* MeshDescription);
 	ENGINE_API void ClearOriginalMeshDescription(int32 LodIndex);
 
 	/**
