Merging //UE4/Dev-Main to Dev-Geometry (//UE4/Dev-Geometry)

diff --git a/Engine/Source/Runtime/AIModule/Private/EnvironmentQuery/Items/EnvQueryItemType_Actor.cpp b/Engine/Source/Runtime/AIModule/Private/EnvironmentQuery/Items/EnvQueryItemType_Actor.cpp
--- a/Engine/Source/Runtime/AIModule/Private/EnvironmentQuery/Items/EnvQueryItemType_Actor.cpp
+++ b/Engine/Source/Runtime/AIModule/Private/EnvironmentQuery/Items/EnvQueryItemType_Actor.cpp
@@ -3,6 +3,7 @@
 #include "EnvironmentQuery/Items/EnvQueryItemType_Actor.h"
 #include "UObject/WeakObjectPtr.h"
 #include "GameFramework/Actor.h"
+#include "AITypes.h"
 #include "EnvironmentQuery/EnvQueryTypes.h"
 
 UEnvQueryItemType_Actor::UEnvQueryItemType_Actor(const FObjectInitializer& ObjectInitializer) : Super(ObjectInitializer)
@@ -25,13 +26,13 @@ void UEnvQueryItemType_Actor::SetValue(uint8* RawData, const FWeakObjectPtr& Val
 FVector UEnvQueryItemType_Actor::GetItemLocation(const uint8* RawData) const
 {
 	AActor* MyActor = UEnvQueryItemType_Actor::GetValue(RawData);
-	return MyActor ? MyActor->GetActorLocation() : FVector::ZeroVector;
+	return MyActor ? MyActor->GetActorLocation() : FAISystem::InvalidLocation;
 }
 
 FRotator UEnvQueryItemType_Actor::GetItemRotation(const uint8* RawData) const
 {
 	AActor* MyActor = UEnvQueryItemType_Actor::GetValue(RawData);
-	return MyActor ? MyActor->GetActorRotation() : FRotator::ZeroRotator;
+	return MyActor ? MyActor->GetActorRotation() : FAISystem::InvalidRotation;
 }
 
 AActor* UEnvQueryItemType_Actor::GetActor(const uint8* RawData) const
diff --git a/Engine/Source/Runtime/AIModule/Private/EnvironmentQuery/Tests/EnvQueryTest_Pathfinding.cpp b/Engine/Source/Runtime/AIModule/Private/EnvironmentQuery/Tests/EnvQueryTest_Pathfinding.cpp
--- a/Engine/Source/Runtime/AIModule/Private/EnvironmentQuery/Tests/EnvQueryTest_Pathfinding.cpp
+++ b/Engine/Source/Runtime/AIModule/Private/EnvironmentQuery/Tests/EnvQueryTest_Pathfinding.cpp
@@ -3,8 +3,8 @@
 #include "EnvironmentQuery/Tests/EnvQueryTest_Pathfinding.h"
 #include "AI/Navigation/NavAgentInterface.h"
 #include "Engine/World.h"
-#include "AI/Navigation/NavigationData.h"
-#include "AI/Navigation/NavigationSystem.h"
+#include "NavigationData.h"
+#include "NavigationSystem.h"
 #include "EnvironmentQuery/Items/EnvQueryItemType_VectorBase.h"
 #include "EnvironmentQuery/Contexts/EnvQueryContext_Querier.h"
 
@@ -39,7 +39,7 @@ void UEnvQueryTest_Pathfinding::RunTest(FEnvQueryInstance& QueryInstance) const
 	float MinThresholdValue = FloatValueMin.GetValue();
 	float MaxThresholdValue = FloatValueMax.GetValue();
 
-	UNavigationSystem* NavSys = QueryInstance.World->GetNavigationSystem();
+	UNavigationSystemV1* NavSys = FNavigationSystem::GetCurrent<UNavigationSystemV1>(QueryInstance.World);
 	if (NavSys == nullptr || QueryOwner == nullptr)
 	{
 		return;
@@ -166,7 +166,7 @@ void UEnvQueryTest_Pathfinding::PostLoad()
 	SetWorkOnFloatValues(TestMode != EEnvTestPathfinding::PathExist);
 }
 
-bool UEnvQueryTest_Pathfinding::TestPathFrom(const FVector& ItemPos, const FVector& ContextPos, EPathFindingMode::Type Mode, const ANavigationData& NavData, UNavigationSystem& NavSys, FSharedConstNavQueryFilter NavFilter, const UObject* PathOwner) const
+bool UEnvQueryTest_Pathfinding::TestPathFrom(const FVector& ItemPos, const FVector& ContextPos, EPathFindingMode::Type Mode, const ANavigationData& NavData, UNavigationSystemV1& NavSys, FSharedConstNavQueryFilter NavFilter, const UObject* PathOwner) const
 {
 	FPathFindingQuery Query(PathOwner, NavData, ItemPos, ContextPos, NavFilter);
 	Query.SetAllowPartialPaths(false);
@@ -175,7 +175,7 @@ bool UEnvQueryTest_Pathfinding::TestPathFrom(const FVector& ItemPos, const FVect
 	return bPathExists;
 }
 
-bool UEnvQueryTest_Pathfinding::TestPathTo(const FVector& ItemPos, const FVector& ContextPos, EPathFindingMode::Type Mode, const ANavigationData& NavData, UNavigationSystem& NavSys, FSharedConstNavQueryFilter NavFilter, const UObject* PathOwner) const
+bool UEnvQueryTest_Pathfinding::TestPathTo(const FVector& ItemPos, const FVector& ContextPos, EPathFindingMode::Type Mode, const ANavigationData& NavData, UNavigationSystemV1& NavSys, FSharedConstNavQueryFilter NavFilter, const UObject* PathOwner) const
 {
 	FPathFindingQuery Query(PathOwner, NavData, ContextPos, ItemPos, NavFilter);
 	Query.SetAllowPartialPaths(false);
@@ -184,7 +184,7 @@ bool UEnvQueryTest_Pathfinding::TestPathTo(const FVector& ItemPos, const FVector
 	return bPathExists;
 }
 
-float UEnvQueryTest_Pathfinding::FindPathCostFrom(const FVector& ItemPos, const FVector& ContextPos, EPathFindingMode::Type Mode, const ANavigationData& NavData, UNavigationSystem& NavSys, FSharedConstNavQueryFilter NavFilter, const UObject* PathOwner) const
+float UEnvQueryTest_Pathfinding::FindPathCostFrom(const FVector& ItemPos, const FVector& ContextPos, EPathFindingMode::Type Mode, const ANavigationData& NavData, UNavigationSystemV1& NavSys, FSharedConstNavQueryFilter NavFilter, const UObject* PathOwner) const
 {
 	FPathFindingQuery Query(PathOwner, NavData, ItemPos, ContextPos, NavFilter);
 	Query.SetAllowPartialPaths(false);
@@ -193,7 +193,7 @@ float UEnvQueryTest_Pathfinding::FindPathCostFrom(const FVector& ItemPos, const
 	return (Result.IsSuccessful()) ? Result.Path->GetCost() : BIG_NUMBER;
 }
 
-float UEnvQueryTest_Pathfinding::FindPathCostTo(const FVector& ItemPos, const FVector& ContextPos, EPathFindingMode::Type Mode, const ANavigationData& NavData, UNavigationSystem& NavSys, FSharedConstNavQueryFilter NavFilter, const UObject* PathOwner) const
+float UEnvQueryTest_Pathfinding::FindPathCostTo(const FVector& ItemPos, const FVector& ContextPos, EPathFindingMode::Type Mode, const ANavigationData& NavData, UNavigationSystemV1& NavSys, FSharedConstNavQueryFilter NavFilter, const UObject* PathOwner) const
 {
 	FPathFindingQuery Query(PathOwner, NavData, ContextPos, ItemPos, NavFilter);
 	Query.SetAllowPartialPaths(false);
@@ -202,7 +202,7 @@ float UEnvQueryTest_Pathfinding::FindPathCostTo(const FVector& ItemPos, const FV
 	return (Result.IsSuccessful()) ? Result.Path->GetCost() : BIG_NUMBER;
 }
 
-float UEnvQueryTest_Pathfinding::FindPathLengthFrom(const FVector& ItemPos, const FVector& ContextPos, EPathFindingMode::Type Mode, const ANavigationData& NavData, UNavigationSystem& NavSys, FSharedConstNavQueryFilter NavFilter, const UObject* PathOwner) const
+float UEnvQueryTest_Pathfinding::FindPathLengthFrom(const FVector& ItemPos, const FVector& ContextPos, EPathFindingMode::Type Mode, const ANavigationData& NavData, UNavigationSystemV1& NavSys, FSharedConstNavQueryFilter NavFilter, const UObject* PathOwner) const
 {
 	FPathFindingQuery Query(PathOwner, NavData, ItemPos, ContextPos, NavFilter);
 	Query.SetAllowPartialPaths(false);
@@ -211,7 +211,7 @@ float UEnvQueryTest_Pathfinding::FindPathLengthFrom(const FVector& ItemPos, cons
 	return (Result.IsSuccessful()) ? Result.Path->GetLength() : BIG_NUMBER;
 }
 
-float UEnvQueryTest_Pathfinding::FindPathLengthTo(const FVector& ItemPos, const FVector& ContextPos, EPathFindingMode::Type Mode, const ANavigationData& NavData, UNavigationSystem& NavSys, FSharedConstNavQueryFilter NavFilter, const UObject* PathOwner) const
+float UEnvQueryTest_Pathfinding::FindPathLengthTo(const FVector& ItemPos, const FVector& ContextPos, EPathFindingMode::Type Mode, const ANavigationData& NavData, UNavigationSystemV1& NavSys, FSharedConstNavQueryFilter NavFilter, const UObject* PathOwner) const
 {
 	FPathFindingQuery Query(PathOwner, NavData, ContextPos, ItemPos, NavFilter);
 	Query.SetAllowPartialPaths(false);
@@ -220,7 +220,7 @@ float UEnvQueryTest_Pathfinding::FindPathLengthTo(const FVector& ItemPos, const
 	return (Result.IsSuccessful()) ? Result.Path->GetLength() : BIG_NUMBER;
 }
 
-ANavigationData* UEnvQueryTest_Pathfinding::FindNavigationData(UNavigationSystem& NavSys, UObject* Owner) const
+ANavigationData* UEnvQueryTest_Pathfinding::FindNavigationData(UNavigationSystemV1& NavSys, UObject* Owner) const
 {
 	INavAgentInterface* NavAgent = Cast<INavAgentInterface>(Owner);
 	if (NavAgent)
@@ -228,7 +228,7 @@ ANavigationData* UEnvQueryTest_Pathfinding::FindNavigationData(UNavigationSystem
 		return NavSys.GetNavDataForProps(NavAgent->GetNavAgentPropertiesRef());
 	}
 
-	return NavSys.GetMainNavData(FNavigationSystem::DontCreate);
+	return NavSys.GetDefaultNavDataInstance(FNavigationSystem::DontCreate);
 }
 
 #undef LOCTEXT_NAMESPACE
diff --git a/Engine/Source/Runtime/AIModule/Private/EnvironmentQuery/Tests/EnvQueryTest_PathfindingBatch.cpp b/Engine/Source/Runtime/AIModule/Private/EnvironmentQuery/Tests/EnvQueryTest_PathfindingBatch.cpp
--- a/Engine/Source/Runtime/AIModule/Private/EnvironmentQuery/Tests/EnvQueryTest_PathfindingBatch.cpp
+++ b/Engine/Source/Runtime/AIModule/Private/EnvironmentQuery/Tests/EnvQueryTest_PathfindingBatch.cpp
@@ -2,8 +2,9 @@
 
 #include "EnvironmentQuery/Tests/EnvQueryTest_PathfindingBatch.h"
 #include "Engine/World.h"
-#include "AI/Navigation/RecastNavMesh.h"
-#include "AI/Navigation/RecastQueryFilter.h"
+#include "NavigationSystem.h"
+#include "NavMesh/RecastNavMesh.h"
+#include "NavMesh/RecastQueryFilter.h"
 
 #define LOCTEXT_NAMESPACE "EnvQueryGenerator"
 
@@ -67,7 +68,7 @@ void UEnvQueryTest_PathfindingBatch::RunTest(FEnvQueryInstance& QueryInstance) c
 	float MaxThresholdValue = FloatValueMax.GetValue();
 	float RangeMultiplierValue = ScanRangeMultiplier.GetValue();
 
-	UNavigationSystem* NavSys = QueryInstance.World->GetNavigationSystem();
+	UNavigationSystemV1* NavSys = FNavigationSystem::GetCurrent<UNavigationSystemV1>(QueryInstance.World);
 	if (NavSys == nullptr || QueryOwner == nullptr)
 	{
 		return;
diff --git a/Engine/Source/Runtime/AIModule/Private/GameplayDebugger/GameplayDebuggerCategory_AI.cpp b/Engine/Source/Runtime/AIModule/Private/GameplayDebugger/GameplayDebuggerCategory_AI.cpp
--- a/Engine/Source/Runtime/AIModule/Private/GameplayDebugger/GameplayDebuggerCategory_AI.cpp
+++ b/Engine/Source/Runtime/AIModule/Private/GameplayDebugger/GameplayDebuggerCategory_AI.cpp
@@ -8,11 +8,11 @@
 #include "ShowFlags.h"
 #include "PrimitiveViewRelevance.h"
 #include "Components/PrimitiveComponent.h"
-#include "AI/Navigation/NavigationSystem.h"
+#include "NavigationSystem.h"
 #include "GameFramework/PlayerController.h"
 #include "Engine/Engine.h"
 #include "EngineGlobals.h"
-#include "AI/Navigation/RecastNavMesh.h"
+#include "NavMesh/RecastNavMesh.h"
 #include "Engine/Canvas.h"
 #include "AIController.h"
 #include "BehaviorTree/BehaviorTreeComponent.h"
@@ -178,7 +178,7 @@ void FGameplayDebuggerCategory_AI::CollectData(APlayerController* OwnerPC, AActo
 
 		DataPack.MontageInfo = MyChar ? GetNameSafe(MyChar->GetCurrentMontage()) : FString();
 
-		UNavigationSystem* NavSys = UNavigationSystem::GetCurrent(MyPawn->GetWorld());
+		UNavigationSystemV1* NavSys = FNavigationSystem::GetCurrent<UNavigationSystemV1>(MyPawn->GetWorld());
 		const ANavigationData* NavData = MyController && NavSys ? NavSys->GetNavDataForProps(MyController->GetNavAgentPropertiesRef()) : nullptr;
 		DataPack.NavDataInfo = NavData ? NavData->GetConfig().Name.ToString() : FString();
 
diff --git a/Engine/Source/Runtime/AIModule/Private/GameplayDebugger/GameplayDebuggerCategory_Navmesh.cpp b/Engine/Source/Runtime/AIModule/Private/GameplayDebugger/GameplayDebuggerCategory_Navmesh.cpp
--- a/Engine/Source/Runtime/AIModule/Private/GameplayDebugger/GameplayDebuggerCategory_Navmesh.cpp
+++ b/Engine/Source/Runtime/AIModule/Private/GameplayDebugger/GameplayDebuggerCategory_Navmesh.cpp
@@ -4,9 +4,9 @@
 
 #if WITH_GAMEPLAY_DEBUGGER
 
-#include "AI/Navigation/NavigationSystem.h"
+#include "NavigationSystem.h"
 #include "GameFramework/PlayerController.h"
-#include "AI/Navigation/RecastNavMesh.h"
+#include "NavMesh/RecastNavMesh.h"
 
 FGameplayDebuggerCategory_Navmesh::FGameplayDebuggerCategory_Navmesh()
 {
@@ -25,18 +25,18 @@ TSharedRef<FGameplayDebuggerCategory> FGameplayDebuggerCategory_Navmesh::MakeIns
 void FGameplayDebuggerCategory_Navmesh::CollectData(APlayerController* OwnerPC, AActor* DebugActor)
 {
 #if WITH_RECAST
-	ARecastNavMesh* NavData = nullptr;
+	const ARecastNavMesh* NavData = nullptr;
 
 	APawn* PlayerPawn = OwnerPC ? OwnerPC->GetPawnOrSpectator() : nullptr;
 	APawn* DebugActorAsPawn = Cast<APawn>(DebugActor);
 	APawn* DestPawn = DebugActorAsPawn ? DebugActorAsPawn : PlayerPawn;
 	if (OwnerPC && DestPawn)
 	{
-		UNavigationSystem* NavSys = UNavigationSystem::GetCurrent(OwnerPC->GetWorld());
+		UNavigationSystemV1* NavSys = FNavigationSystem::GetCurrent<UNavigationSystemV1>(OwnerPC->GetWorld());
 		if (NavSys) 
 		{
 			const FNavAgentProperties& NavAgentProperties = DestPawn->GetNavAgentPropertiesRef();
-			NavData = Cast<ARecastNavMesh>(NavSys->GetNavDataForProps(NavAgentProperties));
+			NavData = Cast<const ARecastNavMesh>(NavSys->GetNavDataForProps(NavAgentProperties));
 		}
 	}
 
diff --git a/Engine/Source/Runtime/AIModule/Private/GameplayDebugger/GameplayDebuggerCategory_Navmesh.h b/Engine/Source/Runtime/AIModule/Private/GameplayDebugger/GameplayDebuggerCategory_Navmesh.h
--- a/Engine/Source/Runtime/AIModule/Private/GameplayDebugger/GameplayDebuggerCategory_Navmesh.h
+++ b/Engine/Source/Runtime/AIModule/Private/GameplayDebugger/GameplayDebuggerCategory_Navmesh.h
@@ -6,7 +6,7 @@
 
 #include "CoreMinimal.h"
 #include "GameplayDebuggerCategory.h"
-#include "AI/Navigation/NavMeshRenderingComponent.h"
+#include "NavMesh/NavMeshRenderingComponent.h"
 
 class APlayerController;
 
diff --git a/Engine/Source/Runtime/AIModule/Private/Navigation/CrowdFollowingComponent.cpp b/Engine/Source/Runtime/AIModule/Private/Navigation/CrowdFollowingComponent.cpp
--- a/Engine/Source/Runtime/AIModule/Private/Navigation/CrowdFollowingComponent.cpp
+++ b/Engine/Source/Runtime/AIModule/Private/Navigation/CrowdFollowingComponent.cpp
@@ -1,14 +1,14 @@
 // Copyright 1998-2018 Epic Games, Inc. All Rights Reserved.
 
 #include "Navigation/CrowdFollowingComponent.h"
-#include "AI/Navigation/NavigationSystem.h"
-#include "AI/Navigation/RecastNavMesh.h"
+#include "NavigationSystem.h"
+#include "NavMesh/RecastNavMesh.h"
 #include "VisualLogger/VisualLoggerTypes.h"
 #include "VisualLogger/VisualLogger.h"
 #include "AIModuleLog.h"
 #include "Navigation/CrowdManager.h"
-#include "AI/Navigation/NavAreas/NavArea.h"
-#include "AI/Navigation/AbstractNavData.h"
+#include "NavAreas/NavArea.h"
+#include "AbstractNavData.h"
 #include "AIConfig.h"
 #include "Navigation/MetaNavMeshPath.h"
 #include "GameFramework/CharacterMovementComponent.h"
@@ -474,7 +474,7 @@ void UCrowdFollowingComponent::Initialize()
 		if (!bRegisteredWithCrowdSimulation)
 		{
 			// crowd manager might not be created yet if this component was spawned during level's begin play (possessing a pawn placed in level)
-			UNavigationSystem* NavSys = UNavigationSystem::GetCurrent(GetWorld());
+			UNavigationSystemV1* NavSys = FNavigationSystem::GetCurrent<UNavigationSystemV1>(GetWorld());
 			if (NavSys && !NavSys->IsInitialized())
 			{
 				NavSys->OnNavigationInitDone.AddUObject(this, &UCrowdFollowingComponent::OnNavigationInitDone);
@@ -734,14 +734,14 @@ void LogPathPartHelper(AActor* LogOwner, FNavMeshPath* NavMeshPath, int32 StartI
 			CenterPt /= Verts.Num();
 
 			const UNavArea* DefArea = AreaClass ? ((UClass*)AreaClass)->GetDefaultObject<UNavArea>() : NULL;
-			const FColor PolygonColor = AreaClass != UNavigationSystem::GetDefaultWalkableArea() ? (DefArea ? DefArea->DrawColor : NavMesh->GetConfig().Color) : FColorList::LightSteelBlue;
+			const FColor PolygonColor = AreaClass != UNavigationSystemV1::GetDefaultWalkableArea() ? (DefArea ? DefArea->DrawColor : NavMesh->GetConfig().Color) : FColorList::LightSteelBlue;
 
 			CorridorPoly.SetColor(PolygonColor.WithAlpha(100));
 			CorridorPoly.Points.Reset();
 			CorridorPoly.Points.Append(Verts);
 			Snapshot->ElementsToDraw.Add(CorridorPoly);
 
-			if (AreaClass && AreaClass != UNavigationSystem::GetDefaultWalkableArea())
+			if (AreaClass && AreaClass != UNavigationSystemV1::GetDefaultWalkableArea())
 			{
 				FVisualLogShapeElement AreaMarkElem(EVisualLoggerShapeElement::Segment);
 				AreaMarkElem.SetColor(FColorList::Orange.WithAlpha(100));
diff --git a/Engine/Source/Runtime/AIModule/Private/Navigation/CrowdManager.cpp b/Engine/Source/Runtime/AIModule/Private/Navigation/CrowdManager.cpp
--- a/Engine/Source/Runtime/AIModule/Private/Navigation/CrowdManager.cpp
+++ b/Engine/Source/Runtime/AIModule/Private/Navigation/CrowdManager.cpp
@@ -4,17 +4,17 @@
 #include "GameFramework/Actor.h"
 #include "GameFramework/Pawn.h"
 #include "GameFramework/Controller.h"
-#include "AI/Navigation/NavigationSystem.h"
-#include "AI/Navigation/RecastNavMesh.h"
+#include "NavigationSystem.h"
+#include "NavMesh/RecastNavMesh.h"
 #include "VisualLogger/VisualLogger.h"
 #include "AIModuleLog.h"
 
 #if WITH_RECAST
-#include "AI/Navigation/RecastHelpers.h"
+#include "NavMesh/RecastHelpers.h"
 #include "DetourCrowd/DetourObstacleAvoidance.h"
 #include "DetourCrowd/DetourCrowd.h"
 #include "Detour/DetourNavMesh.h"
-#include "AI/Navigation/RecastQueryFilter.h"
+#include "NavMesh/RecastQueryFilter.h"
 #endif
 
 #include "Navigation/CrowdFollowingComponent.h"
@@ -691,7 +691,7 @@ void UCrowdManager::AddAgent(const ICrowdAgentInterface* Agent, FCrowdAgentData&
 	const UCrowdFollowingComponent* CrowdComponent = Cast<const UCrowdFollowingComponent>(Agent);
 	if (CrowdComponent)
 	{
-		UNavigationSystem* NavSys = Cast<UNavigationSystem>(GetOuter());
+		UNavigationSystemV1* NavSys = Cast<UNavigationSystemV1>(GetOuter());
 		MyLinkFilter = MakeShareable(new FRecastSpeciaLinkFilter(NavSys, CrowdComponent->GetOuter()));
 	}
 
@@ -805,7 +805,7 @@ void UCrowdManager::ApplyVelocity(UCrowdFollowingComponent* AgentComponent, int3
 
 void UCrowdManager::UpdateAgentPaths()
 {
-	UNavigationSystem* NavSys = Cast<UNavigationSystem>(GetOuter());
+	UNavigationSystemV1* NavSys = Cast<UNavigationSystemV1>(GetOuter());
 	ARecastNavMesh* RecastNavData = Cast<ARecastNavMesh>(MyNavData);
 	if (RecastNavData == NULL)
 	{
@@ -1255,7 +1255,7 @@ void UCrowdManager::UpdateNavData()
 {
 	if (MyNavData == NULL)
 	{
-		UNavigationSystem* NavSys = Cast<UNavigationSystem>(GetOuter());
+		UNavigationSystemV1* NavSys = Cast<UNavigationSystemV1>(GetOuter());
 		if (NavSys)
 		{
 			for (int32 Idx = 0; Idx < NavSys->NavDataSet.Num(); Idx++)
@@ -1372,18 +1372,18 @@ void UCrowdManager::PostMovePointUpdate()
 
 UWorld* UCrowdManager::GetWorld() const
 {
-	UNavigationSystem* NavSys = Cast<UNavigationSystem>(GetOuter());
+	UNavigationSystemV1* NavSys = Cast<UNavigationSystemV1>(GetOuter());
 	return NavSys ? NavSys->GetWorld() : NULL;
 }
 
 UCrowdManager* UCrowdManager::GetCurrent(UObject* WorldContextObject)
 {
-	UNavigationSystem* NavSys = UNavigationSystem::GetCurrent(WorldContextObject);
+	UNavigationSystemV1* NavSys = FNavigationSystem::GetCurrent<UNavigationSystemV1>(WorldContextObject);
 	return NavSys ? Cast<UCrowdManager>(NavSys->GetCrowdManager()) : NULL;
 }
 
 UCrowdManager* UCrowdManager::GetCurrent(UWorld* World)
 {
-	UNavigationSystem* NavSys = UNavigationSystem::GetCurrent(World);
+	UNavigationSystemV1* NavSys = FNavigationSystem::GetCurrent<UNavigationSystemV1>(World);
 	return NavSys ? Cast<UCrowdManager>(NavSys->GetCrowdManager()) : NULL;
 }
diff --git a/Engine/Source/Runtime/AIModule/Private/Navigation/MetaNavMeshPath.cpp b/Engine/Source/Runtime/AIModule/Private/Navigation/MetaNavMeshPath.cpp
--- a/Engine/Source/Runtime/AIModule/Private/Navigation/MetaNavMeshPath.cpp
+++ b/Engine/Source/Runtime/AIModule/Private/Navigation/MetaNavMeshPath.cpp
@@ -2,10 +2,11 @@
 
 #include "Navigation/MetaNavMeshPath.h"
 #include "GameFramework/Controller.h"
-#include "AI/Navigation/NavigationSystem.h"
+#include "NavigationSystem.h"
 #include "VisualLogger/VisualLoggerTypes.h"
 #include "VisualLogger/VisualLogger.h"
-#include "AI/Navigation/NavAreas/NavArea.h"
+#include "NavAreas/NavArea.h"
+#include "NavigationSystemTypes.h"
 #include "DrawDebugHelpers.h"
 
 const FNavPathType FMetaNavMeshPath::Type(&FMetaNavMeshPath::Super::Type);
@@ -35,8 +36,8 @@ FMetaNavMeshPath::FMetaNavMeshPath(const TArray<FMetaPathWayPoint>& InWaypoints,
 FMetaNavMeshPath::FMetaNavMeshPath(const TArray<FMetaPathWayPoint>& InWaypoints, const AController& Owner)
 	: FMetaNavMeshPath()
 {
-	UNavigationSystem* NavSys = UNavigationSystem::GetCurrent(Owner.GetWorld());
-	const ANavigationData* NavData = NavSys ? NavSys->GetNavDataForProps(Owner.GetNavAgentPropertiesRef()) : nullptr;
+	UNavigationSystemV1* NavSys = FNavigationSystem::GetCurrent<UNavigationSystemV1>(Owner.GetWorld());
+	const ANavigationData* NavData = NavSys ? Cast<const ANavigationData>(NavSys->GetNavDataForProps(Owner.GetNavAgentPropertiesRef())) : (const ANavigationData*)nullptr;
 
 	if (ensure(NavData))
 	{
@@ -57,8 +58,8 @@ FMetaNavMeshPath::FMetaNavMeshPath(const TArray<FVector>& InWaypoints, const ANa
 
 FMetaNavMeshPath::FMetaNavMeshPath(const TArray<FVector>& InWaypoints, const AController& Owner) : FMetaNavMeshPath()
 {
-	UNavigationSystem* NavSys = UNavigationSystem::GetCurrent(Owner.GetWorld());
-	const ANavigationData* NavData = NavSys ? NavSys->GetNavDataForProps(Owner.GetNavAgentPropertiesRef()) : nullptr;
+	UNavigationSystemV1* NavSys = FNavigationSystem::GetCurrent<UNavigationSystemV1>(Owner.GetWorld());
+	const ANavigationData* NavData = Cast<const ANavigationData>(NavSys ? NavSys->GetNavDataForProps(Owner.GetNavAgentPropertiesRef()) : nullptr);
 
 	if (ensure(NavData))
 	{
@@ -221,7 +222,7 @@ float FMetaNavMeshPath::GetLengthFromPosition(FVector SegmentStart, uint32 NextP
 float FMetaNavMeshPath::GetCostFromIndex(int32 PathPointIndex) const
 {
 	// return approximation of full path * default cost, there's not enough data to give accurate value
-	const UNavArea* DefaultAreaOb = UNavigationSystem::GetDefaultWalkableArea().GetDefaultObject();
+	const UNavArea* DefaultAreaOb = UNavigationSystemV1::GetDefaultWalkableArea().GetDefaultObject();
 	const float DefaultAreaCost = DefaultAreaOb ? DefaultAreaOb->DefaultCost : 1.0f;
 	return ApproximateLength * DefaultAreaCost;
 }
diff --git a/Engine/Source/Runtime/AIModule/Private/Navigation/NavLocalGridData.cpp b/Engine/Source/Runtime/AIModule/Private/Navigation/NavLocalGridData.cpp
--- a/Engine/Source/Runtime/AIModule/Private/Navigation/NavLocalGridData.cpp
+++ b/Engine/Source/Runtime/AIModule/Private/Navigation/NavLocalGridData.cpp
@@ -3,7 +3,7 @@
 #include "Navigation/NavLocalGridData.h"
 #include "Navigation/NavLocalGridManager.h"
 #include "GraphAStar.h"
-#include "AI/Navigation/NavigationPath.h"
+#include "NavigationPath.h"
 
 namespace FGridMathHelpers
 {
diff --git a/Engine/Source/Runtime/AIModule/Private/Navigation/NavLocalGridManager.cpp b/Engine/Source/Runtime/AIModule/Private/Navigation/NavLocalGridManager.cpp
--- a/Engine/Source/Runtime/AIModule/Private/Navigation/NavLocalGridManager.cpp
+++ b/Engine/Source/Runtime/AIModule/Private/Navigation/NavLocalGridManager.cpp
@@ -2,7 +2,7 @@
 
 #include "Navigation/NavLocalGridManager.h"
 #include "AISystem.h"
-#include "AI/Navigation/NavigationSystem.h"
+#include "NavigationSystem.h"
 #include "Engine/Engine.h"
 
 float UNavLocalGridManager::GridCellSize = 50.0f;
@@ -199,8 +199,8 @@ void UNavLocalGridManager::ProjectGrids(const TArray<int32>& GridIndices)
 {
 	DECLARE_SCOPE_CYCLE_COUNTER(TEXT("NavGrid: Project"), STAT_GridProject, STATGROUP_AI);
 
-	const UNavigationSystem* NavSys = UNavigationSystem::GetCurrent(GetWorld());
-	const ANavigationData* NavData = NavSys ? NavSys->GetMainNavData() : nullptr;
+	const UNavigationSystemV1* NavSys = FNavigationSystem::GetCurrent<UNavigationSystemV1>(GetWorld());
+	const ANavigationData* NavData = NavSys ? NavSys->GetDefaultNavDataInstance() : nullptr;
 	
 	if (NavData)
 	{
diff --git a/Engine/Source/Runtime/AIModule/Private/Navigation/PathFollowingComponent.cpp b/Engine/Source/Runtime/AIModule/Private/Navigation/PathFollowingComponent.cpp
--- a/Engine/Source/Runtime/AIModule/Private/Navigation/PathFollowingComponent.cpp
+++ b/Engine/Source/Runtime/AIModule/Private/Navigation/PathFollowingComponent.cpp
@@ -5,16 +5,16 @@
 #include "TimerManager.h"
 #include "GameFramework/Pawn.h"
 #include "GameFramework/Controller.h"
-#include "AI/Navigation/NavigationSystem.h"
-#include "AI/Navigation/RecastNavMesh.h"
+#include "NavigationSystem.h"
+#include "NavMesh/RecastNavMesh.h"
 #include "AISystem.h"
 #include "BrainComponent.h"
 #include "Engine/Canvas.h"
 #include "AIController.h"
 #include "VisualLogger/VisualLoggerTypes.h"
 #include "VisualLogger/VisualLogger.h"
-#include "AI/Navigation/AbstractNavData.h"
-#include "AI/Navigation/NavLinkCustomInterface.h"
+#include "AbstractNavData.h"
+#include "NavLinkCustomInterface.h"
 #include "Navigation/MetaNavMeshPath.h"
 #include "AIConfig.h"
 
@@ -618,6 +618,22 @@ void UPathFollowingComponent::UpdateCachedComponents()
 	UpdateMovementComponent(/*bForce=*/true);
 }
 
+void UPathFollowingComponent::OnNewPawn(APawn* NewPawn)
+{
+	UpdateCachedComponents();
+}
+
+void UPathFollowingComponent::OnRegister()
+{
+	Super::OnRegister();
+	AController* ControllerOwner = Cast<AController>(GetOwner());
+
+	if (ControllerOwner)
+	{
+		ControllerOwner->GetOnNewPawnNotifier().AddUObject(this, &UPathFollowingComponent::OnNewPawn);
+	}
+}
+
 void UPathFollowingComponent::TickComponent(float DeltaTime, enum ELevelTick TickType, FActorComponentTickFunction *ThisTickFunction)
 {
 	Super::TickComponent(DeltaTime, TickType, ThisTickFunction);
@@ -644,21 +660,22 @@ void UPathFollowingComponent::SetMovementComponent(UNavMovementComponent* MoveCo
 	{
 		const FNavAgentProperties& NavAgentProps = MoveComp->GetNavAgentPropertiesRef();
 		MyDefaultAcceptanceRadius = NavAgentProps.AgentRadius * 0.1f;
-		MoveComp->PathFollowingComp = this;
+		MoveComp->SetPathFollowingAgent(Cast<IPathFollowingAgentInterface>(this));
 
-		UWorld* MyWorld = GetWorld();
-		if (MyWorld && MyWorld->GetNavigationSystem())
+		UNavigationSystemV1* NavSys = FNavigationSystem::GetCurrent<UNavigationSystemV1>(GetWorld());
+		
+		if (NavSys)
 		{
-			MyNavData = MyWorld->GetNavigationSystem()->GetNavDataForProps(NavAgentProps);
+			MyNavData = NavSys->GetNavDataForProps(NavAgentProps);
 			if (MyNavData == nullptr)
 			{
-				if (MyWorld->GetNavigationSystem()->IsInitialized() == false)
+				if (NavSys->IsInitialized() == false)
 				{
-					MyWorld->GetNavigationSystem()->OnNavigationInitDone.AddUObject(this, &UPathFollowingComponent::OnNavigationInitDone);
+					NavSys->OnNavigationInitDone.AddUObject(this, &UPathFollowingComponent::OnNavigationInitDone);
 				}
 				else
 				{
-					MyWorld->GetNavigationSystem()->OnNavDataRegisteredEvent.AddUniqueDynamic(this, &UPathFollowingComponent::OnNavDataRegistered);
+					NavSys->OnNavDataRegisteredEvent.AddUniqueDynamic(this, &UPathFollowingComponent::OnNavDataRegistered);
 				}
 			}
 		}
@@ -670,10 +687,11 @@ void UPathFollowingComponent::OnNavigationInitDone()
 	UWorld* MyWorld = GetWorld();
 	if (MovementComp && MyWorld)
 	{
-		check(MyWorld->GetNavigationSystem());
+		UNavigationSystemV1* NavSys = FNavigationSystem::GetCurrent<UNavigationSystemV1>(MyWorld);
+		check(NavSys);
 		const FNavAgentProperties& NavAgentProps = MovementComp->GetNavAgentPropertiesRef();
-		MyNavData = MyWorld->GetNavigationSystem()->GetNavDataForProps(NavAgentProps);
-		MyWorld->GetNavigationSystem()->OnNavigationInitDone.RemoveAll(this);
+		MyNavData = NavSys->GetNavDataForProps(NavAgentProps);
+		NavSys->OnNavigationInitDone.RemoveAll(this);
 	}
 }
 
@@ -686,8 +704,9 @@ void UPathFollowingComponent::OnNavDataRegistered(ANavigationData* NavData)
 		{
 			MyNavData = NavData;
 			UWorld* MyWorld = GetWorld();
-			check(MyWorld && MyWorld->GetNavigationSystem());
-			MyWorld->GetNavigationSystem()->OnNavDataRegisteredEvent.RemoveAll(this);
+			UNavigationSystemV1* NavSys = FNavigationSystem::GetCurrent<UNavigationSystemV1>(MyWorld);
+			check(NavSys);
+			NavSys->OnNavDataRegisteredEvent.RemoveAll(this);
 		}
 	}
 }
@@ -848,7 +867,7 @@ void UPathFollowingComponent::SetMoveSegment(int32 SegmentStartIndex)
 		// handle moving through custom nav links
 		if (PathPt0.CustomLinkId)
 		{
-			UNavigationSystem* NavSys = UNavigationSystem::GetCurrent(GetWorld());
+			UNavigationSystemV1* NavSys = FNavigationSystem::GetCurrent<UNavigationSystemV1>(GetWorld());
 			INavLinkCustomInterface* CustomNavLink = NavSys->GetCustomLink(PathPt0.CustomLinkId);
 			StartUsingCustomLink(CustomNavLink, SegmentEnd);
 		}
@@ -1559,8 +1578,8 @@ FNavLocation UPathFollowingComponent::GetCurrentNavLocation() const
 		return FNavLocation();
 	}
 
-	UWorld* MyWorld = GetWorld();
-	if (MyWorld == nullptr || MyWorld->GetNavigationSystem() == nullptr)
+	UNavigationSystemV1* NavSys = FNavigationSystem::GetCurrent<UNavigationSystemV1>(GetWorld());
+	if (NavSys == nullptr)
 	{
 		return FNavLocation();
 	}
@@ -1574,7 +1593,7 @@ FNavLocation UPathFollowingComponent::GetCurrentNavLocation() const
 		const AActor* OwnerActor = MovementComp->GetOwner();
 		const FVector OwnerExtent = OwnerActor ? OwnerActor->GetSimpleCollisionCylinderExtent() : FVector::ZeroVector;
 
-		MyWorld->GetNavigationSystem()->ProjectPointToNavigation(OwnerLoc, CurrentNavLocation, OwnerExtent, MyNavData);
+		NavSys->ProjectPointToNavigation(OwnerLoc, CurrentNavLocation, OwnerExtent, MyNavData);
 	}
 
 	return CurrentNavLocation;
@@ -1754,6 +1773,11 @@ bool UPathFollowingComponent::IsPathFollowingAllowed() const
 	return MovementComp && MovementComp->CanStartPathFollowing();
 }
 
+void UPathFollowingComponent::OnUnableToMove(const UObject& Instigator)
+{
+	AbortMove(Instigator, FPathFollowingResultFlags::MovementStop);
+}
+
 void UPathFollowingComponent::OnStartedFalling()
 {
 	bWalkingNavLinkStart = false;
diff --git a/Engine/Source/Runtime/AIModule/Private/Perception/AISense_Sight.cpp b/Engine/Source/Runtime/AIModule/Private/Perception/AISense_Sight.cpp
--- a/Engine/Source/Runtime/AIModule/Private/Perception/AISense_Sight.cpp
+++ b/Engine/Source/Runtime/AIModule/Private/Perception/AISense_Sight.cpp
@@ -633,10 +633,11 @@ void UAISense_Sight::RemoveAllQueriesToTarget(const FAISightTarget::FTargetId& T
 void UAISense_Sight::OnListenerForgetsActor(const FPerceptionListener& Listener, AActor& ActorToForget)
 {
 	const uint32 ListenerId = Listener.GetListenerID();
+	const uint32 TargetId = ActorToForget.GetUniqueID();
 	
 	for (FAISightQuery& SightQuery : SightQueryQueue)
 	{
-		if (SightQuery.ObserverId == ListenerId)
+		if (SightQuery.ObserverId == ListenerId && SightQuery.TargetId == TargetId)
 		{
 			// assuming one query per observer-target pair
 			SightQuery.ForgetPreviousResult();
diff --git a/Engine/Source/Runtime/AssetRegistry/Private/AssetDataGatherer.cpp b/Engine/Source/Runtime/AssetRegistry/Private/AssetDataGatherer.cpp
--- a/Engine/Source/Runtime/AssetRegistry/Private/AssetDataGatherer.cpp
+++ b/Engine/Source/Runtime/AssetRegistry/Private/AssetDataGatherer.cpp
@@ -386,7 +386,7 @@ FAssetDataGatherer::FAssetDataGatherer(const TArray<FString>& InPaths, const TAr
 	, bFinishedInitialDiscovery( false )
 	, Thread(nullptr)
 {
-	bGatherDependsData = (GIsEditor && !FParse::Param( FCommandLine::Get(), TEXT("NoDependsGathering") )) || FParse::Param(FCommandLine::Get(),TEXT("ForceDependsGathering")) ;
+	bGatherDependsData = (GIsEditor && !FParse::Param( FCommandLine::Get(), TEXT("NoDependsGathering") )) || FParse::Param(FCommandLine::Get(),TEXT("ForceDependsGathering"));
 
 	if (FParse::Param(FCommandLine::Get(), TEXT("NoAssetRegistryCache")) || FParse::Param(FCommandLine::Get(), TEXT("multiprocess")))
 	{
@@ -397,7 +397,7 @@ FAssetDataGatherer::FAssetDataGatherer(const TArray<FString>& InPaths, const TAr
 		if (AssetDataCacheMode == EAssetDataCacheMode::UseMonolithicCache)
 		{
 			bLoadAndSaveCache = true;
-			CacheFilename = FPaths::ProjectIntermediateDir() / TEXT("CachedAssetRegistry.bin");
+			CacheFilename = FPaths::ProjectIntermediateDir() / (bGatherDependsData ? TEXT("CachedAssetRegistry.bin") : TEXT("CachedAssetRegistryNoDeps.bin"));
 		}
 		else if (InPaths.Num() > 0)
 		{
@@ -409,7 +409,7 @@ FAssetDataGatherer::FAssetDataGatherer(const TArray<FString>& InPaths, const TAr
 			}
 
 			bLoadAndSaveCache = true;
-			CacheFilename = FPaths::ProjectIntermediateDir() / TEXT("AssetRegistryCache") / FString::Printf(TEXT("%08x.bin"), CacheHash);
+			CacheFilename = FPaths::ProjectIntermediateDir() / TEXT("AssetRegistryCache") / FString::Printf(TEXT("%08x%s.bin"), CacheHash, bGatherDependsData ? TEXT("") : TEXT("NoDeps"));
 		}
 	}
 
diff --git a/Engine/Source/Runtime/Core/Public/UObject/UnrealNames.inl b/Engine/Source/Runtime/Core/Public/UObject/UnrealNames.inl
--- a/Engine/Source/Runtime/Core/Public/UObject/UnrealNames.inl
+++ b/Engine/Source/Runtime/Core/Public/UObject/UnrealNames.inl
@@ -129,6 +129,7 @@ REGISTER_NAME(250,BSP)
 REGISTER_NAME(251,EditorSettings)
 REGISTER_NAME(252,AudioThread)
 REGISTER_NAME(253,ID)
+REGISTER_NAME(254,UserDefinedEnum)
 
 // Online
 REGISTER_NAME(280,DGram)
diff --git a/Engine/Source/Runtime/CoreUObject/Private/Blueprint/BlueprintSupport.cpp b/Engine/Source/Runtime/CoreUObject/Private/Blueprint/BlueprintSupport.cpp
--- a/Engine/Source/Runtime/CoreUObject/Private/Blueprint/BlueprintSupport.cpp
+++ b/Engine/Source/Runtime/CoreUObject/Private/Blueprint/BlueprintSupport.cpp
@@ -1860,8 +1860,23 @@ void FLinkerLoad::ResolveDeferredExports(UClass* LoadClass)
 			TGuardValue<int32> ClearDeferredCDOToPreventDeferExportCreation(DeferredCDOIndex, INDEX_NONE);
 			for(int32 ExportIndex : DeferredTemplateObjects)
 			{
-				ExportMap[ExportIndex].Object = nullptr;
-				CreateExport(ExportIndex);
+				FObjectExport& Export = ExportMap[ExportIndex];
+				ULinkerPlaceholderExportObject* PlaceholderExport = Cast<ULinkerPlaceholderExportObject>(Export.Object);
+				if (ensure(PlaceholderExport))
+				{
+					// replace the placeholder with the proper object instance
+					Export.Object = nullptr;
+					UObject* ExportObj = CreateExport(ExportIndex);
+
+					PlaceholderExport->ResolveAllPlaceholderReferences(ExportObj);
+					ResolvedDeferredSubobjects(PlaceholderExport);
+
+					PlaceholderExport->MarkPendingKill();
+					if (ExportObj != nullptr)
+					{
+						Preload(ExportObj);
+					}
+				}
 			}
 		}
 
diff --git a/Engine/Source/Runtime/CoreUObject/Private/UObject/CoreRedirects.cpp b/Engine/Source/Runtime/CoreUObject/Private/UObject/CoreRedirects.cpp
--- a/Engine/Source/Runtime/CoreUObject/Private/UObject/CoreRedirects.cpp
+++ b/Engine/Source/Runtime/CoreUObject/Private/UObject/CoreRedirects.cpp
@@ -445,7 +445,10 @@ bool FCoreRedirects::RedirectNameAndValues(ECoreRedirectFlags Type, const FCoreR
 				{
 					if (*FoundValueRedirect)
 					{
-						UE_LOG(LogLinker, Error, TEXT("RedirectNameAndValues(%s) found multiple conflicting value redirects, %s and %s!"), *OldObjectName.ToString(), *(*FoundValueRedirect)->OldName.ToString(), *Redirect->OldName.ToString());
+						if (!LegacyCompareEqual((*FoundValueRedirect)->ValueChanges, Redirect->ValueChanges))
+						{
+							UE_LOG(LogLinker, Error, TEXT("RedirectNameAndValues(%s) found multiple conflicting value redirects, %s and %s!"), *OldObjectName.ToString(), *(*FoundValueRedirect)->OldName.ToString(), *Redirect->OldName.ToString());
+						}
 					}
 					else
 					{
@@ -1009,6 +1012,11 @@ ECoreRedirectFlags FCoreRedirects::GetFlagsForTypeName(FName PackageName, FName
 		return ECoreRedirectFlags::Type_Class;
 	}
 
+	if (TypeName == NAME_UserDefinedEnum)
+	{
+		return ECoreRedirectFlags::Type_Enum;
+	}
+
 	return ECoreRedirectFlags::Type_Object;
 }
 
@@ -1059,9 +1067,9 @@ static void RegisterNativeRedirects40(TArray<FCoreRedirect>& Redirects)
 	CLASS_REDIRECT("MovementComp_Character", "CharacterMovementComponent");
 	CLASS_REDIRECT("MovementComp_Projectile", "ProjectileMovementComponent");
 	CLASS_REDIRECT("MovementComp_Rotating", "RotatingMovementComponent");
-	CLASS_REDIRECT("NavAreaDefault", "/Script/Engine.NavArea_Default");
-	CLASS_REDIRECT("NavAreaDefinition", "/Script/Engine.NavArea");
-	CLASS_REDIRECT("NavAreaNull", "/Script/Engine.NavArea_Null");
+	CLASS_REDIRECT("NavAreaDefault", "/Script/NavigationSystem.NavArea_Default");
+	CLASS_REDIRECT("NavAreaDefinition", "/Script/NavigationSystem.NavArea");
+	CLASS_REDIRECT("NavAreaNull", "/Script/NavigationSystem.NavArea_Null");
 	CLASS_REDIRECT("PhysicsActor", "StaticMeshActor");
 	CLASS_REDIRECT("PhysicsBSJointActor", "PhysicsConstraintActor");
 	CLASS_REDIRECT("PhysicsHingeActor", "PhysicsConstraintActor");
@@ -1098,7 +1106,7 @@ static void RegisterNativeRedirects40(TArray<FCoreRedirect>& Redirects)
 	CLASS_REDIRECT("VimGeneratedClass", "AnimBlueprintGeneratedClass");
 	CLASS_REDIRECT("VimInstance", "AnimInstance");
 	CLASS_REDIRECT("WorldInfo", "WorldSettings");
-	CLASS_REDIRECT_INSTANCES("NavAreaMeta", "/Script/Engine.NavArea_Default");
+	CLASS_REDIRECT_INSTANCES("NavAreaMeta", "/Script/NavigationSystem.NavArea_Default");
 
 	STRUCT_REDIRECT("VimDebugData", "AnimBlueprintDebugData");
 
@@ -1388,7 +1396,7 @@ static void RegisterNativeRedirects46(TArray<FCoreRedirect>& Redirects)
 	CLASS_REDIRECT("EmitterSpawnable", "Emitter");
 	CLASS_REDIRECT("SlateWidgetStyleAsset", "/Script/SlateCore.SlateWidgetStyleAsset");
 	CLASS_REDIRECT("SlateWidgetStyleContainerBase", "/Script/SlateCore.SlateWidgetStyleContainerBase");
-	CLASS_REDIRECT("SmartNavLinkComponent", "/Script/Engine.NavLinkCustomComponent");
+	CLASS_REDIRECT("SmartNavLinkComponent", "/Script/NavigationSystem.NavLinkCustomComponent");
 	CLASS_REDIRECT("WidgetBlueprint", "/Script/UMGEditor.WidgetBlueprint");
 
 	PROPERTY_REDIRECT("AnimNotify.Received_Notify.AnimSeq", "Animation");
diff --git a/Engine/Source/Runtime/CoreUObject/Private/UObject/LinkerLoad.cpp b/Engine/Source/Runtime/CoreUObject/Private/UObject/LinkerLoad.cpp
--- a/Engine/Source/Runtime/CoreUObject/Private/UObject/LinkerLoad.cpp
+++ b/Engine/Source/Runtime/CoreUObject/Private/UObject/LinkerLoad.cpp
@@ -1182,7 +1182,7 @@ FLinkerLoad::ELinkerStatus FLinkerLoad::SerializePackageFileSummary()
 				else if (SerializedCustomVersion.Version > LatestVersion->Version)
 				{
 					// Loading a package with a newer custom version than the current one.
-					UE_LOG(LogLinker, Error, TEXT("Package %s was saved with a newer custom version than the current. Tag %s  PackageVersion %d  MaxExpected %d"), *Filename, *SerializedCustomVersion.Key.ToString(), SerializedCustomVersion.Version, LatestVersion->Version);
+					UE_LOG(LogLinker, Error, TEXT("Package %s was saved with a newer custom version than the current. Tag %s Name '%s' PackageVersion %d  MaxExpected %d"), *Filename, *SerializedCustomVersion.Key.ToString(), *LatestVersion->GetFriendlyName().ToString(), SerializedCustomVersion.Version, LatestVersion->Version);
 					return LINKER_Failed;
 				}
 				else if (SerializedCustomVersion.Version != LatestVersion->Version)
@@ -2227,31 +2227,44 @@ FLinkerLoad::EVerifyResult FLinkerLoad::VerifyImport(int32 ImportIndex)
 				UObject* DestObject = Redir->DestinationObject;
 
 				// check to make sure the destination obj was loaded,
-				if ( DestObject == NULL )
+				if (DestObject == nullptr)
 				{
 					Result = VERIFY_Failed;
 				}
-				// check that in fact it was the type we thought it should be
-				else if ( DestObject->GetClass()->GetFName() != OriginalImport.ClassName
-
-					// if the destination object is a CDO, allow class changes
-					&&	!DestObject->HasAnyFlags(RF_ClassDefaultObject) )
+				else
 				{
-					Result = VERIFY_Failed;
-					// if the destination is a ObjectRedirector you've most likely made a nasty circular loop
-					if( Redir->DestinationObject->GetClass() == UObjectRedirector::StaticClass() )
+					// Blueprint CDOs are always allowed to change class, otherwise we need to do a name check for all parent classes
+					bool bIsValidClass = DestObject->HasAnyFlags(RF_ClassDefaultObject);
+					UClass* CheckClass = DestObject->GetClass();
+
+					while (!bIsValidClass && CheckClass)
 					{
-						WarningAppend += LOCTEXT("LoadWarningSuffix_circularredirection", " [circular redirection]").ToString();
+						if (CheckClass->GetFName() == OriginalImport.ClassName)
+						{
+							bIsValidClass = true;
+							break;
+						}
+						CheckClass = CheckClass->GetSuperClass();
 					}
-				}
-				else
-				{
-					Result = VERIFY_Redirected;
 
-					// now, fake our Import to be what the redirector pointed to
-					Import.XObject = Redir->DestinationObject;
-					FUObjectThreadContext::Get().ImportCount++;
-					FLinkerManager::Get().AddLoaderWithNewImports(this);
+					if (!bIsValidClass)
+					{
+						Result = VERIFY_Failed;
+						// if the destination is a ObjectRedirector you've most likely made a nasty circular loop
+						if (Redir->DestinationObject->GetClass() == UObjectRedirector::StaticClass())
+						{
+							WarningAppend += LOCTEXT("LoadWarningSuffix_circularredirection", " [circular redirection]").ToString();
+						}
+					}
+					else
+					{
+						Result = VERIFY_Redirected;
+
+						// now, fake our Import to be what the redirector pointed to
+						Import.XObject = Redir->DestinationObject;
+						FUObjectThreadContext::Get().ImportCount++;
+						FLinkerManager::Get().AddLoaderWithNewImports(this);
+					}
 				}
 			}
 		}
diff --git a/Engine/Source/Runtime/CoreUObject/Private/UObject/SavePackage.cpp b/Engine/Source/Runtime/CoreUObject/Private/UObject/SavePackage.cpp
--- a/Engine/Source/Runtime/CoreUObject/Private/UObject/SavePackage.cpp
+++ b/Engine/Source/Runtime/CoreUObject/Private/UObject/SavePackage.cpp
@@ -2519,6 +2519,11 @@ public:
 		}
 	}
 
+	/** Do nothing when serializing soft references, this is required because the presave on soft references can fix redirectors, which is unsafe at this point */
+	virtual FArchive& operator<<(FLazyObjectPtr& Value) override { return *this; }
+	virtual FArchive& operator<<(FSoftObjectPtr& Value) override { return *this; }
+	virtual FArchive& operator<<(FSoftObjectPath& Value) override { return *this; }
+
 private:
 
 	/**
diff --git a/Engine/Source/Runtime/CoreUObject/Private/UObject/SoftObjectPath.cpp b/Engine/Source/Runtime/CoreUObject/Private/UObject/SoftObjectPath.cpp
--- a/Engine/Source/Runtime/CoreUObject/Private/UObject/SoftObjectPath.cpp
+++ b/Engine/Source/Runtime/CoreUObject/Private/UObject/SoftObjectPath.cpp
@@ -7,6 +7,7 @@
 #include "Misc/PackageName.h"
 #include "UObject/LinkerLoad.h"
 #include "UObject/UObjectThreadContext.h"
+#include "UObject/CoreRedirects.h"
 #include "Misc/RedirectCollector.h"
 
 FSoftObjectPath::FSoftObjectPath(const UObject* InObject)
@@ -48,19 +49,19 @@ void FSoftObjectPath::SetPath(FString Path)
 		{
 			// Possibly an ExportText path. Trim the ClassName.
 			Path = FPackageName::ExportTextPathToObjectPath(Path);
-	}
+		}
 
 		int32 ColonIndex = INDEX_NONE;
 
 		if (Path.FindChar(':', ColonIndex))
-	{
+		{
 			// Has a subobject, split on that then create a name from the temporary path
 			SubPathString = Path.Mid(ColonIndex + 1);
 			Path.RemoveAt(ColonIndex, Path.Len() - ColonIndex);
 			AssetPathName = *Path;
-	}
-	else
-	{
+		}
+		else
+		{
 			// No Subobject
 			AssetPathName = *Path;
 			SubPathString.Empty();
@@ -71,6 +72,11 @@ void FSoftObjectPath::SetPath(FString Path)
 bool FSoftObjectPath::PreSavePath()
 {
 #if WITH_EDITOR
+	if (IsNull())
+	{
+		return false;
+	}
+
 	FName FoundRedirection = GRedirectCollector.GetAssetPathRedirection(AssetPathName);
 
 	if (FoundRedirection != NAME_None)
@@ -78,6 +84,11 @@ bool FSoftObjectPath::PreSavePath()
 		AssetPathName = FoundRedirection;
 		return true;
 	}
+
+	if (FixupCoreRedirects())
+	{
+		return true;
+	}
 #endif // WITH_EDITOR
 	return false;
 }
@@ -153,11 +164,11 @@ void FSoftObjectPath::SerializePath(FArchive& Ar)
 			PostLoadPath();
 		}
 		if (Ar.GetPortFlags()&PPF_DuplicateForPIE)
-	{
-		// Remap unique ID if necessary
-		// only for fixing up cross-level references, inter-level references handled in FDuplicateDataReader
-		FixupForPIE();
-	}
+		{
+			// Remap unique ID if necessary
+			// only for fixing up cross-level references, inter-level references handled in FDuplicateDataReader
+			FixupForPIE();
+		}
 	}
 #endif // WITH_EDITOR
 }
@@ -181,7 +192,7 @@ bool FSoftObjectPath::ExportTextItem(FString& ValueStr, FSoftObjectPath const& D
 		return false;
 	}
 
-	if (IsValid())
+	if (!IsNull())
 	{
 		// Fixup any redirectors
 		FSoftObjectPath Temp = *this;
@@ -232,8 +243,20 @@ bool FSoftObjectPath::ImportTextItem(const TCHAR*& Buffer, int32 PortFlags, UObj
 
 	SetPath(MoveTemp(ImportedPath));
 
-	// Consider this a load, so Config string references get cooked
-	PostLoadPath();
+#if WITH_EDITOR
+	if (Parent && IsEditorOnlyObject(Parent))
+	{
+		// We're probably reading config for an editor only object, we need to mark this reference as editor only
+		FSoftObjectPathSerializationScope SerializationScope(NAME_None, NAME_None, ESoftObjectPathCollectType::EditorOnlyCollect, ESoftObjectPathSerializeType::AlwaysSerialize);
+
+		PostLoadPath();
+	}
+	else
+#endif
+	{
+		// Consider this a load, so Config string references get cooked
+		PostLoadPath();
+	}
 
 	return true;
 }
@@ -259,7 +282,7 @@ bool SerializeFromMismatchedTagTemplate(FString& Output, const FPropertyTag& Tag
 			Output = ObjPtr->GetPathName();
 		}
 		else
-	{
+		{
 			Output = FString();
 		}
 		return true;
@@ -270,8 +293,8 @@ bool SerializeFromMismatchedTagTemplate(FString& Output, const FPropertyTag& Tag
 		Ar << String;
 
 		Output = String;
-	return true;
-}
+		return true;
+	}
 	return false;
 }
 
@@ -300,9 +323,22 @@ UObject* FSoftObjectPath::TryLoad() const
 {
 	UObject* LoadedObject = nullptr;
 
-	if ( IsValid() )
+	if (!IsNull())
 	{
 		LoadedObject = LoadObject<UObject>(nullptr, *ToString());
+
+#if WITH_EDITOR
+		// Look at core redirects if we didn't find the object
+		if (!LoadedObject)
+		{
+			FSoftObjectPath FixupObjectPath = *this;
+			if (FixupObjectPath.FixupCoreRedirects())
+			{
+				LoadedObject = LoadObject<UObject>(nullptr, *FixupObjectPath.ToString());
+			}
+		}
+#endif
+
 		while (UObjectRedirector* Redirector = Cast<UObjectRedirector>(LoadedObject))
 		{
 			LoadedObject = Redirector->DestinationObject;
@@ -316,7 +352,7 @@ UObject* FSoftObjectPath::ResolveObject() const
 {
 	// Don't try to resolve if we're saving a package because StaticFindObject can't be used here
 	// and we usually don't want to force references to weak pointers while saving.
-	if (!IsValid() || GIsSavingPackage)
+	if (IsNull() || GIsSavingPackage)
 	{
 		return nullptr;
 	}
@@ -327,9 +363,7 @@ UObject* FSoftObjectPath::ResolveObject() const
 	{
 		// If we are in PIE and this hasn't already been fixed up, we need to fixup at resolution time. We cannot modify the path as it may be somewhere like a blueprint CDO
 		FSoftObjectPath FixupObjectPath = *this;
-		FixupObjectPath.FixupForPIE();
-
-		if (FixupObjectPath.AssetPathName != AssetPathName)
+		if (FixupObjectPath.FixupForPIE())
 		{
 			PathString = FixupObjectPath.ToString();
 		}
@@ -337,6 +371,20 @@ UObject* FSoftObjectPath::ResolveObject() const
 #endif
 
 	UObject* FoundObject = FindObject<UObject>(nullptr, *PathString);
+
+#if WITH_EDITOR
+	// Look at core redirects if we didn't find the object
+	if (!FoundObject)
+	{
+		FSoftObjectPath FixupObjectPath = *this;
+		if (FixupObjectPath.FixupCoreRedirects())
+		{
+			PathString = FixupObjectPath.ToString();
+			FoundObject = FindObject<UObject>(nullptr, *PathString);
+		}
+	}
+#endif
+
 	while (UObjectRedirector* Redirector = Cast<UObjectRedirector>(FoundObject))
 	{
 		FoundObject = Redirector->DestinationObject;
@@ -361,10 +409,10 @@ void FSoftObjectPath::ClearPIEPackageNames()
 	PIEPackageNames.Empty();
 }
 
-void FSoftObjectPath::FixupForPIE()
+bool FSoftObjectPath::FixupForPIE()
 {
 #if WITH_EDITOR
-	if (GPlayInEditorID != INDEX_NONE && IsValid())
+	if (GPlayInEditorID != INDEX_NONE && !IsNull())
 	{
 		FString Path = ToString();
 
@@ -380,13 +428,36 @@ void FSoftObjectPath::FixupForPIE()
 
 			// Duplicate if this an already registered PIE package or this looks like a level subobject reference
 			if (bIsChildOfLevel || PIEPackageNames.Contains(PIEPackage))
-				{
-					// Need to prepend PIE prefix, as we're in PIE and this refers to an object in a PIE package
-					SetPath(MoveTemp(PIEPath));
+			{
+				// Need to prepend PIE prefix, as we're in PIE and this refers to an object in a PIE package
+				SetPath(MoveTemp(PIEPath));
+
+				return true;
 			}
 		}
 	}
 #endif
+	return false;
+}
+
+bool FSoftObjectPath::FixupCoreRedirects()
+{
+	FCoreRedirectObjectName OldName = FCoreRedirectObjectName(ToString());
+	FCoreRedirectObjectName NewName = FCoreRedirects::GetRedirectedName(ECoreRedirectFlags::Type_Object, OldName);
+
+	// This also might be a class
+	if (OldName == NewName)
+	{
+		NewName = FCoreRedirects::GetRedirectedName(ECoreRedirectFlags::Type_Class, OldName);
+	}
+
+	if (OldName != NewName)
+	{
+		SetPath(NewName.ToString());
+		return true;
+	}
+
+	return false;
 }
 
 bool FSoftClassPath::SerializeFromMismatchedTag(struct FPropertyTag const& Tag, FArchive& Ar)
@@ -406,15 +477,15 @@ bool FSoftClassPath::SerializeFromMismatchedTag(struct FPropertyTag const& Tag,
 	{
 		SetPath(MoveTemp(Path));
 		PostLoadPath();
-				}
+	}
 
 	return bReturn;
-			}
+}
 
 UClass* FSoftClassPath::ResolveClass() const
 {
 	return Cast<UClass>(ResolveObject());
-		}
+}
 
 FSoftClassPath FSoftClassPath::GetOrCreateIDForClass(const UClass *InClass)
 {
@@ -482,7 +553,7 @@ bool FSoftObjectPathThreadContext::GetSerializationOptions(FName& OutPackageName
 			if (bEditorOnly && CurrentCollectType == ESoftObjectPathCollectType::AlwaysCollect)
 			{
 				CurrentCollectType = ESoftObjectPathCollectType::EditorOnlyCollect;
-	}
+			}
 
 			bFoundAnything = true;
 		}
diff --git a/Engine/Source/Runtime/CoreUObject/Private/UObject/UObjectArchetype.cpp b/Engine/Source/Runtime/CoreUObject/Private/UObject/UObjectArchetype.cpp
--- a/Engine/Source/Runtime/CoreUObject/Private/UObject/UObjectArchetype.cpp
+++ b/Engine/Source/Runtime/CoreUObject/Private/UObject/UObjectArchetype.cpp
@@ -98,7 +98,7 @@ UObject* UObject::GetArchetypeFromRequiredInfo(UClass* Class, UObject* Outer, FN
 #if WITH_EDITOR
 	// While compiling we just want to use whatever is in the object hierarchy,
 	// as some instances within the hierarchy may also be compiling:
-	bUseUpToDateClass = Class->GetAuthoritativeClass() == Class && (!GCompilingBlueprint || GIsReinstancing);
+	bUseUpToDateClass = GIsReinstancing && Class->GetAuthoritativeClass() == Class;
 #endif
 	return GetArchetypeFromRequiredInfoImpl(Class, Outer, Name, ObjectFlags, bUseUpToDateClass);
 }
diff --git a/Engine/Source/Runtime/CoreUObject/Public/Misc/RedirectCollector.h b/Engine/Source/Runtime/CoreUObject/Public/Misc/RedirectCollector.h
--- a/Engine/Source/Runtime/CoreUObject/Public/Misc/RedirectCollector.h
+++ b/Engine/Source/Runtime/CoreUObject/Public/Misc/RedirectCollector.h
@@ -29,7 +29,8 @@ private:
 		bool operator==(const FPackagePropertyPair& Other) const
 		{
 			return Package == Other.Package &&
-				Property == Other.Property;
+				Property == Other.Property &&
+				bReferencedByEditorOnlyProperty == Other.bReferencedByEditorOnlyProperty;
 		}
 
 		const FName& GetCachedPackageName() const
diff --git a/Engine/Source/Runtime/CoreUObject/Public/UObject/Class.h b/Engine/Source/Runtime/CoreUObject/Public/UObject/Class.h
--- a/Engine/Source/Runtime/CoreUObject/Public/UObject/Class.h
+++ b/Engine/Source/Runtime/CoreUObject/Public/UObject/Class.h
@@ -2063,7 +2063,7 @@ class COREUOBJECT_API UClass : public UStruct
 #endif
 {
 	DECLARE_CASTED_CLASS_INTRINSIC_NO_CTOR(UClass, UStruct, 0, TEXT("/Script/CoreUObject"), CASTCLASS_UClass, NO_API)
-	DECLARE_WITHIN(UPackage)
+	DECLARE_WITHIN_UPACKAGE()
 
 public:
 #if UCLASS_FAST_ISA_IMPL == UCLASS_ISA_INDEXTREE
@@ -2697,7 +2697,7 @@ protected:
 class COREUOBJECT_API UDynamicClass : public UClass
 {
 	DECLARE_CASTED_CLASS_INTRINSIC_NO_CTOR(UDynamicClass, UClass, 0, TEXT("/Script/CoreUObject"), CASTCLASS_None, NO_API)
-	DECLARE_WITHIN(UPackage)
+	DECLARE_WITHIN_UPACKAGE()
 
 public:
 
diff --git a/Engine/Source/Runtime/CoreUObject/Public/UObject/ObjectMacros.h b/Engine/Source/Runtime/CoreUObject/Public/UObject/ObjectMacros.h
--- a/Engine/Source/Runtime/CoreUObject/Public/UObject/ObjectMacros.h
+++ b/Engine/Source/Runtime/CoreUObject/Public/UObject/ObjectMacros.h
@@ -1484,10 +1484,16 @@ public: \
 	DECLARE_CASTED_CLASS_INTRINSIC_WITH_API( TClass, TSuperClass, TStaticFlags, TPackage, TStaticCastFlags, NO_API) \
 
 // Declare that objects of class being defined reside within objects of the specified class.
-#define DECLARE_WITHIN( TWithinClass ) \
+#define DECLARE_WITHIN_INTERNAL( TWithinClass, bCanUseOnCDO ) \
 	/** The required type of this object's outer ({{ typedef-type }}) */ \
 	typedef class TWithinClass WithinClass;  \
-	TWithinClass* GetOuter##TWithinClass() const { return (TWithinClass*)GetOuter(); }
+	TWithinClass* GetOuter##TWithinClass() const { return (ensure(bCanUseOnCDO || !HasAnyFlags(RF_ClassDefaultObject)) ? (TWithinClass*)GetOuter() : nullptr); }
+
+#define DECLARE_WITHIN( TWithinClass ) \
+	DECLARE_WITHIN_INTERNAL( TWithinClass, false )
+
+#define DECLARE_WITHIN_UPACKAGE() \
+	DECLARE_WITHIN_INTERNAL( UPackage, true )
 
 // Register a class at startup time.
 #define IMPLEMENT_CLASS(TClass, TClassCrc) \
diff --git a/Engine/Source/Runtime/CoreUObject/Public/UObject/SoftObjectPath.h b/Engine/Source/Runtime/CoreUObject/Public/UObject/SoftObjectPath.h
--- a/Engine/Source/Runtime/CoreUObject/Public/UObject/SoftObjectPath.h
+++ b/Engine/Source/Runtime/CoreUObject/Public/UObject/SoftObjectPath.h
@@ -152,8 +152,11 @@ struct COREUOBJECT_API FSoftObjectPath
 	/** Handles when a string asset reference has been loaded, call if loading with a method that skips SerializePath. This does not modify path but might call callbacks */
 	void PostLoadPath() const;
 
-	/** Fixes up this SoftObjectPath to add or remove the PIE prefix depending on what is currently active */
-	void FixupForPIE();
+	/** Fixes up this SoftObjectPath to add or remove the PIE prefix depending on what is currently active, returns true if it was modified */
+	bool FixupForPIE();
+
+	/** Fixes soft object path for CoreRedirects to handle renamed native objects, returns true if it was modified */
+	bool FixupCoreRedirects();
 
 	FORCEINLINE friend uint32 GetTypeHash(FSoftObjectPath const& This)
 	{
diff --git a/Engine/Source/Runtime/Engine/Classes/AI/Navigation/NavAgentInterface.h b/Engine/Source/Runtime/Engine/Classes/AI/Navigation/NavAgentInterface.h
--- a/Engine/Source/Runtime/Engine/Classes/AI/Navigation/NavAgentInterface.h
+++ b/Engine/Source/Runtime/Engine/Classes/AI/Navigation/NavAgentInterface.h
@@ -9,6 +9,7 @@
 #include "NavAgentInterface.generated.h"
 
 class AActor;
+class IPathFollowingAgentInterface;
 
 UINTERFACE(MinimalAPI, meta=(CannotImplementInterfaceInBlueprint))
 class UNavAgentInterface : public UInterface
@@ -49,6 +50,9 @@ class INavAgentInterface
 	/** Checks if the agent is actively following a navigation path */
 	virtual bool IsFollowingAPath() const { return false; }
 
+	/** Retrieves PathFollowingAgent for this NavAgent */
+	virtual IPathFollowingAgentInterface* GetPathFollowingAgent() const { return nullptr; }
+
 	//----------------------------------------------------------------------//
 	// DEPRECATED
 	//----------------------------------------------------------------------//
diff --git a/Engine/Source/Runtime/Engine/Classes/AI/Navigation/NavAreaBase.h b/Engine/Source/Runtime/Engine/Classes/AI/Navigation/NavAreaBase.h
new file mode 100644
--- /dev/null
+++ b/Engine/Source/Runtime/Engine/Classes/AI/Navigation/NavAreaBase.h
@@ -0,0 +1,46 @@
+// Copyright 1998-2018 Epic Games, Inc. All Rights Reserved.
+
+#pragma once
+
+#include "UObject/Object.h"
+#include "AI/Navigation/NavigationTypes.h"
+#include "NavAreaBase.generated.h"
+
+
+// a stub class. Actual implementation in NavigationSystem module.
+UCLASS(DefaultToInstanced, abstract, Config = Engine)
+class ENGINE_API UNavAreaBase : public UObject
+{
+	GENERATED_BODY()
+
+protected:
+	uint8 bIsMetaArea : 1;
+
+public:
+	UNavAreaBase(const FObjectInitializer& ObjectInitializer = FObjectInitializer::Get());
+
+	// leftover from NavigationSystem extraction from the Engine code
+	virtual bool IsLowArea() const { return false; }
+	virtual bool IsMetaArea() const { return (bIsMetaArea != 0); }
+
+	/**
+	*	Picks an navigation area class that should be used for Actor when
+	*	queried by NavAgent. 
+	*/
+	static TSubclassOf<UNavAreaBase> PickAreaClassForAgent(TSubclassOf<UNavAreaBase> AreaClass, const AActor& Actor, const FNavAgentProperties& NavAgent)
+	{
+		const UNavAreaBase* CDO = AreaClass.Get() ? AreaClass->GetDefaultObject<const UNavAreaBase>() : nullptr;
+		return CDO && CDO->IsMetaArea()
+			? CDO->PickAreaClassForAgent(Actor, NavAgent)
+			: AreaClass;
+	}
+	
+
+protected:
+	/**
+	*	Picks an navigation area class that should be used for Actor when
+	*	queried by NavAgent. Call it via the UNavAreaBase::PickAreaClass
+	*/
+	virtual TSubclassOf<UNavAreaBase> PickAreaClassForAgent(const AActor& Actor, const FNavAgentProperties& NavAgent) const;
+};
+
diff --git a/Engine/Source/Runtime/Engine/Classes/AI/Navigation/NavLinkDefinition.h b/Engine/Source/Runtime/Engine/Classes/AI/Navigation/NavLinkDefinition.h
--- a/Engine/Source/Runtime/Engine/Classes/AI/Navigation/NavLinkDefinition.h
+++ b/Engine/Source/Runtime/Engine/Classes/AI/Navigation/NavLinkDefinition.h
@@ -10,6 +10,8 @@
 #include "AI/Navigation/NavigationTypes.h"
 #include "NavLinkDefinition.generated.h"
 
+class UNavAreaBase;
+
 UENUM()
 namespace ENavLinkDirection
 {
@@ -159,7 +161,7 @@ struct ENGINE_API FNavigationLinkBase
 private:
 	/** Area type of this link (empty = default) */
 	UPROPERTY(EditAnywhere, Category = Default)
-	TSubclassOf<class UNavArea> AreaClass;
+	TSubclassOf<UNavAreaBase> AreaClass;
 
 	TWeakObjectPtr<UClass> AreaClassOb;
 };
diff --git a/Engine/Source/Runtime/Engine/Classes/AI/Navigation/NavRelevantInterface.h b/Engine/Source/Runtime/Engine/Classes/AI/Navigation/NavRelevantInterface.h
--- a/Engine/Source/Runtime/Engine/Classes/AI/Navigation/NavRelevantInterface.h
+++ b/Engine/Source/Runtime/Engine/Classes/AI/Navigation/NavRelevantInterface.h
@@ -6,10 +6,113 @@
 #include "UObject/ObjectMacros.h"
 #include "UObject/Interface.h"
 #include "AI/Navigation/NavigationTypes.h"
+#include "AI/NavigationModifier.h"
 #include "NavRelevantInterface.generated.h"
 
 struct FNavigableGeometryExport;
-struct FNavigationRelevantData;
+
+struct FNavigationRelevantDataFilter 
+{
+	/** pass when actor has geometry */
+	uint32 bIncludeGeometry : 1;
+	/** pass when actor has any offmesh link modifier */
+	uint32 bIncludeOffmeshLinks : 1;
+	/** pass when actor has any area modifier */
+	uint32 bIncludeAreas : 1;
+	/** pass when actor has any modifier with meta area */
+	uint32 bIncludeMetaAreas : 1;
+
+	FNavigationRelevantDataFilter() 
+		: bIncludeGeometry(false)
+		, bIncludeOffmeshLinks(false)
+		, bIncludeAreas(false)
+		, bIncludeMetaAreas(false)
+	{}
+};
+
+// @todo consider optional structures that can contain a delegate instead of 
+// actual copy of collision data
+struct ENGINE_API FNavigationRelevantData : public TSharedFromThis<FNavigationRelevantData, ESPMode::ThreadSafe>
+{
+	DECLARE_DELEGATE_RetVal_OneParam(bool, FFilterNavDataDelegate, const struct FNavDataConfig*);
+
+	/** CollisionData should always start with this struct for validation purposes */
+	struct FCollisionDataHeader
+	{
+		int32 DataSize;
+
+		static bool IsValid(const uint8* RawData, int32 RawDataSize);
+	};
+
+	/** exported geometry (used by recast navmesh as FRecastGeometryCache) */
+	TNavStatArray<uint8> CollisionData;
+
+	/** cached voxels (used by recast navmesh as FRecastVoxelCache) */
+	TNavStatArray<uint8> VoxelData;
+
+	/** bounds of geometry (unreal coords) */
+	FBox Bounds;
+
+	/** Gathers per instance data for navigation geometry in a specified area box */
+	FNavDataPerInstanceTransformDelegate NavDataPerInstanceTransformDelegate;
+
+	/** called to check if hosted geometry should be used for given FNavDataConfig. If not set then "true" is assumed. */
+	FFilterNavDataDelegate ShouldUseGeometryDelegate;
+
+	/** additional modifiers: areas and external links */
+	FCompositeNavModifier Modifiers;
+
+	/** UObject these data represents */
+	const TWeakObjectPtr<UObject> SourceObject;
+
+	/** get set to true when lazy navigation exporting is enabled and this navigation data has "potential" of
+	*	containing geometry data. First access will result in gathering the data and setting this flag back to false.
+	*	Mind that this flag can go back to 'true' if related data gets cleared out. */
+	uint32 bPendingLazyGeometryGathering : 1;
+	uint32 bPendingLazyModifiersGathering : 1;
+
+	uint32 bSupportsGatheringGeometrySlices : 1;
+
+	FNavigationRelevantData(UObject& Source)
+		: SourceObject(&Source)
+		, bPendingLazyGeometryGathering(false)
+		, bPendingLazyModifiersGathering(false)
+	{}
+
+	FORCEINLINE bool HasGeometry() const { return VoxelData.Num() || CollisionData.Num(); }
+	FORCEINLINE bool HasModifiers() const { return !Modifiers.IsEmpty(); }
+	FORCEINLINE bool IsPendingLazyGeometryGathering() const { return bPendingLazyGeometryGathering; }
+	FORCEINLINE bool IsPendingLazyModifiersGathering() const { return bPendingLazyModifiersGathering; }
+	FORCEINLINE bool SupportsGatheringGeometrySlices() const { return bSupportsGatheringGeometrySlices; }
+	FORCEINLINE bool IsEmpty() const { return !HasGeometry() && !HasModifiers(); }
+	FORCEINLINE uint32 GetAllocatedSize() const { return CollisionData.GetAllocatedSize() + VoxelData.GetAllocatedSize() + Modifiers.GetAllocatedSize(); }
+	FORCEINLINE uint32 GetGeometryAllocatedSize() const { return CollisionData.GetAllocatedSize() + VoxelData.GetAllocatedSize(); }
+	FORCEINLINE int32 GetDirtyFlag() const
+	{
+		return ((HasGeometry() || IsPendingLazyGeometryGathering()) ? ENavigationDirtyFlag::Geometry : 0) |
+			((HasModifiers() || IsPendingLazyModifiersGathering()) ? ENavigationDirtyFlag::DynamicModifier : 0) |
+			(Modifiers.HasAgentHeightAdjust() ? ENavigationDirtyFlag::UseAgentHeight : 0);
+	}
+
+	bool HasPerInstanceTransforms() const;
+	bool IsMatchingFilter(const FNavigationRelevantDataFilter& Filter) const;
+	void Shrink();
+	bool IsCollisionDataValid() const;
+
+	void ValidateAndShrink()
+	{
+		if (IsCollisionDataValid())
+		{
+			Shrink();
+		}
+		else
+		{
+			CollisionData.Empty();
+		}
+	}
+
+	FORCEINLINE UObject* GetOwner() const { return SourceObject.Get(); }
+};
 
 UINTERFACE(MinimalAPI, meta=(CannotImplementInterfaceInBlueprint))
 class UNavRelevantInterface : public UInterface
diff --git a/Engine/Source/Runtime/Engine/Classes/AI/Navigation/NavigationTypes.h b/Engine/Source/Runtime/Engine/Classes/AI/Navigation/NavigationTypes.h
--- a/Engine/Source/Runtime/Engine/Classes/AI/Navigation/NavigationTypes.h
+++ b/Engine/Source/Runtime/Engine/Classes/AI/Navigation/NavigationTypes.h
@@ -12,7 +12,7 @@
 #include "UObject/WeakObjectPtr.h"
 #include "Misc/CoreStats.h"
 #include "UObject/SoftObjectPath.h"
-#include "AI/Navigation/NavFilters/NavigationQueryFilter.h"
+#include "GameFramework/Actor.h"
 #include "NavigationTypes.generated.h"
 
 #define INVALID_NAVNODEREF (0)
@@ -33,11 +33,6 @@ struct FNavigationPath;
 /** uniform identifier type for navigation data elements may it be a polygon or graph node */
 typedef uint64 NavNodeRef;
 
-class AActor;
-class ANavigationData;
-class INavAgentInterface;
-class INavRelevantInterface;
-
 namespace FNavigationSystem
 {
 	/** used as a fallback value for navigation agent radius, when none specified via UNavigationSystem::SupportedAgents */
@@ -54,13 +49,6 @@ namespace FNavigationSystem
 	{
 		return TestLocation != InvalidLocation;
 	}
-
-	enum ECreateIfEmpty
-	{
-		Invalid = -1,
-		DontCreate = 0,
-		Create = 1,
-	};
 }
 
 UENUM()
@@ -479,13 +467,15 @@ struct ENGINE_API FNavAgentProperties : public FMovementProperties
 	float NavWalkingSearchHeightScale;
 
 	/** Type of navigation data used by agent, null means "any" */
-	UPROPERTY(EditAnywhere, BlueprintReadWrite, Category = MovementProperties)
-	TSubclassOf<ANavigationData> PreferredNavData;
-
+	/*UPROPERTY(EditAnywhere, BlueprintReadWrite, Category = MovementProperties)
+	TSubclassOf<ANavigationData> PreferredNavData;*/
+	UPROPERTY(EditAnywhere, BlueprintReadWrite, Category=MovementProperties, meta=(MetaClass = "NavigationData"))
+	FSoftClassPath PreferredNavData;
+	
 	FNavAgentProperties(float Radius = -1.f, float Height = -1.f)
 		: AgentRadius(Radius), AgentHeight(Height), AgentStepHeight(-1), NavWalkingSearchHeightScale(0.5f)
-	{
-	}
+	{}
+	FNavAgentProperties(const FNavAgentProperties& Other);
 
 	void UpdateWithCollisionComponent(class UShapeComponent* CollisionComponent);
 
@@ -498,7 +488,9 @@ struct ENGINE_API FNavAgentProperties : public FMovementProperties
 	{
 		return FGenericPlatformMath::Abs(AgentRadius - Other.AgentRadius) < Precision
 			&& FGenericPlatformMath::Abs(AgentHeight - Other.AgentHeight) < Precision
-			&& (HasStepHeightOverride() == false || FGenericPlatformMath::Abs(AgentStepHeight - Other.AgentStepHeight) < Precision)
+			&& ((HasStepHeightOverride() == false) 
+				|| (Other.HasStepHeightOverride() == false) 
+				|| FGenericPlatformMath::Abs(AgentStepHeight - Other.AgentStepHeight) < Precision)
 			&& IsNavDataMatching(Other);
 	}
 	
@@ -514,6 +506,11 @@ struct ENGINE_API FNavAgentProperties : public FMovementProperties
 			: INVALID_NAVEXTENT;
 	}
 
+	void SetPreferredNavData(TSubclassOf<AActor> NavDataClass)
+	{
+		PreferredNavData = FSoftClassPath(NavDataClass.Get());
+	}
+
 	static const FNavAgentProperties DefaultProperties;
 };
 
@@ -536,13 +533,15 @@ struct ENGINE_API FNavDataConfig : public FNavAgentProperties
 	UPROPERTY(EditAnywhere, BlueprintReadOnly, Category = Querying)
 	FVector DefaultQueryExtent;
 
-	UPROPERTY(Transient, EditAnywhere, BlueprintReadOnly, Category = Navigation)
-	TSubclassOf<ANavigationData> NavigationDataClass;
+	// mz@todo make sure we handle NavigationDataClass and NavigationDataClassName in PostEditChange
+	UPROPERTY(Transient)
+	TSubclassOf<AActor> NavigationDataClass;
 
-	UPROPERTY(config)
+	UPROPERTY(VisibleAnywhere, BlueprintReadOnly, Category=Navigation, meta=(MetaClass = "NavigationData"))
 	FSoftClassPath NavigationDataClassName;
 
 	FNavDataConfig(float Radius = FNavigationSystem::FallbackAgentRadius, float Height = FNavigationSystem::FallbackAgentHeight);
+	FNavDataConfig(const FNavDataConfig& Other);
 };
 
 struct FNavigationProjectionWork
@@ -606,51 +605,6 @@ namespace ENavigationQueryResult
 	};
 }
 
-struct ENGINE_API FPathFindingQueryData
-{
-	TWeakObjectPtr<const UObject> Owner;
-	FVector StartLocation;
-	FVector EndLocation;
-	FSharedConstNavQueryFilter QueryFilter;
-
-	/** additional flags passed to navigation data handling request */
-	int32 NavDataFlags;
-
-	/** if set, allow partial paths as a result */
-	uint32 bAllowPartialPaths : 1;
-
-	FPathFindingQueryData() : StartLocation(FNavigationSystem::InvalidLocation), EndLocation(FNavigationSystem::InvalidLocation), NavDataFlags(0), bAllowPartialPaths(true) {}
-	FPathFindingQueryData(const UObject* InOwner, const FVector& InStartLocation, const FVector& InEndLocation, FSharedConstNavQueryFilter InQueryFilter = nullptr, int32 InNavDataFlags = 0, bool bInAllowPartialPaths = true) :
-		Owner(InOwner), StartLocation(InStartLocation), EndLocation(InEndLocation), QueryFilter(InQueryFilter), NavDataFlags(InNavDataFlags), bAllowPartialPaths(bInAllowPartialPaths) {}
-};
-
-struct ENGINE_API FPathFindingQuery : public FPathFindingQueryData
-{
-	TWeakObjectPtr<const ANavigationData> NavData;
-	FNavPathSharedPtr PathInstanceToFill;
-	FNavAgentProperties NavAgentProperties;
-
-	FPathFindingQuery() : FPathFindingQueryData() {}
-	FPathFindingQuery(const FPathFindingQuery& Source);
-	FPathFindingQuery(const UObject* InOwner, const ANavigationData& InNavData, const FVector& Start, const FVector& End, FSharedConstNavQueryFilter SourceQueryFilter = NULL, FNavPathSharedPtr InPathInstanceToFill = NULL);
-	FPathFindingQuery(const INavAgentInterface& InNavAgent, const ANavigationData& InNavData, const FVector& Start, const FVector& End, FSharedConstNavQueryFilter SourceQueryFilter = NULL, FNavPathSharedPtr InPathInstanceToFill = NULL);
-
-	explicit FPathFindingQuery(FNavPathSharedRef PathToRecalculate, const ANavigationData* NavDataOverride = NULL);
-
-	FPathFindingQuery& SetPathInstanceToUpdate(FNavPathSharedPtr InPathInstanceToFill) { PathInstanceToFill = InPathInstanceToFill; return *this; }
-	FPathFindingQuery& SetAllowPartialPaths(bool bAllow) { bAllowPartialPaths = bAllow; return *this; }
-	FPathFindingQuery& SetNavAgentProperties(const FNavAgentProperties& InNavAgentProperties) { NavAgentProperties = InNavAgentProperties; return *this; }
-};
-
-namespace EPathFindingMode
-{
-	enum Type
-	{
-		Regular,
-		Hierarchical,
-	};
-};
-
 /**
 *	Delegate used to communicate that path finding query has been finished.
 *	@param uint32 unique Query ID of given query
@@ -661,25 +615,6 @@ namespace EPathFindingMode
 */
 DECLARE_DELEGATE_ThreeParams(FNavPathQueryDelegate, uint32, ENavigationQueryResult::Type, FNavPathSharedPtr);
 
-//////////////////////////////////////////////////////////////////////////
-// Custom path following data
-
-/** Custom data passed to movement requests. */
-struct ENGINE_API FMoveRequestCustomData
-{
-};
-
-typedef TSharedPtr<FMoveRequestCustomData, ESPMode::ThreadSafe> FCustomMoveSharedPtr;
-typedef TWeakPtr<FMoveRequestCustomData, ESPMode::ThreadSafe> FCustomMoveWeakPtr;
-
-UCLASS(Abstract, CustomConstructor)
-class UNavigationTypes : public UObject
-{
-	GENERATED_UCLASS_BODY()
-
-	UNavigationTypes(const FObjectInitializer& ObjectInitializer = FObjectInitializer::Get()) : Super(ObjectInitializer) { }
-};
-
 //////////////////////////////////////////////////////////////////////////
 // Memory stating
 
@@ -758,33 +693,6 @@ struct TContainerTraits<TNavStatArray<InElementType> > : public TContainerTraits
 	enum { MoveWillEmptyContainer = TContainerTraits<typename TNavStatArray<InElementType>::Super>::MoveWillEmptyContainer };
 };
 
-//----------------------------------------------------------------------//
-// Active tiles 
-//----------------------------------------------------------------------//
-struct FNavigationInvokerRaw
-{
-	FVector Location;
-	float RadiusMin;
-	float RadiusMax;
-
-	FNavigationInvokerRaw(const FVector& InLocation, float Min, float Max)
-		: Location(InLocation), RadiusMin(Min), RadiusMax(Max)
-	{}
-};
-
-struct FNavigationInvoker
-{
-	TWeakObjectPtr<AActor> Actor;
-
-	/** tiles GenerationRadius away or close will be generated if they're not already present */
-	float GenerationRadius;
-	/** tiles over RemovalRadius will get removed.
-	*	@Note needs to be >= GenerationRadius or will get clampped */
-	float RemovalRadius;
-
-	FNavigationInvoker();
-	FNavigationInvoker(AActor& InActor, float InGenerationRadius, float InRemovalRadius);
-};
 
 //----------------------------------------------------------------------//
 // generic "landscape" support
@@ -798,3 +706,4 @@ struct ENGINE_API FNavHeightfieldSamples
 
 	FORCEINLINE bool IsEmpty() const { return Heights.Num() == 0; }
 };
+
diff --git a/Engine/Source/Runtime/Engine/Classes/AI/NavigationSystemBase.h b/Engine/Source/Runtime/Engine/Classes/AI/NavigationSystemBase.h
new file mode 100644
--- /dev/null
+++ b/Engine/Source/Runtime/Engine/Classes/AI/NavigationSystemBase.h
@@ -0,0 +1,353 @@
+// Copyright 1998-2018 Epic Games, Inc. All Rights Reserved.
+
+#pragma once
+
+#include "UObject/Object.h"
+#include "Engine/Engine.h"
+#include "AI/Navigation/NavigationTypes.h"
+#include "AI/NavigationModifier.h"
+#include "Engine/World.h"
+#include "NavigationSystemBase.generated.h"
+
+class UNavigationSystemBase;
+class UNavigationSystemConfig;
+class AActor;
+class UActorComponent;
+class USceneComponent;
+class INavigationDataInterface;
+class IPathFollowingAgentInterface;
+class AWorldSettings;
+class ULevel;
+class AController;
+
+ENGINE_API DECLARE_LOG_CATEGORY_EXTERN(LogNavigation, Warning, All);
+
+UENUM()
+enum class FNavigationSystemRunMode : uint8
+{
+	InvalidMode,
+	GameMode,
+	EditorMode,
+	SimulationMode,
+	PIEMode,
+};
+
+namespace ENavigationLockReason
+{
+	enum Type
+	{
+		Unknown = 1 << 0,
+		AllowUnregister = 1 << 1,
+
+		MaterialUpdate = 1 << 2,
+		LightingUpdate = 1 << 3,
+		ContinuousEditorMove = 1 << 4,
+		SpawnOnDragEnter = 1 << 5,
+	};
+}
+
+class ENGINE_API FNavigationLockContext
+{
+public:
+	FNavigationLockContext(ENavigationLockReason::Type Reason = ENavigationLockReason::Unknown, bool bApplyLock = true)
+		: MyWorld(NULL), LockReason(Reason), bSingleWorld(false), bIsLocked(false)
+	{
+		if (bApplyLock)
+		{
+			LockUpdates();
+		}
+	}
+
+	FNavigationLockContext(UWorld* InWorld, ENavigationLockReason::Type Reason = ENavigationLockReason::Unknown, bool bApplyLock = true)
+		: MyWorld(InWorld), LockReason(Reason), bSingleWorld(true), bIsLocked(false)
+	{
+		if (bApplyLock)
+		{
+			LockUpdates();
+		}
+	}
+
+	~FNavigationLockContext()
+	{
+		UnlockUpdates();
+	}
+
+private:
+	UWorld* MyWorld;
+	uint8 LockReason;
+	uint8 bSingleWorld : 1;
+	uint8 bIsLocked : 1;
+
+	void LockUpdates();
+	void UnlockUpdates();
+};
+
+namespace FNavigationSystem
+{
+	/** Creates an instance of NavigationSystem (class being specified by WorldSetting's NavigationSystemConfig)
+	 *	A new instance will be created only if given WorldOwner doesn't have one yet.
+	 *	The new instance will be assigned to the given WorldOwner (via SetNavigationSystem 
+	 *	call) and depending on value of bInitializeForWorld the InitializeForWorld 
+	 *	function will be called on the new NavigationSystem instance.
+	 *	(@see UWorld.NavigationSystem)
+	 *	@param RunMode if set to a valid value (other than FNavigationSystemRunMode::InvalidMode) 
+	 *		will also configure the created NavigationSystem instance for that mode */
+	ENGINE_API void AddNavigationSystemToWorld(UWorld& WorldOwner, const FNavigationSystemRunMode RunMode = FNavigationSystemRunMode::InvalidMode, const bool bInitializeForWorld = true);
+
+	/** Discards all navigation data chunks in all sub-levels */
+	ENGINE_API void DiscardNavigationDataChunks(UWorld& InWorld);
+
+	template<typename TNavSys>
+	FORCEINLINE TNavSys* GetCurrent(UWorld* World)
+	{
+		return World ? Cast<TNavSys>(World->GetNavigationSystem()) : (TNavSys*)nullptr;
+	}
+
+	template<typename TNavSys>
+	FORCEINLINE const TNavSys* GetCurrent(const UWorld* World)
+	{
+		return World ? Cast<TNavSys>(World->GetNavigationSystem()) : (const TNavSys*)nullptr;
+	}
+
+	template<typename TNavSys>
+	TNavSys* GetCurrent(UObject* WorldContextObject)
+	{
+		UWorld* World = nullptr;
+
+		if (WorldContextObject != nullptr)
+		{
+			World = GEngine->GetWorldFromContextObject(WorldContextObject, EGetWorldErrorMode::LogAndReturnNull);
+		}
+
+		return GetCurrent<TNavSys>(World);
+	}
+
+	ENGINE_API void UpdateActorData(AActor& Actor);
+	ENGINE_API void UpdateComponentData(UActorComponent& Comp);
+	ENGINE_API void UpdateActorAndComponentData(AActor& Actor, bool bUpdateAttachedActors = true);
+	ENGINE_API void UpdateComponentDataAfterMove(USceneComponent& Comp);
+	//ENGINE_API bool HasComponentData(UActorComponent& Comp);
+	ENGINE_API void OnActorBoundsChanged(AActor& Actor);
+	ENGINE_API void OnPostEditActorMove(AActor& Actor);
+	ENGINE_API void OnComponentBoundsChanged(UActorComponent& Comp, const FBox& NewBounds, const FBox& DirtyArea);
+	ENGINE_API void OnComponentTransformChanged(USceneComponent& Comp);
+
+	ENGINE_API void OnActorRegistered(AActor& Actor);
+	ENGINE_API void OnActorUnregistered(AActor& Actor);
+
+	ENGINE_API void OnComponentRegistered(UActorComponent& Comp);
+	ENGINE_API void OnComponentUnregistered(UActorComponent& Comp);
+
+	ENGINE_API void RemoveActorData(AActor& Actor);
+
+	ENGINE_API bool HasComponentData(UActorComponent& Comp);
+	
+	ENGINE_API const FNavDataConfig& GetDefaultSupportedAgent();
+
+	ENGINE_API TSubclassOf<UNavAreaBase> GetDefaultWalkableArea();
+
+	ENGINE_API const FTransform& GetCoordTransformTo(const ENavigationCoordSystem::Type CoordType);
+	ENGINE_API const FTransform& GetCoordTransformFrom(const ENavigationCoordSystem::Type CoordType);
+
+	ENGINE_API bool WantsComponentChangeNotifies();
+
+	//ENGINE_API INavigationDataInterface* GetNavDataForProps(const FNavAgentProperties& AgentProperties);
+	ENGINE_API INavigationDataInterface* GetNavDataForActor(const AActor& Actor);
+	ENGINE_API TSubclassOf<AActor> GetDefaultNavDataClass();
+
+	ENGINE_API void VerifyNavigationRenderingComponents(UWorld& World, const bool bShow);
+	ENGINE_API void Build(UWorld& World);
+
+#if WITH_EDITOR
+	ENGINE_API void OnPIEStart(UWorld& World);
+	ENGINE_API void OnPIEEnd(UWorld& World);
+	ENGINE_API void SetNavigationAutoUpdateEnabled(const bool bNewEnable, UNavigationSystemBase* InNavigationSystem);
+	ENGINE_API void UpdateLevelCollision(ULevel& Level);
+#endif // WITH_EDITOR
+
+	enum class ECleanupMode : uint8
+	{
+		CleanupWithWorld,
+		CleanupUnsafe,
+	};
+
+	// pathfollowing
+	ENGINE_API bool IsFollowingAPath(const AController& Controller);
+	ENGINE_API void StopMovement(const AController& Controller);
+	ENGINE_API IPathFollowingAgentInterface* FindPathFollowingAgentForActor(const AActor& Actor);
+
+	DECLARE_DELEGATE_OneParam(FActorBasedSignature, AActor& /*Actor*/);
+	DECLARE_DELEGATE_OneParam(FActorComponentBasedSignature, UActorComponent& /*Comp*/);
+	DECLARE_DELEGATE_OneParam(FSceneComponentBasedSignature, USceneComponent& /*Comp*/);
+	DECLARE_DELEGATE_OneParam(FWorldBasedSignature, UWorld& /*World*/);
+	DECLARE_DELEGATE_OneParam(FLevelBasedSignature, ULevel& /*Level*/);
+	DECLARE_DELEGATE_OneParam(FControllerBasedSignature, const AController& /*Controller*/);
+	DECLARE_DELEGATE_TwoParams(FNavigationAutoUpdateEnableSignature, const bool /*bNewEnable*/, UNavigationSystemBase* /*InNavigationSystem*/);
+	DECLARE_DELEGATE_RetVal_OneParam(bool, FBoolControllerBasedSignature, const AController& /*World*/);
+	DECLARE_DELEGATE_RetVal_OneParam(bool, FBoolActorComponentBasedSignature, UActorComponent& /*Comp*/);
+	DECLARE_DELEGATE_RetVal(TSubclassOf<UNavAreaBase>, FNavAreaBasedSignature);
+	DECLARE_DELEGATE_RetVal(const FNavDataConfig&, FNavDatConfigBasedSignature);
+	DECLARE_DELEGATE_TwoParams(FWorldByteBasedSignature, UWorld& /*World*/, uint8 /*Flags*/);
+	DECLARE_DELEGATE_TwoParams(FActorBooleBasedSignature, AActor& /*Actor*/, bool /*bUpdateAttachedActors*/);
+	DECLARE_DELEGATE_ThreeParams(FComponentBoundsChangeSignature, UActorComponent& /*Comp*/, const FBox& /*NewBounds*/, const FBox& /*DirtyArea*/)
+	DECLARE_DELEGATE_RetVal_OneParam(INavigationDataInterface*, FNavDataForPropsSignature, const FNavAgentProperties& /*AgentProperties*/);
+	DECLARE_DELEGATE_RetVal_OneParam(INavigationDataInterface*, FNavDataForActorSignature, const AActor& /*Actor*/);
+	DECLARE_DELEGATE_RetVal(TSubclassOf<AActor>, FNavDataClassFetchSignature);
+	DECLARE_DELEGATE_TwoParams(FWorldBoolBasedSignature, UWorld& /*World*/, const bool /*bShow*/);
+}
+
+
+UCLASS(Abstract, config = Engine, defaultconfig)
+class ENGINE_API UNavigationSystemBase : public UObject
+{
+	GENERATED_BODY()
+
+public:
+	virtual ~UNavigationSystemBase(){}
+
+	virtual void Tick(float DeltaSeconds) PURE_VIRTUAL(UNavigationSystemBase::Tick, );
+	virtual void CleanUp(const FNavigationSystem::ECleanupMode Mode) PURE_VIRTUAL(UNavigationSystemBase::CleanUp, );
+	virtual void Configure(const UNavigationSystemConfig& Config) PURE_VIRTUAL(UNavigationSystemBase::Configure, );
+
+	// possibly editor-only
+	
+
+	/**
+	*	Called when owner-UWorld initializes actors
+	*/
+	virtual void OnInitializeActors() {}
+
+	virtual bool IsNavigationBuilt(const AWorldSettings* Settings) const { return false; }
+
+	virtual void ApplyWorldOffset(const FVector& InOffset, bool bWorldShift) PURE_VIRTUAL(UNavigationSystemBase::ApplyWorldOffset, );
+
+	virtual void InitializeForWorld(UWorld& World, FNavigationSystemRunMode Mode) PURE_VIRTUAL(UNavigationSystemBase::InitializeForWorld, );
+
+	/** 
+	 *	If you're using NavigationSysstem module consider calling 
+	 *	FNavigationSystem::GetCurrent<UNavigationSystemV1>()->GetDefaultNavDataInstance 
+	 *	instead.
+	 */
+	virtual INavigationDataInterface* GetMainNavData() const { return nullptr; }
+
+	DEPRECATED(4.20, "GetMainNavData is deprecated. Use FNavigationSystem::GetCurrent<UNavigationSystemV1>()->GetDefaultNavDataInstance instead")
+	INavigationDataInterface* GetMainNavData(int) { return nullptr; }
+
+protected:
+	
+	static void SetCoordTransformTo(const ENavigationCoordSystem::Type CoordType, const FTransform& Transform);
+	static void SetCoordTransformFrom(const ENavigationCoordSystem::Type CoordType, const FTransform& Transform);
+
+	static FNavigationSystem::FActorBasedSignature& UpdateActorDataDelegate();
+	static FNavigationSystem::FActorComponentBasedSignature& UpdateComponentDataDelegate();
+	static FNavigationSystem::FSceneComponentBasedSignature& UpdateComponentDataAfterMoveDelegate();
+	static FNavigationSystem::FActorBasedSignature& OnActorBoundsChangedDelegate();
+	static FNavigationSystem::FActorBasedSignature& OnPostEditActorMoveDelegate();
+	static FNavigationSystem::FSceneComponentBasedSignature& OnComponentTransformChangedDelegate();
+	static FNavigationSystem::FActorBasedSignature& OnActorRegisteredDelegate();
+	static FNavigationSystem::FActorBasedSignature& OnActorUnregisteredDelegate();
+	static FNavigationSystem::FActorComponentBasedSignature& OnComponentRegisteredDelegate();
+	static FNavigationSystem::FActorComponentBasedSignature& OnComponentUnregisteredDelegate();
+	static FNavigationSystem::FActorBasedSignature& RemoveActorDataDelegate();
+	static FNavigationSystem::FBoolActorComponentBasedSignature& HasComponentDataDelegate();
+	static FNavigationSystem::FNavAreaBasedSignature& GetDefaultWalkableAreaDelegate();
+	static FNavigationSystem::FNavDatConfigBasedSignature& GetDefaultSupportedAgentDelegate();
+	static FNavigationSystem::FActorBooleBasedSignature& UpdateActorAndComponentDataDelegate();
+	static FNavigationSystem::FComponentBoundsChangeSignature& OnComponentBoundsChangedDelegate();
+	//static FNavigationSystem::FNavDataForPropsSignature& GetNavDataForPropsDelegate();
+	static FNavigationSystem::FNavDataForActorSignature& GetNavDataForActorDelegate();
+	static FNavigationSystem::FNavDataClassFetchSignature& GetDefaultNavDataClassDelegate();
+	static FNavigationSystem::FWorldBoolBasedSignature& VerifyNavigationRenderingComponentsDelegate();
+	static FNavigationSystem::FWorldBasedSignature& BuildDelegate();
+#if WITH_EDITOR
+	static FNavigationSystem::FWorldBasedSignature& OnPIEStartDelegate();
+	static FNavigationSystem::FWorldBasedSignature& OnPIEEndDelegate();
+	static FNavigationSystem::FLevelBasedSignature& UpdateLevelCollisionDelegate();
+	static FNavigationSystem::FNavigationAutoUpdateEnableSignature& SetNavigationAutoUpdateEnableDelegate();
+	static FNavigationSystem::FWorldByteBasedSignature& AddNavigationUpdateLockDelegate();
+	static FNavigationSystem::FWorldByteBasedSignature& RemoveNavigationUpdateLockDelegate();
+#endif // WITH_EDITOR
+};
+
+
+class ENGINE_API IPathFollowingManagerInterface
+{
+protected:
+	static FNavigationSystem::FControllerBasedSignature& StopMovementDelegate();
+	static FNavigationSystem::FBoolControllerBasedSignature& IsFollowingAPathDelegate();
+};
+
+
+UCLASS(Abstract, config = Engine, defaultconfig)
+class ENGINE_API UNavigationSystem : public UObject
+{
+	GENERATED_BODY()
+public:
+	UNavigationSystem(const FObjectInitializer& ObjectInitializer = FObjectInitializer::Get());
+
+	DEPRECATED(4.20, "UpdateActorInNavOctree is deprecated. Use FNavigationSystem::RemoveActorData instead")
+	static void ClearNavOctreeAll(AActor* Actor);
+	DEPRECATED(4.20, "UpdateActorInNavOctree is deprecated. Use FNavigationSystem::UpdateActorData instead")
+	static void UpdateActorInNavOctree(AActor& Actor);
+	DEPRECATED(4.20, "UpdateComponentInNavOctree is deprecated. Use FNavigationSystem::UpdateComponentData instead")
+	static void UpdateComponentInNavOctree(UActorComponent& Comp);
+	DEPRECATED(4.20, "UpdateActorAndComponentsInNavOctree is deprecated. Use FNavigationSystem::UpdateActorAndComponentData instead")
+	static void UpdateActorAndComponentsInNavOctree(AActor& Actor, bool bUpdateAttachedActors = true);
+	DEPRECATED(4.20, "UpdateNavOctreeAfterMove is deprecated. Use FNavigationSystem::UpdateComponentDataAfterMove instead")
+	static void UpdateNavOctreeAfterMove(USceneComponent* Comp);
+	DEPRECATED(4.20, "UpdateNavOctreeBounds is deprecated. Use FNavigationSystem::OnActorBoundsChanged instead")
+	static void UpdateNavOctreeBounds(AActor* Actor);
+	DEPRECATED(4.20, "InitializeForWorld is deprecated. Use FNavigationSystem::CreateNavigationSystem instead")
+	static void InitializeForWorld(UWorld* World, FNavigationSystemRunMode Mode);
+	DEPRECATED(4.20, "CreateNavigationSystem is deprecated. Use FNavigationSystem::CreateNavigationSystem instead")
+	static UNavigationSystem* CreateNavigationSystem(UWorld* WorldOwner);
+	DEPRECATED(4.20, "OnComponentRegistered is deprecated. Use FNavigationSystem::OnComponentRegistered instead")
+	static void OnComponentRegistered(UActorComponent* Comp);
+	DEPRECATED(4.20, "OnComponentUnregistered is deprecated. Use FNavigationSystem::OnComponentUnregistered instead")
+	static void OnComponentUnregistered(UActorComponent* Comp);
+	DEPRECATED(4.20, "OnActorRegistered is deprecated. Use FNavigationSystem::OnActorRegistered instead")
+	static void OnActorRegistered(AActor* Actor);
+	DEPRECATED(4.20, "OnActorUnregistered is deprecated. Use FNavigationSystem::OnActorUnregistered instead")
+	static void OnActorUnregistered(AActor* Actor);
+	DEPRECATED(4.20, "GetCurrent is deprecated. Use FNavigationSystem::GetCurrent instead")
+	static UNavigationSystem* GetCurrent(UWorld* World);
+	DEPRECATED(4.20, "GetCurrent is deprecated. Use FNavigationSystem::GetCurrent instead")
+	static UNavigationSystem* GetCurrent(UObject* WorldContextObject);
+	DEPRECATED(4.20, "ShouldUpdateNavOctreeOnComponentChange is deprecated. Use FNavigationSystem::WantsComponentChangeNotifies instead")
+	static bool ShouldUpdateNavOctreeOnComponentChange();
+	DEPRECATED(4.20, "GetDefaultWalkableArea is deprecated. Use FNavigationSystem::GetDefaultWalkableArea instead")
+	static TSubclassOf<UNavAreaBase> GetDefaultWalkableArea();
+	DEPRECATED(4.20, "UNavigationSystem::GetNavDataForProps is deprecated. Use UNavigationSystemV1::GetDefaultObstacleArea instead")
+	static TSubclassOf<UNavAreaBase> GetDefaultObstacleArea();
+	DEPRECATED(4.20, "UNavigationSystem::K2_GetRandomReachablePointInRadius is deprecated. Use UNavigationSystemV1::K2_GetRandomReachablePointInRadius instead")
+	static bool K2_GetRandomReachablePointInRadius(UObject* WorldContextObject, const FVector& Origin, FVector& RandomLocation, float Radius, UObject* NavData = NULL, TSubclassOf<UObject> FilterClass = NULL) { return false; }
+
+	DEPRECATED(4.20, "UpdateNavOctreeElementBounds is deprecated. Use FNavigationSystem::OnComponentBoundsChanged instead")
+	bool UpdateNavOctreeElementBounds(UActorComponent* Comp, const FBox& NewBounds, const FBox& DirtyArea);
+
+	DEPRECATED(4.20, "GetObjectsNavOctreeId is deprecated. You need to access NavigationSystem's modul NavigationSystem to access this functionality")
+	const class FOctreeElementId* GetObjectsNavOctreeId(const UObject* Object) const { return nullptr; }
+
+	DEPRECATED(4.20, "HasPendingObjectNavOctreeId is deprecated. You need to access NavigationSystem's modul NavigationSystem class to access this functionality")
+	bool HasPendingObjectNavOctreeId(UObject* Object) const { return false; }
+
+	DEPRECATED(4.20, "GetDefaultSupportedAgentConfig is deprecated. Use FNavigationSystem::GetDefaultSupportedAgent instead")
+	const FNavDataConfig& GetDefaultSupportedAgentConfig() const;
+	
+	DEPRECATED(4.20, "GetNavDataForProps is deprecated. Use FNavigationSystem::GetNavDataForProps instead")
+	INavigationDataInterface* GetNavDataForProps(const FNavAgentProperties& AgentProperties);
+
+	DEPRECATED(4.20, "GetSupportedAgents has been moved. Use UNavigationSystemV1::GetSupportedAgents instead")
+	const TArray<FNavDataConfig>& GetSupportedAgents() const { return FakeSupportedAgents; }
+
+	DEPRECATED(4.20, "GetMainNavData has been moved. Use FNavigationSystem::GetCurrent<UNavigationSystemV1>()->GetDefaultNavDataInstance instead")
+	UObject* GetMainNavData(int) { return nullptr; }
+	DEPRECATED(4.20, "GetMainNavData has been moved. Use FNavigationSystem::GetCurrent<UNavigationSystemV1>()->GetDefaultNavDataInstance instead")
+	UObject* GetMainNavData() const { return nullptr; }
+
+	DEPRECATED(4.20, "ProjectPointToNavigation has been deprecated along with the whole UNavigationSystem class. Use UNavigationSystemV1 instead")
+	bool ProjectPointToNavigation(const FVector& Point, FNavLocation& OutLocation, const FVector& Extent = INVALID_NAVEXTENT, const UObject* NavData = NULL, void* QueryFilter = NULL) const { return false; }
+
+private:
+	const TArray<FNavDataConfig> FakeSupportedAgents;
+};
\ No newline at end of file
diff --git a/Engine/Source/Runtime/Engine/Classes/AI/NavigationSystemConfig.h b/Engine/Source/Runtime/Engine/Classes/AI/NavigationSystemConfig.h
new file mode 100644
--- /dev/null
+++ b/Engine/Source/Runtime/Engine/Classes/AI/NavigationSystemConfig.h
@@ -0,0 +1,46 @@
+// Copyright 1998-2018 Epic Games, Inc. All Rights Reserved.
+
+#pragma once
+
+#include "CoreMinimal.h"
+#include "UObject/Object.h"
+#include "UObject/SoftObjectPath.h"
+#include "NavigationSystemConfig.generated.h"
+
+class UNavigationSystemBase;
+class UWorld;
+
+
+UCLASS(config = Engine, defaultconfig, EditInlineNew, DisplayName = "Generic Navigation System Config", collapseCategories)
+class ENGINE_API UNavigationSystemConfig : public UObject
+{
+	GENERATED_BODY()
+
+public:
+	UPROPERTY(EditAnywhere, Category=Navigation, meta = (MetaClass = "NavigationSystemBase", NoResetToDefault))
+	FSoftClassPath NavigationSystemClass;
+
+	/** Whether at game runtime we expect any kind of dynamic navigation generation */
+	UPROPERTY(EditAnywhere, Category = Navigation)
+	uint32 bStrictlyStatic : 1;
+
+	UPROPERTY(EditAnywhere, Category = Navigation)
+	uint32 bCreateOnClient : 1;
+
+public:
+	UNavigationSystemConfig(const FObjectInitializer& ObjectInitializer = FObjectInitializer::Get());
+
+	UNavigationSystemBase* CreateAndConfigureNavigationSystem(UWorld& World) const;
+
+#if WITH_EDITOR
+	virtual void PostEditChangeProperty(FPropertyChangedEvent& PropertyChangedEvent) override;
+#endif // WITH_EDITOR
+};
+
+UCLASS(MinimalAPI, HideCategories=Navigation)
+class UNullNavSysConfig : public UNavigationSystemConfig
+{
+	GENERATED_BODY()
+public:
+	UNullNavSysConfig(const FObjectInitializer& ObjectInitializer = FObjectInitializer::Get());
+};
diff --git a/Engine/Source/Runtime/Engine/Classes/Components/ShapeComponent.h b/Engine/Source/Runtime/Engine/Classes/Components/ShapeComponent.h
--- a/Engine/Source/Runtime/Engine/Classes/Components/ShapeComponent.h
+++ b/Engine/Source/Runtime/Engine/Classes/Components/ShapeComponent.h
@@ -71,7 +71,7 @@ public:
 
 	/** Navigation area type (empty = default obstacle) */
 	UPROPERTY(EditAnywhere, Category = Navigation)
-	TSubclassOf<class UNavArea> AreaClass;
+	TSubclassOf<class UNavAreaBase> AreaClass;
 
 	//~ Begin UPrimitiveComponent Interface.
 	virtual FPrimitiveSceneProxy* CreateSceneProxy() override;
diff --git a/Engine/Source/Runtime/Engine/Classes/EdGraph/EdGraphPin.h b/Engine/Source/Runtime/Engine/Classes/EdGraph/EdGraphPin.h
--- a/Engine/Source/Runtime/Engine/Classes/EdGraph/EdGraphPin.h
+++ b/Engine/Source/Runtime/Engine/Classes/EdGraph/EdGraphPin.h
@@ -342,17 +342,64 @@ public:
 	/** Name of this pin. */
 	FName PinName;
 
+	/** Direction of flow of this pin (input or output) */
+	TEnumAsByte<enum EEdGraphPinDirection> Direction;
+
+#if WITH_EDITORONLY_DATA
+	////////////////////////////////////////////////////////////////////////////////////
+	// ONLY PUT BITFIELD PROPERTIES AFTER THIS TO ENSURE GOOD MEMORY ALIGNMENT
+	////////////////////////////////////////////////////////////////////////////////////
+
+	/** If true, this connector is currently hidden. */
+	uint8 bHidden:1;
+
+	/** If true, this connector is unconnectable, and present only to allow the editing of the default text. */
+	uint8 bNotConnectable:1;
+
+	/** If true, the default value of this connector is fixed and cannot be modified by the user (it's visible for reference only). */
+	uint8 bDefaultValueIsReadOnly:1;
+
+	/** If true, the default value on this pin is ignored and should not be set. */
+	uint8 bDefaultValueIsIgnored:1;
+
+	/** If true, this pin is the focus of a diff. This is transient. */
+	uint8 bIsDiffing:1;
+
+	/** If true, the pin may be hidden by user. */
+	uint8 bAdvancedView:1;
+
+	/** If true, the pin is displayed as ref. This is transient. */
+	uint8 bDisplayAsMutableRef:1;
+
+	/** 
+	* If true, this pin existed on an older version of the owning node, but when the node was reconstructed a matching pin was not found.
+	* This pin must be linked to other pins or have a non-default value and will be removed if disconnected, reset to default, or the node is refreshed.
+	**/
+	uint8 bOrphanedPin:1;
+
+	bool ShouldSavePinIfOrphaned() const { return bSavePinIfOrphaned; }
+	ENGINE_API void SetSavePinIfOrphaned(bool bShouldSave);
+
+private:
+	/** If true, this pin will be retained when reconstructing a node if there is no matching pin on the new version of the pin. This is transient. */
+	uint8 bSavePinIfOrphaned:1;
+
+public:
+	/** Older content sometimes had an empty autogenerated default value string in cases where that does not mean the property default value (0, none, false, etc.) */
+	uint8 bUseBackwardsCompatForEmptyAutogeneratedValue:1;
+
+#endif // WITH_EDITORONLY_DATA
+
+	/** True when InvalidateAndTrash was called. This pin is intended to be discarded and destroyed. */
+	uint8 bWasTrashed:1;	/** Used as the display name if set. */
+
 #if WITH_EDITORONLY_DATA
-	/** Used as the display name if set. */
 	FText PinFriendlyName;
 #endif
 
 	/** The tool-tip describing this pin's purpose */
 	FString PinToolTip;
 
-	/** Direction of flow of this pin (input or output) */
-	TEnumAsByte<enum EEdGraphPinDirection> Direction;
-
 	/** The type of information carried on this pin */
 	FEdGraphPinType PinType;
 
@@ -384,47 +431,7 @@ public:
 	/** Pin name could be changed, so whenever possible it's good to have a persistent GUID identifying Pin to reconstruct Node seamlessly */
 	FGuid PersistentGuid;
 
-////////////////////////////////////////////////////////////////////////////////////
-// ONLY PUT BITFIELD PROPERTIES AFTER THIS TO ENSURE GOOD MEMORY ALIGNMENT
-////////////////////////////////////////////////////////////////////////////////////
-
-	/** If true, this connector is currently hidden. */
-	uint32 bHidden:1;
-
-	/** If true, this connector is unconnectable, and present only to allow the editing of the default text. */
-	uint32 bNotConnectable:1;
-
-	/** If true, the default value of this connector is fixed and cannot be modified by the user (it's visible for reference only). */
-	uint32 bDefaultValueIsReadOnly:1;
-
-	/** If true, the default value on this pin is ignored and should not be set. */
-	uint32 bDefaultValueIsIgnored:1;
-
-	/** If true, this pin is the focus of a diff. This is transient. */
-	uint32 bIsDiffing:1;
-
-	/** If true, the pin may be hidden by user. */
-	uint32 bAdvancedView:1;
-
-	/** If true, the pin is displayed as ref. This is transient. */
-	uint32 bDisplayAsMutableRef:1;
-
-	/** 
-	 * If true, this pin existed on an older version of the owning node, but when the node was reconstructed a matching pin was not found.
-	 * This pin must be linked to other pins or have a non-default value and will be removed if disconnected, reset to default, or the node is refreshed.
-	 **/
-	uint32 bOrphanedPin:1;
-
-	/** If true, this pin will be retained when reconstructing a node if there is no matching pin on the new version of the pin. This is transient. */
-	uint32 bSavePinIfOrphaned:1;
-
-	/** Older content sometimes had an empty autogenerated default value string in cases where that does not mean the property default value (0, none, false, etc.) */
-	uint32 bUseBackwardsCompatForEmptyAutogeneratedValue:1;
-
-#endif // WITH_EDITORONLY_DATA
-
-	/** True when InvalidateAndTrash was called. This pin is intended to be discarded and destroyed. */
-	uint32 bWasTrashed:1;
+#endif
 
 public:
 
