Merging //UE4/Dev-Main to Dev-Geometry (//UE4/Dev-Geometry)

diff --git a/Engine/Source/Runtime/Engine/Classes/Engine/Blueprint.h b/Engine/Source/Runtime/Engine/Classes/Engine/Blueprint.h
--- a/Engine/Source/Runtime/Engine/Classes/Engine/Blueprint.h
+++ b/Engine/Source/Runtime/Engine/Classes/Engine/Blueprint.h
@@ -296,8 +296,9 @@ struct FEditedDocumentInfo
 {
 	GENERATED_USTRUCT_BODY()
 
+	/** Edited object */
 	UPROPERTY()
-	UObject* EditedObject;
+	FSoftObjectPath EditedObjectPath;
 
 	/** Saved view position */
 	UPROPERTY()
@@ -308,29 +309,78 @@ struct FEditedDocumentInfo
 	float SavedZoomAmount;
 
 	FEditedDocumentInfo()
-		: EditedObject(nullptr)
-		, SavedViewOffset(0.0f, 0.0f)
+		: SavedViewOffset(0.0f, 0.0f)
 		, SavedZoomAmount(-1.0f)
+		, EditedObject_DEPRECATED(nullptr)
 	{ }
 
 	FEditedDocumentInfo(UObject* InEditedObject)
-		: EditedObject(InEditedObject)
+		: EditedObjectPath(InEditedObject)
 		, SavedViewOffset(0.0f, 0.0f)
 		, SavedZoomAmount(-1.0f)
+		, EditedObject_DEPRECATED(nullptr)
 	{ }
 
 	FEditedDocumentInfo(UObject* InEditedObject, FVector2D& InSavedViewOffset, float InSavedZoomAmount)
-		: EditedObject(InEditedObject)
+		: EditedObjectPath(InEditedObject)
 		, SavedViewOffset(InSavedViewOffset)
 		, SavedZoomAmount(InSavedZoomAmount)
+		, EditedObject_DEPRECATED(nullptr)
 	{ }
 
-	friend bool operator==( const FEditedDocumentInfo& LHS, const FEditedDocumentInfo& RHS )
+	void PostSerialize(const FArchive& Ar)
 	{
-		return LHS.EditedObject == RHS.EditedObject && LHS.SavedViewOffset == RHS.SavedViewOffset && LHS.SavedZoomAmount == RHS.SavedZoomAmount;
+		if (Ar.IsLoading() && EditedObject_DEPRECATED)
+		{
+			// Convert hard to soft reference.
+			EditedObjectPath = EditedObject_DEPRECATED;
+			EditedObject_DEPRECATED = nullptr;
+		}
 	}
+
+	friend bool operator==(const FEditedDocumentInfo& LHS, const FEditedDocumentInfo& RHS)
+	{
+		return LHS.EditedObjectPath == RHS.EditedObjectPath && LHS.SavedViewOffset == RHS.SavedViewOffset && LHS.SavedZoomAmount == RHS.SavedZoomAmount;
+	}
+
+private:
+	// Legacy hard reference is now serialized as a soft reference (see above).
+	UPROPERTY()
+	UObject* EditedObject_DEPRECATED;
 };
 
+template<>
+struct TStructOpsTypeTraits<FEditedDocumentInfo> : public TStructOpsTypeTraitsBase2<FEditedDocumentInfo>
+{
+	enum
+	{
+		WithPostSerialize = true
+	};
+};
+
+/** Bookmark node info */
+USTRUCT()
+struct FBPEditorBookmarkNode
+{
+	GENERATED_USTRUCT_BODY()
+
+	/** Node ID */
+	UPROPERTY()
+	FGuid NodeGuid;
+
+	/** Parent ID */
+	UPROPERTY()
+	FGuid ParentGuid;
+
+	/** Display name */
+	UPROPERTY()
+	FText DisplayName;
+
+	friend bool operator==(const FBPEditorBookmarkNode& LHS, const FBPEditorBookmarkNode& RHS)
+	{
+		return LHS.NodeGuid == RHS.NodeGuid;
+	}
+};
 
 UENUM()
 enum class EBlueprintNativizationFlag : uint8
@@ -340,6 +390,7 @@ enum class EBlueprintNativizationFlag : uint8
 	ExplicitlyEnabled
 };
 
+
 /**
  * Blueprints are special assets that provide an intuitive, node-based interface that can be used to create new types of Actors
  * and script level events; giving designers and gameplay programmers the tools to quickly create and iterate gameplay from
@@ -529,6 +580,14 @@ class ENGINE_API UBlueprint : public UBlueprintCore
 	UPROPERTY()
 	TArray<struct FEditedDocumentInfo> LastEditedDocuments;
 
+	/** Bookmark data */
+	UPROPERTY()
+	TMap<FGuid, struct FEditedDocumentInfo> Bookmarks;
+
+	/** Bookmark nodes (for display) */
+	UPROPERTY()
+	TArray<FBPEditorBookmarkNode> BookmarkNodes;
+
 	/** Persistent debugging options */
 	UPROPERTY()
 	TArray<class UBreakpoint*> Breakpoints;
diff --git a/Engine/Source/Runtime/Engine/Classes/Engine/Engine.h b/Engine/Source/Runtime/Engine/Classes/Engine/Engine.h
--- a/Engine/Source/Runtime/Engine/Classes/Engine/Engine.h
+++ b/Engine/Source/Runtime/Engine/Classes/Engine/Engine.h
@@ -710,7 +710,7 @@ public:
 
 	/** The class for NavigationSystem **/
 	UPROPERTY()
-	TSubclassOf<class UNavigationSystem>  NavigationSystemClass;
+	TSubclassOf<class UNavigationSystemBase>  NavigationSystemClass;
 	
 	/** Name of behavior tree manager class */
 	UPROPERTY(globalconfig, noclear, meta=(MetaClass="AvoidanceManager", DisplayName="Avoidance Manager Class"))
diff --git a/Engine/Source/Runtime/Engine/Classes/Engine/Level.h b/Engine/Source/Runtime/Engine/Classes/Engine/Level.h
--- a/Engine/Source/Runtime/Engine/Classes/Engine/Level.h
+++ b/Engine/Source/Runtime/Engine/Classes/Engine/Level.h
@@ -651,7 +651,7 @@ public:
 	virtual void BeginDestroy() override;
 	virtual bool IsReadyForFinishDestroy() override;
 	virtual void FinishDestroy() override;
-	virtual UWorld* GetWorld() const override;
+	ENGINE_API virtual UWorld* GetWorld() const override final;
 
 #if	WITH_EDITOR
 	virtual void PreEditUndo() override;
diff --git a/Engine/Source/Runtime/Engine/Classes/Engine/StaticMesh.h b/Engine/Source/Runtime/Engine/Classes/Engine/StaticMesh.h
--- a/Engine/Source/Runtime/Engine/Classes/Engine/StaticMesh.h
+++ b/Engine/Source/Runtime/Engine/Classes/Engine/StaticMesh.h
@@ -24,7 +24,9 @@
 class FSpeedTreeWind;
 class UAssetUserData;
 class UMaterialInterface;
+class UNavCollisionBase;
 class UMeshDescription;
+
 struct FStaticMeshLODResources;
 
 /*-----------------------------------------------------------------------------
@@ -710,7 +712,7 @@ public:
 
 	/** Pre-build navigation collision */
 	UPROPERTY(VisibleAnywhere, transient, duplicatetransient, Instanced, Category = Navigation)
-	class UNavCollision* NavCollision;
+	UNavCollisionBase* NavCollision;
 public:
 	/**
 	 * Default constructor
@@ -883,9 +885,9 @@ public:
 	 */
 	ENGINE_API void CreateNavCollision(const bool bIsUpdate = false);
 
-	FORCEINLINE const UNavCollision* GetNavCollision() const { return NavCollision; }
+	FORCEINLINE const UNavCollisionBase* GetNavCollision() const { return NavCollision; }
 
-	/** Configures this SM as bHasNavigationData = false and clears stored UNavCollision */
+	/** Configures this SM as bHasNavigationData = false and clears stored NavCollision */
 	ENGINE_API void MarkAsNotHavingNavigationData();
 
 	const FGuid& GetLightingGuid() const
diff --git a/Engine/Source/Runtime/Engine/Classes/Engine/StreamableManager.h b/Engine/Source/Runtime/Engine/Classes/Engine/StreamableManager.h
--- a/Engine/Source/Runtime/Engine/Classes/Engine/StreamableManager.h
+++ b/Engine/Source/Runtime/Engine/Classes/Engine/StreamableManager.h
@@ -64,11 +64,18 @@ struct ENGINE_API FStreamableHandle : public TSharedFromThis<FStreamableHandle,
 		return Priority;
 	}
 
-	/** Release this handle. This can be called from normal gameplay code to indicate that the loaded assets are no longer needed. Will be called implicitly if all shared pointers
-	to this handle are destroyed. If called before the completion delegate, the release will be delayed until after completion. */
+	/**
+	 * Release this handle. This can be called from normal gameplay code to indicate that the loaded assets are no longer needed
+	 * This will be called implicitly if all shared pointers to this handle are destroyed
+	 * If called before the completion delegate, the release will be delayed until after completion
+	 */
 	void ReleaseHandle();
 
-	/** Cancel a request, callable from within the manager or externally. Will stop delegate from being called */
+	/**
+	 * Cancel a request, callable from within the manager or externally
+	 * This will immediately release the handle even if it is still in progress
+	 * This will normally stop the completion delegate from getting called, unless that delegate is already in the callback delay queue
+	 */
 	void CancelHandle();
 
 	/** Tells a stalled handle to start its actual request. */
@@ -285,6 +292,9 @@ struct ENGINE_API FStreamableManager : public FGCObject
 	/** This will release any managed active handles pointing to the target string asset reference, even if they include other requested assets in the same load */
 	void Unload(const FSoftObjectPath& Target);
 
+	/** Checks for any redirectors that were previously loaded, and returns the redirected target if found. This will not handle redirects that it doesn't yet know about */
+	FSoftObjectPath ResolveRedirects(const FSoftObjectPath& Target) const;
+
 	DEPRECATED(4.16, "Call LoadSynchronous with bManageActiveHandle=true instead if you want the manager to keep the handle alive")
 	UObject* SynchronousLoad(FSoftObjectPath const& Target);
 
@@ -317,7 +327,6 @@ private:
 
 	void RemoveReferencedAsset(const FSoftObjectPath& Target, TSharedRef<FStreamableHandle> Handle);
 	void StartHandleRequests(TSharedRef<FStreamableHandle> Handle);
-	FSoftObjectPath ResolveRedirects(const FSoftObjectPath& Target) const;
 	void FindInMemory(FSoftObjectPath& InOutTarget, struct FStreamable* Existing);
 	struct FStreamable* FindStreamable(const FSoftObjectPath& Target) const;
 	struct FStreamable* StreamInternal(const FSoftObjectPath& Target, TAsyncLoadPriority Priority, TSharedRef<FStreamableHandle> Handle);
diff --git a/Engine/Source/Runtime/Engine/Classes/Engine/World.h b/Engine/Source/Runtime/Engine/Classes/Engine/World.h
--- a/Engine/Source/Runtime/Engine/Classes/Engine/World.h
+++ b/Engine/Source/Runtime/Engine/Classes/Engine/World.h
@@ -896,9 +896,9 @@ public:
 	uint32 bTriggerPostLoadMap:1;
 
 private:
-	/** The world's navmesh */
+	/** The world's navigation data manager */
 	UPROPERTY(Transient)
-	class UNavigationSystem*					NavigationSystem;
+	class UNavigationSystemBase*				NavigationSystem;
 
 	/** The current GameMode, valid only on the server */
 	UPROPERTY(Transient)
@@ -961,8 +961,8 @@ private:
 	UCanvas* CanvasForDrawMaterialToRenderTarget;
 
 public:
-	/** Set the pointer to the Navgation system. */
-	void SetNavigationSystem( UNavigationSystem* InNavigationSystem);
+	/** Set the pointer to the Navigation System instance. */
+	void SetNavigationSystem(UNavigationSystemBase* InNavigationSystem);
 
 	/** The interface to the scene manager for this world. */
 	class FSceneInterface*						Scene;
@@ -1934,9 +1934,9 @@ public:
 	bool IsTraceHandleValid(const FTraceHandle& Handle, bool bOverlapTrace);
 
 	/** NavigationSystem getter */
-	FORCEINLINE UNavigationSystem* GetNavigationSystem() { return NavigationSystem; }
+	FORCEINLINE UNavigationSystemBase* GetNavigationSystem() { return NavigationSystem; }
 	/** NavigationSystem const getter */
-	FORCEINLINE const UNavigationSystem* GetNavigationSystem() const { return NavigationSystem; }
+	FORCEINLINE const UNavigationSystemBase* GetNavigationSystem() const { return NavigationSystem; }
 
 	/** AISystem getter. if AISystem is missing it tries to create one and returns the result.
 	 *	@NOTE the result can be NULL, for example on client games or if no AI module or AISystem class have not been specified
diff --git a/Engine/Source/Runtime/Engine/Classes/GameFramework/Actor.h b/Engine/Source/Runtime/Engine/Classes/GameFramework/Actor.h
--- a/Engine/Source/Runtime/Engine/Classes/GameFramework/Actor.h
+++ b/Engine/Source/Runtime/Engine/Classes/GameFramework/Actor.h
@@ -38,6 +38,7 @@ ENGINE_API DECLARE_LOG_CATEGORY_EXTERN(LogActor, Log, Warning);
 // Delegate signatures
 DECLARE_DYNAMIC_MULTICAST_DELEGATE_FiveParams( FTakeAnyDamageSignature, AActor*, DamagedActor, float, Damage, const class UDamageType*, DamageType, class AController*, InstigatedBy, AActor*, DamageCauser );
 DECLARE_DYNAMIC_MULTICAST_DELEGATE_NineParams( FTakePointDamageSignature, AActor*, DamagedActor, float, Damage, class AController*, InstigatedBy, FVector, HitLocation, class UPrimitiveComponent*, FHitComponent, FName, BoneName, FVector, ShotFromDirection, const class UDamageType*, DamageType, AActor*, DamageCauser );
+DECLARE_DYNAMIC_MULTICAST_DELEGATE_SevenParams( FTakeRadialDamageSignature, AActor*, DamagedActor, float, Damage, const class UDamageType*, DamageType, FVector, Origin, FHitResult, HitInfo, class AController*, InstigatedBy, AActor*, DamageCauser );
 DECLARE_DYNAMIC_MULTICAST_DELEGATE_TwoParams( FActorBeginOverlapSignature, AActor*, OverlappedActor, AActor*, OtherActor );
 DECLARE_DYNAMIC_MULTICAST_DELEGATE_TwoParams( FActorEndOverlapSignature, AActor*, OverlappedActor, AActor*, OtherActor );
 DECLARE_DYNAMIC_MULTICAST_DELEGATE_FourParams( FActorHitSignature, AActor*, SelfActor, AActor*, OtherActor, FVector, NormalImpulse, const FHitResult&, Hit );
@@ -621,6 +622,10 @@ public:
 	/** Called when the actor is damaged by point damage. */
 	UPROPERTY(BlueprintAssignable, Category="Game|Damage")
 	FTakePointDamageSignature OnTakePointDamage;
+
+	/** Called when the actor is damaged by radial damage. */
+	UPROPERTY(BlueprintAssignable, Category="Game|Damage")
+	FTakeRadialDamageSignature OnTakeRadialDamage;
 	
 	/** 
 	 * Called when another actor begins to overlap this actor, for example a player walking into a trigger.
@@ -2678,7 +2683,7 @@ public:
 			{
 				OutComponents.Add(Component);
 			}
-			else if (bIncludeFromChildActors)
+			if (bIncludeFromChildActors)
 			{
 				if (UChildActorComponent* ChildActorComponent = Cast<UChildActorComponent>(OwnedComponent))
 				{
@@ -2726,12 +2731,13 @@ public:
 			if (Component)
 			{
 				OutComponents.Add(Component);
-			}
-			else if (bIncludeFromChildActors)
-			{
-				if (UChildActorComponent* ChildActorComponent = Cast<UChildActorComponent>(Component))
+
+				if (bIncludeFromChildActors)
 				{
-					ChildActorComponents.Add(ChildActorComponent);
+					if (UChildActorComponent* ChildActorComponent = Cast<UChildActorComponent>(Component))
+					{
+						ChildActorComponents.Add(ChildActorComponent);
+					}
 				}
 			}
 		}
diff --git a/Engine/Source/Runtime/Engine/Classes/GameFramework/CharacterMovementComponent.h b/Engine/Source/Runtime/Engine/Classes/GameFramework/CharacterMovementComponent.h
--- a/Engine/Source/Runtime/Engine/Classes/GameFramework/CharacterMovementComponent.h
+++ b/Engine/Source/Runtime/Engine/Classes/GameFramework/CharacterMovementComponent.h
@@ -24,6 +24,7 @@ class FDebugDisplayInfo;
 class FNetworkPredictionData_Server_Character;
 class FSavedMove_Character;
 class UPrimitiveComponent;
+class INavigationData;
 class UCharacterMovementComponent;
 
 DECLARE_DELEGATE_RetVal_TwoParams(FTransform, FOnProcessRootMotion, const FTransform&, UCharacterMovementComponent*)
@@ -486,6 +487,14 @@ public:
 	UPROPERTY(Category="Character Movement (Networking)", EditDefaultsOnly)
 	uint32 bNetworkSkipProxyPredictionOnNetUpdate:1;
 
+	/**
+	 * Flag used on the server to determine whether to always replicate ReplicatedServerLastTransformUpdateTimeStamp to clients.
+	 * Normally this is only sent when the network smoothing mode on character movement is set to Linear smoothing (on the server), to save bandwidth.
+	 * Setting this to true will force the timestamp to replicate regardless, in case the server doesn't know about the smoothing mode, or if the timestamp is used for another purpose.
+	 */
+	UPROPERTY(Category="Character Movement (Networking)", EditDefaultsOnly, AdvancedDisplay)
+	uint32 bNetworkAlwaysReplicateTransformUpdateTimestamp:1;
+
 public:
 
 	/** true to update CharacterOwner and UpdatedComponent after movement ends */
@@ -1497,7 +1506,7 @@ protected:
 	virtual void SetNavWalkingPhysics(bool bEnable);
 
 	/** Get Navigation data for the Character. Returns null if there is no associated nav data. */
-	const class ANavigationData* GetNavData() const;
+	const class INavigationDataInterface* GetNavData() const;
 
 	/** 
 	 * Checks to see if the current location is not encroaching blocking geometry so the character can leave NavWalking.
diff --git a/Engine/Source/Runtime/Engine/Classes/GameFramework/CheatManager.h b/Engine/Source/Runtime/Engine/Classes/GameFramework/CheatManager.h
--- a/Engine/Source/Runtime/Engine/Classes/GameFramework/CheatManager.h
+++ b/Engine/Source/Runtime/Engine/Classes/GameFramework/CheatManager.h
@@ -252,14 +252,6 @@ class ENGINE_API UCheatManager : public UObject
 	UFUNCTION(exec)
 	virtual void TestCollisionDistance();
 
-	/** Builds the navigation mesh (or rebuilds it). **/
-	UFUNCTION(exec)
-	virtual void RebuildNavigation();
-
-	/** Sets navigation drawing distance. Relevant only in non-editor modes. **/
-	UFUNCTION(exec)
-	void SetNavDrawDistance(float DrawDistance);
-
 	/** Dump online session information */
 	UFUNCTION(exec)
 	virtual void DumpOnlineSessionState();
diff --git a/Engine/Source/Runtime/Engine/Classes/GameFramework/Controller.h b/Engine/Source/Runtime/Engine/Classes/GameFramework/Controller.h
--- a/Engine/Source/Runtime/Engine/Classes/GameFramework/Controller.h
+++ b/Engine/Source/Runtime/Engine/Classes/GameFramework/Controller.h
@@ -14,10 +14,10 @@ class APawn;
 class APlayerState;
 class FDebugDisplayInfo;
 class UDamageType;
-class UPathFollowingComponent;
 
 UDELEGATE(BlueprintAuthorityOnly)
 DECLARE_DYNAMIC_MULTICAST_DELEGATE_FourParams( FInstigatedAnyDamageSignature, float, Damage, const UDamageType*, DamageType, AActor*, DamagedActor, AActor*, DamageCauser );
+DECLARE_MULTICAST_DELEGATE_OneParam(FPawnChangedSignature, APawn* /*NewPawn*/);
 
 /**
  * Controllers are non-physical actors that can possess a Pawn to control
@@ -80,6 +80,8 @@ private:
 	USceneComponent* TransformComponent;
 
 protected:
+	FPawnChangedSignature OnNewPawn;
+
 	/** The control rotation of the Controller. See GetControlRotation. */
 	UPROPERTY()
 	FRotator ControlRotation;
@@ -187,6 +189,10 @@ public:
 	UFUNCTION(BlueprintCallable, Category=Pawn, meta=(DisplayName="Get Controlled Pawn", ScriptName="GetControlledPawn"))
 	APawn* K2_GetPawn() const;
 
+	FPawnChangedSignature GetOnNewPawnNotifier() { return OnNewPawn; }
+
+public:
+
 	virtual void GetLifetimeReplicatedProps(TArray<FLifetimeProperty>& OutLifetimeProps) const override;
 
 	//~ Begin AActor Interface
@@ -289,14 +295,9 @@ public:
 	virtual void GetMoveGoalReachTest(const AActor* MovingActor, const FVector& MoveOffset, FVector& GoalOffset, float& GoalRadius, float& GoalHalfHeight) const override;
 	virtual bool ShouldPostponePathUpdates() const override;
 	virtual bool IsFollowingAPath() const override;
+	virtual IPathFollowingAgentInterface* GetPathFollowingAgent() const override;
 	//~ End INavAgentInterface Interface
-
-	/** Prepares path following component */
-	virtual void InitNavigationControl(UPathFollowingComponent*& PathFollowingComp);
-
-	/** If controller has any navigation-related components then this function makes them update their cached data */
-	virtual void UpdateNavigationComponents();
-
+	
 	/** Aborts the move the controller is currently performing */
 	UFUNCTION(BlueprintCallable, Category="AI|Navigation")
 	virtual void StopMovement();
@@ -352,6 +353,20 @@ protected:
 private:
 	// Hidden functions that don't make sense to use on this class.
 	HIDE_ACTOR_TRANSFORM_FUNCTIONS();
+
+public:
+
+	/** prepares path following component */
+	DEPRECATED(4.20, "InitNavigationControl is depracated. The function lost its meaning due to better Engine-AIModule code separation")
+	virtual void InitNavigationControl(UObject*& PathFollowingComp) {}
+
+	/** If controller has any navigation-related components then this function
+	 *	makes them update their cached data 
+	 *	DEPRECATED: this functionality has been taken over by PathFollowingComponent 
+	 *		by observing newly possessed pawns (via OnNewPawn)
+	 */
+	DEPRECATED(4.20, "UpdateNavigationComponents is depracated. The function lost its meaning due to better Engine-AIModule code separation.")
+	virtual void UpdateNavigationComponents() {}
 };
 
 
diff --git a/Engine/Source/Runtime/Engine/Classes/GameFramework/GameNetworkManager.h b/Engine/Source/Runtime/Engine/Classes/GameFramework/GameNetworkManager.h
--- a/Engine/Source/Runtime/Engine/Classes/GameFramework/GameNetworkManager.h
+++ b/Engine/Source/Runtime/Engine/Classes/GameFramework/GameNetworkManager.h
@@ -129,6 +129,10 @@ class ENGINE_API AGameNetworkManager : public AInfo
 	UPROPERTY(GlobalConfig)
 	float ClientNetSendMoveDeltaTimeThrottled;
 
+	/** ClientNetSendMoveDeltaTimeStationary is used when players are determined to not be moving or changing their view. See ClientNetSendMoveDeltaTime for more info. */
+	UPROPERTY(GlobalConfig)
+	float ClientNetSendMoveDeltaTimeStationary;
+
 	/** When player net speed (CurrentNetSpeed, based on ConfiguredInternetSpeed or ConfiguredLanSpeed) is less than or equal to this amount, ClientNetSendMoveDeltaTimeThrottled is used instead of ClientNetSendMoveDeltaTime. */
 	UPROPERTY(GlobalConfig)
 	int32 ClientNetSendMoveThrottleAtNetSpeed;
diff --git a/Engine/Source/Runtime/Engine/Classes/GameFramework/NavMovementComponent.h b/Engine/Source/Runtime/Engine/Classes/GameFramework/NavMovementComponent.h
--- a/Engine/Source/Runtime/Engine/Classes/GameFramework/NavMovementComponent.h
+++ b/Engine/Source/Runtime/Engine/Classes/GameFramework/NavMovementComponent.h
@@ -12,11 +12,11 @@
 #include "GameFramework/Actor.h"
 #include "AI/Navigation/NavigationTypes.h"
 #include "AI/Navigation/NavAgentInterface.h"
+#include "AI/Navigation/PathFollowingAgentInterface.h"
 #include "GameFramework/MovementComponent.h"
 #include "NavMovementComponent.generated.h"
 
 class UCapsuleComponent;
-class UPathFollowingComponent;
 
 /**
  * NavMovementComponent defines base functionality for MovementComponents that move any 'agent' that may be involved in AI pathfinding.
@@ -47,18 +47,20 @@ protected:
 	UPROPERTY(EditAnywhere, Category = NavMovement, meta = (EditCondition = "bUseAccelerationForPaths"))
 	uint32 bUseFixedBrakingDistanceForPaths : 1;
 
-protected:
 	/** If set, StopActiveMovement call will abort current path following request */
 	uint32 bStopMovementAbortPaths:1;
 
+private:
+	/** object implementing IPathFollowingAgentInterface. Is private to control access to it.
+	 *	@see SetPathFollowingAgent, GetPathFollowingAgent */
+	UPROPERTY()
+	UObject* PathFollowingComp;
+
 public:
 	/** Expresses runtime state of character's movement. Put all temporal changes to movement properties here */
 	UPROPERTY()
 	FMovementProperties MovementState;
 
-	/** associated path following component */
-	TWeakObjectPtr<class UPathFollowingComponent> PathFollowingComp;
-
 	/** Stops applying further movement (usually zeros acceleration). */
 	UFUNCTION(BlueprintCallable, Category="Pawn|Components|PawnMovement")
 	virtual void StopActiveMovement();
@@ -98,6 +100,10 @@ public:
 	/** @returns braking distance for acceleration driven path following */
 	virtual float GetPathFollowingBrakingDistance(float MaxSpeed) const;
 
+	void SetPathFollowingAgent(IPathFollowingAgentInterface* InPathFollowingAgent) { PathFollowingComp = Cast<UObject>(InPathFollowingAgent); }
+	IPathFollowingAgentInterface* GetPathFollowingAgent() { return Cast<IPathFollowingAgentInterface>(PathFollowingComp); }
+	const IPathFollowingAgentInterface* GetPathFollowingAgent() const { return Cast<const IPathFollowingAgentInterface>(PathFollowingComp); }
+
 	/** set fixed braking distance */
 	void SetFixedBrakingDistance(float DistanceToEndOfPath);
 
diff --git a/Engine/Source/Runtime/Engine/Classes/GameFramework/WorldSettings.h b/Engine/Source/Runtime/Engine/Classes/GameFramework/WorldSettings.h
--- a/Engine/Source/Runtime/Engine/Classes/GameFramework/WorldSettings.h
+++ b/Engine/Source/Runtime/Engine/Classes/GameFramework/WorldSettings.h
@@ -16,6 +16,7 @@
 
 class UAssetUserData;
 class UNetConnection;
+class UNavigationSystemConfig;
 
 UENUM()
 enum EVisibilityAggressiveness
@@ -387,10 +388,15 @@ class ENGINE_API AWorldSettings : public AInfo, public IInterface_AssetUserData
 	UPROPERTY(EditAnywhere, BlueprintReadOnly, Category=World, AdvancedDisplay)
 	uint32 bEnableWorldBoundsChecks:1;
 
-	/** if set to false navigation system will not get created (and all navigation functionality won't be accessible)*/
-	UPROPERTY(EditAnywhere, BlueprintReadOnly, config, Category=World, AdvancedDisplay)
+protected:
+	/** if set to false navigation system will not get created (and all 
+	 *	navigation functionality won't be accessible).
+	 *	This flag is now deprecated. Use NavigationSystemConfig property to 
+	 *	determine navigation system's properties or existence */
+	UPROPERTY(VisibleAnywhere, BlueprintReadOnly, config, Category=World, AdvancedDisplay, meta=(DisplayName = "DEPRECATED_bEnableNavigationSystem"))
 	uint32 bEnableNavigationSystem:1;
 
+public:
 	/** if set to false AI system will not get created. Use it to disable all AI-related activity on a map */
 	UPROPERTY(EditAnywhere, BlueprintReadOnly, config, Category = World, AdvancedDisplay)
 	uint32 bEnableAISystem:1;
@@ -421,6 +427,11 @@ class ENGINE_API AWorldSettings : public AInfo, public IInterface_AssetUserData
 	UPROPERTY(EditAnywhere, BlueprintReadOnly, meta=(DisplayName = "Override World Gravity"), Category = Physics)
 	uint32 bGlobalGravitySet:1;
 
+	/** Holds parameters for NavigationSystem's creation. Set to Null will result
+	 *	in NavigationSystem instance not being created for this world */
+	UPROPERTY(EditAnywhere, BlueprintReadOnly, Category=World, AdvancedDisplay, Instanced, NoClear, meta=(NoResetToDefault))
+	UNavigationSystemConfig* NavigationSystemConfig;
+
 	// any actor falling below this level gets destroyed
 	UPROPERTY(EditAnywhere, BlueprintReadOnly, Category=World, meta=(editcondition = "bEnableWorldBoundsChecks"))
 	float KillZ;
@@ -675,6 +686,7 @@ public:
 #if WITH_EDITOR
 	virtual void CheckForErrors() override;
 #endif // WITH_EDITOR
+	virtual void PostInitProperties() override;
 	virtual void PreInitializeComponents() override;
 	virtual void PostInitializeComponents() override;
 	virtual void PostRegisterAllComponents() override;
@@ -701,6 +713,13 @@ public:
 	/** Sets the global time dilation value (subject to clamping). Returns the final value that was set. */
 	virtual float SetTimeDilation(float NewTimeDilation);
 
+	/** @return configuration for NavigationSystem's creation. Null means 
+	 *	no navigation system will be created*/
+	UNavigationSystemConfig* const GetNavigationSystemConfig() const { return NavigationSystemConfig; }
+
+	/** @return whether given world is configured to host any NavigationSystem */
+	bool IsNavigationSystemEnabled() const;
+
 	/**
 	 * Called from GameStateBase, calls BeginPlay on all actors
 	 */
diff --git a/Engine/Source/Runtime/Engine/Classes/Interfaces/Interface_CollisionDataProvider.h b/Engine/Source/Runtime/Engine/Classes/Interfaces/Interface_CollisionDataProvider.h
--- a/Engine/Source/Runtime/Engine/Classes/Interfaces/Interface_CollisionDataProvider.h
+++ b/Engine/Source/Runtime/Engine/Classes/Interfaces/Interface_CollisionDataProvider.h
@@ -60,7 +60,7 @@ struct FTriMeshCollisionData
 	}
 };
 
-UINTERFACE(meta=(CannotImplementInterfaceInBlueprint))
+UINTERFACE(MinimalAPI, meta=(CannotImplementInterfaceInBlueprint))
 class UInterface_CollisionDataProvider : public UInterface
 {
 	GENERATED_UINTERFACE_BODY()
diff --git a/Engine/Source/Runtime/Engine/Classes/Kismet/KismetArrayLibrary.h b/Engine/Source/Runtime/Engine/Classes/Kismet/KismetArrayLibrary.h
--- a/Engine/Source/Runtime/Engine/Classes/Kismet/KismetArrayLibrary.h
+++ b/Engine/Source/Runtime/Engine/Classes/Kismet/KismetArrayLibrary.h
@@ -1,4 +1,4 @@
-﻿// Copyright 1998-2018 Epic Games, Inc. All Rights Reserved.
+// Copyright 1998-2018 Epic Games, Inc. All Rights Reserved.
 
 #pragma once
 
@@ -249,11 +249,10 @@ public:
  
  		Stack.MostRecentPropertyAddress = NULL;
  		Stack.StepCompiledIn<UProperty>(StorageSpace);
-		void* NewItemPtr = (Stack.MostRecentPropertyAddress != NULL && Stack.MostRecentProperty->GetClass() == InnerProp->GetClass()) ? Stack.MostRecentPropertyAddress : StorageSpace;
  
  		P_FINISH;
 		P_NATIVE_BEGIN;
-		*(int32*)RESULT_PARAM = GenericArray_Add(ArrayAddr, ArrayProperty, NewItemPtr);
+		*(int32*)RESULT_PARAM = GenericArray_Add(ArrayAddr, ArrayProperty, StorageSpace);
 		P_NATIVE_END;
 		InnerProp->DestroyValue(StorageSpace);
 	}
@@ -278,11 +277,10 @@ public:
 
 		Stack.MostRecentPropertyAddress = NULL;
 		Stack.StepCompiledIn<UProperty>(StorageSpace);
-		void* NewItemPtr = (Stack.MostRecentPropertyAddress != NULL && Stack.MostRecentProperty->GetClass() == InnerProp->GetClass()) ? Stack.MostRecentPropertyAddress : StorageSpace;
 
 		P_FINISH;
 		P_NATIVE_BEGIN;
-		*(int32*)RESULT_PARAM = GenericArray_AddUnique(ArrayAddr, ArrayProperty, NewItemPtr);
+		*(int32*)RESULT_PARAM = GenericArray_AddUnique(ArrayAddr, ArrayProperty, StorageSpace);
 		P_NATIVE_END;
 		InnerProp->DestroyValue(StorageSpace);
 	}
@@ -354,12 +352,11 @@ public:
 
 		Stack.MostRecentPropertyAddress = NULL;
 		Stack.StepCompiledIn<UProperty>(StorageSpace);
-		void* NewItemPtr = (Stack.MostRecentPropertyAddress != NULL && Stack.MostRecentProperty->GetClass() == InnerProp->GetClass()) ? Stack.MostRecentPropertyAddress : StorageSpace;
 
 		P_GET_PROPERTY(UIntProperty, Index);
 		P_FINISH;
 		P_NATIVE_BEGIN;
-		GenericArray_Insert(ArrayAddr, ArrayProperty, NewItemPtr, Index);
+		GenericArray_Insert(ArrayAddr, ArrayProperty, StorageSpace, Index);
 		P_NATIVE_END;
 		InnerProp->DestroyValue(StorageSpace);
 	}
@@ -535,14 +532,13 @@ public:
 
 		Stack.MostRecentPropertyAddress = NULL;
 		Stack.StepCompiledIn<UProperty>(StorageSpace);
-		void* NewItemPtr = (Stack.MostRecentPropertyAddress != NULL && Stack.MostRecentProperty->GetClass() == InnerProp->GetClass()) ? Stack.MostRecentPropertyAddress : StorageSpace;
 
 		P_GET_UBOOL(bSizeToFit);
 
 		P_FINISH;
 
 		P_NATIVE_BEGIN;
-		GenericArray_Set(ArrayAddr, ArrayProperty, Index, NewItemPtr, bSizeToFit);
+		GenericArray_Set(ArrayAddr, ArrayProperty, Index, StorageSpace, bSizeToFit);
 		P_NATIVE_END;
 		InnerProp->DestroyValue(StorageSpace);
 	}
diff --git a/Engine/Source/Runtime/Engine/Classes/Kismet/KismetMathLibrary.h b/Engine/Source/Runtime/Engine/Classes/Kismet/KismetMathLibrary.h
--- a/Engine/Source/Runtime/Engine/Classes/Kismet/KismetMathLibrary.h
+++ b/Engine/Source/Runtime/Engine/Classes/Kismet/KismetMathLibrary.h
@@ -313,7 +313,7 @@ class ENGINE_API UKismetMathLibrary : public UBlueprintFunctionLibrary
 	static int32 Not_Int(int32 A);
 
 	/** Sign (integer, returns -1 if A < 0, 0 if A is zero, and +1 if A > 0) */
-	UFUNCTION(BlueprintPure, meta=(DisplayName = "Sign (int)"), Category="Math|Integer")
+	UFUNCTION(BlueprintPure, meta=(DisplayName = "Sign (integer)"), Category="Math|Integer")
 	static int32 SignOfInteger(int32 A);
 
 	/** Returns a uniformly distributed random number between 0 and Max - 1 */
@@ -325,19 +325,19 @@ class ENGINE_API UKismetMathLibrary : public UBlueprintFunctionLibrary
 	static int32 RandomIntegerInRange(int32 Min, int32 Max);
 
 	/** Returns the minimum value of A and B */
-	UFUNCTION(BlueprintPure, meta=(DisplayName = "Min (int)", CompactNodeTitle = "MIN", CommutativeAssociativeBinaryOperator = "true"), Category="Math|Integer")
+	UFUNCTION(BlueprintPure, meta=(DisplayName = "Min (integer)", CompactNodeTitle = "MIN", CommutativeAssociativeBinaryOperator = "true"), Category="Math|Integer")
 	static int32 Min(int32 A, int32 B);
 
 	/** Returns the maximum value of A and B */
-	UFUNCTION(BlueprintPure, meta=(DisplayName = "Max (int)", CompactNodeTitle = "MAX", CommutativeAssociativeBinaryOperator = "true"), Category="Math|Integer")
+	UFUNCTION(BlueprintPure, meta=(DisplayName = "Max (integer)", CompactNodeTitle = "MAX", CommutativeAssociativeBinaryOperator = "true"), Category="Math|Integer")
 	static int32 Max(int32 A, int32 B);
 
 	/** Returns Value clamped to be between A and B (inclusive) */
-	UFUNCTION(BlueprintPure, meta=(DisplayName = "Clamp (int)"), Category="Math|Integer")
+	UFUNCTION(BlueprintPure, meta=(DisplayName = "Clamp (integer)"), Category="Math|Integer")
 	static int32 Clamp(int32 Value, int32 Min, int32 Max);
 
 	/** Returns the absolute (positive) value of A */
-	UFUNCTION(BlueprintPure, meta=(DisplayName = "Absolute (int)", CompactNodeTitle = "ABS"), Category="Math|Integer")
+	UFUNCTION(BlueprintPure, meta=(DisplayName = "Absolute (integer)", CompactNodeTitle = "ABS"), Category="Math|Integer")
 	static int32 Abs_Int(int32 A);
 
 	//
@@ -1016,7 +1016,7 @@ class ENGINE_API UKismetMathLibrary : public UBlueprintFunctionLibrary
 	static FRotator Multiply_RotatorFloat(FRotator A, float B);
 	
 	/** Returns rotator representing rotator A scaled by B */
-	UFUNCTION(BlueprintPure, meta=(DisplayName = "ScaleRotator (int)", CompactNodeTitle = "*", Keywords = "* multiply rotate rotation"), Category="Math|Rotator")
+	UFUNCTION(BlueprintPure, meta=(DisplayName = "ScaleRotator (integer)", CompactNodeTitle = "*", Keywords = "* multiply rotate rotation"), Category="Math|Rotator")
 	static FRotator Multiply_RotatorInt(FRotator A, int32 B);
 
 	/** Combine 2 rotations to give you the resulting rotation of first applying A, then B. */
@@ -1403,19 +1403,19 @@ class ENGINE_API UKismetMathLibrary : public UBlueprintFunctionLibrary
 	static float Conv_ByteToFloat(uint8 InByte);
 
 	/** Converts an integer to a float */
-	UFUNCTION(BlueprintPure, meta=(DisplayName = "ToFloat (int)", CompactNodeTitle = "->", Keywords="cast convert", BlueprintAutocast), Category="Math|Conversions")
+	UFUNCTION(BlueprintPure, meta=(DisplayName = "ToFloat (integer)", CompactNodeTitle = "->", Keywords="cast convert", BlueprintAutocast), Category="Math|Conversions")
 	static float Conv_IntToFloat(int32 InInt);
 
 	/** Converts an integer to a byte (if the integer is too large, returns the low 8 bits) */
-	UFUNCTION(BlueprintPure, meta=(DisplayName = "ToByte (int)", CompactNodeTitle = "->", Keywords="cast convert", BlueprintAutocast), Category="Math|Conversions")
+	UFUNCTION(BlueprintPure, meta=(DisplayName = "ToByte (integer)", CompactNodeTitle = "->", Keywords="cast convert", BlueprintAutocast), Category="Math|Conversions")
 	static uint8 Conv_IntToByte(int32 InInt);
 
 	/** Converts an integer to an IntVector*/
-	UFUNCTION(BlueprintPure, meta = (DisplayName = "ToIntVector (int)", CompactNodeTitle = "->", Keywords = "cast convert", BlueprintAutocast), Category = "Math|Conversions")
+	UFUNCTION(BlueprintPure, meta = (DisplayName = "ToIntVector (integer)", CompactNodeTitle = "->", Keywords = "cast convert", BlueprintAutocast), Category = "Math|Conversions")
 	static FIntVector Conv_IntToIntVector(int32 InInt);
 
 	/** Converts a int to a bool*/
-	UFUNCTION(BlueprintPure, meta=(DisplayName = "ToBool (int)", CompactNodeTitle = "->", Keywords="cast convert", BlueprintAutocast), Category="Math|Conversions")
+	UFUNCTION(BlueprintPure, meta=(DisplayName = "ToBool (integer)", CompactNodeTitle = "->", Keywords="cast convert", BlueprintAutocast), Category="Math|Conversions")
 	static bool Conv_IntToBool(int32 InInt);
 
 	/** Converts a bool to an int */
diff --git a/Engine/Source/Runtime/Engine/Classes/Kismet/KismetStringLibrary.h b/Engine/Source/Runtime/Engine/Classes/Kismet/KismetStringLibrary.h
--- a/Engine/Source/Runtime/Engine/Classes/Kismet/KismetStringLibrary.h
+++ b/Engine/Source/Runtime/Engine/Classes/Kismet/KismetStringLibrary.h
@@ -17,7 +17,7 @@ class ENGINE_API UKismetStringLibrary : public UBlueprintFunctionLibrary
 	static FString Conv_FloatToString(float InFloat);
 
 	/** Converts an integer value to a string */
-	UFUNCTION(BlueprintPure, meta=(DisplayName = "ToString (int)", CompactNodeTitle = "->", BlueprintAutocast), Category="Utilities|String")
+	UFUNCTION(BlueprintPure, meta=(DisplayName = "ToString (integer)", CompactNodeTitle = "->", BlueprintAutocast), Category="Utilities|String")
 	static FString Conv_IntToString(int32 InInt);
 
 	/** Converts a byte value to a string */
@@ -107,7 +107,7 @@ class ENGINE_API UKismetStringLibrary : public UBlueprintFunctionLibrary
 	 * @param Suffix - A suffix to append to the end of the conversion string
 	 * @return A new string built from the passed parameters
 	 */
-	UFUNCTION(BlueprintPure, meta=(DisplayName = "BuildString (int)"), Category="Utilities|String")
+	UFUNCTION(BlueprintPure, meta=(DisplayName = "BuildString (integer)"), Category="Utilities|String")
 	static FString BuildString_Int(const FString& AppendTo, const FString& Prefix, int32 InInt, const FString& Suffix);
 
 	/** 
diff --git a/Engine/Source/Runtime/Engine/Classes/Kismet/KismetTextLibrary.h b/Engine/Source/Runtime/Engine/Classes/Kismet/KismetTextLibrary.h
--- a/Engine/Source/Runtime/Engine/Classes/Kismet/KismetTextLibrary.h
+++ b/Engine/Source/Runtime/Engine/Classes/Kismet/KismetTextLibrary.h
@@ -193,7 +193,7 @@ class ENGINE_API UKismetTextLibrary : public UBlueprintFunctionLibrary
 
 	// Default values are duplicated from FNumberFormattingOptions and should be replicated in all functions and in the struct when changed!
 	/* Converts a passed in integer to text based on formatting options */
-	UFUNCTION(BlueprintPure, meta=(DisplayName = "ToText (int)", AdvancedDisplay = "1", BlueprintAutocast), Category="Utilities|Text")
+	UFUNCTION(BlueprintPure, meta=(DisplayName = "ToText (integer)", AdvancedDisplay = "1", BlueprintAutocast), Category="Utilities|Text")
 	static FText Conv_IntToText(int32 Value, bool bAlwaysSign = false, bool bUseGrouping = true, int32 MinimumIntegralDigits = 1, int32 MaximumIntegralDigits = 324);
 
 	// Default values are duplicated from FNumberFormattingOptions and should be replicated in all functions and in the struct when changed!
@@ -212,7 +212,7 @@ class ENGINE_API UKismetTextLibrary : public UBlueprintFunctionLibrary
 
 	// Default values are duplicated from FNumberFormattingOptions and should be replicated in all functions and in the struct when changed!
 	/* Converts a passed in integer to a text formatted as a currency */
-	UFUNCTION(BlueprintPure, meta=(DisplayName = "AsCurrency (int) - DEPRECATED (use AsCurrency)"), Category="Utilities|Text")
+	UFUNCTION(BlueprintPure, meta=(DisplayName = "AsCurrency (integer) - DEPRECATED (use AsCurrency)"), Category="Utilities|Text")
 	static FText AsCurrency_Integer(int32 Value, TEnumAsByte<ERoundingMode> RoundingMode, bool bAlwaysSign = false, bool bUseGrouping = true, int32 MinimumIntegralDigits = 1, int32 MaximumIntegralDigits = 324, int32 MinimumFractionalDigits = 0, int32 MaximumFractionalDigits = 3, const FString& CurrencyCode = TEXT(""));
 
 	// Default values are duplicated from FNumberFormattingOptions and should be replicated in all functions and in the struct when changed!
diff --git a/Engine/Source/Runtime/Engine/Engine.Build.cs b/Engine/Source/Runtime/Engine/Engine.Build.cs
--- a/Engine/Source/Runtime/Engine/Engine.Build.cs
+++ b/Engine/Source/Runtime/Engine/Engine.Build.cs
@@ -76,7 +76,6 @@ public class Engine : ModuleRules
 				"EngineSettings",
 				"SynthBenchmark",
 				"GameplayTags",
-				"AIModule",
 				"DatabaseSupport",
 				"PacketHandler",
 				"HardwareSurvey",
@@ -164,7 +163,6 @@ public class Engine : ModuleRules
 		}
 
 		CircularlyReferencedDependentModules.Add("GameplayTags");
-		CircularlyReferencedDependentModules.Add("AIModule");
 		CircularlyReferencedDependentModules.Add("Landscape");
 		CircularlyReferencedDependentModules.Add("UMG");
 		CircularlyReferencedDependentModules.Add("MaterialShaderQualitySettings");
@@ -417,19 +415,6 @@ public class Engine : ModuleRules
 				"libOpus"
 				);
 		}
-
-		if (Target.bCompileRecast)
-		{
-			PrivateDependencyModuleNames.Add("Navmesh");
-			PublicDefinitions.Add("WITH_RECAST=1");
-		}
-		else
-		{
-			// Because we test WITH_RECAST in public Engine header files, we need to make sure that modules
-			// that import this also have this definition set appropriately.  Recast is a private dependency
-			// module, so it's definitions won't propagate to modules that import Engine.
-			PublicDefinitions.Add("WITH_RECAST=0");
-		}
 /*
 		ConfigHierarchy Ini = ConfigCache.ReadHierarchy(ConfigHierarchyType.Engine, System.IO.DirectoryReference.FromFile(Target.ProjectFile), Target.Platform);
 		bool bLocalVectorFieldOnly = false;
diff --git a/Engine/Source/Runtime/Engine/Private/AI/Navigation/NavigationInterfaces.cpp b/Engine/Source/Runtime/Engine/Private/AI/Navigation/NavigationInterfaces.cpp
--- a/Engine/Source/Runtime/Engine/Private/AI/Navigation/NavigationInterfaces.cpp
+++ b/Engine/Source/Runtime/Engine/Private/AI/Navigation/NavigationInterfaces.cpp
@@ -5,16 +5,11 @@
 #include "Templates/Casts.h"
 #include "AI/Navigation/NavAgentInterface.h"
 #include "AI/Navigation/NavRelevantInterface.h"
-#include "AI/Navigation/NavigationPathGenerator.h"
-#include "AI/Navigation/NavNodeInterface.h"
-#include "AI/Navigation/NavLinkDefinition.h"
-#include "AI/Navigation/NavLinkHostInterface.h"
 #include "AI/Navigation/NavPathObserverInterface.h"
-#include "AI/Navigation/NavLinkCustomInterface.h"
 #include "AI/Navigation/NavEdgeProviderInterface.h"
 #include "AI/RVOAvoidanceInterface.h"
-
-uint32 INavLinkCustomInterface::NextUniqueId = 1;
+#include "AI/Navigation/NavigationDataInterface.h"
+#include "AI/Navigation/PathFollowingAgentInterface.h"
 
 URVOAvoidanceInterface::URVOAvoidanceInterface(const FObjectInitializer& ObjectInitializer)
 	: Super(ObjectInitializer)
@@ -26,65 +21,27 @@ UNavAgentInterface::UNavAgentInterface(const FObjectInitializer& ObjectInitializ
 {
 }
 
-UNavigationPathGenerator::UNavigationPathGenerator(const FObjectInitializer& ObjectInitializer)
-	: Super(ObjectInitializer)
-{
-}
-
-UNavNodeInterface::UNavNodeInterface(const FObjectInitializer& ObjectInitializer)
+UNavRelevantInterface::UNavRelevantInterface(const FObjectInitializer& ObjectInitializer)
 	: Super(ObjectInitializer)
 {
 }
 
-UNavLinkHostInterface::UNavLinkHostInterface(const FObjectInitializer& ObjectInitializer)
+UNavEdgeProviderInterface::UNavEdgeProviderInterface(const FObjectInitializer& ObjectInitializer) 
 	: Super(ObjectInitializer)
 {
 }
 
-UNavLinkCustomInterface::UNavLinkCustomInterface(const FObjectInitializer& ObjectInitializer)
-: Super(ObjectInitializer)
-{
-}
-
-UObject* INavLinkCustomInterface::GetLinkOwner() const
-{
-	return Cast<UObject>((INavLinkCustomInterface*)this);
-}
-
-uint32 INavLinkCustomInterface::GetUniqueId()
-{
-	return NextUniqueId++;
-}
-
-void INavLinkCustomInterface::UpdateUniqueId(uint32 AlreadyUsedId)
-{
-	NextUniqueId = FMath::Max(NextUniqueId, AlreadyUsedId + 1);
-}
-
-FNavigationLink INavLinkCustomInterface::GetModifier(const INavLinkCustomInterface* CustomNavLink)
-{
-	FNavigationLink LinkMod;
-	LinkMod.SetAreaClass(CustomNavLink->GetLinkAreaClass());
-	LinkMod.UserId = CustomNavLink->GetLinkId();
-
-	ENavLinkDirection::Type LinkDirection = ENavLinkDirection::BothWays;
-	CustomNavLink->GetLinkData(LinkMod.Left, LinkMod.Right, LinkDirection);
-	LinkMod.Direction = LinkDirection;
-
-	return LinkMod;
-}
-
 UNavPathObserverInterface::UNavPathObserverInterface(const FObjectInitializer& ObjectInitializer)
 	: Super(ObjectInitializer)
 {
 }
 
-UNavRelevantInterface::UNavRelevantInterface(const FObjectInitializer& ObjectInitializer)
+UNavigationDataInterface::UNavigationDataInterface(const FObjectInitializer& ObjectInitializer)
 	: Super(ObjectInitializer)
 {
 }
 
-UNavEdgeProviderInterface::UNavEdgeProviderInterface(const FObjectInitializer& ObjectInitializer) : Super(ObjectInitializer)
+UPathFollowingAgentInterface::UPathFollowingAgentInterface(const FObjectInitializer& ObjectInitializer)
+	: Super(ObjectInitializer)
 {
 }
-
diff --git a/Engine/Source/Runtime/Engine/Private/AI/Navigation/NavigationModifier.cpp b/Engine/Source/Runtime/Engine/Private/AI/Navigation/NavigationModifier.cpp
--- a/Engine/Source/Runtime/Engine/Private/AI/Navigation/NavigationModifier.cpp
+++ b/Engine/Source/Runtime/Engine/Private/AI/Navigation/NavigationModifier.cpp
@@ -5,12 +5,9 @@
 #include "EngineStats.h"
 #include "GameFramework/Actor.h"
 #include "Components/BrushComponent.h"
-#include "AI/Navigation/RecastHelpers.h"
-#include "AI/Navigation/NavLinkTrivial.h"
-#include "AI/Navigation/NavAreas/NavAreaMeta.h"
 #include "PhysicsEngine/BodySetup.h"
-#include "AI/Navigation/NavigationSystem.h"
-#include "AI/Navigation/NavAreas/NavArea_LowHeight.h"
+#include "AI/NavigationSystemBase.h"
+#include "AI/Navigation/NavAreaBase.h"
 
 // if square distance between two points is less than this the those points
 // will be considered identical when calculating convex hull
@@ -42,7 +39,7 @@ void FNavigationLinkBase::SetAreaClass(UClass* InAreaClass)
 UClass* FNavigationLinkBase::GetAreaClass() const
 {
 	UClass* ClassOb = AreaClassOb.Get();
-	return ClassOb ? ClassOb : (UClass*)UNavigationSystem::GetDefaultWalkableArea();
+	return ClassOb ? ClassOb : *FNavigationSystem::GetDefaultWalkableArea();
 }
 
 void FNavigationLinkBase::InitializeAreaClass(const bool bForceRefresh)
@@ -52,8 +49,7 @@ void FNavigationLinkBase::InitializeAreaClass(const bool bForceRefresh)
 
 bool FNavigationLinkBase::HasMetaArea() const
 {
-	UClass* ClassOb = GetAreaClass();
-	return ClassOb && ClassOb->IsChildOf(UNavAreaMeta::StaticClass());
+	return AreaClass && AreaClass->GetDefaultObject<UNavAreaBase>()->IsMetaArea();
 }
 
 void FNavigationLinkBase::PostSerialize(const FArchive& Ar)
@@ -279,7 +275,7 @@ bool UNavLinkDefinition::HasAdjustableLinks() const
 // FAreaNavModifier
 //----------------------------------------------------------------------//
 
-FAreaNavModifier::FAreaNavModifier(float Radius, float Height, const FTransform& LocalToWorld, const TSubclassOf<UNavArea> InAreaClass)
+FAreaNavModifier::FAreaNavModifier(float Radius, float Height, const FTransform& LocalToWorld, const TSubclassOf<UNavAreaBase> InAreaClass)
 {
 	Init(InAreaClass);
 	
@@ -296,25 +292,25 @@ FAreaNavModifier::FAreaNavModifier(float Radius, float Height, const FTransform&
 	Bounds = FBox::BuildAABB(LocalToWorld.GetLocation(), FVector(Radius, Radius, Height));
 }
 
-FAreaNavModifier::FAreaNavModifier(const FVector& Extent, const FTransform& LocalToWorld, const TSubclassOf<UNavArea> InAreaClass)
+FAreaNavModifier::FAreaNavModifier(const FVector& Extent, const FTransform& LocalToWorld, const TSubclassOf<UNavAreaBase> InAreaClass)
 {
 	Init(InAreaClass);
 	SetBox(FBox::BuildAABB(FVector::ZeroVector, Extent), LocalToWorld);
 }
 
-FAreaNavModifier::FAreaNavModifier(const FBox& Box, const FTransform& LocalToWorld, const TSubclassOf<UNavArea> InAreaClass)
+FAreaNavModifier::FAreaNavModifier(const FBox& Box, const FTransform& LocalToWorld, const TSubclassOf<UNavAreaBase> InAreaClass)
 {
 	Init(InAreaClass);
 	SetBox(Box, LocalToWorld);
 }
 
-FAreaNavModifier::FAreaNavModifier(const TArray<FVector>& InPoints, ENavigationCoordSystem::Type CoordType, const FTransform& LocalToWorld, const TSubclassOf<UNavArea> InAreaClass)
+FAreaNavModifier::FAreaNavModifier(const TArray<FVector>& InPoints, ENavigationCoordSystem::Type CoordType, const FTransform& LocalToWorld, const TSubclassOf<UNavAreaBase> InAreaClass)
 {
 	Init(InAreaClass);
 	SetConvex(InPoints.GetData(), 0, InPoints.Num(), CoordType, LocalToWorld);
 }
 
-FAreaNavModifier::FAreaNavModifier(const TArray<FVector>& InPoints, const int32 FirstIndex, const int32 LastIndex, ENavigationCoordSystem::Type CoordType, const FTransform& LocalToWorld, const TSubclassOf<UNavArea> InAreaClass)
+FAreaNavModifier::FAreaNavModifier(const TArray<FVector>& InPoints, const int32 FirstIndex, const int32 LastIndex, ENavigationCoordSystem::Type CoordType, const FTransform& LocalToWorld, const TSubclassOf<UNavAreaBase> InAreaClass)
 {
 	check(InPoints.IsValidIndex(FirstIndex) && InPoints.IsValidIndex(LastIndex-1));
 
@@ -322,7 +318,7 @@ FAreaNavModifier::FAreaNavModifier(const TArray<FVector>& InPoints, const int32
 	SetConvex(InPoints.GetData(), FirstIndex, LastIndex, CoordType, LocalToWorld);
 }
 
-FAreaNavModifier::FAreaNavModifier(const TNavStatArray<FVector>& InPoints, const int32 FirstIndex, const int32 LastIndex, ENavigationCoordSystem::Type CoordType, const FTransform& LocalToWorld, const TSubclassOf<UNavArea> InAreaClass)
+FAreaNavModifier::FAreaNavModifier(const TNavStatArray<FVector>& InPoints, const int32 FirstIndex, const int32 LastIndex, ENavigationCoordSystem::Type CoordType, const FTransform& LocalToWorld, const TSubclassOf<UNavAreaBase> InAreaClass)
 {
 	check(InPoints.IsValidIndex(FirstIndex) && InPoints.IsValidIndex(LastIndex-1));
 
@@ -330,7 +326,7 @@ FAreaNavModifier::FAreaNavModifier(const TNavStatArray<FVector>& InPoints, const
 	SetConvex(InPoints.GetData(), FirstIndex, LastIndex, CoordType, LocalToWorld);
 }
 
-FAreaNavModifier::FAreaNavModifier(const UBrushComponent* BrushComponent, const TSubclassOf<UNavArea> InAreaClass)
+FAreaNavModifier::FAreaNavModifier(const UBrushComponent* BrushComponent, const TSubclassOf<UNavAreaBase> InAreaClass)
 {
 	check(BrushComponent != NULL);
 
@@ -375,7 +371,7 @@ void FAreaNavModifier::GetConvex(FConvexNavAreaData& Data) const
 	Data.MaxZ = LastPoint.Y;
 }
 
-void FAreaNavModifier::Init(const TSubclassOf<UNavArea> InAreaClass)
+void FAreaNavModifier::Init(const TSubclassOf<UNavAreaBase> InAreaClass)
 {
 	bIncludeAgentHeight = false;
 	ApplyMode = ENavigationAreaMode::Apply;
@@ -385,24 +381,32 @@ void FAreaNavModifier::Init(const TSubclassOf<UNavArea> InAreaClass)
 	SetAreaClass(InAreaClass);
 }
 
-void FAreaNavModifier::SetAreaClass(const TSubclassOf<UNavArea> InAreaClass)
+bool IsMetaAreaClass(UClass& AreaClass)
+{
+	const UNavAreaBase* AreaClassCDO = GetDefault<UNavAreaBase>(&AreaClass);
+	return AreaClassCDO && AreaClassCDO->IsMetaArea();
+}
+
+void FAreaNavModifier::SetAreaClass(const TSubclassOf<UNavAreaBase> InAreaClass)
 {
 	AreaClassOb = (UClass*)InAreaClass;
 
-	const UClass* AreaClass1 = AreaClassOb.Get();
-	const UClass* AreaClass2 = ReplaceAreaClassOb.Get();
-	bHasMetaAreas = (AreaClass1 && AreaClass1->IsChildOf(UNavAreaMeta::StaticClass())) || (AreaClass2 && AreaClass2->IsChildOf(UNavAreaMeta::StaticClass()));
+	UClass* AreaClass1 = AreaClassOb.Get();
+	UClass* AreaClass2 = ReplaceAreaClassOb.Get();
+	bHasMetaAreas = (AreaClass1 && IsMetaAreaClass(*AreaClass1))
+		|| (AreaClass2 && IsMetaAreaClass(*AreaClass2));
 }
 
-void FAreaNavModifier::SetAreaClassToReplace(const TSubclassOf<UNavArea> InAreaClass)
+void FAreaNavModifier::SetAreaClassToReplace(const TSubclassOf<UNavAreaBase> InAreaClass)
 {
 	ReplaceAreaClassOb = (UClass*)InAreaClass;
 
-	const UClass* AreaClass1 = AreaClassOb.Get();
-	const UClass* AreaClass2 = ReplaceAreaClassOb.Get();
-	bHasMetaAreas = (AreaClass1 && AreaClass1->IsChildOf(UNavAreaMeta::StaticClass())) || (AreaClass2 && AreaClass2->IsChildOf(UNavAreaMeta::StaticClass()));
+	UClass* AreaClass1 = AreaClassOb.Get();
+	UClass* AreaClass2 = ReplaceAreaClassOb.Get();
+	bHasMetaAreas = (AreaClass1 && IsMetaAreaClass(*AreaClass1))
+		|| (AreaClass2 && IsMetaAreaClass(*AreaClass2));
 
-	bIsLowAreaModifier = (AreaClass2 == UNavArea_LowHeight::StaticClass());
+	bIsLowAreaModifier = (AreaClass2 && AreaClass2->GetDefaultObject<UNavAreaBase>()->IsLowArea());
 	ApplyMode = bIsLowAreaModifier ? ENavigationAreaMode::ReplaceInLowPass : AreaClass2 ? ENavigationAreaMode::Replace : ENavigationAreaMode::Apply;
 }
 
@@ -464,9 +468,12 @@ void FAreaNavModifier::SetConvex(const FVector* InPoints, const int32 FirstIndex
 	TArray<FVector, TInlineAllocator<MaxConvexPoints>> HullVertices;
 	HullVertices.Empty(MaxConvexPoints);
 
+	const FTransform& NavCoordTransform = FNavigationSystem::GetCoordTransformTo(CoordType);
+
 	for (int32 i = FirstIndex; i < LastIndex; i++)
 	{
-		const FVector Point = (CoordType == ENavigationCoordSystem::Recast) ? Recast2UnrealPoint(InPoints[i]) : InPoints[i];
+		//const FVector Point = (CoordType == ENavigationCoordSystem::Recast) ? Recast2UnrealPoint(InPoints[i]) : InPoints[i];
+		const FVector Point = NavCoordTransform.TransformPosition(InPoints[i]);
 
 		FVector TransformedPoint = LocalToWorld.TransformPosition(Point);
 		ConvexData.MinZ = FMath::Min( ConvexData.MinZ, TransformedPoint.Z );
@@ -705,8 +712,8 @@ FCompositeNavModifier FCompositeNavModifier::GetInstantiatedMetaModifier(const F
 		{
 			if (Area->HasMetaAreas())
 			{
-				Area->SetAreaClass(UNavAreaMeta::PickAreaClass(Area->GetAreaClass(), ActorOwner, *NavAgent));
-				Area->SetAreaClassToReplace(UNavAreaMeta::PickAreaClass(Area->GetAreaClassToReplace(), ActorOwner, *NavAgent));
+				Area->SetAreaClass(UNavAreaBase::PickAreaClassForAgent(Area->GetAreaClass(), *ActorOwner, *NavAgent));
+				Area->SetAreaClassToReplace(UNavAreaBase::PickAreaClassForAgent(Area->GetAreaClassToReplace(), *ActorOwner, *NavAgent));
 			}
 		}
 	}
@@ -720,13 +727,13 @@ FCompositeNavModifier FCompositeNavModifier::GetInstantiatedMetaModifier(const F
 				for (int32 LinkIndex = 0; LinkIndex < SimpleLink->Links.Num(); ++LinkIndex)
 				{
 					FNavigationLink& Link = SimpleLink->Links[LinkIndex];
-					Link.SetAreaClass(UNavAreaMeta::PickAreaClass(Link.GetAreaClass(), ActorOwner, *NavAgent));
+					Link.SetAreaClass(UNavAreaBase::PickAreaClassForAgent(Link.GetAreaClass(), *ActorOwner, *NavAgent));
 				}
 
 				for (int32 LinkIndex = 0; LinkIndex < SimpleLink->SegmentLinks.Num(); ++LinkIndex)
 				{
 					FNavigationSegmentLink& Link = SimpleLink->SegmentLinks[LinkIndex];
-					Link.SetAreaClass(UNavAreaMeta::PickAreaClass(Link.GetAreaClass(), ActorOwner, *NavAgent));
+					Link.SetAreaClass(UNavAreaBase::PickAreaClassForAgent(Link.GetAreaClass(), *ActorOwner, *NavAgent));
 				}
 			}
 		}
@@ -754,7 +761,7 @@ FCompositeNavModifier FCompositeNavModifier::GetInstantiatedMetaModifier(const F
 				{
 					const int32 AddedIdx = SimpleLink.Links.Add(Links[LinkIndex]);
 					FNavigationLink& NavLink = SimpleLink.Links[AddedIdx];
-					NavLink.SetAreaClass(UNavAreaMeta::PickAreaClass(NavLink.GetAreaClass(), ActorOwner, *NavAgent));
+					NavLink.SetAreaClass(UNavAreaBase::PickAreaClassForAgent(NavLink.GetAreaClass(), *ActorOwner, *NavAgent));
 				}				
 
 				const TArray<FNavigationSegmentLink>& SegmentLinks = UNavLinkDefinition::GetSegmentLinksDefinition(CustomLink->GetNavLinkClass());
@@ -769,7 +776,7 @@ FCompositeNavModifier FCompositeNavModifier::GetInstantiatedMetaModifier(const F
 				{
 					const int32 AddedIdx = SimpleSegLink.SegmentLinks.Add(SegmentLinks[LinkIndex]);
 					FNavigationSegmentLink& NavLink = SimpleSegLink.SegmentLinks[AddedIdx];
-					NavLink.SetAreaClass(UNavAreaMeta::PickAreaClass(NavLink.GetAreaClass(), ActorOwner, *NavAgent));
+					NavLink.SetAreaClass(UNavAreaBase::PickAreaClassForAgent(NavLink.GetAreaClass(), *ActorOwner, *NavAgent));
 				}
 
 				Result.CustomLinks.RemoveAtSwap(Index, 1, false);
@@ -780,7 +787,7 @@ FCompositeNavModifier FCompositeNavModifier::GetInstantiatedMetaModifier(const F
 	return Result;
 }
 
-void FCompositeNavModifier::CreateAreaModifiers(const UPrimitiveComponent* PrimComp, const TSubclassOf<UNavArea> AreaClass)
+void FCompositeNavModifier::CreateAreaModifiers(const UPrimitiveComponent* PrimComp, const TSubclassOf<UNavAreaBase> AreaClass)
 {
 	UBodySetup* BodySetup = PrimComp ? ((UPrimitiveComponent*)PrimComp)->GetBodySetup() : nullptr;
 	if (BodySetup == nullptr)
@@ -841,11 +848,3 @@ bool FCompositeNavModifier::HasPerInstanceTransforms() const
 {
 	return NavDataPerInstanceTransformDelegate.IsBound();
 }
-
-//----------------------------------------------------------------------//
-// UNavLinkTrivial
-//----------------------------------------------------------------------//
-UNavLinkTrivial::UNavLinkTrivial(const FObjectInitializer& ObjectInitializer) : Super(ObjectInitializer)
-{
-	FNavigationLink& Link = Links[Links.Add(FNavigationLink(FVector(0,100,0), FVector(0,-100,0)))];
-}
diff --git a/Engine/Source/Runtime/Engine/Private/Actor.cpp b/Engine/Source/Runtime/Engine/Private/Actor.cpp
--- a/Engine/Source/Runtime/Engine/Private/Actor.cpp
+++ b/Engine/Source/Runtime/Engine/Private/Actor.cpp
@@ -9,7 +9,7 @@
 #include "TimerManager.h"
 #include "GameFramework/Pawn.h"
 #include "Components/PrimitiveComponent.h"
-#include "AI/Navigation/NavigationSystem.h"
+#include "AI/NavigationSystemBase.h"
 #include "Components/InputComponent.h"
 #include "Engine/Engine.h"
 #include "UObject/UObjectHash.h"
@@ -752,8 +752,8 @@ void AActor::ApplyWorldOffset(const FVector& InOffset, bool bWorldShift)
 	{
 		if (RootComponent != nullptr && RootComponent->IsRegistered())
 		{
-			UNavigationSystem::UpdateNavOctreeBounds(this);
-			UNavigationSystem::UpdateActorAndComponentsInNavOctree(*this);
+			FNavigationSystem::OnActorBoundsChanged(*this);
+			FNavigationSystem::UpdateActorAndComponentData(*this);
 		}
 	}
 }
@@ -850,23 +850,41 @@ float AActor::GetActorTickInterval() const
 bool AActor::Rename( const TCHAR* InName, UObject* NewOuter, ERenameFlags Flags )
 {
 	const bool bRenameTest = ((Flags & REN_Test) != 0);
+	const bool bChangingOuters = (NewOuter && (NewOuter != GetOuter()));
 
-	if (!bRenameTest && NewOuter)
+	if (!bRenameTest && bChangingOuters)
 	{
 		RegisterAllActorTickFunctions(false, true); // unregister all tick functions
 		UnregisterAllComponents();
+
+		if (ULevel* MyLevel = GetLevel())
+		{
+			int32 ActorIndex;
+			if (MyLevel->Actors.Find(this, ActorIndex))
+			{
+				MyLevel->Actors[ActorIndex] = nullptr;
+				MyLevel->ActorsForGC.Remove(this);
+				// TODO: There may need to be some consideration about removing this actor from the level cluster, but that would probably require destroying the entire cluster, so defer for now
+			}
+		}
 	}
 
 	const bool bSuccess = Super::Rename( InName, NewOuter, Flags );
 
-	if (!bRenameTest && NewOuter && NewOuter->IsA<ULevel>())
+	if (!bRenameTest && bChangingOuters)
 	{
-		UWorld* World = NewOuter->GetWorld();
-		if (World && World->bIsWorldInitialized)
+		if (ULevel* MyLevel = GetLevel())
 		{
-			RegisterAllComponents();
+			MyLevel->Actors.Add(this);
+			MyLevel->ActorsForGC.Add(this);
+
+			UWorld* World = MyLevel->GetWorld();
+			if (World && World->bIsWorldInitialized)
+			{
+				RegisterAllComponents();
+			}
+			RegisterAllActorTickFunctions(true, true); // register all tick functions
 		}
-		RegisterAllActorTickFunctions(true, true); // register all tick functions
 	}
 	return bSuccess;
 }
@@ -1485,7 +1503,6 @@ void AActor::SetOwner( AActor *NewOwner )
 		}
 
 		// Sets this actor's parent to the specified actor.
-		AActor* OldOwner = Owner;
 		if( Owner != nullptr )
 		{
 			// remove from old owner's Children array
@@ -1979,7 +1996,7 @@ void AActor::RouteEndPlay(const EEndPlayReason::Type EndPlayReason)
 			SetLifeSpan(0.f);
 		}
 
-		UNavigationSystem::OnActorUnregistered(this);
+		FNavigationSystem::OnActorUnregistered(*this);
 	}
 
 	UninitializeComponents();
@@ -2135,6 +2152,7 @@ float AActor::TakeDamage(float DamageAmount, FDamageEvent const& DamageEvent, AC
 		{
 			FHitResult const& Hit = (RadialDamageEvent->ComponentHits.Num() > 0) ? RadialDamageEvent->ComponentHits[0] : FHitResult();
 			ReceiveRadialDamage(ActualDamage, DamageTypeCDO, RadialDamageEvent->Origin, Hit, EventInstigator, DamageCauser);
+			OnTakeRadialDamage.Broadcast(this, ActualDamage, DamageTypeCDO, RadialDamageEvent->Origin, Hit, EventInstigator, DamageCauser);
 
 			// add any desired physics impulses to our components
 			for (int HitIdx = 0; HitIdx < RadialDamageEvent->ComponentHits.Num(); ++HitIdx)
@@ -4406,7 +4424,7 @@ void AActor::PostInitializeComponents()
 	{
 		bActorInitialized = true;
 
-		UNavigationSystem::OnActorRegistered(this);
+		FNavigationSystem::OnActorRegistered(*this);
 		
 		UpdateAllReplicatedComponents();
 	}
diff --git a/Engine/Source/Runtime/Engine/Private/ActorEditor.cpp b/Engine/Source/Runtime/Engine/Private/ActorEditor.cpp
--- a/Engine/Source/Runtime/Engine/Private/ActorEditor.cpp
+++ b/Engine/Source/Runtime/Engine/Private/ActorEditor.cpp
@@ -10,7 +10,7 @@
 #include "GameFramework/Actor.h"
 #include "Components/ChildActorComponent.h"
 #include "Components/PrimitiveComponent.h"
-#include "AI/Navigation/NavigationSystem.h"
+#include "AI/NavigationSystemBase.h"
 #include "Engine/BlueprintGeneratedClass.h"
 #include "EditorSupportDelegates.h"
 #include "Logging/TokenizedMessage.h"
@@ -172,19 +172,7 @@ void AActor::PostEditMove(bool bFinished)
 
 	if (bFinished)
 	{
-		// update actor and all its components in navigation system after finishing move
-		// USceneComponent::UpdateNavigationData works only in game world
-		UNavigationSystem::UpdateNavOctreeBounds(this);
-
-		TArray<AActor*> ParentedActors;
-		GetAttachedActors(ParentedActors);
-		for (int32 Idx = 0; Idx < ParentedActors.Num(); Idx++)
-		{
-			UNavigationSystem::UpdateNavOctreeBounds(ParentedActors[Idx]);
-		}
-
-		// not doing manual update of all attached actors since UpdateActorAndComponentsInNavOctree should take care of it
-		UNavigationSystem::UpdateActorAndComponentsInNavOctree(*this);
+		FNavigationSystem::OnPostEditActorMove(*this);
 	}
 }
 
@@ -372,7 +360,7 @@ void AActor::PreEditUndo()
 	}
 
 	// let navigation system know to not care about this actor anymore
-	UNavigationSystem::ClearNavOctreeAll(this);
+	FNavigationSystem::RemoveActorData(*this);
 
 	Super::PreEditUndo();
 }
@@ -411,11 +399,11 @@ bool AActor::InternalPostEditUndo()
 		BlueprintCreatedComponents.Reset();
 
 		// notify navigation system
-		UNavigationSystem::UpdateActorAndComponentsInNavOctree(*this);
+		FNavigationSystem::UpdateActorAndComponentData(*this);
 	}
 	else
 	{
-		UNavigationSystem::ClearNavOctreeAll(this);
+		FNavigationSystem::RemoveActorData(*this);
 	}
 
 	// This is a normal undo, so call super
diff --git a/Engine/Source/Runtime/Engine/Private/Blueprint.cpp b/Engine/Source/Runtime/Engine/Private/Blueprint.cpp
--- a/Engine/Source/Runtime/Engine/Private/Blueprint.cpp
+++ b/Engine/Source/Runtime/Engine/Private/Blueprint.cpp
@@ -1810,7 +1810,7 @@ UEdGraph* UBlueprint::GetLastEditedUberGraph() const
 {
 	for ( int32 LastEditedIndex = LastEditedDocuments.Num() - 1; LastEditedIndex >= 0; LastEditedIndex-- )
 	{
-		if ( UObject* Obj = LastEditedDocuments[LastEditedIndex].EditedObject )
+		if ( UObject* Obj = LastEditedDocuments[LastEditedIndex].EditedObjectPath.ResolveObject() )
 		{
 			if ( UEdGraph* Graph = Cast<UEdGraph>(Obj) )
 			{
diff --git a/Engine/Source/Runtime/Engine/Private/BlueprintGeneratedClass.cpp b/Engine/Source/Runtime/Engine/Private/BlueprintGeneratedClass.cpp
--- a/Engine/Source/Runtime/Engine/Private/BlueprintGeneratedClass.cpp
+++ b/Engine/Source/Runtime/Engine/Private/BlueprintGeneratedClass.cpp
@@ -30,7 +30,7 @@
 DEFINE_STAT(STAT_PersistentUberGraphFrameMemory);
 
 int32 GBlueprintClusteringEnabled = 0;
-static FAutoConsoleVariableRef CVarUseBackgroundLevelStreaming(
+static FAutoConsoleVariableRef CVarBlueprintClusteringEnabled(
 	TEXT("gc.BlueprintClusteringEnabled"),
 	GBlueprintClusteringEnabled,
 	TEXT("Whether to allow Blueprint classes to create GC clusters."),
diff --git a/Engine/Source/Runtime/Engine/Private/Character.cpp b/Engine/Source/Runtime/Engine/Private/Character.cpp
--- a/Engine/Source/Runtime/Engine/Private/Character.cpp
+++ b/Engine/Source/Runtime/Engine/Private/Character.cpp
@@ -1407,6 +1407,12 @@ void ACharacter::PreReplication( IRepChangedPropertyTracker & ChangedPropertyTra
 			ReplicatedBasedMovement.Rotation = GetActorRotation();
 		}
 	}
+
+	// Save bandwidth by not replicating this value unless it is necessary, since it changes every update.
+	if ((CharacterMovement->NetworkSmoothingMode != ENetworkSmoothingMode::Linear) && !CharacterMovement->bNetworkAlwaysReplicateTransformUpdateTimestamp)
+	{
+		ReplicatedServerLastTransformUpdateTimeStamp = 0.f;
+	}
 }
 
 void ACharacter::PreReplicationForReplay(IRepChangedPropertyTracker & ChangedPropertyTracker)
diff --git a/Engine/Source/Runtime/Engine/Private/CheatManager.cpp b/Engine/Source/Runtime/Engine/Private/CheatManager.cpp
--- a/Engine/Source/Runtime/Engine/Private/CheatManager.cpp
+++ b/Engine/Source/Runtime/Engine/Private/CheatManager.cpp
@@ -17,7 +17,7 @@
 #include "CollisionQueryParams.h"
 #include "WorldCollision.h"
 #include "Engine/World.h"
-#include "AI/Navigation/NavigationSystem.h"
+#include "AI/NavigationSystemBase.h"
 #include "UObject/Package.h"
 #include "GameFramework/PlayerController.h"
 #include "GameFramework/Volume.h"
@@ -29,7 +29,6 @@
 #include "EngineUtils.h"
 #include "Net/OnlineEngineInterface.h"
 #include "VisualLogger/VisualLogger.h"
-#include "AI/Navigation/RecastNavMesh.h"
 #include "GameFramework/Character.h"
 #include "Engine/Console.h"
 #include "Engine/DebugCameraController.h"
@@ -915,28 +914,6 @@ void UCheatManager::TestCollisionDistance()
 #endif
 }
 
-void UCheatManager::RebuildNavigation()
-{
-	UNavigationSystem* NavSys = UNavigationSystem::GetCurrent(GetWorld());
-	if (NavSys)
-	{
-		NavSys->Build();
-	}
-}
-
-void UCheatManager::SetNavDrawDistance(float DrawDistance)
-{
-	if (GIsEditor)
-	{
-		APlayerController* PC = GetOuterAPlayerController();
-		if (PC != NULL)
-		{
-			PC->ClientMessage(TEXT("Setting Nav Rendering Draw Distance is not supported while in Edior"));
-		}
-	}
-	ARecastNavMesh::SetDrawDistance(DrawDistance);
-}
-
 void UCheatManager::DumpOnlineSessionState()
 {
 	UOnlineEngineInterface::Get()->DumpSessionState(GetWorld());
@@ -967,6 +944,11 @@ void UCheatManager::DumpVoiceMutingState()
 
 UWorld* UCheatManager::GetWorld() const
 {
+	if (IsTemplate())
+	{
+		return nullptr;
+	}
+
 	return GetOuterAPlayerController()->GetWorld();
 }
 
diff --git a/Engine/Source/Runtime/Engine/Private/Components/ActorComponent.cpp b/Engine/Source/Runtime/Engine/Private/Components/ActorComponent.cpp
--- a/Engine/Source/Runtime/Engine/Private/Components/ActorComponent.cpp
+++ b/Engine/Source/Runtime/Engine/Private/Components/ActorComponent.cpp
@@ -1,4 +1,4 @@
-﻿// Copyright 1998-2018 Epic Games, Inc. All Rights Reserved.
+// Copyright 1998-2018 Epic Games, Inc. All Rights Reserved.
 // ActorComponent.cpp: Actor component implementation.
 
 #include "Components/ActorComponent.h"
@@ -11,7 +11,7 @@
 #include "GameFramework/Actor.h"
 #include "Engine/World.h"
 #include "Components/PrimitiveComponent.h"
-#include "AI/Navigation/NavigationSystem.h"
+#include "AI/NavigationSystemBase.h"
 #include "Engine/BlueprintGeneratedClass.h"
 #include "ContentStreaming.h"
 #include "ComponentReregisterContext.h"
@@ -1017,10 +1017,12 @@ void UActorComponent::RegisterComponentWithWorld(UWorld* InWorld)
 
 		for (UObject* Child : Children)
 		{
-			UActorComponent* ChildComponent = Cast<UActorComponent>(Child);
-			if (ChildComponent && !ChildComponent->IsRegistered() && ChildComponent->GetOwner() == MyOwner)
+			if (UActorComponent* ChildComponent = Cast<UActorComponent>(Child))
 			{
-				ChildComponent->RegisterComponentWithWorld(InWorld);
+				if (ChildComponent->bAutoRegister && !ChildComponent->IsRegistered() && ChildComponent->GetOwner() == MyOwner)
+				{
+					ChildComponent->RegisterComponentWithWorld(InWorld);
+				}
 			}
 		}
 
@@ -1780,11 +1782,11 @@ void UActorComponent::HandleCanEverAffectNavigationChange(bool bForceUpdate)
 		if (bCanEverAffectNavigation)
 		{
 			bNavigationRelevant = IsNavigationRelevant();
-			UNavigationSystem::OnComponentRegistered(this);
+			FNavigationSystem::OnComponentRegistered(*this);
 		}
 		else
 		{
-			UNavigationSystem::OnComponentUnregistered(this);
+			FNavigationSystem::OnComponentUnregistered(*this);
 		}
 	}
 }
diff --git a/Engine/Source/Runtime/Engine/Private/Components/CharacterMovementComponent.cpp b/Engine/Source/Runtime/Engine/Private/Components/CharacterMovementComponent.cpp
--- a/Engine/Source/Runtime/Engine/Private/Components/CharacterMovementComponent.cpp
+++ b/Engine/Source/Runtime/Engine/Private/Components/CharacterMovementComponent.cpp
@@ -8,7 +8,8 @@
 #include "GameFramework/CharacterMovementComponent.h"
 #include "EngineStats.h"
 #include "Components/PrimitiveComponent.h"
-#include "AI/Navigation/NavigationSystem.h"
+#include "AI/NavigationSystemBase.h"
+#include "AI/Navigation/NavigationDataInterface.h"
 #include "UObject/Package.h"
 #include "GameFramework/PlayerController.h"
 #include "GameFramework/PhysicsVolume.h"
@@ -20,10 +21,7 @@
 #include "Components/CapsuleComponent.h"
 #include "GameFramework/GameStateBase.h"
 #include "Engine/Canvas.h"
-
-// @todo this is here only due to circular dependency to AIModule. To be removed
-#include "Navigation/PathFollowingComponent.h"
-#include "AI/Navigation/RecastNavMesh.h"
+#include "AI/Navigation/PathFollowingAgentInterface.h"
 #include "AI/Navigation/AvoidanceManager.h"
 #include "Components/BrushComponent.h"
 
@@ -1012,9 +1010,13 @@ void UCharacterMovementComponent::OnMovementModeChanged(EMovementMode PreviousMo
 		}
 	}
 
-	if (MovementMode == MOVE_Falling && PreviousMovementMode != MOVE_Falling && PathFollowingComp.IsValid())
+	if (MovementMode == MOVE_Falling && PreviousMovementMode != MOVE_Falling)
 	{
-		PathFollowingComp->OnStartedFalling();
+		IPathFollowingAgentInterface* PFAgent = GetPathFollowingAgent();
+		if (PFAgent)
+		{
+			PFAgent->OnStartedFalling();
+		}
 	}
 
 	CharacterOwner->OnMovementModeChanged(PreviousMovementMode, PreviousCustomMode);
@@ -2382,7 +2384,7 @@ bool UCharacterMovementComponent::CanCrouchInCurrentState() const
 		return false;
 	}
 
-	return IsFalling() || IsMovingOnGround();
+	return (IsFalling() || IsMovingOnGround()) && UpdatedComponent && !UpdatedComponent->IsSimulatingPhysics();
 }
 
 
@@ -2626,12 +2628,12 @@ void UCharacterMovementComponent::UnCrouch(bool bClientSimulation)
 void UCharacterMovementComponent::UpdateCharacterStateBeforeMovement()
 {
 	// Check for a change in crouch state. Players toggle crouch by changing bWantsToCrouch.
-	const bool bAllowedToCrouch = CanCrouchInCurrentState();
-	if ((!bAllowedToCrouch || !bWantsToCrouch) && IsCrouching())
+	const bool bIsCrouching = IsCrouching();
+	if (bIsCrouching && (!bWantsToCrouch || !CanCrouchInCurrentState()))
 	{
 		UnCrouch(false);
 	}
-	else if (bWantsToCrouch && bAllowedToCrouch && !IsCrouching())
+	else if (!bIsCrouching && bWantsToCrouch && CanCrouchInCurrentState())
 	{
 		Crouch(false);
 	}
@@ -3278,15 +3280,12 @@ void UCharacterMovementComponent::UpdateDefaultAvoidance()
 	SCOPE_CYCLE_COUNTER(STAT_AI_ObstacleAvoidance);
 
 	UAvoidanceManager* AvoidanceManager = GetWorld()->GetAvoidanceManager();
-	if (AvoidanceManager && !bWasAvoidanceUpdated)
+	if (AvoidanceManager && !bWasAvoidanceUpdated && GetCharacterOwner()->GetCapsuleComponent())
 	{
-		if (UCapsuleComponent *OurCapsule = GetCharacterOwner()->GetCapsuleComponent())
-		{
-			AvoidanceManager->UpdateRVO(this);
+		AvoidanceManager->UpdateRVO(this);
 
-			//Consider this a clean move because we didn't even try to avoid.
-			SetAvoidanceVelocityLock(AvoidanceManager, AvoidanceManager->LockTimeAfterClean);
-		}
+		//Consider this a clean move because we didn't even try to avoid.
+		SetAvoidanceVelocityLock(AvoidanceManager, AvoidanceManager->LockTimeAfterClean);
 	}
 
 	bWasAvoidanceUpdated = false;		//Reset for next frame
@@ -4860,7 +4859,7 @@ void UCharacterMovementComponent::PhysNavWalking(float deltaTime, int32 Iteratio
 
 bool UCharacterMovementComponent::FindNavFloor(const FVector& TestLocation, FNavLocation& NavFloorLocation) const
 {
-	const ANavigationData* NavData = GetNavData();
+	const INavigationDataInterface* NavData = GetNavData();
 	if (NavData == nullptr)
 	{
 		return false;
@@ -5029,32 +5028,17 @@ void UCharacterMovementComponent::FindBestNavMeshLocation(const FVector& TraceSt
 	}
 }
 
-const ANavigationData* UCharacterMovementComponent::GetNavData() const
+const INavigationDataInterface* UCharacterMovementComponent::GetNavData() const
 {
-	UNavigationSystem* NavSys = UNavigationSystem::GetCurrent(GetWorld());
-	if (NavSys == nullptr || !HasValidData())
+	const UWorld* World = GetWorld();
+	if (World == nullptr || World->GetNavigationSystem() == nullptr 
+		|| !HasValidData()
+		|| CharacterOwner == nullptr)
 	{
 		return nullptr;
 	}
 
-	const ANavigationData* NavData = nullptr;
-	INavAgentInterface* MyNavAgent = CastChecked<INavAgentInterface>(CharacterOwner);
-	if (MyNavAgent)
-	{
-		const FNavAgentProperties& AgentProps = MyNavAgent->GetNavAgentPropertiesRef();
-		NavData = NavSys->GetNavDataForProps(AgentProps);
-	}
-	if (NavData == nullptr)
-	{
-		NavData = NavSys->GetMainNavData();
-	}
-
-	// Only RecastNavMesh supported
-	const ARecastNavMesh* NavMeshData = Cast<const ARecastNavMesh>(NavData);
-	if (NavMeshData == nullptr)
-	{
-		return nullptr;
-	}
+	const INavigationDataInterface* NavData = FNavigationSystem::GetNavDataForActor(*CharacterOwner);
 
 	return NavData;
 }
@@ -5178,9 +5162,11 @@ void UCharacterMovementComponent::ProcessLanded(const FHitResult& Hit, float rem
 
 		SetPostLandedPhysics(Hit);
 	}
-	if (PathFollowingComp.IsValid())
+	
+	IPathFollowingAgentInterface* PFAgent = GetPathFollowingAgent();
+	if (PFAgent)
 	{
-		PathFollowingComp->OnLanded();
+		PFAgent->OnLanded();
 	}
 
 	StartNewPhysics(remainingTime, Iterations);
@@ -5461,9 +5447,11 @@ void UCharacterMovementComponent::PhysicsVolumeChanged( APhysicsVolume* NewVolum
 		if ( !CanEverSwim() )
 		{
 			// AI needs to stop any current moves
-			if (PathFollowingComp.IsValid())
+			IPathFollowingAgentInterface* PFAgent = GetPathFollowingAgent();
+			if (PFAgent)
 			{
-				PathFollowingComp->AbortMove(*this, FPathFollowingResultFlags::MovementStop);
+				//PathFollowingComp->AbortMove(*this, FPathFollowingResultFlags::MovementStop);
+				PFAgent->OnUnableToMove(*this);
 			}			
 		}
 		else if ( !IsSwimming() )
@@ -6424,9 +6412,11 @@ void UCharacterMovementComponent::HandleImpact(const FHitResult& Impact, float T
 		CharacterOwner->MoveBlockedBy(Impact);
 	}
 
-	if (PathFollowingComp.IsValid())
-	{	// Also notify path following!
-		PathFollowingComp->OnMoveBlockedBy(Impact);
+	IPathFollowingAgentInterface* PFAgent = GetPathFollowingAgent();
+	if (PFAgent)
+	{
+		// Also notify path following!
+		PFAgent->OnMoveBlockedBy(Impact);
 	}
 
 	APawn* OtherPawn = Cast<APawn>(Impact.GetActor());
@@ -10189,7 +10179,6 @@ void FSavedMove_Character::PostUpdate(ACharacter* Character, FSavedMove_Characte
 		if ((StartActorOverlapCounter != EndActorOverlapCounter) || (StartComponentOverlapCounter != EndComponentOverlapCounter))
 		{
 			bForceNoCombine = true;
-			UE_LOG(LogNetPlayerMovement, Log, TEXT("SavedMove: Setting bForceNoCombine true because of an overlap event!"));
 		}
 	}
 	else if (PostUpdateMode == PostUpdate_Replay)
@@ -10256,14 +10245,23 @@ float UCharacterMovementComponent::GetClientNetSendDeltaTime(const APlayerContro
 	const AGameNetworkManager* const GameNetworkManager = GetDefault<AGameNetworkManager>();
 	float NetMoveDelta = GameNetworkManager->ClientNetSendMoveDeltaTime;
 
-	// send moves more frequently in small games where server isn't likely to be saturated
-	if (Player && (Player->CurrentNetSpeed > GameNetworkManager->ClientNetSendMoveThrottleAtNetSpeed) && (GameState != nullptr) && (GameState->PlayerArray.Num() <= GameNetworkManager->ClientNetSendMoveThrottleOverPlayerCount))
-	{
-		NetMoveDelta = GameNetworkManager->ClientNetSendMoveDeltaTime;
-	}
-	else if (Player)
+	if (PC && Player)
 	{
-		NetMoveDelta = FMath::Max(GameNetworkManager->ClientNetSendMoveDeltaTimeThrottled, 2 * GameNetworkManager->MoveRepSize / Player->CurrentNetSpeed);
+		// send moves more frequently in small games where server isn't likely to be saturated
+		if ((Player->CurrentNetSpeed > GameNetworkManager->ClientNetSendMoveThrottleAtNetSpeed) && (GameState != nullptr) && (GameState->PlayerArray.Num() <= GameNetworkManager->ClientNetSendMoveThrottleOverPlayerCount))
+		{
+			NetMoveDelta = GameNetworkManager->ClientNetSendMoveDeltaTime;
+		}
+		else
+		{
+			NetMoveDelta = FMath::Max(GameNetworkManager->ClientNetSendMoveDeltaTimeThrottled, 2 * GameNetworkManager->MoveRepSize / Player->CurrentNetSpeed);
+		}
+
+		// Lower frequency for standing still and not rotating camera
+		if (Acceleration.IsZero() && Velocity.IsZero() && ClientData->LastAckedMove.IsValid() && ClientData->LastAckedMove->StartControlRotation.Equals(PC->GetControlRotation()))
+		{
+			NetMoveDelta = FMath::Max(GameNetworkManager->ClientNetSendMoveDeltaTimeStationary, NetMoveDelta);
+		}
 	}
 	
 	return NetMoveDelta;
diff --git a/Engine/Source/Runtime/Engine/Private/Components/NavMovementComponent.cpp b/Engine/Source/Runtime/Engine/Private/Components/NavMovementComponent.cpp
--- a/Engine/Source/Runtime/Engine/Private/Components/NavMovementComponent.cpp
+++ b/Engine/Source/Runtime/Engine/Private/Components/NavMovementComponent.cpp
@@ -1,25 +1,10 @@
 // Copyright 1998-2018 Epic Games, Inc. All Rights Reserved.
 
 #include "GameFramework/NavMovementComponent.h"
-#include "AI/Navigation/NavigationSystem.h"
+#include "AI/NavigationSystemBase.h"
 #include "Components/CapsuleComponent.h"
-#include "Navigation/PathFollowingComponent.h"
+#include "AI/Navigation/PathFollowingAgentInterface.h"
 
-//----------------------------------------------------------------------//
-// FNavAgentProperties
-//----------------------------------------------------------------------//
-const FNavAgentProperties FNavAgentProperties::DefaultProperties;
-
-void FNavAgentProperties::UpdateWithCollisionComponent(UShapeComponent* CollisionComponent)
-{
-	check( CollisionComponent != NULL);
-	AgentRadius = CollisionComponent->Bounds.SphereRadius;
-}
-
-bool FNavAgentProperties::IsNavDataMatching(const FNavAgentProperties& Other) const
-{
-	return (PreferredNavData == Other.PreferredNavData || PreferredNavData == nullptr || Other.PreferredNavData == nullptr);
-}
 
 //----------------------------------------------------------------------//
 // UMovementComponent
@@ -74,9 +59,15 @@ void UNavMovementComponent::ClearFixedBrakingDistance()
 
 void UNavMovementComponent::StopActiveMovement()
 {
-	if (PathFollowingComp.IsValid() && bStopMovementAbortPaths)
+	if (!bStopMovementAbortPaths)
+	{
+		return;
+	}
+
+	IPathFollowingAgentInterface* PFAgent = GetPathFollowingAgent();
+	if (PFAgent)
 	{
-		PathFollowingComp->AbortMove(*this, FPathFollowingResultFlags::MovementStop);
+		PFAgent->OnUnableToMove(*this);
 	}
 }
 
@@ -89,12 +80,8 @@ void UNavMovementComponent::UpdateNavAgent(const AActor& Owner)
 	}
 
 	// initialize properties from navigation system
-	UNavigationSystem* NavSys = UNavigationSystem::GetCurrent(GetWorld());
-	if (NavSys != nullptr)
-	{
-		NavAgentProps.NavWalkingSearchHeightScale = NavSys->GetDefaultSupportedAgentConfig().NavWalkingSearchHeightScale;
-	}
-
+	NavAgentProps.NavWalkingSearchHeightScale = FNavigationSystem::GetDefaultSupportedAgent().NavWalkingSearchHeightScale;
+	
 	// Can't call GetSimpleCollisionCylinder(), because no components will be registered.
 	float BoundRadius, BoundHalfHeight;	
 	Owner.GetSimpleCollisionCylinder(BoundRadius, BoundHalfHeight);
@@ -110,11 +97,7 @@ void UNavMovementComponent::UpdateNavAgent(const UCapsuleComponent& CapsuleCompo
 	}
 
 	// initialize properties from navigation system
-	UNavigationSystem* NavSys = UNavigationSystem::GetCurrent(GetWorld());
-	if (NavSys != nullptr)
-	{
-		NavAgentProps.NavWalkingSearchHeightScale = NavSys->GetDefaultSupportedAgentConfig().NavWalkingSearchHeightScale;
-	}
+	NavAgentProps.NavWalkingSearchHeightScale = FNavigationSystem::GetDefaultSupportedAgent().NavWalkingSearchHeightScale;
 
 	NavAgentProps.AgentRadius = CapsuleComponent.GetScaledCapsuleRadius();
 	NavAgentProps.AgentHeight = CapsuleComponent.GetScaledCapsuleHalfHeight() * 2.f;
@@ -124,3 +107,4 @@ void UNavMovementComponent::SetUpdateNavAgentWithOwnersCollisions(bool bUpdateWi
 {
 	bUpdateNavAgentWithOwnersCollision = bUpdateWithOwner;
 }
+
diff --git a/Engine/Source/Runtime/Engine/Private/Components/PrimitiveComponent.cpp b/Engine/Source/Runtime/Engine/Private/Components/PrimitiveComponent.cpp
--- a/Engine/Source/Runtime/Engine/Private/Components/PrimitiveComponent.cpp
+++ b/Engine/Source/Runtime/Engine/Private/Components/PrimitiveComponent.cpp
@@ -9,7 +9,7 @@
 #include "GameFramework/DamageType.h"
 #include "GameFramework/Pawn.h"
 #include "WorldCollision.h"
-#include "AI/Navigation/NavigationSystem.h"
+#include "AI/NavigationSystemBase.h"
 #include "GameFramework/PlayerController.h"
 #include "GameFramework/PhysicsVolume.h"
 #include "GameFramework/WorldSettings.h"
@@ -536,7 +536,7 @@ void UPrimitiveComponent::OnRegister()
 		const bool bNavRelevant = bNavigationRelevant = IsNavigationRelevant();
 		if (bNavRelevant)
 		{
-			UNavigationSystem::OnComponentRegistered(this);
+			FNavigationSystem::OnComponentRegistered(*this);
 		}
 	}
 	else
@@ -568,7 +568,7 @@ void UPrimitiveComponent::OnUnregister()
 
 	if (bCanEverAffectNavigation)
 	{
-		UNavigationSystem::OnComponentUnregistered(this);
+		FNavigationSystem::OnComponentUnregistered(*this);
 	}
 }
 
diff --git a/Engine/Source/Runtime/Engine/Private/Components/SceneComponent.cpp b/Engine/Source/Runtime/Engine/Private/Components/SceneComponent.cpp
--- a/Engine/Source/Runtime/Engine/Private/Components/SceneComponent.cpp
+++ b/Engine/Source/Runtime/Engine/Private/Components/SceneComponent.cpp
@@ -13,7 +13,7 @@
 #include "WorldCollision.h"
 #include "Components/PrimitiveComponent.h"
 #include "Components/StaticMeshComponent.h"
-#include "AI/Navigation/NavigationSystem.h"
+#include "AI/NavigationSystemBase.h"
 #include "Engine/MapBuildDataRegistry.h"
 #include "GameFramework/PhysicsVolume.h"
 #include "Components/BillboardComponent.h"
@@ -508,7 +508,7 @@ void USceneComponent::PostEditChangeProperty(FPropertyChangedEvent& PropertyChan
 
 	if (PropertyName == LocationName || PropertyName == RotationName || PropertyName == ScaleName)
 	{
-		UNavigationSystem::UpdateComponentInNavOctree(*this);
+		FNavigationSystem::UpdateComponentData(*this);
 	}
 }
 
@@ -1824,11 +1824,11 @@ bool USceneComponent::AttachToComponent(USceneComponent* Parent, const FAttachme
 		{
 			if (Parent->IsTemplate())
 			{
-				ensureMsgf(false, TEXT("Template Mismatch during attachment. Attaching instanced component to template component. Parent '%s' Self '%s'"), *Parent->GetName(), *GetName());
+				ensureMsgf(false, TEXT("Template Mismatch during attachment. Attaching instanced component to template component. Parent '%s' (Owner '%s') Self '%s' (Owner '%s')."), *Parent->GetName(), *GetNameSafe(Parent->GetOwner()), *GetName(), *GetNameSafe(GetOwner()));
 			}
 			else
 			{
-				ensureMsgf(false, TEXT("Template Mismatch during attachment. Attaching template component to instanced component. Parent '%s' Self '%s'"), *Parent->GetName(), *GetName());
+				ensureMsgf(false, TEXT("Template Mismatch during attachment. Attaching template component to instanced component. Parent '%s' (Owner '%s') Self '%s' (Owner '%s')."), *Parent->GetName(), *GetNameSafe(Parent->GetOwner()), *GetName(), *GetNameSafe(GetOwner()));
 			}
 			return false;
 		}
@@ -2728,11 +2728,7 @@ bool USceneComponent::InternalSetWorldLocationAndRotation(FVector NewLocation, c
 		UpdateComponentToWorldWithParent(GetAttachParent(),GetAttachSocketName(), SkipPhysicsToEnum(bNoPhysics), RelativeRotationCache.GetCachedQuat(), Teleport);
 
 		// we need to call this even if this component itself is not navigation relevant
-		// checking ShouldUpdateNavOctreeOnComponentChange here is an optimization for static navigation users
-		if (UNavigationSystem::ShouldUpdateNavOctreeOnComponentChange())
-		{
-			PostUpdateNavigationData();
-		}
+		PostUpdateNavigationData();
 
 		return true;
 	}
@@ -3357,7 +3353,7 @@ bool FScopedMovementUpdate::IsTransformDirty() const
 {
 	if (IsValid(Owner))
 	{
-		return !InitialTransform.Equals(Owner->GetComponentToWorld());
+		return !InitialTransform.Equals(Owner->GetComponentToWorld(), SMALL_NUMBER);
 	}
 
 	return false;
@@ -3579,13 +3575,13 @@ void USceneComponent::UpdateNavigationData()
 {
 	SCOPE_CYCLE_COUNTER(STAT_ComponentUpdateNavData);
 
-	if (UNavigationSystem::ShouldUpdateNavOctreeOnComponentChange() && IsRegistered())
+	if (IsRegistered())
 	{
 		UWorld* MyWorld = GetWorld();
-		if ((MyWorld == nullptr) || !MyWorld->IsGameWorld() || !MyWorld->IsNetMode(ENetMode::NM_Client))
+		if ((MyWorld != nullptr) && (!MyWorld->IsGameWorld() || !MyWorld->IsNetMode(ENetMode::NM_Client)))
 		{
 			// use propagated component's transform update in editor OR server game with additional navsys check
-			UNavigationSystem::UpdateComponentInNavOctree(*this);
+			FNavigationSystem::UpdateComponentData(*this);
 		}
 	}
 }
@@ -3594,15 +3590,9 @@ void USceneComponent::PostUpdateNavigationData()
 {
 	SCOPE_CYCLE_COUNTER(STAT_ComponentPostUpdateNavData);
 
-	if (UNavigationSystem::ShouldUpdateNavOctreeOnComponentChange() && IsRegistered())
+	if (IsRegistered())
 	{
-		UWorld* MyWorld = GetWorld();
-		if (MyWorld != nullptr && MyWorld->GetNavigationSystem() != nullptr 
-			&& (MyWorld->GetNavigationSystem()->ShouldAllowClientSideNavigation() || !MyWorld->IsNetMode(ENetMode::NM_Client)))
-		{
-			// use propagated component's transform update in editor OR server game with additional navsys check
-			UNavigationSystem::UpdateNavOctreeAfterMove(this);
-		}
+		FNavigationSystem::OnComponentTransformChanged(*this);
 	}
 }
 
diff --git a/Engine/Source/Runtime/Engine/Private/Components/ShapeComponent.cpp b/Engine/Source/Runtime/Engine/Private/Components/ShapeComponent.cpp
--- a/Engine/Source/Runtime/Engine/Private/Components/ShapeComponent.cpp
+++ b/Engine/Source/Runtime/Engine/Private/Components/ShapeComponent.cpp
@@ -2,8 +2,7 @@
 
 
 #include "Components/ShapeComponent.h"
-#include "AI/Navigation/NavAreas/NavArea_Obstacle.h"
-#include "AI/NavigationOctree.h"
+#include "AI/Navigation/NavAreaBase.h"
 #include "PhysicsEngine/BoxElem.h"
 #include "PhysicsEngine/SphereElem.h"
 #include "PhysicsEngine/SphylElem.h"
@@ -39,7 +38,8 @@ UShapeComponent::UShapeComponent(const FObjectInitializer& ObjectInitializer)
 	bHasCustomNavigableGeometry = EHasCustomNavigableGeometry::Yes;
 	bCanEverAffectNavigation = true;
 	bDynamicObstacle = false;
-	AreaClass = UNavArea_Obstacle::StaticClass();
+	// mz@todo
+	//AreaClass = UNavArea_Obstacle::StaticClass();
 }
 
 FPrimitiveSceneProxy* UShapeComponent::CreateSceneProxy()
diff --git a/Engine/Source/Runtime/Engine/Private/Components/SplineMeshComponent.cpp b/Engine/Source/Runtime/Engine/Private/Components/SplineMeshComponent.cpp
--- a/Engine/Source/Runtime/Engine/Private/Components/SplineMeshComponent.cpp
+++ b/Engine/Source/Runtime/Engine/Private/Components/SplineMeshComponent.cpp
@@ -11,7 +11,7 @@
 #include "SplineMeshSceneProxy.h"
 #include "ShaderParameterUtils.h"
 #include "AI/NavigationSystemHelpers.h"
-#include "AI/Navigation/NavCollision.h"
+#include "AI/Navigation/NavCollisionBase.h"
 #include "Engine/StaticMeshSocket.h"
 #include "Engine/StaticMesh.h"
 #include "PhysicsEngine/ConvexElem.h"
@@ -941,37 +941,37 @@ bool USplineMeshComponent::DoCustomNavigableGeometryExport(FNavigableGeometryExp
 
 	if (GetStaticMesh() != nullptr && GetStaticMesh()->NavCollision != nullptr)
 	{
-		UNavCollision* NavCollision = GetStaticMesh()->NavCollision;
+		UNavCollisionBase* NavCollision = GetStaticMesh()->NavCollision;
 		
-		if (ensure(!NavCollision->bIsDynamicObstacle))
+		if (ensure(!NavCollision->IsDynamicObstacle()))
 		{
-			if (NavCollision->bHasConvexGeometry)
+			if (NavCollision->HasConvexGeometry())
 			{
 				FVector Mask = FVector(1, 1, 1);
 				GetAxisValue(Mask, ForwardAxis) = 0;
 
 				TArray<FVector> VertexBuffer;
-				VertexBuffer.Reserve(FMath::Max(NavCollision->ConvexCollision.VertexBuffer.Num(), NavCollision->TriMeshCollision.VertexBuffer.Num()));
+				VertexBuffer.Reserve(FMath::Max(NavCollision->GetConvexCollision().VertexBuffer.Num(), NavCollision->GetTriMeshCollision().VertexBuffer.Num()));
 
-				for (int32 i = 0; i < NavCollision->ConvexCollision.VertexBuffer.Num(); ++i)
+				for (int32 i = 0; i < NavCollision->GetConvexCollision().VertexBuffer.Num(); ++i)
 				{
-					FVector Vertex = NavCollision->ConvexCollision.VertexBuffer[i];
+					FVector Vertex = NavCollision->GetConvexCollision().VertexBuffer[i];
 					Vertex = CalcSliceTransform(GetAxisValue(Vertex, ForwardAxis)).TransformPosition(Vertex * Mask);
 					VertexBuffer.Add(Vertex);
 				}
 				GeomExport.ExportCustomMesh(VertexBuffer.GetData(), VertexBuffer.Num(),
-					NavCollision->ConvexCollision.IndexBuffer.GetData(), NavCollision->ConvexCollision.IndexBuffer.Num(),
+					NavCollision->GetConvexCollision().IndexBuffer.GetData(), NavCollision->GetConvexCollision().IndexBuffer.Num(),
 					GetComponentTransform());
 
 				VertexBuffer.Reset();
-				for (int32 i = 0; i < NavCollision->TriMeshCollision.VertexBuffer.Num(); ++i)
+				for (int32 i = 0; i < NavCollision->GetTriMeshCollision().VertexBuffer.Num(); ++i)
 				{
-					FVector Vertex = NavCollision->TriMeshCollision.VertexBuffer[i];
+					FVector Vertex = NavCollision->GetTriMeshCollision().VertexBuffer[i];
 					Vertex = CalcSliceTransform(GetAxisValue(Vertex, ForwardAxis)).TransformPosition(Vertex * Mask);
 					VertexBuffer.Add(Vertex);
 				}
 				GeomExport.ExportCustomMesh(VertexBuffer.GetData(), VertexBuffer.Num(),
-					NavCollision->TriMeshCollision.IndexBuffer.GetData(), NavCollision->TriMeshCollision.IndexBuffer.Num(),
+					NavCollision->GetTriMeshCollision().IndexBuffer.GetData(), NavCollision->GetTriMeshCollision().IndexBuffer.Num(),
 					GetComponentTransform());
 
 				return false;
diff --git a/Engine/Source/Runtime/Engine/Private/Components/StaticMeshComponent.cpp b/Engine/Source/Runtime/Engine/Private/Components/StaticMeshComponent.cpp
--- a/Engine/Source/Runtime/Engine/Private/Components/StaticMeshComponent.cpp
+++ b/Engine/Source/Runtime/Engine/Private/Components/StaticMeshComponent.cpp
@@ -28,11 +28,10 @@
 #include "LightMap.h"
 #include "ShadowMap.h"
 #include "Engine/ShadowMapTexture2D.h"
-#include "AI/Navigation/NavCollision.h"
+#include "AI/Navigation/NavCollisionBase.h"
 #include "Engine/StaticMeshSocket.h"
 #include "AI/NavigationSystemHelpers.h"
-#include "AI/NavigationOctree.h"
-#include "AI/Navigation/NavigationSystem.h"
+#include "AI/NavigationSystemBase.h"
 #include "PhysicsEngine/BodySetup.h"
 #include "EngineGlobals.h"
 #include "ComponentRecreateRenderStateContext.h"
@@ -612,14 +611,14 @@ void UStaticMeshComponent::OnCreatePhysicsState()
 	Super::OnCreatePhysicsState();
 
 	bNavigationRelevant = IsNavigationRelevant();
-	UNavigationSystem::UpdateComponentInNavOctree(*this);
+	FNavigationSystem::UpdateComponentData(*this);
 }
 
 void UStaticMeshComponent::OnDestroyPhysicsState()
 {
 	Super::OnDestroyPhysicsState();
 
-	UNavigationSystem::UpdateComponentInNavOctree(*this);
+	FNavigationSystem::UpdateComponentData(*this);
 	bNavigationRelevant = IsNavigationRelevant();
 }
 
@@ -2240,8 +2239,8 @@ bool UStaticMeshComponent::DoCustomNavigableGeometryExport(FNavigableGeometryExp
 	const FVector Scale3D = GetComponentToWorld().GetScale3D();
 	if (GetStaticMesh() && GetStaticMesh()->NavCollision && !Scale3D.IsZero())
 	{
-		UNavCollision* NavCollision = GetStaticMesh()->NavCollision;
-		const bool bExportAsObstacle = bOverrideNavigationExport ? bForceNavigationObstacle : NavCollision->bIsDynamicObstacle;
+		UNavCollisionBase* NavCollision = GetStaticMesh()->NavCollision;
+		const bool bExportAsObstacle = bOverrideNavigationExport ? bForceNavigationObstacle : NavCollision->IsDynamicObstacle();
 
 		if (bExportAsObstacle)
 		{
@@ -2313,8 +2312,8 @@ void UStaticMeshComponent::GetNavigationData(FNavigationRelevantData& Data) cons
 	const FVector Scale3D = GetComponentToWorld().GetScale3D();
 	if (GetStaticMesh() && GetStaticMesh()->NavCollision && !Scale3D.IsZero())
 	{
-		UNavCollision* NavCollision = GetStaticMesh()->NavCollision;
-		const bool bExportAsObstacle = bOverrideNavigationExport ? bForceNavigationObstacle : NavCollision->bIsDynamicObstacle;
+		UNavCollisionBase* NavCollision = GetStaticMesh()->NavCollision;
+		const bool bExportAsObstacle = bOverrideNavigationExport ? bForceNavigationObstacle : NavCollision->IsDynamicObstacle();
 
 		if (bExportAsObstacle)
 		{
diff --git a/Engine/Source/Runtime/Engine/Private/Controller.cpp b/Engine/Source/Runtime/Engine/Private/Controller.cpp
--- a/Engine/Source/Runtime/Engine/Private/Controller.cpp
+++ b/Engine/Source/Runtime/Engine/Private/Controller.cpp
@@ -8,7 +8,7 @@
 #include "GameFramework/Pawn.h"
 #include "CollisionQueryParams.h"
 #include "Engine/World.h"
-#include "AI/Navigation/NavigationSystem.h"
+#include "AI/NavigationSystemBase.h"
 #include "GameFramework/PlayerController.h"
 #include "Net/UnrealNetwork.h"
 #include "NetworkingDistanceConstants.h"
@@ -21,7 +21,6 @@
 
 // @todo this is here only due to circular dependency to AIModule. To be removed
 #include "Engine/Canvas.h"
-#include "Navigation/PathFollowingComponent.h"
 
 #include "GameFramework/PlayerState.h"
 
@@ -286,9 +285,11 @@ void AController::Possess(APawn* InPawn)
 
 	REDIRECT_OBJECT_TO_VLOG(InPawn, this);
 
+	const bool bNewPawn = GetPawn() != InPawn;
+
 	if (InPawn != NULL)
 	{
-		if (GetPawn() && GetPawn() != InPawn)
+		if (bNewPawn && GetPawn())
 		{
 			UnPossess();
 		}
@@ -306,6 +307,11 @@ void AController::Possess(APawn* InPawn)
 
 		Pawn->Restart();
 	}
+
+	if (bNewPawn)
+	{
+		OnNewPawn.Broadcast(GetPawn());
+	}
 }
 
 void AController::UnPossess()
@@ -641,39 +647,17 @@ bool AController::ShouldPostponePathUpdates() const
 
 bool AController::IsFollowingAPath() const
 {
-	UPathFollowingComponent* PathFollowingComp = FindComponentByClass<UPathFollowingComponent>();
-	return (PathFollowingComp != nullptr) && (PathFollowingComp->GetStatus() != EPathFollowingStatus::Idle);
-}
-
-void AController::UpdateNavigationComponents()
-{
-	UPathFollowingComponent* PathFollowingComp = FindComponentByClass<UPathFollowingComponent>();
-	if (PathFollowingComp != NULL)
-	{
-		PathFollowingComp->UpdateCachedComponents();
-	}
+	return FNavigationSystem::IsFollowingAPath(*this);
 }
 
-void AController::InitNavigationControl(UPathFollowingComponent*& PathFollowingComp)
+IPathFollowingAgentInterface* AController::GetPathFollowingAgent() const
 {
-	PathFollowingComp = FindComponentByClass<UPathFollowingComponent>();
-	if (PathFollowingComp == NULL)
-	{
-		PathFollowingComp = NewObject<UPathFollowingComponent>(this);
-		PathFollowingComp->RegisterComponentWithWorld(GetWorld());
-		PathFollowingComp->Initialize();
-	}
+	return FNavigationSystem::FindPathFollowingAgentForActor(*this);
 }
 
 void AController::StopMovement()
 {
-	UE_VLOG(this, LogNavigation, Log, TEXT("AController::StopMovement: %s STOP MOVEMENT"), *GetNameSafe(GetPawn()));
-
-	UPathFollowingComponent* PathFollowingComp = FindComponentByClass<UPathFollowingComponent>();
-	if (PathFollowingComp != NULL)
-	{
-		PathFollowingComp->AbortMove(*this, FPathFollowingResultFlags::MovementStop);
-	}
+	FNavigationSystem::StopMovement(*this);
 }
 
 void AController::GetLifetimeReplicatedProps( TArray< FLifetimeProperty > & OutLifetimeProps ) const
diff --git a/Engine/Source/Runtime/Engine/Private/EdGraph/EdGraphPin.cpp b/Engine/Source/Runtime/Engine/Private/EdGraph/EdGraphPin.cpp
--- a/Engine/Source/Runtime/Engine/Private/EdGraph/EdGraphPin.cpp
+++ b/Engine/Source/Runtime/Engine/Private/EdGraph/EdGraphPin.cpp
@@ -1312,11 +1312,24 @@ UEdGraphPin::UEdGraphPin(UEdGraphNode* InOwningNode, const FGuid& PinIdGuid)
 	: OwningNode(InOwningNode)
 	, PinId(PinIdGuid)
 	, PinName()
+	, Direction(EGPD_Input)
+#if WITH_EDITORONLY_DATA
+	, bHidden(false)
+	, bNotConnectable(false)
+	, bDefaultValueIsReadOnly(false)
+	, bDefaultValueIsIgnored(false)
+	, bIsDiffing(false)
+	, bAdvancedView(false)
+	, bDisplayAsMutableRef(false)
+	, bOrphanedPin(false)
+	, bSavePinIfOrphaned(true)
+	, bUseBackwardsCompatForEmptyAutogeneratedValue(false)
+#endif
+	, bWasTrashed(false)
 #if WITH_EDITORONLY_DATA
 	, PinFriendlyName()
 #endif // WITH_EDITORONLY_DATA
 	, PinToolTip()
-	, Direction(EGPD_Input)
 	, PinType()
 	, DefaultValue()
 	, AutogeneratedDefaultValue()
@@ -1328,17 +1341,7 @@ UEdGraphPin::UEdGraphPin(UEdGraphNode* InOwningNode, const FGuid& PinIdGuid)
 	, ReferencePassThroughConnection(nullptr)
 #if WITH_EDITORONLY_DATA
 	, PersistentGuid()
-	, bHidden(false)
-	, bNotConnectable(false)
-	, bDefaultValueIsReadOnly(false)
-	, bDefaultValueIsIgnored(false)
-	, bIsDiffing(false)
-	, bAdvancedView(false)
-	, bDisplayAsMutableRef(false)
-	, bOrphanedPin(false)
-	, bSavePinIfOrphaned(true)
 #endif
-	, bWasTrashed(false)
 {
 	PinHelpers::NumPinsInMemory++;
 #ifdef TRACK_PINS
@@ -2101,6 +2104,24 @@ void UEdGraphPin::SanitizePinsPostUndoRedo()
 	GraphConnectionSanitizer.SanitizeConnectionsForCorruptGraphs();
 }
 
+#if WITH_EDITORONLY_DATA
+void UEdGraphPin::SetSavePinIfOrphaned(const bool bShouldSave)
+{
+	if (bSavePinIfOrphaned != bShouldSave)
+	{
+		bSavePinIfOrphaned = bShouldSave;
+		// If we are not saving this pin, we cannot save any children either
+		if (!bSavePinIfOrphaned)
+		{
+			for (UEdGraphPin* SubPin : SubPins)
+			{
+				SubPin->SetSavePinIfOrphaned(false);
+			}
+		}
+	}
+}
+#endif
+
 /////////////////////////////////////////////////////
 // UEdGraphPin
 
diff --git a/Engine/Source/Runtime/Engine/Private/EnginePrivatePCH.h b/Engine/Source/Runtime/Engine/Private/EnginePrivatePCH.h
--- a/Engine/Source/Runtime/Engine/Private/EnginePrivatePCH.h
+++ b/Engine/Source/Runtime/Engine/Private/EnginePrivatePCH.h
@@ -392,7 +392,7 @@
 #include "EngineDefines.h"
 #include "CollisionQueryParams.h"
 #include "AI/Navigation/NavigationTypes.h"
-#include "AI/Navigation/NavFilters/NavigationQueryFilter.h"
+#include "AI/Navigation/NavQueryFilter.h"
 #include "WorldCollision.h"
 #include "AI/Navigation/NavAgentInterface.h"
 #include "Engine/Blueprint.h"
@@ -499,9 +499,8 @@
 #include "Materials/MaterialFunction.h"
 #include "GameFramework/PlayerController.h"
 #include "GameFramework/PlayerMuteList.h"
-#include "AI/Navigation/NavigationData.h"
 #include "Animation/AnimSequence.h"
-#include "AI/Navigation/NavigationSystem.h"
+#include "AI/NavigationSystemBase.h"
 #include "Sound/SoundConcurrency.h"
 #include "ClothSimData.h"
 #include "AnimEncoding.h"
diff --git a/Engine/Source/Runtime/Engine/Private/GameInstance.cpp b/Engine/Source/Runtime/Engine/Private/GameInstance.cpp
--- a/Engine/Source/Runtime/Engine/Private/GameInstance.cpp
+++ b/Engine/Source/Runtime/Engine/Private/GameInstance.cpp
@@ -9,7 +9,7 @@
 #include "TimerManager.h"
 #include "Engine/LatentActionManager.h"
 #include "Engine/World.h"
-#include "AI/Navigation/NavigationSystem.h"
+#include "AI/NavigationSystemBase.h"
 #include "Misc/Paths.h"
 #include "UObject/CoreOnline.h"
 #include "GameFramework/PlayerController.h"
@@ -320,7 +320,7 @@ FGameInstancePIEResult UGameInstance::StartPlayInEditorGameInstance(ULocalPlayer
 
 		PlayWorld->InitializeActorsForPlay(URL);
 		// calling it after InitializeActorsForPlay has been called to have all potential bounding boxed initialized
-		UNavigationSystem::InitializeForWorld(PlayWorld, LocalPlayers.Num() > 0 ? FNavigationSystemRunMode::PIEMode : FNavigationSystemRunMode::SimulationMode);
+		FNavigationSystem::AddNavigationSystemToWorld(*PlayWorld, LocalPlayers.Num() > 0 ? FNavigationSystemRunMode::PIEMode : FNavigationSystemRunMode::SimulationMode);
 
 		// @todo, just use WorldContext.GamePlayer[0]?
 		if (LocalPlayer)
diff --git a/Engine/Source/Runtime/Engine/Private/GameNetworkManager.cpp b/Engine/Source/Runtime/Engine/Private/GameNetworkManager.cpp
--- a/Engine/Source/Runtime/Engine/Private/GameNetworkManager.cpp
+++ b/Engine/Source/Runtime/Engine/Private/GameNetworkManager.cpp
@@ -23,8 +23,9 @@ AGameNetworkManager::AGameNetworkManager(const FObjectInitializer& ObjectInitial
 	CLIENTADJUSTUPDATECOST = 180.0f;
 	MAXCLIENTUPDATEINTERVAL = 0.25f;
 	MaxMoveDeltaTime = 0.125f;
-	ClientNetSendMoveDeltaTime = 0.0111f;
+	ClientNetSendMoveDeltaTime = 0.0166f;
 	ClientNetSendMoveDeltaTimeThrottled = 0.0222f;
+	ClientNetSendMoveDeltaTimeStationary = 0.0833f;
 	ClientNetSendMoveThrottleAtNetSpeed = 10000;
 	ClientNetSendMoveThrottleOverPlayerCount = 10;
 	ClientAuthorativePosition = false;
