Fix merge actor material issues
Fix some automation tests

diff --git a/Engine/Source/Runtime/Engine/Classes/Engine/StaticMesh.h b/Engine/Source/Runtime/Engine/Classes/Engine/StaticMesh.h
--- a/Engine/Source/Runtime/Engine/Classes/Engine/StaticMesh.h
+++ b/Engine/Source/Runtime/Engine/Classes/Engine/StaticMesh.h
@@ -364,7 +364,17 @@ struct FStaticMaterial
 		, ImportedMaterialSlotName(InImportedMaterialSlotName)
 #endif //WITH_EDITORONLY_DATA
 	{
-
+		//If not specified add some valid material slot name
+		if (MaterialInterface && MaterialSlotName == NAME_None)
+		{
+			MaterialSlotName = MaterialInterface->GetFName();
+		}
+#if WITH_EDITORONLY_DATA
+		if (ImportedMaterialSlotName == NAME_None)
+		{
+			ImportedMaterialSlotName = MaterialSlotName;
+		}
+#endif
 	}
 
 	friend FArchive& operator<<(FArchive& Ar, FStaticMaterial& Elem);
diff --git a/Engine/Source/Runtime/Engine/Private/StaticMesh.cpp b/Engine/Source/Runtime/Engine/Private/StaticMesh.cpp
--- a/Engine/Source/Runtime/Engine/Private/StaticMesh.cpp
+++ b/Engine/Source/Runtime/Engine/Private/StaticMesh.cpp
@@ -1316,7 +1316,7 @@ FArchive& operator<<(FArchive& Ar, FMeshBuildSettings& BuildSettings)
 // differences, etc.) replace the version GUID below with a new one.
 // In case of merge conflicts with DDC versions, you MUST generate a new GUID
 // and set this new GUID as the version.                                       
-#define STATICMESH_DERIVEDDATA_VER TEXT("23E9164C89904B29B9A209F087DF9C0D")
+#define STATICMESH_DERIVEDDATA_VER TEXT("AD8FEAB182864DDDAFD9AC29A7B2B600")
 
 static const FString& GetStaticMeshDerivedDataVersion()
 {
@@ -2680,7 +2680,7 @@ void UStaticMesh::FixupMaterialSlotName()
 // differences, etc.) replace the version GUID below with a new one.
 // In case of merge conflicts with DDC versions, you MUST generate a new GUID
 // and set this new GUID as the version.                                       
-#define MESHDATAKEY_STATICMESH_DERIVEDDATA_VER TEXT("1536BB4EBA6845F5AECFDA4C12FA216B")
+#define MESHDATAKEY_STATICMESH_DERIVEDDATA_VER TEXT("EEEE375ED2AD4000BE7EEEC11D331333")
 
 static const FString& GetMeshDataKeyStaticMeshDerivedDataVersion()
 {
