Fix partially the Merge actor tests

diff --git a/Engine/Source/Developer/MeshBuilder/Private/StaticMeshBuilder.cpp b/Engine/Source/Developer/MeshBuilder/Private/StaticMeshBuilder.cpp
--- a/Engine/Source/Developer/MeshBuilder/Private/StaticMeshBuilder.cpp
+++ b/Engine/Source/Developer/MeshBuilder/Private/StaticMeshBuilder.cpp
@@ -92,10 +92,11 @@ bool FStaticMeshBuilder::Build(UStaticMesh* StaticMesh, const FStaticMeshLODGrou
 		//Reduce LODs
 		if (bUseReduction)
 		{
+			int32 BaseLODIndex = 0;
 			float OverlappingThreshold = LODBuildSettings.bRemoveDegenerates ? THRESH_POINTS_ARE_SAME : 0.0f;
 			TMultiMap<int32, int32> OverlappingCorners;
-			FMeshDescriptionOperations::FindOverlappingCorners(OverlappingCorners, StaticMesh->GetMeshDescription(0), OverlappingThreshold);
-			MeshDescriptionHelper.ReduceLOD(StaticMesh->GetMeshDescription(0), StaticMesh->GetMeshDescription(LodIndex), ReductionSettings, OverlappingCorners);
+			FMeshDescriptionOperations::FindOverlappingCorners(OverlappingCorners, StaticMesh->GetMeshDescription(BaseLODIndex), OverlappingThreshold);
+			MeshDescriptionHelper.ReduceLOD(StaticMesh->GetMeshDescription(BaseLODIndex), StaticMesh->GetMeshDescription(LodIndex), ReductionSettings, OverlappingCorners);
 			// Recompute adjacency information. Since we change the vertices when we reduce
 			MeshDescriptionHelper.FindOverlappingCorners(StaticMesh->GetMeshDescription(LodIndex), OverlappingThreshold);
 		}
@@ -429,9 +430,9 @@ void BuildVertexBuffer(
 			}
 		}
 
-		if (LodIndex > 0)
+		if (LodIndex > 0 && StaticMesh->SectionInfoMap.GetSectionNumber(LodIndex) <= SectionIndex)
 		{
-			//Set the overwrite section info map
+			//Set the overwrite section info map in case there is not already one
 			FMeshSectionInfo SectionInfo = StaticMesh->SectionInfoMap.Get(LodIndex, SectionIndex);
 			SectionInfo.MaterialIndex = StaticMeshSection.MaterialIndex;
 			StaticMesh->SectionInfoMap.Set(LodIndex, SectionIndex, SectionInfo);
diff --git a/Engine/Source/Developer/MeshDescriptionOperations/Private/MeshDescriptionOperations.cpp b/Engine/Source/Developer/MeshDescriptionOperations/Private/MeshDescriptionOperations.cpp
--- a/Engine/Source/Developer/MeshDescriptionOperations/Private/MeshDescriptionOperations.cpp
+++ b/Engine/Source/Developer/MeshDescriptionOperations/Private/MeshDescriptionOperations.cpp
@@ -447,6 +447,33 @@ void FMeshDescriptionOperations::ConverFromRawMesh(const FRawMesh &SourceRawMesh
 	bool bHasNormals = SourceRawMesh.WedgeTangentZ.Num() > 0;
 
 	TArray<FPolygonGroupID> PolygonGroups;
+	TMap<int32, FPolygonGroupID> MaterialIndexToPolygonGroup;
+
+	//Find the maximum material index
+	int32 MaxMaterialIndex = 0;
+	for(int32 MaterialIndex : SourceRawMesh.FaceMaterialIndices)
+	{
+		if (MaterialIndex > MaxMaterialIndex)
+		{
+			MaxMaterialIndex = MaterialIndex;
+		}
+	}
+
+	//Create the known Materials, this will ensure we create them in the correct order, since it is possible to have FRawMesh with unordered FaceMaterialIndices
+	for (auto Kvp : MaterialMap)
+	{
+		//Avoid creating empty sections
+		if (Kvp.Key > MaxMaterialIndex)
+		{
+			continue;
+		}
+		FPolygonGroupID PolygonGroupID = DestinationMeshDescription->CreatePolygonGroup();
+		check(Kvp.Key == PolygonGroupID.GetValue());
+		FName PolygonGroupImportedMaterialSlotName = Kvp.Value;
+		PolygonGroupImportedMaterialSlotNames[PolygonGroupID] = PolygonGroupImportedMaterialSlotName == NAME_None ? FName(*FString::Printf(TEXT("MaterialSlot_%d"), Kvp.Key)) : PolygonGroupImportedMaterialSlotName;
+		PolygonGroups.Add(PolygonGroupID);
+		MaterialIndexToPolygonGroup.Add(Kvp.Key, PolygonGroupID);
+	}
 
 	//Triangles
 	int32 TriangleCount = SourceRawMesh.WedgeIndices.Num() / 3;
@@ -463,15 +490,19 @@ void FMeshDescriptionOperations::ConverFromRawMesh(const FRawMesh &SourceRawMesh
 		FPolygonGroupID PolygonGroupID = FPolygonGroupID::Invalid;
 		FName PolygonGroupImportedMaterialSlotName = NAME_None;
 		int32 MaterialIndex = SourceRawMesh.FaceMaterialIndices[TriangleIndex];
-
-		if (MaterialMap.Num() > 0 && MaterialMap.Contains(SourceRawMesh.FaceMaterialIndices[TriangleIndex]))
+		if (MaterialIndexToPolygonGroup.Contains(MaterialIndex))
+		{
+			PolygonGroupID = MaterialIndexToPolygonGroup[MaterialIndex];
+		}
+		else if (MaterialMap.Num() > 0 && MaterialMap.Contains(MaterialIndex))
 		{
-			PolygonGroupImportedMaterialSlotName = MaterialMap[SourceRawMesh.FaceMaterialIndices[TriangleIndex]];
+			PolygonGroupImportedMaterialSlotName = MaterialMap[MaterialIndex];
 			for (const FPolygonGroupID& SearchPolygonGroupID : DestinationMeshDescription->PolygonGroups().GetElementIDs())
 			{
 				if (PolygonGroupImportedMaterialSlotNames[SearchPolygonGroupID] == PolygonGroupImportedMaterialSlotName)
 				{
 					PolygonGroupID = SearchPolygonGroupID;
+					break;
 				}
 			}
 		}
@@ -481,6 +512,7 @@ void FMeshDescriptionOperations::ConverFromRawMesh(const FRawMesh &SourceRawMesh
 			PolygonGroupID = DestinationMeshDescription->CreatePolygonGroup();
 			PolygonGroupImportedMaterialSlotNames[PolygonGroupID] = PolygonGroupImportedMaterialSlotName == NAME_None ? FName(*FString::Printf(TEXT("MaterialSlot_%d"), MaterialIndex)) : PolygonGroupImportedMaterialSlotName;
 			PolygonGroups.Add(PolygonGroupID);
+			MaterialIndexToPolygonGroup.Add(MaterialIndex, PolygonGroupID);
 		}
 		FVertexInstanceID TriangleVertexInstanceIDs[3];
 		for (int32 Corner = 0; Corner < 3; ++Corner)
diff --git a/Engine/Source/Developer/MeshMergeUtilities/Private/MeshMergeHelpers.cpp b/Engine/Source/Developer/MeshMergeUtilities/Private/MeshMergeHelpers.cpp
--- a/Engine/Source/Developer/MeshMergeUtilities/Private/MeshMergeHelpers.cpp
+++ b/Engine/Source/Developer/MeshMergeUtilities/Private/MeshMergeHelpers.cpp
@@ -48,10 +48,19 @@
 void FMeshMergeHelpers::ExtractSections(const UStaticMeshComponent* Component, int32 LODIndex, TArray<FSectionInfo>& OutSections)
 {
 	static UMaterialInterface* DefaultMaterial = UMaterial::GetDefaultMaterial(MD_Surface);
+	
+	const UStaticMesh* StaticMesh = Component->GetStaticMesh();
 
-	TArray<FName> MaterialSlotNames = Component->GetMaterialSlotNames();
+	TArray<FName> MaterialSlotNames;
+	for (const FStaticMaterial& StaticMaterial : StaticMesh->StaticMaterials)
+	{
+#if WITH_EDITOR
+		MaterialSlotNames.Add(StaticMaterial.ImportedMaterialSlotName);
+#else
+		MaterialSlotNames.Add(StaticMaterial.MaterialSlotName);
+#endif
+	}
 
-	const UStaticMesh* StaticMesh = Component->GetStaticMesh();
 	for (const FStaticMeshSection& MeshSection : StaticMesh->RenderData->LODResources[LODIndex].Sections)
 	{
 		// Retrieve material for this section
@@ -133,7 +142,12 @@ void FMeshMergeHelpers::ExtractSections(const UStaticMesh* StaticMesh, int32 LOD
 		FSectionInfo SectionInfo;
 		SectionInfo.Material = StoredMaterial;
 		SectionInfo.MaterialIndex = MeshSection.MaterialIndex;
+#if WITH_EDITOR
+		SectionInfo.MaterialSlotName = StaticMesh->StaticMaterials.IsValidIndex(MeshSection.MaterialIndex) ? StaticMesh->StaticMaterials[MeshSection.MaterialIndex].ImportedMaterialSlotName : NAME_None;
+#else
 		SectionInfo.MaterialSlotName = StaticMesh->StaticMaterials.IsValidIndex(MeshSection.MaterialIndex) ? StaticMesh->StaticMaterials[MeshSection.MaterialIndex].MaterialSlotName : NAME_None;
+#endif
+		
 
 		if (MeshSection.bEnableCollision)
 		{
diff --git a/Engine/Source/Developer/MeshMergeUtilities/Private/MeshMergeUtilities.cpp b/Engine/Source/Developer/MeshMergeUtilities/Private/MeshMergeUtilities.cpp
--- a/Engine/Source/Developer/MeshMergeUtilities/Private/MeshMergeUtilities.cpp
+++ b/Engine/Source/Developer/MeshMergeUtilities/Private/MeshMergeUtilities.cpp
@@ -2588,18 +2588,50 @@ void FMeshMergeUtilities::MergeComponentsToStaticMesh(const TArray<UPrimitiveCom
 			}
 		}
 		
+		auto IsMaterialImportedNameUnique = [&StaticMesh](FName ImportedMaterialSlotName)
+		{
+			for (const FStaticMaterial& StaticMaterial : StaticMesh->StaticMaterials)
+			{
+#if WITH_EDITOR
+				if (StaticMaterial.ImportedMaterialSlotName == ImportedMaterialSlotName)
+#else
+				if (StaticMaterial.MaterialSlotName == ImportedMaterialSlotName)
+#endif
+				{
+					return false;
+				}
+			}
+			return true;
+		};
+		
+
 		for (UMaterialInterface* Material : UniqueMaterials)
 		{
 			if (Material && (!Material->IsAsset() && InOuter != GetTransientPackage()))
 			{
 				Material = nullptr; // do not save non-asset materials
 			}
-			StaticMesh->StaticMaterials.Add(FStaticMaterial(Material, DataTracker.GetMaterialSlotName(Material)));
+			//Make sure we have unique slot name here
+			FName MaterialSlotName = DataTracker.GetMaterialSlotName(Material);
+			int32 Counter = 1;
+			while (!IsMaterialImportedNameUnique(MaterialSlotName))
+			{
+				MaterialSlotName = *(DataTracker.GetMaterialSlotName(Material).ToString() + TEXT("_") + FString::FromInt(Counter++));
+			}
+
+			StaticMesh->StaticMaterials.Add(FStaticMaterial(Material, MaterialSlotName));
 		}
 
 		for(UMaterialInterface* ImposterMaterial : ImposterMaterials)
 		{
-			StaticMesh->StaticMaterials.Add(FStaticMaterial(ImposterMaterial));
+			//Make sure we have unique slot name here
+			FName MaterialSlotName = ImposterMaterial->GetFName();
+			int32 Counter = 1;
+			while (!IsMaterialImportedNameUnique(MaterialSlotName))
+			{
+				MaterialSlotName = *(ImposterMaterial->GetName() + TEXT("_") + FString::FromInt(Counter++));
+			}
+			StaticMesh->StaticMaterials.Add(FStaticMaterial(ImposterMaterial, MaterialSlotName));
 		}
 
 		if (InSettings.bMergePhysicsData)
diff --git a/Engine/Source/Developer/MeshSimplifier/Private/QuadricMeshReduction.cpp b/Engine/Source/Developer/MeshSimplifier/Private/QuadricMeshReduction.cpp
--- a/Engine/Source/Developer/MeshSimplifier/Private/QuadricMeshReduction.cpp
+++ b/Engine/Source/Developer/MeshSimplifier/Private/QuadricMeshReduction.cpp
@@ -502,6 +502,9 @@ public:
 		const TVertexInstanceAttributeArray<float>& InVertexBinormalSigns = InMesh->VertexInstanceAttributes().GetAttributes<float>(MeshAttribute::VertexInstance::BinormalSign);
 		const TVertexInstanceAttributeArray<FVector4>& InVertexColors = InMesh->VertexInstanceAttributes().GetAttributes<FVector4>(MeshAttribute::VertexInstance::Color);
 		const TVertexInstanceAttributeIndicesArray<FVector2D>& InVertexUVs = InMesh->VertexInstanceAttributes().GetAttributesSet<FVector2D>(MeshAttribute::VertexInstance::TextureCoordinate);
+		
+		const TPolygonGroupAttributeArray<FName>& InPolygonGroupMaterialNames = InMesh->PolygonGroupAttributes().GetAttributes<FName>(MeshAttribute::PolygonGroup::ImportedMaterialSlotName);
+		TPolygonGroupAttributeArray<FName>& OutPolygonGroupMaterialNames = OutReducedMesh->PolygonGroupAttributes().GetAttributes<FName>(MeshAttribute::PolygonGroup::ImportedMaterialSlotName);
 
 		// Process each face, build vertex buffer and index buffer
 		for (int32 FaceIndex = 0; FaceIndex < NumFaces; FaceIndex++)
@@ -674,19 +677,19 @@ public:
 
 		{
 			//Empty the destination mesh
-			//We want to keep the polygon group information, the reduce do not impact the material list
-			// @todo refactor this: PolygonGroup.Polygons shouldn't be mutated like this (it is likely to be made private soon).
-			// Prefer creating an empty mesh description, and copying the polygon group attributes over.
-			for (const FPolygonGroupID& PolygonGroupID : OutReducedMesh->PolygonGroups().GetElementIDs())
-			{
-				FMeshPolygonGroup& PolygonGroup = OutReducedMesh->GetPolygonGroup(PolygonGroupID);
-				PolygonGroup.Polygons.Empty();
-			}
+			OutReducedMesh->PolygonGroups().Reset();
 			OutReducedMesh->Polygons().Reset();
 			OutReducedMesh->Edges().Reset();
 			OutReducedMesh->VertexInstances().Reset();
 			OutReducedMesh->Vertices().Reset();
 
+			//Fill the PolygonGroups from the InMesh
+			for (const FPolygonGroupID& PolygonGroupID : InMesh->PolygonGroups().GetElementIDs())
+			{
+				OutReducedMesh->CreatePolygonGroupWithID(PolygonGroupID);
+				OutPolygonGroupMaterialNames[PolygonGroupID] = InPolygonGroupMaterialNames[PolygonGroupID];
+			}
+
 			TVertexAttributeArray<FVector>& OutVertexPositions = OutReducedMesh->VertexAttributes().GetAttributes<FVector>(MeshAttribute::Vertex::Position);
 
 			//Fill the vertex array
diff --git a/Engine/Source/Runtime/Engine/Private/StaticMesh.cpp b/Engine/Source/Runtime/Engine/Private/StaticMesh.cpp
--- a/Engine/Source/Runtime/Engine/Private/StaticMesh.cpp
+++ b/Engine/Source/Runtime/Engine/Private/StaticMesh.cpp
@@ -123,6 +123,23 @@ namespace StaticMeshCookStats
 #endif
 
 
+#if WITH_EDITOR
+static void FillMaterialName(const TArray<FStaticMaterial>& StaticMaterials, TMap<int32, FName>& OutMaterialMap)
+{
+	OutMaterialMap.Empty(StaticMaterials.Num());
+
+	for (int32 MaterialIndex = 0; MaterialIndex < StaticMaterials.Num(); ++MaterialIndex)
+	{
+		FName MaterialName = StaticMaterials[MaterialIndex].ImportedMaterialSlotName;
+		if (MaterialName == NAME_None)
+		{
+			MaterialName = *(TEXT("MaterialSlot_") + FString::FromInt(MaterialIndex));
+		}
+		OutMaterialMap.Add(MaterialIndex, MaterialName);
+	}
+}
+#endif
+
 
 /*-----------------------------------------------------------------------------
 	FStaticMeshLODResources
@@ -1316,7 +1333,7 @@ FArchive& operator<<(FArchive& Ar, FMeshBuildSettings& BuildSettings)
 // differences, etc.) replace the version GUID below with a new one.
 // In case of merge conflicts with DDC versions, you MUST generate a new GUID
 // and set this new GUID as the version.                                       
-#define STATICMESH_DERIVEDDATA_VER TEXT("AD8FEAB182864DDDAFD9AC29A7B2B600")
+#define STATICMESH_DERIVEDDATA_VER TEXT("74CAFC5D929A4BB8AC6B2690A29A1E99")
 
 static const FString& GetStaticMeshDerivedDataVersion()
 {
@@ -2357,10 +2374,7 @@ void FStaticMeshSourceModel::SaveRawMesh(FRawMesh& InRawMesh, bool bConvertToMes
 	{
 		TMap<int32, FName> MaterialMap;
 		check(StaticMeshOwner != nullptr);
-		for (int32 MaterialIndex = 0; MaterialIndex < StaticMeshOwner->StaticMaterials.Num(); ++MaterialIndex)
-		{
-			MaterialMap.Add(MaterialIndex, StaticMeshOwner->StaticMaterials[MaterialIndex].ImportedMaterialSlotName);
-		}
+		FillMaterialName(StaticMeshOwner->StaticMaterials, MaterialMap);
 		FMeshDescriptionOperations::ConverFromRawMesh(InRawMesh, OriginalMeshDescription, MaterialMap);
 	}
 }
@@ -2592,10 +2606,7 @@ UMeshDescription* UStaticMesh::GetOriginalMeshDescription(int32 LodIndex)
 			FRawMesh LodRawMesh;
 			SourceModels[LodIndex].LoadRawMesh(LodRawMesh);
 			TMap<int32, FName> MaterialMap;
-			for (int32 MaterialIndex = 0; MaterialIndex < StaticMaterials.Num(); ++MaterialIndex)
-			{
-				MaterialMap.Add(MaterialIndex, StaticMaterials[MaterialIndex].ImportedMaterialSlotName);
-			}
+			FillMaterialName(StaticMaterials, MaterialMap);
 			FMeshDescriptionOperations::ConverFromRawMesh(LodRawMesh, SourceModels[LodIndex].OriginalMeshDescription, MaterialMap);
 		}
 		return SourceModels[LodIndex].OriginalMeshDescription;
@@ -2642,22 +2653,20 @@ void UStaticMesh::FixupMaterialSlotName()
 	//Make sure we have non empty imported material slot names
 	for (FStaticMaterial& Material : StaticMaterials)
 	{
-		if (Material.ImportedMaterialSlotName != NAME_None)
+		if (Material.ImportedMaterialSlotName == NAME_None)
 		{
-			continue;
-		}
-		
-		if (Material.MaterialSlotName != NAME_None)
-		{
-			Material.ImportedMaterialSlotName = Material.MaterialSlotName;
-		}
-		else if (Material.MaterialInterface != nullptr)
-		{
-			Material.ImportedMaterialSlotName = Material.MaterialInterface->GetFName();
-		}
-		else
-		{
-			Material.ImportedMaterialSlotName = FName(TEXT("MaterialSlot"));
+			if (Material.MaterialSlotName != NAME_None)
+			{
+				Material.ImportedMaterialSlotName = Material.MaterialSlotName;
+			}
+			else if (Material.MaterialInterface != nullptr)
+			{
+				Material.ImportedMaterialSlotName = Material.MaterialInterface->GetFName();
+			}
+			else
+			{
+				Material.ImportedMaterialSlotName = FName(TEXT("MaterialSlot"));
+			}
 		}
 
 		FString UniqueName = Material.ImportedMaterialSlotName.ToString();
@@ -2680,7 +2689,7 @@ void UStaticMesh::FixupMaterialSlotName()
 // differences, etc.) replace the version GUID below with a new one.
 // In case of merge conflicts with DDC versions, you MUST generate a new GUID
 // and set this new GUID as the version.                                       
-#define MESHDATAKEY_STATICMESH_DERIVEDDATA_VER TEXT("EEEE375ED2AD4000BE7EEEC11D331333")
+#define MESHDATAKEY_STATICMESH_DERIVEDDATA_VER TEXT("7681842ECB2C4CBDB21613591CDC5B40")
 
 static const FString& GetMeshDataKeyStaticMeshDerivedDataVersion()
 {
@@ -2715,7 +2724,7 @@ bool UStaticMesh::GetMeshDataKey(FString& OutKey)
 		{
 			LodIndexString += TEXT("REDUCELOD");
 		}
-		TArray<TCHAR> LodIndexArray = LodIndexString.GetCharArray();
+		const TArray<TCHAR>& LodIndexArray = LodIndexString.GetCharArray();
 		Sha.Update((uint8*)LodIndexArray.GetData(), LodIndexArray.Num() * LodIndexArray.GetTypeSize());
 	}
 	Sha.Final();
@@ -2873,10 +2882,7 @@ void UStaticMesh::CacheMeshData()
 
 				//Convert the raw mesh to mesh description and save the mesh description in the ddc
 				TMap<int32, FName> MaterialMap;
-				for (int32 MaterialIndex = 0; MaterialIndex < StaticMaterials.Num(); ++MaterialIndex)
-				{
-					MaterialMap.Add(MaterialIndex, StaticMaterials[MaterialIndex].ImportedMaterialSlotName);
-				}
+				FillMaterialName(StaticMaterials, MaterialMap);
 				FMeshDescriptionOperations::ConverFromRawMesh(TempRawMesh, SourceModel.OriginalMeshDescription, MaterialMap);
 			}
 		}
