Fix all material baking automation tests

diff --git a/Engine/Source/Developer/MeshBuilder/Private/StaticMeshBuilder.cpp b/Engine/Source/Developer/MeshBuilder/Private/StaticMeshBuilder.cpp
--- a/Engine/Source/Developer/MeshBuilder/Private/StaticMeshBuilder.cpp
+++ b/Engine/Source/Developer/MeshBuilder/Private/StaticMeshBuilder.cpp
@@ -425,14 +425,7 @@ void BuildVertexBuffer(
 	StaticMeshLOD.VertexBuffers.StaticMeshVertexBuffer.SetUseFullPrecisionUVs(LODBuildSettings.bUseFullPrecisionUVs);
 	StaticMeshLOD.VertexBuffers.StaticMeshVertexBuffer.Init(StaticMeshBuildVertices, NumTextureCoord);
 	StaticMeshLOD.VertexBuffers.PositionVertexBuffer.Init(StaticMeshBuildVertices);
-	if (bHasColor)
-	{
-		StaticMeshLOD.VertexBuffers.ColorVertexBuffer.Init(StaticMeshBuildVertices);
-	}
-	else
-	{
-		StaticMeshLOD.VertexBuffers.ColorVertexBuffer.InitFromSingleColor(FColor::White, StaticMeshBuildVertices.Num());
-	}
+	StaticMeshLOD.VertexBuffers.ColorVertexBuffer.Init(StaticMeshBuildVertices);
 }
 
 void BuildAllBufferOptimizations(FStaticMeshLODResources& StaticMeshLOD, const FMeshBuildSettings& LODBuildSettings, TArray< uint32 >& IndexBuffer, bool bNeeds32BitIndices, TArray< FStaticMeshBuildVertex >& StaticMeshBuildVertices)
diff --git a/Engine/Source/Runtime/Engine/Private/StaticMesh.cpp b/Engine/Source/Runtime/Engine/Private/StaticMesh.cpp
--- a/Engine/Source/Runtime/Engine/Private/StaticMesh.cpp
+++ b/Engine/Source/Runtime/Engine/Private/StaticMesh.cpp
@@ -1333,7 +1333,7 @@ FArchive& operator<<(FArchive& Ar, FMeshBuildSettings& BuildSettings)
 // differences, etc.) replace the version GUID below with a new one.
 // In case of merge conflicts with DDC versions, you MUST generate a new GUID
 // and set this new GUID as the version.                                       
-#define STATICMESH_DERIVEDDATA_VER TEXT("FE87C30744E84D589D23AC2FEC516511")
+#define STATICMESH_DERIVEDDATA_VER TEXT("8A8FE832D30D4F72A111D4FF9BCF4346")
 
 static const FString& GetStaticMeshDerivedDataVersion()
 {
@@ -2689,7 +2689,7 @@ void UStaticMesh::FixupMaterialSlotName()
 // differences, etc.) replace the version GUID below with a new one.
 // In case of merge conflicts with DDC versions, you MUST generate a new GUID
 // and set this new GUID as the version.                                       
-#define MESHDATAKEY_STATICMESH_DERIVEDDATA_VER TEXT("6EDD0652F49745F1B8FAEB8813333E45")
+#define MESHDATAKEY_STATICMESH_DERIVEDDATA_VER TEXT("24E2D709773D4503B6873C4BED85ECB2")
 
 static const FString& GetMeshDataKeyStaticMeshDerivedDataVersion()
 {
