Fix the scaling of generated LOD when the base LOD use to generate the LOD has BuildScale different then 1.0
I also bump the staticmesh version since I change the build

diff --git a/Engine/Source/Developer/MeshBuilder/Private/StaticMeshBuilder.cpp b/Engine/Source/Developer/MeshBuilder/Private/StaticMeshBuilder.cpp
--- a/Engine/Source/Developer/MeshBuilder/Private/StaticMeshBuilder.cpp
+++ b/Engine/Source/Developer/MeshBuilder/Private/StaticMeshBuilder.cpp
@@ -92,6 +92,7 @@ bool FStaticMeshBuilder::Build(FStaticMeshRenderData& StaticMeshRenderData, USta
 			}
 			if (LodIndex > 0)
 			{
+				LODBuildSettings = StaticMesh->SourceModels[BaseReduceLodIndex].BuildSettings;
 				//Make sure the SectionInfoMap is taken from the Base RawMesh
 				int32 SectionNumber = StaticMesh->OriginalSectionInfoMap.GetSectionNumber(BaseReduceLodIndex);
 				for (int32 SectionIndex = 0; SectionIndex < SectionNumber; ++SectionIndex)
diff --git a/Engine/Source/Runtime/Engine/Private/StaticMesh.cpp b/Engine/Source/Runtime/Engine/Private/StaticMesh.cpp
--- a/Engine/Source/Runtime/Engine/Private/StaticMesh.cpp
+++ b/Engine/Source/Runtime/Engine/Private/StaticMesh.cpp
@@ -1333,7 +1333,7 @@ FArchive& operator<<(FArchive& Ar, FMeshBuildSettings& BuildSettings)
 // differences, etc.) replace the version GUID below with a new one.
 // In case of merge conflicts with DDC versions, you MUST generate a new GUID
 // and set this new GUID as the version.                                       
-#define STATICMESH_DERIVEDDATA_VER TEXT("B330D0DD8F214096995A2A9527F4E0DD")
+#define STATICMESH_DERIVEDDATA_VER TEXT("D81A3AD754C34B6EA3F822DDB06000D8")
 
 static const FString& GetStaticMeshDerivedDataVersion()
 {
@@ -2689,7 +2689,7 @@ void UStaticMesh::FixupMaterialSlotName()
 // differences, etc.) replace the version GUID below with a new one.
 // In case of merge conflicts with DDC versions, you MUST generate a new GUID
 // and set this new GUID as the version.                                       
-#define MESHDATAKEY_STATICMESH_DERIVEDDATA_VER TEXT("22EA1C022AC745028AB8AB4904B15B75")
+#define MESHDATAKEY_STATICMESH_DERIVEDDATA_VER TEXT("ECB687DFD4B34464B85F5C854A13A309")
 
 static const FString& GetMeshDataKeyStaticMeshDerivedDataVersion()
 {
