Fix the lightmap issue. The UVindex count for the reduce mesh was offset by -1, so the light map was often set to 0 instead of 1. Using UVChannel 0 is really bad in case like this.

diff --git a/Engine/Source/Developer/MeshSimplifier/Private/QuadricMeshReduction.cpp b/Engine/Source/Developer/MeshSimplifier/Private/QuadricMeshReduction.cpp
--- a/Engine/Source/Developer/MeshSimplifier/Private/QuadricMeshReduction.cpp
+++ b/Engine/Source/Developer/MeshSimplifier/Private/QuadricMeshReduction.cpp
@@ -563,7 +563,7 @@ public:
 					if (UVIndex < InVertexUVs.GetNumIndices())
 					{
 						NewVert.TexCoords[UVIndex] = InVertexUVs.GetArrayForIndex(UVIndex)[VertexInstanceIDs[CornerIndex]];
-						InMeshNumTexCoords = FMath::Max(UVIndex, InMeshNumTexCoords);
+						InMeshNumTexCoords = FMath::Max(UVIndex+1, InMeshNumTexCoords);
 					}
 					else
 					{
diff --git a/Engine/Source/Runtime/Engine/Private/StaticMesh.cpp b/Engine/Source/Runtime/Engine/Private/StaticMesh.cpp
--- a/Engine/Source/Runtime/Engine/Private/StaticMesh.cpp
+++ b/Engine/Source/Runtime/Engine/Private/StaticMesh.cpp
@@ -1333,7 +1333,7 @@ FArchive& operator<<(FArchive& Ar, FMeshBuildSettings& BuildSettings)
 // differences, etc.) replace the version GUID below with a new one.
 // In case of merge conflicts with DDC versions, you MUST generate a new GUID
 // and set this new GUID as the version.                                       
-#define STATICMESH_DERIVEDDATA_VER TEXT("D81A3AD754C34B6EA3F822DDB06000D8")
+#define STATICMESH_DERIVEDDATA_VER TEXT("AEC7F25C2C9141A99EC26695EE5BA392")
 
 static const FString& GetStaticMeshDerivedDataVersion()
 {
@@ -2689,7 +2689,7 @@ void UStaticMesh::FixupMaterialSlotName()
 // differences, etc.) replace the version GUID below with a new one.
 // In case of merge conflicts with DDC versions, you MUST generate a new GUID
 // and set this new GUID as the version.                                       
-#define MESHDATAKEY_STATICMESH_DERIVEDDATA_VER TEXT("ECB687DFD4B34464B85F5C854A13A309")
+#define MESHDATAKEY_STATICMESH_DERIVEDDATA_VER TEXT("8BB01BA6862D43818AE1E78DA6BF67EC")
 
 static const FString& GetMeshDataKeyStaticMeshDerivedDataVersion()
 {
