Bump the skeletal mesh version to force ddc refresh

diff --git a/Engine/Source/Runtime/Engine/Private/SkeletalMeshRenderData.cpp b/Engine/Source/Runtime/Engine/Private/SkeletalMeshRenderData.cpp
--- a/Engine/Source/Runtime/Engine/Private/SkeletalMeshRenderData.cpp
+++ b/Engine/Source/Runtime/Engine/Private/SkeletalMeshRenderData.cpp
@@ -29,7 +29,7 @@ namespace SkeletalMeshCookStats
 // differences, etc.) replace the version GUID below with a new one.
 // In case of merge conflicts with DDC versions, you MUST generate a new GUID
 // and set this new GUID as the version.                                       
-#define SKELETALMESH_DERIVEDDATA_VER TEXT("979A598F4D5F4686AD99644DCC83DCF2")
+#define SKELETALMESH_DERIVEDDATA_VER TEXT("9DD6B50371AE4517A1050CE83E4D2765")
 
 static const FString& GetSkeletalMeshDerivedDataVersion()
 {
