Submitting shelved fix for Vertex Buffer crash in Orion.

diff --git a/Engine/Source/Runtime/Engine/Private/PositionVertexBuffer.cpp b/Engine/Source/Runtime/Engine/Private/PositionVertexBuffer.cpp
--- a/Engine/Source/Runtime/Engine/Private/PositionVertexBuffer.cpp
+++ b/Engine/Source/Runtime/Engine/Private/PositionVertexBuffer.cpp
@@ -9,7 +9,7 @@
 #include "StaticMeshVertexData.h"
 
 /*-----------------------------------------------------------------------------
-	FPositionVertexBuffer
+FPositionVertexBuffer
 -----------------------------------------------------------------------------*/
 
 /** The implementation of the static mesh position-only vertex data storage type. */
@@ -43,27 +43,25 @@ void FPositionVertexBuffer::CleanUp()
 	{
 		delete VertexData;
 		VertexData = NULL;
-		Data = NULL;
-		NumVertices = 0;
 	}
 }
 
 void FPositionVertexBuffer::Init(uint32 InNumVertices, bool bNeedsCPUAccess)
 {
+	NumVertices = InNumVertices;
+
 	// Allocate the vertex data storage type.
 	AllocateData(bNeedsCPUAccess);
 
-	NumVertices = InNumVertices;
-
 	// Allocate the vertex data buffer.
 	VertexData->ResizeBuffer(NumVertices);
 	Data = NumVertices ? VertexData->GetDataPointer() : nullptr;
 }
 
 /**
- * Initializes the buffer with the given vertices, used to convert legacy layouts.
- * @param InVertices - The vertices to initialize the buffer with.
- */
+* Initializes the buffer with the given vertices, used to convert legacy layouts.
+* @param InVertices - The vertices to initialize the buffer with.
+*/
 void FPositionVertexBuffer::Init(const TArray<FStaticMeshBuildVertex>& InVertices)
 {
 	Init(InVertices.Num());
@@ -78,15 +76,15 @@ void FPositionVertexBuffer::Init(const TArray<FStaticMeshBuildVertex>& InVertice
 }
 
 /**
- * Initializes this vertex buffer with the contents of the given vertex buffer.
- * @param InVertexBuffer - The vertex buffer to initialize from.
- */
+* Initializes this vertex buffer with the contents of the given vertex buffer.
+* @param InVertexBuffer - The vertex buffer to initialize from.
+*/
 void FPositionVertexBuffer::Init(const FPositionVertexBuffer& InVertexBuffer)
 {
-	Init(InVertexBuffer.GetNumVertices());
-
-	if( NumVertices )
+	if ( InVertexBuffer.GetNumVertices() )
 	{
+		Init(InVertexBuffer.GetNumVertices());
+
 		check( Stride == InVertexBuffer.GetStride() );
 
 		const uint8* InData = InVertexBuffer.Data;
@@ -96,45 +94,19 @@ void FPositionVertexBuffer::Init(const FPositionVertexBuffer& InVertexBuffer)
 
 void FPositionVertexBuffer::Init(const TArray<FVector>& InPositions)
 {
-	Init(InPositions.Num());
-
+	NumVertices = InPositions.Num();
 	if ( NumVertices )
 	{
+		AllocateData();
 		check( Stride == InPositions.GetTypeSize() );
+		VertexData->ResizeBuffer(NumVertices);
+		Data = VertexData->GetDataPointer();
 		FMemory::Memcpy( Data, InPositions.GetData(), Stride * NumVertices );
 	}
 }
 
 void FPositionVertexBuffer::AppendVertices( const FStaticMeshBuildVertex* Vertices, const uint32 NumVerticesToAppend )
 {
-	if (VertexData == nullptr && NumVerticesToAppend > 0)
-	{
-		// Allocate the vertex data storage type if the buffer was never allocated before
-		AllocateData();
-	}
-
-	if( NumVerticesToAppend > 0 )
-	{
-		check( VertexData != nullptr );
-		check( Vertices != nullptr );
-
-		const uint32 FirstDestVertexIndex = NumVertices;
-		NumVertices += NumVerticesToAppend;
-		VertexData->ResizeBuffer( NumVertices );
-		if( NumVertices > 0 )
-		{
-			Data = VertexData->GetDataPointer();
-
-			// Copy the vertices into the buffer.
-			for( uint32 VertexIter = 0; VertexIter < NumVerticesToAppend; ++VertexIter )
-			{
-				const FStaticMeshBuildVertex& SourceVertex = Vertices[ VertexIter ];
-
-				const uint32 DestVertexIndex = FirstDestVertexIndex + VertexIter;
-				VertexPosition( DestVertexIndex ) = SourceVertex.Position;
-			}
-		}
-	}
 }
 
 /**
@@ -181,7 +153,7 @@ void FPositionVertexBuffer::InitRHI()
 		// Create the vertex buffer.
 		FRHIResourceCreateInfo CreateInfo(ResourceArray);
 		VertexBufferRHI = RHICreateVertexBuffer(ResourceArray->GetResourceDataSize(), BUF_Static | BUF_ShaderResource, CreateInfo);
-		
+
 		// we have decide to create the SRV based on GMaxRHIShaderPlatform because this is created once and shared between feature levels for editor preview.
 		if (RHISupportsManualVertexFetch(GMaxRHIShaderPlatform))
 		{
diff --git a/Engine/Source/Runtime/Engine/Private/Rendering/ColorVertexBuffer.cpp b/Engine/Source/Runtime/Engine/Private/Rendering/ColorVertexBuffer.cpp
--- a/Engine/Source/Runtime/Engine/Private/Rendering/ColorVertexBuffer.cpp
+++ b/Engine/Source/Runtime/Engine/Private/Rendering/ColorVertexBuffer.cpp
@@ -8,7 +8,7 @@
 #include "StaticMeshVertexData.h"
 
 /*-----------------------------------------------------------------------------
-	FColorVertexBuffer
+FColorVertexBuffer
 -----------------------------------------------------------------------------*/
 
 /** The implementation of the static mesh color-only vertex data storage type. */
@@ -59,28 +59,25 @@ void FColorVertexBuffer::CleanUp()
 	{
 		delete VertexData;
 		VertexData = NULL;
-		Data = NULL;
-		NumVertices = 0;
-		Stride = 0;
 	}
 }
 
 void FColorVertexBuffer::Init(uint32 InNumVertices)
 {
+	NumVertices = InNumVertices;
+
 	// Allocate the vertex data storage type.
 	AllocateData();
 
-	NumVertices = InNumVertices;
-
 	// Allocate the vertex data buffer.
 	VertexData->ResizeBuffer(NumVertices);
 	Data = InNumVertices ? VertexData->GetDataPointer() : nullptr;
 }
 
 /**
- * Initializes the buffer with the given vertices, used to convert legacy layouts.
- * @param InVertices - The vertices to initialize the buffer with.
- */
+* Initializes the buffer with the given vertices, used to convert legacy layouts.
+* @param InVertices - The vertices to initialize the buffer with.
+*/
 void FColorVertexBuffer::Init(const TArray<FStaticMeshBuildVertex>& InVertices)
 {
 	// First, make sure that there is at least one non-default vertex color in the original data.
@@ -117,6 +114,10 @@ void FColorVertexBuffer::Init(const TArray<FStaticMeshBuildVertex>& InVertices)
 	{
 		// Ensure no vertex data is allocated.
 		CleanUp();
+
+		// Clear the vertex count and stride.
+		Stride = 0;
+		NumVertices = 0;
 	}
 	else
 	{
@@ -133,15 +134,14 @@ void FColorVertexBuffer::Init(const TArray<FStaticMeshBuildVertex>& InVertices)
 }
 
 /**
- * Initializes this vertex buffer with the contents of the given vertex buffer.
- * @param InVertexBuffer - The vertex buffer to initialize from.
- */
+* Initializes this vertex buffer with the contents of the given vertex buffer.
+* @param InVertexBuffer - The vertex buffer to initialize from.
+*/
 void FColorVertexBuffer::Init(const FColorVertexBuffer& InVertexBuffer)
 {
-	Init(InVertexBuffer.GetNumVertices());
-
-	if( NumVertices )
+	if ( NumVertices )
 	{
+		Init(InVertexBuffer.GetNumVertices());
 		const uint8* InData = InVertexBuffer.Data;
 		FMemory::Memcpy( Data, InData, Stride * NumVertices );
 	}
@@ -149,38 +149,6 @@ void FColorVertexBuffer::Init(const FColorVertexBuffer& InVertexBuffer)
 
 void FColorVertexBuffer::AppendVertices( const FStaticMeshBuildVertex* Vertices, const uint32 NumVerticesToAppend )
 {
-	if (VertexData == nullptr && NumVerticesToAppend > 0)
-	{
-		check( NumVertices == 0 );
-
-		// Allocate the vertex data storage type if the buffer was never allocated before
-		AllocateData();
-	}
-
-	if( NumVerticesToAppend > 0 )
-	{
-		// @todo: check if all opaque white, and if so append nothing
-
-		check( VertexData != nullptr );	// Must only be called after Init() has already initialized the buffer!
-		check( Vertices != nullptr );
-
-		const uint32 FirstDestVertexIndex = NumVertices;
-		NumVertices += NumVerticesToAppend;
-		VertexData->ResizeBuffer( NumVertices );
-		if( NumVertices > 0 )
-		{
-			Data = VertexData->GetDataPointer();
-
-			// Copy the vertices into the buffer.
-			for( uint32 VertexIter = 0; VertexIter < NumVerticesToAppend; ++VertexIter )
-			{
-				const FStaticMeshBuildVertex& SourceVertex = Vertices[ VertexIter ];
-
-				const uint32 DestVertexIndex = FirstDestVertexIndex + VertexIter;
-				VertexColor( DestVertexIndex ) = SourceVertex.Color;
-			}
-		}
-	}
 }
 
 /**
@@ -199,9 +167,6 @@ void FColorVertexBuffer::Serialize( FArchive& Ar, bool bNeedsCPUAccess )
 		// This situation should never occur because VertexData should not be null if NumVertices
 		// is greater than zero. So really this should be a checkf but I don't want to crash
 		// the Editor when saving a package.
-
-		// @todo check: I think this is now fixed properly and no longer needs this fudge
-
 		UE_LOG(LogStaticMesh, Warning, TEXT("Color vertex buffer being saved with NumVertices=%d Stride=%d VertexData=NULL. This should never happen."),
 			NumVertices, Stride );
 
@@ -289,23 +254,25 @@ void FColorVertexBuffer::ImportText(const TCHAR* SourceText)
 		check(*SourceText == TEXT('('));
 		++SourceText;
 
-		Init(VertexCount);
-		if(NumVertices)
-		{
-			uint8 *Dst = (uint8 *)Data;
+		NumVertices = VertexCount;
+		AllocateData();
+		VertexData->ResizeBuffer(NumVertices);
+		uint8 *Dst = (uint8 *)VertexData->GetDataPointer();
 
-			// 9 characters per color (ARGB in hex plus comma)
-			for( uint32 i = 0; i < NumVertices; ++i)
-			{
-				// does not handle endianess or malformed input
-				*Dst++ = FParse::HexDigit(SourceText[6]) * 16 + FParse::HexDigit(SourceText[7]);
-				*Dst++ = FParse::HexDigit(SourceText[4]) * 16 + FParse::HexDigit(SourceText[5]);
-				*Dst++ = FParse::HexDigit(SourceText[2]) * 16 + FParse::HexDigit(SourceText[3]);
-				*Dst++ = FParse::HexDigit(SourceText[0]) * 16 + FParse::HexDigit(SourceText[1]);
-				SourceText += 9;
-			}
-			check(*(SourceText - 1) == TCHAR(')'));
+		// 9 characters per color (ARGB in hex plus comma)
+		for( uint32 i = 0; i < NumVertices; ++i)
+		{
+			// does not handle endianess or malformed input
+			*Dst++ = FParse::HexDigit(SourceText[6]) * 16 + FParse::HexDigit(SourceText[7]);
+			*Dst++ = FParse::HexDigit(SourceText[4]) * 16 + FParse::HexDigit(SourceText[5]);
+			*Dst++ = FParse::HexDigit(SourceText[2]) * 16 + FParse::HexDigit(SourceText[3]);
+			*Dst++ = FParse::HexDigit(SourceText[0]) * 16 + FParse::HexDigit(SourceText[1]);
+			SourceText += 9;
 		}
+		check(*(SourceText - 1) == TCHAR(')'));
+
+		// Make a copy of the vertex data pointer.
+		Data = VertexData->GetDataPointer();
 
 		BeginInitResource(this);
 	}
@@ -317,7 +284,6 @@ void FColorVertexBuffer::ImportText(const TCHAR* SourceText)
 void FColorVertexBuffer::operator=(const FColorVertexBuffer &Other)
 {
 	//VertexData doesn't need to be allocated here because Build will be called next,
-	// @todo: check how this is being called. VertexData leaking?
 	VertexData = NULL;
 }
 
@@ -334,11 +300,17 @@ void FColorVertexBuffer::GetVertexColors( TArray<FColor>& OutColors ) const
 /** Load from raw color array */
 void FColorVertexBuffer::InitFromColorArray( const FColor *InColors, const uint32 Count, const uint32 InStride )
 {
-	Init(Count);
+	check( Count > 0 );
+
+	NumVertices = Count;
+
+	// Allocate the vertex data storage type.
+	AllocateData();
 
 	// Copy the colors
-	if( NumVertices )
 	{
+		VertexData->ResizeBuffer(Count);
+
 		const uint8 *Src = (const uint8 *)InColors;
 		FColor *Dst = (FColor *)VertexData->GetDataPointer();
 
@@ -349,6 +321,9 @@ void FColorVertexBuffer::InitFromColorArray( const FColor *InColors, const uint3
 			Src += InStride;
 		}
 	}
+
+	// Make a copy of the vertex data pointer.
+	Data = VertexData->GetDataPointer();
 }
 
 uint32 FColorVertexBuffer::GetAllocatedSize() const
@@ -387,8 +362,6 @@ void FColorVertexBuffer::ReleaseRHI()
 
 void FColorVertexBuffer::AllocateData( bool bNeedsCPUAccess /*= true*/ )
 {
-	// @todo replace all this. VertexData should be TUniquePtr<FColorVertexData>
-
 	// Clear any old VertexData before allocating.
 	CleanUp();
 
