Reinstated AppendVertices methods for the vertex buffer classes, as used by [REDACTED_PERSON]. Thus we maintain the original API specs, but with the additional functionality required by [REDACTED_PERSON].

diff --git a/Engine/Source/Runtime/Engine/Private/PositionVertexBuffer.cpp b/Engine/Source/Runtime/Engine/Private/PositionVertexBuffer.cpp
--- a/Engine/Source/Runtime/Engine/Private/PositionVertexBuffer.cpp
+++ b/Engine/Source/Runtime/Engine/Private/PositionVertexBuffer.cpp
@@ -107,14 +107,42 @@ void FPositionVertexBuffer::Init(const TArray<FVector>& InPositions)
 
 void FPositionVertexBuffer::AppendVertices( const FStaticMeshBuildVertex* Vertices, const uint32 NumVerticesToAppend )
 {
+	if (VertexData == nullptr && NumVerticesToAppend > 0)
+	{
+		// Allocate the vertex data storage type if the buffer was never allocated before
+		AllocateData();
+	}
+
+	if( NumVerticesToAppend > 0 )
+	{
+		check( VertexData != nullptr );
+		check( Vertices != nullptr );
+
+		const uint32 FirstDestVertexIndex = NumVertices;
+		NumVertices += NumVerticesToAppend;
+		VertexData->ResizeBuffer( NumVertices );
+		if( NumVertices > 0 )
+		{
+			Data = VertexData->GetDataPointer();
+
+			// Copy the vertices into the buffer.
+			for( uint32 VertexIter = 0; VertexIter < NumVerticesToAppend; ++VertexIter )
+			{
+				const FStaticMeshBuildVertex& SourceVertex = Vertices[ VertexIter ];
+
+				const uint32 DestVertexIndex = FirstDestVertexIndex + VertexIter;
+				VertexPosition( DestVertexIndex ) = SourceVertex.Position;
+			}
+		}
+	}
 }
 
 /**
-* Serializer
-*
-* @param	Ar				Archive to serialize with
-* @param	bNeedsCPUAccess	Whether the elements need to be accessed by the CPU
-*/
+ * Serializer
+ *
+ * @param	Ar				Archive to serialize with
+ * @param	bNeedsCPUAccess	Whether the elements need to be accessed by the CPU
+ */
 void FPositionVertexBuffer::Serialize( FArchive& Ar, bool bNeedsCPUAccess )
 {
 	Ar << Stride << NumVertices;
diff --git a/Engine/Source/Runtime/Engine/Private/Rendering/ColorVertexBuffer.cpp b/Engine/Source/Runtime/Engine/Private/Rendering/ColorVertexBuffer.cpp
--- a/Engine/Source/Runtime/Engine/Private/Rendering/ColorVertexBuffer.cpp
+++ b/Engine/Source/Runtime/Engine/Private/Rendering/ColorVertexBuffer.cpp
@@ -149,14 +149,46 @@ void FColorVertexBuffer::Init(const FColorVertexBuffer& InVertexBuffer)
 
 void FColorVertexBuffer::AppendVertices( const FStaticMeshBuildVertex* Vertices, const uint32 NumVerticesToAppend )
 {
+	if (VertexData == nullptr && NumVerticesToAppend > 0)
+	{
+		check( NumVertices == 0 );
+
+		// Allocate the vertex data storage type if the buffer was never allocated before
+		AllocateData();
+	}
+
+	if( NumVerticesToAppend > 0 )
+	{
+		// @todo: check if all opaque white, and if so append nothing
+
+		check( VertexData != nullptr );	// Must only be called after Init() has already initialized the buffer!
+		check( Vertices != nullptr );
+
+		const uint32 FirstDestVertexIndex = NumVertices;
+		NumVertices += NumVerticesToAppend;
+		VertexData->ResizeBuffer( NumVertices );
+		if( NumVertices > 0 )
+		{
+			Data = VertexData->GetDataPointer();
+
+			// Copy the vertices into the buffer.
+			for( uint32 VertexIter = 0; VertexIter < NumVerticesToAppend; ++VertexIter )
+			{
+				const FStaticMeshBuildVertex& SourceVertex = Vertices[ VertexIter ];
+
+				const uint32 DestVertexIndex = FirstDestVertexIndex + VertexIter;
+				VertexColor( DestVertexIndex ) = SourceVertex.Color;
+			}
+		}
+	}
 }
 
 /**
-* Serializer
-*
-* @param	Ar				Archive to serialize with
-* @param	bNeedsCPUAccess	Whether the elements need to be accessed by the CPU
-*/
+ * Serializer
+ *
+ * @param	Ar				Archive to serialize with
+ * @param	bNeedsCPUAccess	Whether the elements need to be accessed by the CPU
+ */
 void FColorVertexBuffer::Serialize( FArchive& Ar, bool bNeedsCPUAccess )
 {
 	FStripDataFlags StripFlags(Ar, 0, VER_UE4_STATIC_SKELETAL_MESH_SERIALIZATION_FIX);
diff --git a/Engine/Source/Runtime/Engine/Private/Rendering/StaticMeshVertexBuffer.cpp b/Engine/Source/Runtime/Engine/Private/Rendering/StaticMeshVertexBuffer.cpp
--- a/Engine/Source/Runtime/Engine/Private/Rendering/StaticMeshVertexBuffer.cpp
+++ b/Engine/Source/Runtime/Engine/Private/Rendering/StaticMeshVertexBuffer.cpp
@@ -37,16 +37,15 @@ void FStaticMeshVertexBuffer::CleanUp()
 
 void FStaticMeshVertexBuffer::Init(uint32 InNumVertices, uint32 InNumTexCoords, bool bNeedsCPUAccess)
 {
-	// Allocate the vertex data storage type.
-	AllocateData(bNeedsCPUAccess);
-
 	NumTexCoords = InNumTexCoords;
 	NumVertices = InNumVertices;
 
+	// Allocate the vertex data storage type.
+	AllocateData(bNeedsCPUAccess);
+
 	// Allocate the vertex data buffer.
 	TangentsData->ResizeBuffer(NumVertices);
 	TangentsDataPtr = NumVertices ? TangentsData->GetDataPointer() : nullptr;
-
 	TexcoordData->ResizeBuffer(NumVertices * GetNumTexCoords());
 	TexcoordDataPtr = NumVertices ? TexcoordData->GetDataPointer() : nullptr;
 }
@@ -80,15 +79,18 @@ void FStaticMeshVertexBuffer::Init(const TArray<FStaticMeshBuildVertex>& InVerti
 */
 void FStaticMeshVertexBuffer::Init(const FStaticMeshVertexBuffer& InVertexBuffer)
 {
+	NumTexCoords = InVertexBuffer.GetNumTexCoords();
+	NumVertices = InVertexBuffer.GetNumVertices();
 	bUseFullPrecisionUVs = InVertexBuffer.GetUseFullPrecisionUVs();
 	bUseHighPrecisionTangentBasis = InVertexBuffer.GetUseHighPrecisionTangentBasis();
 
-	Init(InVertexBuffer.GetNumVertices(), InVertexBuffer.GetNumTexCoords());
-
 	if (NumVertices)
 	{
+		AllocateData();
 		{
 			check(TangentsData->GetStride() == InVertexBuffer.TangentsData->GetStride());
+			TangentsData->ResizeBuffer(NumVertices);
+			TangentsDataPtr = TangentsData->GetDataPointer();
 			const uint8* InData = InVertexBuffer.TangentsDataPtr;
 			FMemory::Memcpy(TangentsDataPtr, InData, TangentsData->GetStride() * NumVertices);
 		}
@@ -104,6 +106,8 @@ void FStaticMeshVertexBuffer::Init(const FStaticMeshVertexBuffer& InVertexBuffer
 			}
 			else
 			{
+				TexcoordData->ResizeBuffer(NumVertices * GetNumTexCoords());
+				TexcoordDataPtr = TexcoordData->GetDataPointer();
 				FMemory::Memcpy(TexcoordDataPtr, InData, TexcoordData->GetStride() * NumVertices * GetNumTexCoords());
 			}
 		}
@@ -131,6 +135,7 @@ void FStaticMeshVertexBuffer::ConvertHalfTexcoordsToFloat(const uint8* InData)
 	}
 
 	delete OriginalTexcoordData;
+	OriginalTexcoordData = nullptr;
 }
 
 
@@ -218,7 +223,7 @@ void FStaticMeshVertexBuffer::Serialize(FArchive& Ar, bool bNeedsCPUAccess)
 
 			// Make a copy of the vertex data pointer.
 			TexcoordDataPtr = TexcoordData->GetDataPointer();
-			
+
 			// convert half float data to full float if the HW requires it.
 			if (!GetUseFullPrecisionUVs() && !GVertexElementTypeSupport.IsSupported(VET_Half2))
 			{
@@ -365,7 +370,7 @@ void FStaticMeshVertexBuffer::BindTangentVertexBuffer(const FVertexFactory* Vert
 	{
 		Data.TangentsSRV = TangentsSRV;
 	}
-	
+
 	{
 		uint32 TangentSizeInBytes = 0;
 		uint32 TangentXOffset = 0;
@@ -416,7 +421,7 @@ void FStaticMeshVertexBuffer::BindPackedTexCoordVertexBuffer(const FVertexFactor
 	{
 		Data.TextureCoordinatesSRV = TextureCoordinatesSRV;
 	}
-	
+
 	{
 		EVertexElementType UVDoubleWideVertexElementType = VET_None;
 		EVertexElementType UVVertexElementType = VET_None;
@@ -470,7 +475,7 @@ void FStaticMeshVertexBuffer::BindTexCoordVertexBuffer(const FVertexFactory* Ver
 	{
 		Data.TextureCoordinatesSRV = TextureCoordinatesSRV;
 	}
-	
+
 	{
 		EVertexElementType UVVertexElementType = VET_None;
 		uint32 UVSizeInBytes = 0;
@@ -521,7 +526,7 @@ void FStaticMeshVertexBuffer::BindLightMapVertexBuffer(const FVertexFactory* Ver
 	{
 		Data.TextureCoordinatesSRV = TextureCoordinatesSRV;
 	}
-	
+
 	{
 		EVertexElementType UVVertexElementType = VET_None;
 		uint32 UVSizeInBytes = 0;
