Handle the case where an unparented material instance attempts to get material relevancy (using the fallback material).

diff --git a/Engine/Source/Runtime/Engine/Private/Materials/MaterialInterface.cpp b/Engine/Source/Runtime/Engine/Private/Materials/MaterialInterface.cpp
--- a/Engine/Source/Runtime/Engine/Private/Materials/MaterialInterface.cpp
+++ b/Engine/Source/Runtime/Engine/Private/Materials/MaterialInterface.cpp
@@ -106,6 +106,18 @@ FMaterialRelevance UMaterialInterface::GetRelevance_Internal(const UMaterial* Ma
 	if(Material)
 	{
 		const FMaterialResource* MaterialResource = GetMaterialResource(InFeatureLevel);
+
+		// If material is invalid e.g. unparented instance, fallback to the passed in material
+		if (!MaterialResource && Material)
+		{
+			MaterialResource = Material->GetMaterialResource(InFeatureLevel);	
+		}
+
+		if (!MaterialResource)
+		{
+			return FMaterialRelevance();
+		}
+
 		const EBlendMode BlendMode = (EBlendMode)GetBlendMode();
 		const bool bIsTranslucent = IsTranslucentBlendMode(BlendMode);
 
