Copying //UE4/Dev-Networking to //UE4/Dev-Main (Source: //UE4/Dev-Networking @ 4051526)

diff --git a/Engine/Source/Runtime/Engine/Private/Actor.cpp b/Engine/Source/Runtime/Engine/Private/Actor.cpp
--- a/Engine/Source/Runtime/Engine/Private/Actor.cpp
+++ b/Engine/Source/Runtime/Engine/Private/Actor.cpp
@@ -1030,6 +1030,10 @@ void AActor::PreReplicationForReplay(IRepChangedPropertyTracker & ChangedPropert
 {
 }
 
+void AActor::RewindForReplay()
+{
+}
+
 void AActor::PostActorCreated()
 {
 	// nothing at the moment
@@ -1827,12 +1831,15 @@ bool AActor::IsRelevancyOwnerFor(const AActor* ReplicatedActor, const AActor* Ac
 
 void AActor::ForceNetUpdate()
 {
-	if (NetDormancy > DORM_Awake)
+	if (UNetDriver* NetDriver = GetNetDriver())
 	{
-		FlushNetDormancy(); 
-	}
+		NetDriver->ForceNetUpdate(this);
 
-	SetNetUpdateTime(GetWorld()->TimeSeconds - 0.01f);
+		if (NetDormancy > DORM_Awake)
+		{
+			FlushNetDormancy(); 
+		}
+	}
 }
 
 bool AActor::IsReplicationPausedForConnection(const FNetViewer& ConnectionOwnerNetViewer)
@@ -1855,8 +1862,15 @@ void AActor::SetNetDormancy(ENetDormancy NewDormancy)
 	UNetDriver* NetDriver = GEngine->FindNamedNetDriver(MyWorld, NetDriverName);
 	if (NetDriver)
 	{
+		ENetDormancy OldDormancy = NetDormancy;
 		NetDormancy = NewDormancy;
 
+		// Tell driver about change
+		if (OldDormancy != NewDormancy)
+		{
+			NetDriver->NotifyActorDormancyChange(this, OldDormancy);
+		}
+
 		// If not dormant, flush actor from NetDriver's dormant list
 		if (NewDormancy <= DORM_Awake)
 		{
@@ -1881,10 +1895,12 @@ void AActor::FlushNetDormancy()
 		return;
 	}
 
+	bool bWasDormInitial = false;
 	if (NetDormancy == DORM_Initial)
 	{
 		// No longer initially dormant
 		NetDormancy = DORM_DormantAll;
+		bWasDormInitial = true;
 	}
 
 	// Don't proceed with network operations if not actually set to replicate
@@ -1905,7 +1921,7 @@ void AActor::FlushNetDormancy()
 
 		if (MyWorld->DemoNetDriver && MyWorld->DemoNetDriver!= NetDriver)
 		{
-			MyWorld->DemoNetDriver->FlushActorDormancy(this);
+			MyWorld->DemoNetDriver->FlushActorDormancy(this, bWasDormInitial);
 		}
 	}
 }
@@ -2064,7 +2080,13 @@ void AActor::TearOff()
 
 	if (NetMode == NM_ListenServer || NetMode == NM_DedicatedServer)
 	{
+		PRAGMA_DISABLE_DEPRECATION_WARNINGS
 		bTearOff = true;
+		PRAGMA_ENABLE_DEPRECATION_WARNINGS
+		if (UNetDriver* NetDriver = GetNetDriver())
+		{
+			NetDriver->NotifyActorTearOff(this);
+		}
 	}
 }
 
@@ -2286,7 +2308,7 @@ void AActor::DisplayDebug(UCanvas* Canvas, const FDebugDisplayInfo& DebugDisplay
 			// networking attributes
 			T = FString::Printf(TEXT("ROLE: %i RemoteRole: %i NetNode: %i"), (int32)Role, (int32)RemoteRole, (int32)GetNetMode());
 
-			if( bTearOff )
+			if( GetTearOff() )
 			{
 				T = T + FString(TEXT(" Tear Off"));
 			}
@@ -4417,7 +4439,7 @@ void AActor::SetLifeSpan( float InLifespan )
 	// Store the new value
 	InitialLifeSpan = InLifespan;
 	// Initialize a timer for the actors lifespan if there is one. Otherwise clear any existing timer
-	if ((Role == ROLE_Authority || bTearOff) && !IsPendingKill())
+	if ((Role == ROLE_Authority || GetTearOff()) && !IsPendingKill())
 	{
 		if( InLifespan > 0.0f)
 		{
diff --git a/Engine/Source/Runtime/Engine/Private/DataBunch.cpp b/Engine/Source/Runtime/Engine/Private/DataBunch.cpp
--- a/Engine/Source/Runtime/Engine/Private/DataBunch.cpp
+++ b/Engine/Source/Runtime/Engine/Private/DataBunch.cpp
@@ -8,9 +8,7 @@
 #include "Engine/NetConnection.h"
 #include "Engine/ControlChannel.h"
 
-static const int32 MAX_BUNCH_SIZE = 1024 * 1024; 
-
-
+const int32 MAX_BUNCH_SIZE = 1024 * 1024; 
 
 
 /*-----------------------------------------------------------------------------
@@ -33,6 +31,8 @@ FInBunch::FInBunch( UNetConnection* InConnection, uint8* Src, int64 CountBits )
 ,	bPartialInitial ( 0 )
 ,	bPartialFinal ( 0 )
 ,	bHasPackageMapExports ( 0 )
+,	bHasMustBeMappedGUIDs ( 0 )
+,	bIgnoreRPCs ( 0 )
 {
 	check(Connection);
 	// Match the byte swapping settings of the connection
@@ -68,6 +68,7 @@ FInBunch::FInBunch( FInBunch &InBunch, bool CopyBuffer )
 	bPartialFinal =	InBunch.bPartialFinal;
 	bHasPackageMapExports = InBunch.bHasPackageMapExports;
 	bHasMustBeMappedGUIDs =	InBunch.bHasMustBeMappedGUIDs;
+	bIgnoreRPCs = InBunch.bIgnoreRPCs;
 
 	// Copy network version info
 	ArEngineNetVer = InBunch.ArEngineNetVer;
diff --git a/Engine/Source/Runtime/Engine/Private/DataChannel.cpp b/Engine/Source/Runtime/Engine/Private/DataChannel.cpp
--- a/Engine/Source/Runtime/Engine/Private/DataChannel.cpp
+++ b/Engine/Source/Runtime/Engine/Private/DataChannel.cpp
@@ -18,6 +18,8 @@
 #include "Engine/PackageMapClient.h"
 #include "Engine/DemoNetDriver.h"
 #include "Engine/NetworkObjectList.h"
+#include "Engine/ReplicationDriver.h"
+#include "Stats/StatsMisc.h"
 
 DEFINE_LOG_CATEGORY(LogNet);
 DEFINE_LOG_CATEGORY(LogRep);
@@ -95,11 +97,13 @@ void UChannel::SetClosingFlag()
 }
 
 
-void UChannel::Close()
+int64 UChannel::Close()
 {
 	check(OpenedLocally || ChIndex == 0);		// We are only allowed to close channels that we opened locally (except channel 0, so the server can notify disconnected clients)
 	check(Connection->Channels[ChIndex]==this);
 
+	int64 NumBits = 0;
+
 	if ( !Closing && ( Connection->State == USOCK_Open || Connection->State == USOCK_Pending ) )
 	{
 		if ( ChIndex == 0 )
@@ -124,8 +128,11 @@ void UChannel::Close()
 			CloseBunch.bReliable = 1;
 			CloseBunch.bDormant = Dormant;
 			SendBunch( &CloseBunch, 0 );
+			NumBits = CloseBunch.GetNumBits();
 		}
 	}
+
+	return NumBits;
 }
 
 void UChannel::ConditionalCleanUp( const bool bForDestroy )
@@ -793,7 +800,11 @@ FPacketIdRange UChannel::SendBunch( FOutBunch* Bunch, bool Merge )
 	OutgoingBunches.Reset();
 
 	// Add any export bunches
-	AppendExportBunches( OutgoingBunches );
+	// Replay connections will manage export bunches separately.
+	if (!Connection->InternalAck)
+	{
+		AppendExportBunches( OutgoingBunches );
+	}
 
 	if ( OutgoingBunches.Num() )
 	{
@@ -1538,15 +1549,16 @@ void UActorChannel::Init( UNetConnection* InConnection, int32 InChannelIndex, bo
 void UActorChannel::SetClosingFlag()
 {
 	if( Actor )
-		Connection->ActorChannels.Remove( Actor );
+	{
+		Connection->RemoveActorChannel( Actor );
+	}
 	UChannel::SetClosingFlag();
 }
 
-void UActorChannel::Close()
+int64 UActorChannel::Close()
 {
 	UE_LOG(LogNetTraffic, Log, TEXT("UActorChannel::Close: ChIndex: %d, Actor: %s"), ChIndex, Actor ? *Actor->GetFullName() : TEXT("NULL") );
-
-	UChannel::Close();
+	int64 NumBits = UChannel::Close();
 
 	if (Actor != nullptr)
 	{
@@ -1562,7 +1574,7 @@ void UActorChannel::Close()
 				}
 
 				check( Actor->NetDormancy > DORM_Awake ); // Dormancy should have been canceled if game code changed NetDormancy
-				Connection->Driver->GetNetworkObjectList().MarkDormant(Actor, Connection, Connection->Driver->ClientConnections.Num(), Connection->Driver->NetDriverName);
+				Connection->Driver->NotifyActorFullyDormantForConnection(Actor, Connection);
 			}
 
 			// Validation checking
@@ -1576,12 +1588,14 @@ void UActorChannel::Close()
 		// SetClosingFlag() might have already done this, but we need to make sure as that won't get called if the connection itself has already been closed
 		if (Connection)
 		{
-			Connection->ActorChannels.Remove( Actor );
+			Connection->RemoveActorChannel( Actor );
 		}
 
 		Actor = nullptr;
 		CleanupReplicators( bKeepReplicators );
 	}
+
+	return NumBits;
 }
 
 void UActorChannel::CleanupReplicators( const bool bKeepReplicators )
@@ -1712,6 +1726,12 @@ bool UActorChannel::CleanUp( const bool bForDestroy )
 	checkf(Connection != nullptr, TEXT("UActorChannel::CleanUp: Connection is null!"));
 	checkf(Connection->Driver != nullptr, TEXT("UActorChannel::CleanUp: Connection->Driver is null!"));
 
+	UReplicationConnectionDriver* const ConnectionDriver = Connection->GetReplicationConnectionDriver();
+	if (ConnectionDriver)
+	{
+		ConnectionDriver->NotifyActorChannelCleanedUp(this);
+	}
+
 	const bool bIsServer = Connection->Driver->IsServer();
 
 	UE_LOG( LogNetTraffic, Log, TEXT( "UActorChannel::CleanUp: %s" ), *Describe() );
@@ -1759,7 +1779,7 @@ bool UActorChannel::CleanUp( const bool bForDestroy )
 		checkSlow(Connection->Driver->IsValidLowLevel());
 		if (Actor != NULL)
 		{
-			if (Actor->bTearOff && !Connection->Driver->ShouldClientDestroyTearOffActors())
+			if (Actor->GetTearOff() && !Connection->Driver->ShouldClientDestroyTearOffActors())
 			{
 				if (!bTornOff)
 				{
@@ -1772,11 +1792,11 @@ bool UActorChannel::CleanUp( const bool bForDestroy )
 					}
 				}
 			}
-			else if (Dormant && !Actor->bTearOff)
+			else if (Dormant && !Actor->GetTearOff())
 			{
 				Actor->NetDormancy = DORM_DormantAll;
 
-				Connection->Driver->GetNetworkObjectList().MarkDormant(Actor, Connection, 1, Connection->Driver->NetDriverName);
+				Connection->Driver->NotifyActorFullyDormantForConnection(Actor, Connection);
 				bWasDormant = true;
 			}
 			else if (!Actor->bNetTemporary && Actor->GetWorld() != NULL && !GIsRequestingExit)
@@ -1918,7 +1938,7 @@ void UActorChannel::SetChannelActor( AActor* InActor )
 	if (Actor)
 	{
 		// Add to map.
-		Connection->ActorChannels.Add( Actor, this );
+		Connection->AddActorChannel( Actor, this );
 
 		check( !ReplicationMap.Contains( Actor ) );
 
@@ -1933,35 +1953,42 @@ void UActorChannel::SetChannelActor( AActor* InActor )
 
 void UActorChannel::NotifyActorChannelOpen(AActor* InActor, FInBunch& InBunch)
 {
+	UNetDriver* const NetDriver = (Connection && Connection->Driver) ? Connection->Driver : nullptr;
+	UWorld* const World = (NetDriver && NetDriver->World) ? NetDriver->World : nullptr;
+	UDemoNetDriver* const DemoNetDriver = (World && World->DemoNetDriver) ? World->DemoNetDriver : nullptr;
+	
+	if (DemoNetDriver)
+	{
+		DemoNetDriver->PreNotifyActorChannelOpen(this, InActor);
+	}
+	
 	Actor->OnActorChannelOpen(InBunch, Connection);
 
-	if (Connection && Connection->Driver && !Connection->Driver->IsServer())
+	if (NetDriver && !NetDriver->IsServer())
 	{
 		if (Actor->NetDormancy > DORM_Awake)
 		{
 			Actor->NetDormancy = DORM_Awake;
 
 			// if recording on client, make sure the actor is marked active
-			if (Connection->Driver->World && Connection->Driver->World->IsRecordingClientReplay())
+			if (World && World->IsRecordingClientReplay() && DemoNetDriver)
 			{
-				UDemoNetDriver* DemoDriver = Connection->Driver->World->DemoNetDriver;
-				if (DemoDriver)
-				{
-					DemoDriver->GetNetworkObjectList().FindOrAdd(Actor, DemoDriver->NetDriverName);
+				DemoNetDriver->GetNetworkObjectList().FindOrAdd(Actor, DemoNetDriver->NetDriverName);
 
-					if (DemoDriver->ClientConnections.Num() > 0)
-					{
-						DemoDriver->GetNetworkObjectList().MarkActive(Actor, DemoDriver->ClientConnections[0], DemoDriver->NetDriverName);
-						DemoDriver->GetNetworkObjectList().ClearRecentlyDormantConnection(Actor, DemoDriver->ClientConnections[0], DemoDriver->NetDriverName);
-					}
+				UNetConnection* DemoClientConnection = (DemoNetDriver->ClientConnections.Num() > 0) ? DemoNetDriver->ClientConnections[0] : nullptr;
+				if (DemoClientConnection)
+				{
+					DemoNetDriver->GetNetworkObjectList().MarkActive(Actor, DemoClientConnection, DemoNetDriver->NetDriverName);
+					DemoNetDriver->GetNetworkObjectList().ClearRecentlyDormantConnection(Actor, DemoClientConnection, DemoNetDriver->NetDriverName);
 				}
 			}
 		}
 	}
 }
 
-void UActorChannel::SetChannelActorForDestroy( FActorDestructionInfo *DestructInfo )
+int64 UActorChannel::SetChannelActorForDestroy( FActorDestructionInfo *DestructInfo )
 {
+	int64 NumBits = 0;
 	check(Connection->Channels[ChIndex]==this);
 	if
 	(	!Closing
@@ -1982,7 +2009,10 @@ void UActorChannel::SetChannelActorForDestroy( FActorDestructionInfo *DestructIn
 		UE_LOG(LogNetDormancy, Verbose, TEXT("SetChannelActorForDestroy: Channel %d. NetGUID <%s> Path: %s. Bits: %d"), ChIndex, *DestructInfo->NetGUID.ToString(), *DestructInfo->PathName, CloseBunch.GetNumBits() );
 
 		SendBunch( &CloseBunch, 0 );
+		NumBits = CloseBunch.GetNumBits();
 	}
+
+	return NumBits;
 }
 
 void UActorChannel::Tick()
@@ -2281,6 +2311,8 @@ void UActorChannel::ProcessBunch( FInBunch & Bunch )
 		RepFlags.bNetOwner = true;
 	}
 
+	RepFlags.bIgnoreRPCs = Bunch.bIgnoreRPCs;
+
 	// ----------------------------------------------
 	//	Read chunks of actor content
 	// ----------------------------------------------
@@ -2413,22 +2445,50 @@ private:
 	const ENetRole	ActualRemoteRole;
 };
 
-bool UActorChannel::ReplicateActor()
+double GReplicateActorTimeSeconds = 0.0;
+int32 GNumReplicateActorCalls = 0;
+
+int64 UActorChannel::ReplicateActor()
 {
+	check(Actor);
+	
+	const UWorld* const ActorWorld = Actor->GetWorld();
+	check(ActorWorld);
+
+#if STATS
 	SCOPE_CYCLE_COUNTER(STAT_NetReplicateActorTime);
 
-	check(Actor);
+	UClass* ParentNativeClass = Actor->GetClass();
+	while(ParentNativeClass->IsNative() == false)
+	{
+		if (ParentNativeClass->GetSuperClass() == nullptr || ParentNativeClass->GetSuperClass() == AActor::StaticClass())
+		{
+			break;
+		}
+
+		ParentNativeClass = ParentNativeClass->GetSuperClass();
+	}
+
+	SCOPE_CYCLE_UOBJECT(ParentNativeClass, ParentNativeClass);
+#endif
+
 	check(!Closing);
 	check(Connection);
 	check(Connection->PackageMap);
-
-	const UWorld* const ActorWorld = Actor->GetWorld();
+	
+	const bool bReplay = ActorWorld && (ActorWorld->DemoNetDriver == Connection->GetDriver());
+	
+	FSimpleScopeSecondsCounter ScopedSecondsCounter(GReplicateActorTimeSeconds, !bReplay);
+	if (!bReplay)
+	{
+		GNumReplicateActorCalls++;
+	}
 
 	if (bPausedUntilReliableACK )
 	{
 		if (NumOutRec > 0)
 		{			
-			return false;
+			return 0;
 		}
 		bPausedUntilReliableACK = 0;
 		UE_LOG(LogNet, Log, TEXT( "ReplicateActor: bPausedUntilReliableACK is ending now that reliables have been ACK'd. %s"), *Describe());
@@ -2456,7 +2516,7 @@ bool UActorChannel::ReplicateActor()
 		// We were paused and still are, don't do anything.
 		if (bOldPaused && bNewPaused)
 		{
-			return false;
+			return 0;
 		}
 
 		bIsNewlyReplicationUnpaused = bOldPaused && !bNewPaused;
@@ -2470,8 +2530,6 @@ bool UActorChannel::ReplicateActor()
 		UE_LOG( LogNet, Warning, TEXT( "ReplicateActor: PackageMap->GetMustBeMappedGuidsInLastBunch().Num() != 0: %i" ), CastChecked< UPackageMapClient >( Connection->PackageMap )->GetMustBeMappedGuidsInLastBunch().Num() );
 	}
 
-	// Time how long it takes to replicate this particular actor
-	SCOPE_CYCLE_UOBJECT(Actor, Actor);
 
 	bool WroteSomethingImportant = bIsNewlyReplicationUnpaused || bIsNewlyReplicationPaused;
 
@@ -2481,14 +2539,15 @@ bool UActorChannel::ReplicateActor()
 		FString Error(FString::Printf(TEXT("Attempt to replicate '%s' while already replicating that Actor!"), *Actor->GetName()));
 		UE_LOG(LogNet, Log, TEXT("%s"), *Error);
 		ensureMsgf(false, TEXT("%s"), *Error);
-		return false;
+		return 0;
 	}
 
 	// Create an outgoing bunch, and skip this actor if the channel is saturated.
 	FOutBunch Bunch( this, 0 );
+
 	if( Bunch.IsError() )
 	{
-		return false;
+		return 0;
 	}
 
 	if (bIsNewlyReplicationPaused)
@@ -2544,6 +2603,7 @@ bool UActorChannel::ReplicateActor()
 	// ----------------------------------------------------------
 	// If initial, send init data.
 	// ----------------------------------------------------------
+
 	if( RepFlags.bNetInitial && OpenedLocally )
 	{
 		Connection->PackageMap->SerializeNewActor(Bunch, this, Actor);
@@ -2557,7 +2617,7 @@ bool UActorChannel::ReplicateActor()
 
 	RepFlags.bNetSimulated	= ( Actor->GetRemoteRole() == ROLE_SimulatedProxy );
 	RepFlags.bRepPhysics	= Actor->ReplicatedMovement.bRepPhysics;
-	RepFlags.bReplay		= ActorWorld && (ActorWorld->DemoNetDriver == Connection->GetDriver());
+	RepFlags.bReplay		= bReplay;
 	//RepFlags.bNetInitial	= RepFlags.bNetInitial;
 
 	UE_LOG(LogNetTraffic, Log, TEXT("Replicate %s, bNetInitial: %d, bNetOwner: %d"), *Actor->GetName(), RepFlags.bNetInitial, RepFlags.bNetOwner );
@@ -2566,7 +2626,7 @@ bool UActorChannel::ReplicateActor()
 
 	// ----------------------------------------------------------
 	// Replicate Actor and Component properties and RPCs
-	// ----------------------------------------------------------
+	// ---------------------------------------------------
 
 #if USE_NETWORK_PROFILER 
 	const uint32 ActorReplicateStartTime = GNetworkProfiler.IsTrackingEnabled() ? FPlatformTime::Cycles() : 0;
@@ -2611,13 +2671,13 @@ bool UActorChannel::ReplicateActor()
 		}
 	}
 
-
 	NETWORK_PROFILER(GNetworkProfiler.TrackReplicateActor(Actor, RepFlags, FPlatformTime::Cycles() - ActorReplicateStartTime, Connection));
 
 	// -----------------------------
 	// Send if necessary
 	// -----------------------------
-	bool SentBunch = false;
+
+	int64 NumBitsWrote = 0;
 	if( WroteSomethingImportant )
 	{
 		FPacketIdRange PacketRange = SendBunch( &Bunch, 1 );
@@ -2660,7 +2720,7 @@ bool UActorChannel::ReplicateActor()
 				Connection->SentTemporaries.Add(Actor);
 			}
 		}
-		SentBunch = true;
+		NumBitsWrote = Bunch.GetNumBits();
 	}
 
 	PendingObjKeys.Empty();
@@ -2674,7 +2734,9 @@ bool UActorChannel::ReplicateActor()
 
 	bForceCompareProperties = false;		// Only do this once per frame when set
 
-	return SentBunch;
+	
+	INC_DWORD_STAT_BY(STAT_NumReplicatedActorBytes, (NumBitsWrote + 7) >> 3);
+	return NumBitsWrote;
 }
 
 
@@ -2738,6 +2800,11 @@ bool UActorChannel::ReadyForDormancy(bool suppressLogs)
 
 void UActorChannel::StartBecomingDormant()
 {
+	if (bPendingDormancy || Dormant)
+	{
+		return;
+	}
+
 	UE_LOG(LogNetDormancy, Verbose, TEXT("StartBecomingDormant: %s"), *Describe() );
 
 	for (auto MapIt = ReplicationMap.CreateIterator(); MapIt; ++MapIt)
@@ -3380,6 +3447,8 @@ bool UActorChannel::KeyNeedsToReplicate(int32 ObjID, int32 RepKey)
 
 bool UActorChannel::ReplicateSubobject(UObject *Obj, FOutBunch &Bunch, const FReplicationFlags &RepFlags)
 {
+	SCOPE_CYCLE_UOBJECT(ActorChannelRepSubObj, Obj);
+
 	// Hack for now: subobjects are SupportsObject==false until they are replicated via ::ReplicateSUbobject, and then we make them supported
 	// here, by forcing the packagemap to give them a NetGUID.
 	//
@@ -3458,7 +3527,7 @@ static void	ListOpenActorChannels( UWorld* InWorld )
 
 	TMap<UClass*, int32> ClassMap;
 	
-	for (auto It = Connection->ActorChannels.CreateIterator(); It; ++It)
+	for (auto It = Connection->ActorChannelConstIterator(); It; ++It)
 	{
 		UActorChannel* Chan = It.Value();
 		UClass *ThisClass = Chan->Actor->GetClass();
diff --git a/Engine/Source/Runtime/Engine/Private/DataReplication.cpp b/Engine/Source/Runtime/Engine/Private/DataReplication.cpp
--- a/Engine/Source/Runtime/Engine/Private/DataReplication.cpp
+++ b/Engine/Source/Runtime/Engine/Private/DataReplication.cpp
@@ -131,16 +131,16 @@ void FObjectReplicator::InitRecentProperties( uint8* Source )
 {
 	check( GetObject() != NULL );
 	check( Connection != NULL );
-	check( RepState == NULL );
+	check( !RepState.IsValid() );
 
 	UClass * InObjectClass = GetObject()->GetClass();
 
-	RepState = new FRepState;
+	RepState = MakeShareable(new FRepState());
 
 	// Initialize the RepState memory
 	TSharedPtr<FRepChangedPropertyTracker> RepChangedPropertyTracker = Connection->Driver->FindOrCreateRepChangedPropertyTracker( GetObject() );
 
-	RepLayout->InitRepState( RepState, InObjectClass, Source, RepChangedPropertyTracker );
+	RepLayout->InitRepState( RepState.Get(), InObjectClass, Source, RepChangedPropertyTracker );
 	RepState->RepLayout = RepLayout;
 
 	if ( !Connection->Driver->IsServer() )
@@ -184,13 +184,13 @@ bool FObjectReplicator::ValidateAgainstState( const UObject* ObjectState )
 		return false;
 	}
 
-	if ( RepState == NULL )
+	if ( !RepState.IsValid() )
 	{
-		UE_LOG(LogRep, Warning, TEXT("ValidateAgainstState: RepState == NULL"));
+		UE_LOG(LogRep, Warning, TEXT("ValidateAgainstState: RepState.IsValid() == false"));
 		return false;
 	}
 
-	if ( RepLayout->DiffProperties( &(RepState->RepNotifies), RepState->StaticBuffer.GetData(), ObjectState, false ) )
+	if ( RepLayout->DiffProperties( &(RepState->RepNotifies), RepState->StaticBuffer.GetData(), ObjectState, EDiffPropertiesFlags::None ) )
 	{
 		UE_LOG(LogRep, Warning, TEXT("ValidateAgainstState: Properties changed for %s"), *ObjectState->GetName());
 		return false;
@@ -206,7 +206,7 @@ void FObjectReplicator::InitWithObject( UObject* InObject, UNetConnection * InCo
 	check( bLastUpdateEmpty == false );
 	check( Connection == NULL );
 	check( OwningChannel == NULL );
-	check( RepState == NULL );
+	check( !RepState.IsValid() );
 	check( RemoteFunctions == NULL );
 	check( !RepLayout.IsValid() );
 
@@ -291,11 +291,7 @@ void FObjectReplicator::CleanUp()
 	LifetimeCustomDeltaProperties.Empty();
 	LifetimeCustomDeltaPropertyConditions.Empty();
 
-	if ( RepState != NULL )
-	{
-		delete RepState;
-		RepState = NULL;
-	}
+	RepState = nullptr;
 }
 
 void FObjectReplicator::StartReplicating( class UActorChannel * InActorChannel )
@@ -454,7 +450,7 @@ void FObjectReplicator::ReceivedNak( int32 NakPacketId )
 
 	if ( Object != NULL && ObjectClass != NULL )
 	{
-		RepLayout->ReceivedNak( RepState, NakPacketId );
+		RepLayout->ReceivedNak( RepState.Get(), NakPacketId );
 
 		for ( int32 i = Retirement.Num() - 1; i >= 0; i-- )
 		{
@@ -560,7 +556,7 @@ bool FObjectReplicator::ReceivedBunch( FNetBitReader& Bunch, const FReplicationF
 
 		bool bLocalHasUnmapped = false;
 
-		if ( !RepLayout->ReceiveProperties( OwningChannel, ObjectClass, RepState, ( void* )Object, Bunch, bLocalHasUnmapped, bShouldReceiveRepNotifies, bGuidsChanged ) )
+		if ( !RepLayout->ReceiveProperties( OwningChannel, ObjectClass, RepState.Get(), ( void* )Object, Bunch, bLocalHasUnmapped, bShouldReceiveRepNotifies, bGuidsChanged ) )
 		{
 			UE_LOG( LogRep, Error, TEXT( "RepLayout->ReceiveProperties FAILED: %s" ), *Object->GetFullName() );
 			return false;
@@ -801,7 +797,8 @@ bool FObjectReplicator::ReceivedRPC(FNetBitReader& Reader, const FReplicationFla
 	UE_LOG(LogRepTraffic, Log, TEXT("      Received RPC: %s"), *FunctionName.ToString());
 
 	// validate that the function is callable here
-	const bool bCanExecute = (!bIsServer || RepFlags.bNetOwner);		// we are client or net owner
+	// we are client or net owner and shouldn't be ignoring rpcs
+	const bool bCanExecute = ((!bIsServer || RepFlags.bNetOwner) && !RepFlags.bIgnoreRPCs);
 
 	if (bCanExecute)
 	{
@@ -892,9 +889,9 @@ void FObjectReplicator::UpdateGuidToReplicatorMap()
 	int32 LocalTrackedGuidMemoryBytes = 0;
 
 	// Gather guids on rep layout
-	if ( RepLayout.IsValid() && RepState )
+	if ( RepLayout.IsValid() && RepState.IsValid() )
 	{
-		RepLayout->GatherGuidReferences( RepState, LocalReferencedGuids, LocalTrackedGuidMemoryBytes );
+		RepLayout->GatherGuidReferences( RepState.Get(), LocalReferencedGuids, LocalTrackedGuidMemoryBytes );
 	}
 
 	UObject* Object = GetObject();
@@ -974,7 +971,7 @@ bool FObjectReplicator::MoveMappedObjectToUnmapped( const FNetworkGUID& GUID )
 
 	if ( RepLayout.IsValid() )
 	{
-		if ( RepLayout->MoveMappedObjectToUnmapped( RepState, GUID ) )
+		if ( RepLayout->MoveMappedObjectToUnmapped( RepState.Get(), GUID ) )
 		{
 			bFound = true;
 		}
@@ -1215,7 +1212,7 @@ bool FObjectReplicator::ReplicateProperties( FOutBunch & Bunch, FReplicationFlag
 	{
 		check( OwningChannel );
 		check( RepLayout.IsValid() );
-		check( RepState )
+		check( RepState.IsValid() )
 		check( RepState->StaticBuffer.Num() );
 	}
 
@@ -1227,7 +1224,7 @@ bool FObjectReplicator::ReplicateProperties( FOutBunch & Bunch, FReplicationFlag
 	ChangelistMgr->Update( Object, Connection->Driver->ReplicationFrame, RepState->LastCompareIndex, RepFlags, OwningChannel->bForceCompareProperties );
 
 	// Replicate properties in the layout
-	const bool bHasRepLayout = RepLayout->ReplicateProperties( RepState, ChangelistMgr->GetRepChangelistState(), ( uint8* )Object, ObjectClass, OwningChannel, Writer, RepFlags );
+	const bool bHasRepLayout = RepLayout->ReplicateProperties( RepState.Get(), ChangelistMgr->GetRepChangelistState(), ( uint8* )Object, ObjectClass, OwningChannel, Writer, RepFlags );
 
 	// Replicate all the custom delta properties (fast arrays, etc)
 	ReplicateCustomDeltaProperties( Writer, RepFlags );
@@ -1290,7 +1287,7 @@ void FObjectReplicator::ForceRefreshUnreliableProperties()
 
 	check( !bOpenAckCalled );
 
-	RepLayout->OpenAcked( RepState );
+	RepLayout->OpenAcked( RepState.Get() );
 
 	bOpenAckCalled = true;
 }
@@ -1313,7 +1310,7 @@ void FObjectReplicator::PostSendBunch( FPacketIdRange & PacketRange, uint8 bReli
 		// Don't call if reliable, since the bunch will be resent. We dont want this to end up in the changelist history
 		// But is that enough? How does it know to delta against this latest state?
 
-		RepLayout->PostReplicate( RepState, PacketRange, bReliable ? true : false );
+		RepLayout->PostReplicate( RepState.Get(), PacketRange, bReliable ? true : false );
 	}
 
 	for ( int32 i = 0; i < LifetimeCustomDeltaProperties.Num(); i++ )
@@ -1366,6 +1363,11 @@ void FObjectReplicator::Serialize(FArchive& Ar)
 
 void FObjectReplicator::QueueRemoteFunctionBunch( UFunction* Func, FOutBunch &Bunch )
 {
+	if (Connection == nullptr)
+	{
+		return;
+	}
+
 	// This is a pretty basic throttling method - just don't let same func be called more than
 	// twice in one network update period.
 	//
@@ -1391,19 +1393,22 @@ void FObjectReplicator::QueueRemoteFunctionBunch( UFunction* Func, FOutBunch &Bu
 	{
 		UE_LOG(LogRep, Verbose, TEXT("Too many calls (%d) to RPC %s within a single netupdate. Skipping. %s.  LastCallTime: %.2f. CurrentTime: %.2f. LastRelevantTime: %.2f. LastUpdateTime: %.2f "),
 			RemoteFuncInfo[InfoIdx].Calls, *Func->GetName(), *GetPathNameSafe(GetObject()), RemoteFuncInfo[InfoIdx].LastCallTime, OwningChannel->Connection->Driver->Time, OwningChannel->RelevantTime, OwningChannel->LastUpdateTime);
+
+		// The MustBeMappedGuids can just be dropped, because we aren't actually going to send a bunch. If we don't clear it, then we will get warnings when the next channel tries to replicate
+		CastChecked< UPackageMapClient >( Connection->PackageMap )->GetMustBeMappedGuidsInLastBunch().Reset();
 		return;
 	}
 	
 	RemoteFuncInfo[InfoIdx].LastCallTime = OwningChannel->Connection->Driver->Time;
 
-	if (RemoteFunctions == NULL)
+	if (RemoteFunctions == nullptr)
 	{
 		RemoteFunctions = new FOutBunch(OwningChannel, 0);
 	}
 
 	RemoteFunctions->SerializeBits( Bunch.GetData(), Bunch.GetNumBits() );
 
-	if ( Connection != NULL && Connection->PackageMap != NULL )
+	if ( Connection->PackageMap != nullptr )
 	{
 		UPackageMapClient * PackageMapClient = CastChecked< UPackageMapClient >( Connection->PackageMap );
 
@@ -1411,11 +1416,14 @@ void FObjectReplicator::QueueRemoteFunctionBunch( UFunction* Func, FOutBunch &Bu
 		if ( PackageMapClient->GetMustBeMappedGuidsInLastBunch().Num() )
 		{
 			OwningChannel->QueuedMustBeMappedGuidsInLastBunch.Append( PackageMapClient->GetMustBeMappedGuidsInLastBunch() );
-			PackageMapClient->GetMustBeMappedGuidsInLastBunch().Empty();
+			PackageMapClient->GetMustBeMappedGuidsInLastBunch().Reset();
 		}
 
-		// Copy over any exported bunches
-		PackageMapClient->AppendExportBunches( OwningChannel->QueuedExportBunches );
+		if (!Connection->InternalAck)
+		{
+			// Copy over any exported bunches
+			PackageMapClient->AppendExportBunches( OwningChannel->QueuedExportBunches );
+		}
 	}
 }
 
@@ -1451,7 +1459,7 @@ bool FObjectReplicator::ReadyForDormancy(bool suppressLogs)
 		}
 	}
 
-	return RepLayout->ReadyForDormancy( RepState );
+	return RepLayout->ReadyForDormancy( RepState.Get() );
 }
 
 void FObjectReplicator::StartBecomingDormant()
@@ -1484,7 +1492,7 @@ void FObjectReplicator::CallRepNotifies(bool bSkipIfChannelHasQueuedBunches)
 		return;
 	}
 
-	RepLayout->CallRepNotifies( RepState, Object );
+	RepLayout->CallRepNotifies( RepState.Get(), Object );
 
 	if ( RepNotifies.Num() > 0 )
 	{
@@ -1575,7 +1583,7 @@ void FObjectReplicator::UpdateUnmappedObjects( bool & bOutHasMoreUnmapped )
 	bool bSomeObjectsWereMapped = false;
 
 	// Let the rep layout update any unmapped properties
-	RepLayout->UpdateUnmappedObjects( RepState, Connection->PackageMap, Object, bSomeObjectsWereMapped, bOutHasMoreUnmapped );
+	RepLayout->UpdateUnmappedObjects( RepState.Get(), Connection->PackageMap, Object, bSomeObjectsWereMapped, bOutHasMoreUnmapped );
 
 	// Update unmapped objects for custom properties (currently just fast tarray)
 	for ( auto It = UnmappedCustomProperties.CreateIterator(); It; ++It )
