Copying //UE4/Dev-Networking to //UE4/Dev-Main (Source: //UE4/Dev-Networking @ 4051526)

diff --git a/Engine/Source/Runtime/Engine/Private/DemoNetDriver.cpp b/Engine/Source/Runtime/Engine/Private/DemoNetDriver.cpp
--- a/Engine/Source/Runtime/Engine/Private/DemoNetDriver.cpp
+++ b/Engine/Source/Runtime/Engine/Private/DemoNetDriver.cpp
@@ -33,6 +33,9 @@
 #include "GameFramework/PlayerState.h"
 #include "HAL/LowLevelMemTracker.h"
 #include "Stats/StatsMisc.h"
+#include "Kismet/GameplayStatics.h"
+#include "ProfilingDebugging/CsvProfiler.h"
+#include "Misc/EngineVersion.h"
 
 DEFINE_LOG_CATEGORY( LogDemo );
 
@@ -51,13 +54,15 @@ static TAutoConsoleVariable<float> CVarCheckpointUploadDelayInSeconds( TEXT( "de
 static TAutoConsoleVariable<int32> CVarDemoLoadCheckpointGarbageCollect( TEXT( "demo.LoadCheckpointGarbageCollect" ), 1, TEXT("If nonzero, CollectGarbage will be called during LoadCheckpoint after the old actors and connection are cleaned up." ) );
 static TAutoConsoleVariable<float> CVarCheckpointSaveMaxMSPerFrameOverride( TEXT( "demo.CheckpointSaveMaxMSPerFrameOverride" ), -1.0f, TEXT( "If >= 0, this value will override the CheckpointSaveMaxMSPerFrame member variable, which is the maximum time allowed each frame to spend on saving a checkpoint. If 0, it will save the checkpoint in a single frame, regardless of how long it takes." ) );
 static TAutoConsoleVariable<int32> CVarDemoClientRecordAsyncEndOfFrame( TEXT( "demo.ClientRecordAsyncEndOfFrame" ), 0, TEXT( "If true, TickFlush will be called on a thread in parallel with Slate." ) );
-
 static TAutoConsoleVariable<int32> CVarForceDisableAsyncPackageMapLoading( TEXT( "demo.ForceDisableAsyncPackageMapLoading" ), 0, TEXT( "If true, async package map loading of network assets will be disabled." ) );
-
 static TAutoConsoleVariable<int32> CVarDemoUseNetRelevancy( TEXT( "demo.UseNetRelevancy" ), 0, TEXT( "If 1, will enable relevancy checks and distance culling, using all connected clients as reference." ) );
 static TAutoConsoleVariable<float> CVarDemoCullDistanceOverride( TEXT( "demo.CullDistanceOverride" ), 0.0f, TEXT( "If > 0, will represent distance from any viewer where actors will stop being recorded." ) );
 static TAutoConsoleVariable<float> CVarDemoRecordHzWhenNotRelevant( TEXT( "demo.RecordHzWhenNotRelevant" ), 2.0f, TEXT( "Record at this frequency when actor is not relevant." ) );
 static TAutoConsoleVariable<int32> CVarLoopDemo(TEXT("demo.Loop"), 0, TEXT("<1> : play replay from beginning once it reaches the end / <0> : stop replay at the end"));
+static TAutoConsoleVariable<int32> CVarDemoFastForwardIgnoreRPCs( TEXT( "demo.FastForwardIgnoreRPCs" ), 1, TEXT( "If true, RPCs will be discarded during playback fast forward." ) );
+static TAutoConsoleVariable<int32> CVarDemoSaveRollbackActorState( TEXT( "demo.SaveRollbackActorState" ), 1, TEXT( "If true, rollback actors will save some replicated state to apply when respawned." ) );
+static TAutoConsoleVariable<int32> CVarWithLevelStreamingFixes(TEXT("demo.WithLevelStreamingFixes"), 0, TEXT("If 1, provides fixes for level streaming (but breaks backwards compatibility)."));
+static TAutoConsoleVariable<int32> CVarWithDemoTimeBurnIn(TEXT("demo.WithTimeBurnIn"), 0, TEXT("If true, adds an on screen message with the current DemoTime and Changelist."));
 
 #if !(UE_BUILD_SHIPPING || UE_BUILD_TEST)
 static TAutoConsoleVariable<int32> CVarDemoForceFailure( TEXT( "demo.ForceFailure" ), 0, TEXT( "" ) );
@@ -70,10 +75,119 @@ namespace ReplayTaskNames
 	static FName SkipTimeInSecondsTask(TEXT("SkipTimeInSecondsTask"));
 	static FName JumpToLiveReplayTask(TEXT("JumpToLiveReplayTask"));
 	static FName GotoTimeInSecondsTask(TEXT("GotoTimeInSecondsTask"));
+	static FName FastForwardLevelsTask(TEXT("FastForwardLevelsTask"));
 };
 
 #define DEMO_CHECKSUMS 0		// When setting this to 1, this will invalidate all demos, you will need to re-record and playback
 
+// This is only intended for testing purposes
+// A "better" way might be to throw together a GameplayDebuggerComponent or Category, so we could populate
+// more than just the DemoTime.
+static void ConditionallyDisplayBurnInTime(uint32 RecordedCL, float CurrentDemoTime)
+{
+	if (CVarWithDemoTimeBurnIn.GetValueOnAnyThread() != 0)
+	{
+		GEngine->AddOnScreenDebugMessage(INDEX_NONE, 0.f, FColor::Red, FString::Printf(TEXT("Current CL: %lu | Recorded CL: %lu | Time: %f"), FEngineVersion::Current().GetChangelist(), RecordedCL, CurrentDemoTime), true, FVector2D(3.f, 3.f));
+	}
+}
+
+static FString GetLevelPackageName(const ULevel& InLevel)
+{
+	// TODO: Do we need to call NetworkRemapPath here?
+	//			If so, consider passing in a bool indicating if this is
+	//			for playback / recording, or maybe move it into DemoNetDriver.
+	return InLevel.GetOutermost()->GetFName().ToString();
+}
+
+static void FlushNetChecked(UNetConnection& NetConnection)
+{
+	NetConnection.FlushNet();
+	check(NetConnection.SendBuffer.GetNumBits() == 0);
+}
+
+// Helps manage packets, and any associations with streaming levels or exported GUIDs / fields.
+class FScopedPacketManager
+{
+public:
+	FScopedPacketManager(UDemoNetConnection& InConnection, const uint32 InSeenLevelIndex):
+		Connection(InConnection),
+		Packets(Connection.bResendAllDataSinceOpen ? InConnection.QueuedCheckpointPackets : InConnection.QueuedDemoPackets),
+		SeenLevelIndex(InSeenLevelIndex)
+	{
+		FlushNetChecked(Connection);
+		StartPacketCount = Packets.Num();
+	}
+
+	~FScopedPacketManager()
+	{
+		FlushNetChecked(Connection);
+		AssociatePacketsWithLevel();
+	}
+
+private:
+
+	void AssociatePacketsWithLevel()
+	{
+		for (int32 i = StartPacketCount; i < Packets.Num(); i++)
+		{
+			Packets[i].SeenLevelIndex = SeenLevelIndex;
+		}
+	}
+
+	UDemoNetConnection& Connection;
+	TArray<FQueuedDemoPacket>& Packets;
+	const uint32 SeenLevelIndex;
+	int32 StartPacketCount;
+};
+
+class FPendingTaskHelper
+{
+public:
+
+	static bool FastForwardLevels(UDemoNetDriver* DemoNetDriver, FArchive* CheckpointArchive, int64 ExtraTime)
+	{
+		return DemoNetDriver->FastForwardLevels(CheckpointArchive, ExtraTime);
+	}
+
+	static float GetLastProcessedPacketTime(UDemoNetDriver* DemoNetDriver)
+	{
+		return DemoNetDriver->LastProcessedPacketTime;
+	}
+};
+
+typedef int64 FArchivePos;
+
+/**
+ * Helps track Offsets in an Archive before the actual size of the offset is known.
+ * This relies on serialization always used a fixed number of bytes for primitive types,
+ * and Sane implementations of Seek and Tell.
+ */
+struct FScopedStoreArchiveOffset
+{
+	FScopedStoreArchiveOffset(FArchive& InAr):
+		Ar(InAr),
+		StartPosition(Ar.Tell())
+	{
+		// Save room for the offset here.
+		FArchivePos TempOffset = 0;
+		Ar << TempOffset;
+	}
+
+	~FScopedStoreArchiveOffset()
+	{
+		const FArchivePos CurrentPosition = Ar.Tell();
+		FArchivePos Offset = CurrentPosition - (StartPosition + sizeof(FArchivePos));
+		Ar.Seek(StartPosition);
+		Ar << Offset;
+		Ar.Seek(CurrentPosition);
+	}
+
+private:
+
+	FArchive& Ar;
+	const FArchivePos StartPosition;
+};
+
 // RAII object to swap the Role and RemoteRole of an actor within a scope. Used for recording replays on a client.
 class FScopedActorRoleSwap
 {
@@ -220,7 +334,7 @@ public:
 
 	void CheckpointReady( const bool bSuccess, const int64 SkipExtraTimeInMS )
 	{
-		check( GotoCheckpointArchive == NULL );
+		check( GotoCheckpointArchive == nullptr );
 		check( GotoCheckpointSkipExtraTimeInMS == -1 );
 
 		if ( !bSuccess )
@@ -279,6 +393,100 @@ public:
 	float SecondsToSkip;
 };
 
+class FFastForwardLevelsTask : public FQueuedReplayTask
+{
+public:
+
+	FFastForwardLevelsTask( UDemoNetDriver* InDriver ) : FQueuedReplayTask( InDriver ), bSkipWork(false), CheckpointArchive(nullptr), CheckpointSkipExtraTimeInMS(-1) { }
+
+	virtual void StartTask() override
+	{
+		check(!Driver->IsFastForwarding());
+
+		// If there's a GotoTimeInSeconds task pending, we don't need to do any work.
+		// That task should trigger a full checkpoint load.
+		// Only check the next task, to avoid issues with SkipTime / JumpToLive not having updated levels.
+		if (Driver->GetNextQueuedTaskName() == ReplayTaskNames::GotoTimeInSecondsTask)
+		{
+			bSkipWork = true;
+		}
+		else
+		{
+			// Tell the streamer to start going to this time.
+			// This should force the streamer to get the correct checkpoint and set of packets.
+			CheckpointReadyDelegate = FOnCheckpointReadyDelegate::CreateRaw(this, &FFastForwardLevelsTask::CheckpointReady);
+
+			const float LastProcessedPacketTime = FPendingTaskHelper::GetLastProcessedPacketTime(Driver);
+			const uint32 GotoTime = LastProcessedPacketTime * 1000;
+			Driver->ReplayStreamer->GotoTimeInMS(GotoTime, CheckpointReadyDelegate);
+
+			// Pause channels while we wait (so the world is paused while we wait for the new stream location to load)
+			Driver->PauseChannels(true);
+		}
+	}
+
+	virtual bool Tick() override
+	{
+		if (CheckpointSkipExtraTimeInMS == -2)
+		{
+			// Detect failure case
+			return true;
+		}
+
+		if (CheckpointArchive != nullptr)
+		{
+			// Skip loading anything, because there's a pending GotoTime.
+			// Note, we still needed to wait until the checkpoint was loaded so our delegates were
+			// properly cleared from the streamer.
+			if (bSkipWork)
+			{
+				return true;
+			}
+		
+			if (CheckpointSkipExtraTimeInMS > 0 && !Driver->ReplayStreamer->IsDataAvailable())
+			{
+				// Wait for rest of stream before loading checkpoint
+				// We do this so we can load the checkpoint and fastforward the stream all at once
+				// We do this so that the OnReps don't stay queued up outside of this frame
+				return false;
+			}
+
+			return FPendingTaskHelper::FastForwardLevels(Driver, CheckpointArchive, CheckpointSkipExtraTimeInMS);
+		}
+
+		return false;
+	}
+
+	virtual FName GetName() const override
+	{
+		return ReplayTaskNames::FastForwardLevelsTask;
+	}
+
+	void CheckpointReady( const bool bSuccess, const int64 SkipExtraTimeInMS )
+	{
+		check( CheckpointArchive == nullptr );
+		check( CheckpointSkipExtraTimeInMS == -1 );
+
+		if ( !bSuccess )
+		{
+			UE_LOG( LogDemo, Warning, TEXT( "FFastForwardLevelsTask::CheckpointReady: Faled to get checkpoint." ) );
+			CheckpointSkipExtraTimeInMS = -2;
+		}
+		else
+		{
+			CheckpointArchive = Driver->ReplayStreamer->GetCheckpointArchive();
+			CheckpointSkipExtraTimeInMS = SkipExtraTimeInMS;
+		}
+	}
+
+private:
+
+	bool						bSkipWork;
+	FArchive*					CheckpointArchive;				
+	int64						CheckpointSkipExtraTimeInMS;
+	FOnCheckpointReadyDelegate	CheckpointReadyDelegate;
+};
+
 /*-----------------------------------------------------------------------------
 	UDemoNetDriver.
 -----------------------------------------------------------------------------*/
@@ -371,6 +579,11 @@ bool UDemoNetDriver::IsNamedTaskInQueue( const FName& Name ) const
 	return false;
 }
 
+FName UDemoNetDriver::GetNextQueuedTaskName() const
+{
+	return QueuedReplayTasks.Num() > 0 ? QueuedReplayTasks[0]->GetName() : NAME_None;
+}
+
 bool UDemoNetDriver::InitBase( bool bInitAsClient, FNetworkNotify* InNotify, const FURL& URL, bool bReuseAddressAndPort, FString& Error )
 {
 	if ( Super::InitBase( bInitAsClient, InNotify, URL, bReuseAddressAndPort, Error ) )
@@ -417,6 +630,7 @@ void UDemoNetDriver::FinishDestroy()
 		}
 	}
 
+	FCoreUObjectDelegates::PostLoadMapWithWorld.RemoveAll(this);
 	Super::FinishDestroy();
 }
 
@@ -427,14 +641,17 @@ FString UDemoNetDriver::LowLevelGetNetworkNumber()
 
 void UDemoNetDriver::ResetDemoState()
 {
-	DemoFrameNum		= 0;
-	LastCheckpointTime	= 0.0f;
-	DemoTotalTime		= 0;
-	DemoCurrentTime		= 0;
-	DemoTotalFrames		= 0;
+	DemoFrameNum			= 0;
+	LastCheckpointTime		= 0.0f;
+	DemoTotalTime			= 0;
+	DemoCurrentTime			= 0;
+	DemoTotalFrames			= 0;
+	LatestReadFrameTime		= 0.0f;
+	LastProcessedPacketTime	= 0.0f;
 
 	ExternalDataToObjectMap.Empty();
 	PlaybackPackets.Empty();
+	ClearLevelStreamingState();
 }
 
 bool UDemoNetDriver::InitConnect( FNetworkNotify* InNotify, const FURL& ConnectURL, FString& Error )
@@ -548,6 +765,9 @@ bool UDemoNetDriver::ReadPlaybackDemoHeader(FString& Error)
 		GameInstance->HandleDemoPlaybackFailure( EDemoPlayFailure::Corrupt, Error );
 		return false;
 	}
+	
+	// Check whether or not we need to process streaming level fixes.
+	bHasLevelStreamingFixes = !!(PlaybackDemoHeader.HeaderFlags & EReplayHeaderFlags::HasStreamingFixes);
 
 	// Set network version on connection
 	ServerConnection->EngineNetworkProtocolVersion = PlaybackDemoHeader.EngineNetworkProtocolVersion;
@@ -585,6 +805,9 @@ bool UDemoNetDriver::InitConnectInternal( FString& Error )
 		bAsyncLoadWorld = FCString::ToBool(AsyncLoadWorldOverrideOption);
 	}
 
+	// Hook up to get notifications so we know when a travel is complete (LoadMap or Seamless).
+	FCoreUObjectDelegates::PostLoadMapWithWorld.AddUObject(this, &ThisClass::OnPostLoadMapWithWorld);
+
 	if (GetDuplicateLevelID() == -1)
 	{
 		if ( bAsyncLoadWorld && GetWorld()->WorldType != EWorldType::PIE ) // Editor doesn't support async map travel
@@ -604,7 +827,7 @@ bool UDemoNetDriver::InitConnectInternal( FString& Error )
 
 			FWorldContext * WorldContext = GEngine->GetWorldContextFromWorld( GetWorld() );
 
-			if ( WorldContext == NULL )
+			if ( WorldContext == nullptr )
 			{
 				UGameInstance* GameInstance = GetWorld()->GetGameInstance();
 
@@ -614,8 +837,8 @@ bool UDemoNetDriver::InitConnectInternal( FString& Error )
 				return false;
 			}
 
-			GetWorld()->DemoNetDriver = NULL;
-			SetWorld( NULL );
+			GetWorld()->DemoNetDriver = nullptr;
+			SetWorld( nullptr );
 
 			auto NewPendingNetGame = NewObject<UDemoPendingNetGame>();
 
@@ -627,6 +850,10 @@ bool UDemoNetDriver::InitConnectInternal( FString& Error )
 			WorldContext->PendingNetGame = NewPendingNetGame;
 		}
 	}
+	else
+	{
+		ResetLevelStatuses();
+	}
 
 	return true;
 }
@@ -640,7 +867,7 @@ bool UDemoNetDriver::InitListen( FNetworkNotify* InNotify, FURL& ListenURL, bool
 
 	GuidCache->SetNetworkChecksumMode( FNetGUIDCache::ENetworkChecksumMode::SaveButIgnore );
 
-	check( World != NULL );
+	check( World != nullptr );
 
 	class AWorldSettings * WorldSettings = World->GetWorldSettings(); 
 
@@ -650,6 +877,10 @@ bool UDemoNetDriver::InitListen( FNetworkNotify* InNotify, FURL& ListenURL, bool
 		return false;
 	}
 
+	// We'll only check these CVars here, because we don't want to break the replay if they change part way through recording.
+	// During playback the CVars won't be used. Instead, we'll rely on the DemoPacketHeader value.
+	bHasLevelStreamingFixes = !!CVarWithLevelStreamingFixes.GetValueOnAnyThread();
+
 	// Recording, local machine is server, demo stream acts "as if" it's a client.
 	UDemoNetConnection* Connection = NewObject<UDemoNetConnection>();
 	Connection->InitConnection( this, USOCK_Open, ListenURL, 1000000 );
@@ -670,8 +901,7 @@ bool UDemoNetDriver::InitListen( FNetworkNotify* InNotify, FURL& ListenURL, bool
 		for ( int32 i = 0; i < GameState->PlayerArray.Num(); i++ )
 		{
 			APlayerState* PlayerState = GameState->PlayerArray[i];
-
-			if ( !PlayerState->bIsABot && !PlayerState->bIsSpectator )
+			if ( PlayerState && !PlayerState->bIsABot && !PlayerState->bIsSpectator )
 			{
 				UserNames.Add( PlayerState->UniqueId.ToString() );
 			}
@@ -696,6 +926,119 @@ bool UDemoNetDriver::InitListen( FNetworkNotify* InNotify, FURL& ListenURL, bool
 	return Result;
 }
 
+void UDemoNetDriver::OnLevelAddedToWorld(ULevel* InLevel, UWorld* InWorld)
+{
+	Super::OnLevelAddedToWorld(InLevel, InWorld);
+
+	if (InLevel && !InLevel->bClientOnlyVisible && GetWorld() == InWorld && HasLevelStreamingFixes() && IsPlaying())
+	{
+		if (!NewStreamingLevelsThisFrame.Contains(InLevel) && !LevelsPendingFastForward.Contains(InLevel))
+		{
+			FLevelStatus& LevelStatus = FindOrAddLevelStatus(GetLevelPackageName(*InLevel));
+
+			// If we haven't processed any packets for this level yet, immediately mark it as ready.
+			if (!LevelStatus.bHasBeenSeen)
+			{
+				LevelStatus.bIsReady = true;
+			}
+
+			// If the level isn't ready, go ahead and queue it up to get fast-forwarded.
+			// Note, we explicitly check the visible flag in case same the level gets notified multiple times.
+			else if (!LevelStatus.bIsReady)
+			{
+				NewStreamingLevelsThisFrame.Add(InLevel);
+			}
+		}
+	}
+}
+
+void UDemoNetDriver::OnLevelRemovedFromWorld(ULevel* InLevel, UWorld* InWorld)
+{
+	Super::OnLevelRemovedFromWorld(InLevel, InWorld);
+
+	if (InLevel && !InLevel->bClientOnlyVisible && GetWorld() == InWorld && HasLevelStreamingFixes() && IsPlaying())
+	{
+		const FString LevelPackageName = GetLevelPackageName(*InLevel);
+		if (LevelStatusesByName.Contains(LevelPackageName))
+		{
+			FLevelStatus& LevelStatus = GetLevelStatus(LevelPackageName);
+			LevelStatus.bIsReady = false;
+
+			// Make sure we don't try to fast-forward this level later.
+			LevelsPendingFastForward.Remove(InLevel);
+			NewStreamingLevelsThisFrame.Remove(InLevel);
+
+			// We can't just iterate over the levels actors, because the ones in the queue will already have been destroyed.
+			for (TMap<FString, FRollbackNetStartupActorInfo>::TIterator It = RollbackNetStartupActors.CreateIterator(); It; ++It)
+			{
+				if (It.Value().Level == InLevel)
+				{
+					It.RemoveCurrent();
+				}
+			}
+		}
+	}
+}
+
+void UDemoNetDriver::OnPostLoadMapWithWorld(UWorld* InWorld)
+{
+	if (InWorld != nullptr && InWorld == World && HasLevelStreamingFixes())
+	{
+		if (IsPlaying())
+		{
+			ResetLevelStatuses();
+		}
+		else
+		{
+			ClearLevelStreamingState();
+		}
+	}
+}
+
+TUniquePtr<FScopedPacketManager> UDemoNetDriver::ConditionallyCreatePacketManager(ULevel& Level)
+{
+	if(IsRecording() && HasLevelStreamingFixes())
+	{
+		// Indices need to be 1 based, so +1.
+		return MakeUnique<FScopedPacketManager>(*(UDemoNetConnection*)ClientConnections[0], FindOrAddLevelStatus(GetLevelPackageName(Level)).LevelIndex + 1);
+	}
+
+	return nullptr;
+}
+
+TUniquePtr<FScopedPacketManager> UDemoNetDriver::ConditionallyCreatePacketManager(int32 LevelIndex)
+{
+	if(IsRecording() && HasLevelStreamingFixes())
+	{
+		// Indices need to be 1 based, so +1.
+		return MakeUnique<FScopedPacketManager>(*(UDemoNetConnection*)ClientConnections[0], LevelIndex);
+	}
+
+	return nullptr;
+}
+
+void UDemoNetDriver::ResetLevelStatuses()
+{
+	ClearLevelStreamingState();
+
+	// There are times (e.g., during travel) when we may not have a valid level.
+	// This **should never** be called during those times.
+	check(World);
+
+	// ResetLevelStatuses should only ever be called before receiving *any* data from the Replay stream,
+	// immediately before processing checkpoint data, or after a level transistion (in which case no data
+	// will be relevant to the new sublevels).
+	// In any case, we can just flag these sublevels as ready immediately.
+	FindOrAddLevelStatus(GetLevelPackageName(*World->PersistentLevel)).bIsReady = true;
+	for (ULevelStreaming* LevelStreaming : World->GetStreamingLevels())
+	{
+		if (LevelStreaming && LevelStreaming->IsLevelVisible())
+		{
+			FindOrAddLevelStatus(GetLevelPackageName(*LevelStreaming->GetLoadedLevel())).bIsReady = true;
+		}
+	}
+}
+
 bool UDemoNetDriver::ContinueListen(FURL& ListenURL)
 {
 	if (IsRecording() && ensure(IsRecordingPaused()))
@@ -742,6 +1085,10 @@ bool UDemoNetDriver::WriteNetworkDemoHeader(FString& Error)
 	{
 		DemoHeader.HeaderFlags |= EReplayHeaderFlags::ClientRecorded;
 	}
+	if (HasLevelStreamingFixes())
+	{
+		DemoHeader.HeaderFlags |= EReplayHeaderFlags::HasStreamingFixes;
+	}
 
 	// Write the header
 	(*FileAr) << DemoHeader;
@@ -757,7 +1104,8 @@ bool UDemoNetDriver::IsRecording() const
 
 bool UDemoNetDriver::IsPlaying() const
 {
-	return ServerConnection != nullptr && ServerConnection->State != USOCK_Closed;
+	// ServerConnection may be deleted / recreated during checkpoint loading.
+	return IsLoadingCheckpoint() || (ServerConnection != nullptr && ServerConnection->State != USOCK_Closed);
 }
 
 bool UDemoNetDriver::ShouldTickFlushAsyncEndOfFrame() const
@@ -841,6 +1189,9 @@ void UDemoNetDriver::TickFlushInternal( float DeltaSeconds )
 
 	const double RecordTotalTime = ( EndTime - StartTime );
 
+	// While recording, the CurrentCL is the same as the recording CL.
+	ConditionallyDisplayBurnInTime(FEngineVersion::Current().GetChangelist(), DemoCurrentTime);
+
 	MaxRecordTime = FMath::Max( MaxRecordTime, RecordTotalTime );
 
 	AccumulatedRecordTime += RecordTotalTime;
@@ -923,15 +1274,18 @@ void UDemoNetDriver::TickDispatch( float DeltaSeconds )
 		return;
 	}
 
-	// Wait until all levels are streamed in
-	for (ULevelStreaming* StreamingLevel : World->GetStreamingLevels())
+	if (!HasLevelStreamingFixes())
 	{
-		if ( StreamingLevel && StreamingLevel->ShouldBeLoaded() && (!StreamingLevel->IsLevelLoaded() || !StreamingLevel->GetLoadedLevel()->GetOutermost()->IsFullyLoaded() || !StreamingLevel->IsLevelVisible() ) )
+		// Wait until all levels are streamed in
+		for (ULevelStreaming* StreamingLevel : World->GetStreamingLevels())
 		{
-			// Abort, we have more streaming levels to load
-			return;
+			if ( StreamingLevel && StreamingLevel->ShouldBeLoaded() && (!StreamingLevel->IsLevelLoaded() || !StreamingLevel->GetLoadedLevel()->GetOutermost()->IsFullyLoaded() || !StreamingLevel->IsLevelVisible() ) )
+			{
+				// Abort, we have more streaming levels to load
+				return;
+			}
 		}
-	}
+	}	
 
 	if ( CVarDemoTimeDilation.GetValueOnGameThread() >= 0.0f )
 	{
@@ -943,7 +1297,7 @@ void UDemoNetDriver::TickDispatch( float DeltaSeconds )
 
 	// Update time dilation on spectator pawn to compensate for any demo dilation 
 	//	(we want to continue to fly around in real-time)
-	if ( SpectatorController != NULL )
+	if ( SpectatorController != nullptr )
 	{
 		if ( World->GetWorldSettings()->DemoPlayTimeDilation > KINDA_SMALL_NUMBER )
 		{
@@ -954,7 +1308,7 @@ void UDemoNetDriver::TickDispatch( float DeltaSeconds )
 			SpectatorController->CustomTimeDilation = 1.0f;
 		}
 
-		if ( SpectatorController->GetSpectatorPawn() != NULL )
+		if ( SpectatorController->GetSpectatorPawn() != nullptr )
 		{
 			SpectatorController->GetSpectatorPawn()->CustomTimeDilation = SpectatorController->CustomTimeDilation;
 					
@@ -971,6 +1325,9 @@ void UDemoNetDriver::TickDispatch( float DeltaSeconds )
 	}
 
 	TickDemoPlayback( DeltaSeconds );
+
+	// Used LastProcessedPacketTime because it will correlate better with recorded frame time.
+	ConditionallyDisplayBurnInTime(PlaybackDemoHeader.Changelist, LastProcessedPacketTime);
 }
 
 void UDemoNetDriver::ProcessRemoteFunction( class AActor* Actor, class UFunction* Function, void* Parameters, struct FOutParmRec* OutParms, struct FFrame* Stack, class UObject* SubObject )
@@ -985,6 +1342,8 @@ void UDemoNetDriver::ProcessRemoteFunction( class AActor* Actor, class UFunction
 	{
 		if ( IsRecording() )
 		{
+			TUniquePtr<FScopedPacketManager> PacketManager(ConditionallyCreatePacketManager(*Actor->GetLevel()));
+
 			if ((Function->FunctionFlags & FUNC_NetMulticast))
 			{
 				// Handle role swapping if this is a client-recorded replay.
@@ -1029,7 +1388,7 @@ void UDemoNetDriver::StopDemo()
 	if ( !ServerConnection )
 	{
 		// let GC cleanup the object
-		if ( ClientConnections.Num() > 0 && ClientConnections[0] != NULL )
+		if ( ClientConnections.Num() > 0 && ClientConnections[0] != nullptr )
 		{
 			ClientConnections[0]->Close();
 		}
@@ -1067,29 +1426,29 @@ bool UDemoNetDriver::DemoReplicateActor(AActor* Actor, UNetConnection* Connectio
 	// Handle role swapping if this is a client-recorded replay.
 	FScopedActorRoleSwap RoleSwap(Actor);
 
-	if ((Actor->GetRemoteRole() != ROLE_None || Actor->bTearOff) && (Actor == Connection->PlayerController || Cast< APlayerController >(Actor) == NULL))
+	if ((Actor->GetRemoteRole() != ROLE_None || Actor->GetTearOff()) && (Actor == Connection->PlayerController || Cast< APlayerController >(Actor) == nullptr))
 	{
 		const bool bShouldHaveChannel =
 			Actor->bRelevantForNetworkReplays &&
-			!Actor->bTearOff &&
+			!Actor->GetTearOff() &&
 			(!Actor->IsNetStartupActor() || Connection->ClientHasInitializedLevelFor(Actor));
 
-		UActorChannel* Channel = Connection->ActorChannels.FindRef(Actor);
+		UActorChannel* Channel = Connection->FindActorChannelRef(Actor);
 
-		if (bShouldHaveChannel && Channel == NULL)
+		if (bShouldHaveChannel && Channel == nullptr)
 		{
 			// Create a new channel for this actor.
 			Channel = (UActorChannel*)Connection->CreateChannel(CHTYPE_Actor, 1);
-			if (Channel != NULL)
+			if (Channel != nullptr)
 			{
 				Channel->SetChannelActor(Actor);
 			}
 		}
 
-		if (Channel != NULL && !Channel->Closing)
+		if (Channel != nullptr && !Channel->Closing)
 		{
 			// Send it out!
-			bDidReplicateActor = Channel->ReplicateActor();
+			bDidReplicateActor = (Channel->ReplicateActor() > 0);
 
 			// Close the channel if this actor shouldn't have one
 			if (!bShouldHaveChannel)
@@ -1118,7 +1477,7 @@ static void SerializeGuidCache( TSharedPtr< class FNetGUIDCache > GuidCache, FAr
 	{
 		FNetGuidCacheObject& CacheObject = It.Value();
 
-		if ( CacheObject.Object == NULL || !CacheObject.Object->IsNameStableForNetworking() )
+		if ( CacheObject.Object == nullptr || !CacheObject.Object->IsNameStableForNetworking() )
 		{
 			continue;
 		}
@@ -1189,19 +1548,57 @@ void UDemoNetDriver::SaveCheckpoint()
 	check( ClientConnections[0]->SendBuffer.GetNumBits() == 0 );
 
 	check( PendingCheckpointActors.Num() == 0 );
-
-	// Add any actor with a valid channel to the PendingCheckpointActors list
-	for ( const TSharedPtr<FNetworkObjectInfo>& ObjectInfo : GetNetworkObjectList().GetAllObjects() )
+	
+	if (HasLevelStreamingFixes())
 	{
-		AActor* Actor = ObjectInfo.Get()->Actor;
+		// TODO: Consider storing the TMap to prevent re-allocation each frame (like PendingActors list).
+		// That has potential downsides, such as clearing all the lists each frame, managing levels as they're removed,
+		// and needing to skip empty lists of actors.
+
+		TMap<ULevel*, TArray<AActor*>> LevelToActor;
+		{
+			DECLARE_SCOPE_CYCLE_COUNTER(TEXT("Replay actor level sorting time."), STAT_ReplayLevelSorting, STATGROUP_Net);
+
+			if (bPrioritizeActors)
+			{
+				UE_LOG(LogDemo, Verbose, TEXT("bPrioritizeActors and HasLevelStreamingFixes are both enabled. This will undo some prioritization work."));
+			}
+
+			for (const TSharedPtr<FNetworkObjectInfo>& ObjectInfo : GetNetworkObjectList().GetAllObjects())
+			{
+				AActor* Actor = ObjectInfo.Get()->Actor;
+				if (ClientConnections[0]->FindActorChannelRef(Actor))
+				{
+					LevelToActor.FindOrAdd(Actor->GetLevel()).Add(Actor);
+				}
+			}
+		}
+
+		NumActorsToProcessForLevel.Reset(LevelToActor.Num());
+		for (auto It = LevelToActor.CreateIterator(); It; ++It)
+		{
+			TArray<AActor*>& ActorArray = It.Value();
+			const FLevelStatus& LevelStatus = FindOrAddLevelStatus(GetLevelPackageName(*It.Key()));
 
-		if ( ClientConnections[0]->ActorChannels.FindRef( Actor ) )
+			NumActorsToProcessForLevel.Emplace(LevelStatus.LevelIndex + 1, ActorArray.Num());
+			PendingCheckpointActors.Append(MoveTemp(ActorArray));
+		}
+	}
+	else
+	{
+		// Add any actor with a valid channel to the PendingCheckpointActors list
+		for (const TSharedPtr<FNetworkObjectInfo>& ObjectInfo : GetNetworkObjectList().GetAllObjects())
 		{
-			PendingCheckpointActors.Add( Actor );
+			AActor* Actor = ObjectInfo.Get()->Actor;
+
+			if ( ClientConnections[0]->FindActorChannelRef( Actor ) )
+			{
+				PendingCheckpointActors.Add( Actor );
+			}
 		}
 	}
 
-	if ( !PendingCheckpointActors.Num() )
+	if ( PendingCheckpointActors.Num() == 0 )
 	{
 		return;
 	}
@@ -1219,6 +1616,14 @@ void UDemoNetDriver::SaveCheckpoint()
 	TickCheckpoint();
 }
 
+class FRepActorsCheckpointParams
+{
+public:
+
+	const double StartCheckpointTime;
+	const double CheckpointMaxUploadTimePerFrame;
+};
+
 void UDemoNetDriver::TickCheckpoint()
 {
 	if ( !PendingCheckpointActors.Num() )
@@ -1235,14 +1640,19 @@ void UDemoNetDriver::TickCheckpoint()
 		return;
 	}
 
-	const double StartCheckpointTime = FPlatformTime::Seconds();
+	FRepActorsCheckpointParams Params
+	{
+		FPlatformTime::Seconds(),
+		(double)GetCheckpointSaveMaxMSPerFrame() / 1000
+	};
+
+	UDemoNetConnection* ClientConnection = CastChecked<UDemoNetConnection>(ClientConnections[0]);
 
 	TotalCheckpointSaveFrames++;
 
-	ClientConnections[0]->FlushNet();
-	check( ClientConnections[0]->SendBuffer.GetNumBits() == 0 );
+	FlushNetChecked(*ClientConnection);
 
-	UPackageMapClient* PackageMapClient = ( ( UPackageMapClient* )ClientConnections[0]->PackageMap );
+	UPackageMapClient* PackageMapClient = ( ( UPackageMapClient* )ClientConnection->PackageMap );
 
 	// Save package map ack status in case we export stuff during the checkpoint (so we can restore the connection back to what it was before we saved the checkpoint)
 	PackageMapClient->OverridePackageMapExportAckStatus( &CheckpointAckState );
@@ -1261,74 +1671,81 @@ void UDemoNetDriver::TickCheckpoint()
 		GameState->UpdateServerTimeSeconds();
 	}
 
-	// Re-use the existing connection to record all properties that have changed since channels were first opened
-	// Set bResendAllDataSinceOpen to true to signify that we want to do this
-	ClientConnections[0]->bResendAllDataSinceOpen = true;
-
-	const double CheckpointMaxUploadTimePerFrame = ( double )GetCheckpointSaveMaxMSPerFrame() / 1000;
-
-	for ( auto ActorIt = PendingCheckpointActors.CreateIterator(); ActorIt; ++ActorIt )
 	{
-		AActor* Actor = ( *ActorIt ).Get();
+		// Re-use the existing connection to record all properties that have changed since channels were first opened
+		// Set bResendAllDataSinceOpen to true to signify that we want to do this
+		TGuardValue<bool> ResendAllData(ClientConnection->bResendAllDataSinceOpen, true);
 
-		ActorIt.RemoveCurrent();	// We're done with this now
-
-		UActorChannel* ActorChannel = ClientConnections[0]->ActorChannels.FindRef( Actor );
+		// Can't use conditionally create here, because NumActorsToProcess will be empty when HasLevelStreamingFixes is false.
+		TUniquePtr<FScopedPacketManager> PacketManager;
+		if (HasLevelStreamingFixes())
+		{
+			const int32 LevelIndex = NumActorsToProcessForLevel[0].Key;
+			PacketManager.Reset(new FScopedPacketManager(*ClientConnection, LevelIndex));
+		}
 
-		if ( ActorChannel )
+		bool bContinue = true;
+		for (auto ActorIt = PendingCheckpointActors.CreateIterator(); ActorIt && bContinue; ++ActorIt)
 		{
-			Actor->CallPreReplication( this );
-			DemoReplicateActor( Actor, ClientConnections[0], true );
+			AActor* Actor = (*ActorIt).Get();
+			ActorIt.RemoveCurrent();
 
-			const double CheckpointTime = FPlatformTime::Seconds();
+			bContinue = ReplicateCheckpointActor(Actor, ClientConnection, Params);
 
-			if ( CheckpointMaxUploadTimePerFrame > 0 && CheckpointTime - StartCheckpointTime > CheckpointMaxUploadTimePerFrame )
+			if (PacketManager && --NumActorsToProcessForLevel[0].Value == 0)
 			{
-				break;
+				NumActorsToProcessForLevel.RemoveAt(0);
+				if (bContinue && NumActorsToProcessForLevel.Num() > 0)
+				{
+					const int32 LevelIndex = NumActorsToProcessForLevel[0].Key;
+					PacketManager.Reset(new FScopedPacketManager(*ClientConnection, LevelIndex));
+				}
 			}
 		}
-	}
-
-	if ( GameState )
-	{
-		// Restore the game state's replicated world time
-		GameState->ReplicatedWorldTimeSeconds = SavedReplicatedServerTimeSeconds;
-	}
 
-	// Make sure to flush the connection (fill up QueuedCheckpointPackets)
-	// This also frees up the connection to be used for normal streaming again
-	ClientConnections[0]->FlushNet();
-	check( ClientConnections[0]->SendBuffer.GetNumBits() == 0 );
+		if ( GameState )
+		{
+			// Restore the game state's replicated world time
+			GameState->ReplicatedWorldTimeSeconds = SavedReplicatedServerTimeSeconds;
+		}
 
-	PackageMapClient->OverridePackageMapExportAckStatus( nullptr );
+		FlushNetChecked(*ClientConnection);
 
-	ClientConnections[0]->bResendAllDataSinceOpen = false;
+		PackageMapClient->OverridePackageMapExportAckStatus( nullptr );
+	}
 
 	const double EndCheckpointTime = FPlatformTime::Seconds();
 
-	TotalCheckpointSaveTimeSeconds += ( EndCheckpointTime - StartCheckpointTime );
+	TotalCheckpointSaveTimeSeconds += ( EndCheckpointTime - Params.StartCheckpointTime );
 
 	if ( PendingCheckpointActors.Num() == 0 )
 	{
-		//
-		// We're done saving this checkpoint
-		//
-		*CheckpointArchive << CurrentLevelIndex;
+		{
+			TUniquePtr<FScopedStoreArchiveOffset> ArchiveOffset(HasLevelStreamingFixes() ? new FScopedStoreArchiveOffset(*CheckpointArchive) : nullptr);
 
-		// Save deleted startup actors
-		*CheckpointArchive << DeletedNetStartupActors;
+			//
+			// We're done saving this checkpoint
+			//
+			*CheckpointArchive << CurrentLevelIndex;
 
-		// Save the current guid cache
-		SerializeGuidCache( GuidCache, CheckpointArchive );
+			// Save deleted startup actors
+			*CheckpointArchive << DeletedNetStartupActors;
 
-		// Save the compatible rep layout map
-		PackageMapClient->SerializeNetFieldExportGroupMap( *CheckpointArchive );
+			// Save the current guid cache
+			SerializeGuidCache( GuidCache, CheckpointArchive );
+
+			// Save the compatible rep layout map
+			PackageMapClient->SerializeNetFieldExportGroupMap( *CheckpointArchive );
+		}
 
 		// Get the size of the guid data saved
 		const uint32 GuidCacheSize = CheckpointArchive->TotalSize();
 
+		// This will cause the entire name list to be written out again.
+		NumLevelsAddedThisFrame = AllLevelStatuses.Num();
+
 		// Write out all of the queued up packets generated while saving the checkpoint
-		WriteDemoFrameFromQueuedDemoPackets( *CheckpointArchive, CastChecked< UDemoNetConnection>( ClientConnections[0] )->QueuedCheckpointPackets );
+		WriteDemoFrameFromQueuedDemoPackets( *CheckpointArchive, CastChecked< UDemoNetConnection>( ClientConnection )->QueuedCheckpointPackets );
 
 		// Get the total checkpoint size
 		const int32 TotalCheckpointSize = CheckpointArchive->TotalSize();
@@ -1344,18 +1761,36 @@ void UDemoNetDriver::TickCheckpoint()
 	}
 }
 
-void UDemoNetDriver::SaveExternalData( FArchive& Ar )
+bool UDemoNetDriver::ReplicateCheckpointActor(AActor* ToReplicate, UDemoNetConnection* ClientConnection, FRepActorsCheckpointParams& Params)
 {
-	for ( auto It = RepChangedPropertyTrackerMap.CreateIterator(); It; ++It )
+	if (UActorChannel* ActorChannel = ClientConnection->FindActorChannelRef(ToReplicate))
 	{
-		if ( It.Key().IsValid() )
-		{
-			if ( It.Value()->ExternalDataNumBits > 0 )
-			{
-				// Save payload size (in bits)
-				Ar.SerializeIntPacked( It.Value()->ExternalDataNumBits );
+		ToReplicate->CallPreReplication(this);
+		DemoReplicateActor(ToReplicate, ClientConnection, true);
 
-				FNetworkGUID NetworkGUID = GuidCache->NetGUIDLookup.FindChecked( It.Key() );
+		const double CheckpointTime = FPlatformTime::Seconds();
+
+		if (Params.CheckpointMaxUploadTimePerFrame > 0 && CheckpointTime - Params.StartCheckpointTime > Params.CheckpointMaxUploadTimePerFrame)
+		{
+			return false;
+		}
+	}
+
+	return true;
+}
+
+void UDemoNetDriver::SaveExternalData( FArchive& Ar )
+{
+	for ( auto It = RepChangedPropertyTrackerMap.CreateIterator(); It; ++It )
+	{
+		if ( It.Key().IsValid() )
+		{
+			if ( It.Value()->ExternalDataNumBits > 0 )
+			{
+				// Save payload size (in bits)
+				Ar.SerializeIntPacked( It.Value()->ExternalDataNumBits );
+
+				FNetworkGUID NetworkGUID = GuidCache->NetGUIDLookup.FindChecked( It.Key() );
 
 				// Save GUID
 				Ar << NetworkGUID;
@@ -1454,6 +1889,17 @@ public:
 	FVector		Location;
 };
 
+class FRepActorsParams
+{
+public:
+
+	bool bUseAdapativeNetFrequency;
+	float MinRecordHz;
+	float MaxRecordHz;
+	float ServerTickTime;
+	double ReplicationStartTimeSeconds;
+};
+
 void UDemoNetDriver::TickDemoRecord( float DeltaSeconds )
 {
 	if ( !IsRecording() || bPauseRecording )
@@ -1470,25 +1916,28 @@ void UDemoNetDriver::TickDemoRecord( float DeltaSeconds )
 
 	FArchive* FileAr = ReplayStreamer->GetStreamingArchive();
 
-	if ( FileAr == NULL )
+	if ( FileAr == nullptr )
 	{
 		return;
 	}
 
 	// Mark any new streaming levels, so that they are saved out this frame
-	for (ULevelStreaming* StreamingLevel : World->GetStreamingLevels())
+	if (!HasLevelStreamingFixes())
 	{
-		if ( StreamingLevel == nullptr || !StreamingLevel->ShouldBeLoaded() || StreamingLevel->ShouldBeAlwaysLoaded() )
+		for (ULevelStreaming* StreamingLevel : World->GetStreamingLevels())
 		{
-			continue;
-		}
+			if ( StreamingLevel == nullptr || !StreamingLevel->ShouldBeLoaded() || StreamingLevel->ShouldBeAlwaysLoaded() )
+			{
+				continue;
+			}
 
-		TWeakObjectPtr<UObject> WeakStreamingLevel;
-		WeakStreamingLevel = StreamingLevel;
-		if ( !UniqueStreamingLevels.Contains( WeakStreamingLevel ) )
-		{
-			UniqueStreamingLevels.Add( WeakStreamingLevel );
-			NewStreamingLevelsThisFrame.Add( StreamingLevel );
+			TWeakObjectPtr<UObject> WeakStreamingLevel;
+			WeakStreamingLevel = StreamingLevel;
+			if ( !UniqueStreamingLevels.Contains( WeakStreamingLevel ) )
+			{
+				UniqueStreamingLevels.Add( WeakStreamingLevel );
+				NewStreamingLevelsThisFrame.Add( WeakStreamingLevel );
+			}
 		}
 	}
 
@@ -1497,23 +1946,14 @@ void UDemoNetDriver::TickDemoRecord( float DeltaSeconds )
 
 	ReplayStreamer->UpdateTotalDemoTime( GetDemoCurrentTimeInMS() );
 
-	float MAX_RECORD_HZ	= CVarDemoRecordHz.GetValueOnAnyThread();
-	float MIN_RECORD_HZ = CVarDemoMinRecordHz.GetValueOnAnyThread();
-	if (MIN_RECORD_HZ > MAX_RECORD_HZ)
-	{
-		// make sure min and max are sane
-		Swap(MIN_RECORD_HZ, MAX_RECORD_HZ);
-	}
-
 	// Save out a frame
 	DemoFrameNum++;
 	ReplicationFrame++;
 
-	// flush out any pending network traffic
-	ClientConnections[0]->FlushNet();
+	UDemoNetConnection* ClientConnection = CastChecked<UDemoNetConnection>(ClientConnections[0]);
 
-	// Make sure we don't have anything in the buffer for this new frame
-	check( ClientConnections[0]->SendBuffer.GetNumBits() == 0 );
+	// flush out any pending network traffic
+	FlushNetChecked(*ClientConnection);
 
 	float ServerTickTime = GEngine->GetMaxTickRate( DeltaSeconds );
 	if ( ServerTickTime == 0.0 )
@@ -1525,8 +1965,6 @@ void UDemoNetDriver::TickDemoRecord( float DeltaSeconds )
 		ServerTickTime	= 1.0 / ServerTickTime;
 	}
 
-	const bool bUseAdapativeNetFrequency = CVarUseAdaptiveReplayUpdateFrequency.GetValueOnAnyThread() > 0;
-
 	// Build priority list
 	PrioritizedActors.Reset( GetNetworkObjectList().GetActiveObjects().Num() );
 
@@ -1534,7 +1972,7 @@ void UDemoNetDriver::TickDemoRecord( float DeltaSeconds )
 	FVector ViewLocation = FVector::ZeroVector;
 	FVector ViewDirection = FVector::ZeroVector;
 	APlayerController* CachedViewerOverride = ViewerOverride.Get();
-	APlayerController* Viewer = CachedViewerOverride ? CachedViewerOverride : ClientConnections[0]->GetPlayerController(World);
+	APlayerController* Viewer = CachedViewerOverride ? CachedViewerOverride : ClientConnection->GetPlayerController(World);
 	AActor* ViewTarget = Viewer ? Viewer->GetViewTarget() : nullptr;
 	
 	if (ViewTarget)
@@ -1588,7 +2026,7 @@ void UDemoNetDriver::TickDemoRecord( float DeltaSeconds )
 
 				// During client recording, a torn-off actor will already have its remote role set to None, but
 				// we still need to replicate it one more time so that the recorded replay knows it's been torn-off as well.
-				if ( Actor->GetRemoteRole() == ROLE_None && !Actor->bTearOff)
+				if ( Actor->GetRemoteRole() == ROLE_None && !Actor->GetTearOff())
 				{
 					ActorsToRemove.Add( Actor );
 					continue;
@@ -1603,7 +2041,7 @@ void UDemoNetDriver::TickDemoRecord( float DeltaSeconds )
 				// We check ActorInfo->LastNetUpdateTime < KINDA_SMALL_NUMBER to force at least one update for each actor
 				const bool bWasRecentlyRelevant = (ActorInfo->LastNetUpdateTime < KINDA_SMALL_NUMBER) || ((Time - ActorInfo->LastNetUpdateTime) < RelevantTimeout);
 
-				bool bIsRelevant = !bUseNetRelevancy || Actor->bAlwaysRelevant || Actor == ClientConnections[0]->PlayerController || ActorInfo->bForceRelevantNextUpdate;
+				bool bIsRelevant = !bUseNetRelevancy || Actor->bAlwaysRelevant || Actor == ClientConnection->PlayerController || ActorInfo->bForceRelevantNextUpdate;
 
 				ActorInfo->bForceRelevantNextUpdate = false;
 
@@ -1627,7 +2065,7 @@ void UDemoNetDriver::TickDemoRecord( float DeltaSeconds )
 					continue;
 				}
 
-				UActorChannel* Channel = ClientConnections[0]->ActorChannels.FindRef( Actor );
+				UActorChannel* Channel = ClientConnection->FindActorChannelRef( Actor );
 				FActorPriority ActorPriority;
 				ActorPriority.ActorInfo = ActorInfo;
 				ActorPriority.Channel = Channel;
@@ -1648,17 +2086,17 @@ void UDemoNetDriver::TickDemoRecord( float DeltaSeconds )
 		}
 
 		// Also add destroyed actors that the client may not have a channel for
-		for ( auto It = ClientConnections[0]->DestroyedStartupOrDormantActors.CreateIterator(); It; ++It )
+		for ( auto It = ClientConnection->GetDestroyedStartupOrDormantActorGUIDs().CreateIterator(); It; ++It )
 		{
-			FActorDestructionInfo& DInfo = DestroyedStartupOrDormantActors.FindChecked( *It );
+			TUniquePtr<FActorDestructionInfo>& DInfo = DestroyedStartupOrDormantActors.FindChecked( *It );
 			FActorPriority Priority;
-			Priority.DestructionInfo = &DInfo;
+			Priority.DestructionInfo = DInfo.Get();
 			PrioritizedActors.Add(Priority);
 		}
 
 		for ( AActor* Actor : ActorsToRemove )
 		{
-			GetNetworkObjectList().Remove( Actor );
+			RemoveNetworkActor( Actor );
 		}
 
 		if ( bPrioritizeActors )
@@ -1667,154 +2105,249 @@ void UDemoNetDriver::TickDemoRecord( float DeltaSeconds )
 		}
 	}
 
-	const double ReplicationStartTimeSeconds = FPlatformTime::Seconds();
-	TArray<AActor*> ActorsToGoDormant;
+	FRepActorsParams Params
+	{
+		CVarUseAdaptiveReplayUpdateFrequency.GetValueOnAnyThread() > 0,
+		CVarDemoMinRecordHz.GetValueOnAnyThread(),
+		CVarDemoRecordHz.GetValueOnAnyThread(),
+		ServerTickTime,
+		FPlatformTime::Seconds()
+	};
 
-	for ( const FActorPriority& ActorPriority : PrioritizedActors )
+	if (Params.MaxRecordHz < Params.MinRecordHz)
 	{
-		FNetworkObjectInfo* ActorInfo = ActorPriority.ActorInfo;
-		FActorDestructionInfo* DestructionInfo = ActorPriority.DestructionInfo;
+		Swap(Params.MaxRecordHz, Params.MinRecordHz);
+	}
 
-		const double RecordStartTimeSeconds = FPlatformTime::Seconds();
+	if (HasLevelStreamingFixes())
+	{
+		// TODO: Consider storing the TMap to prevent re-allocation each frame (like PendingActors list).
+		// That has potential downsides, such as clearing all the lists each frame, managing levels as they're removed,
+		// and needing to skip empty lists of actors.
 
-		// Deletion entry
-		if (ActorInfo == nullptr && DestructionInfo != nullptr)
+		TMap<ULevel*, TArray<FActorPriority*>> LevelToActorPriority;
 		{
-			UActorChannel* Channel = (UActorChannel*)ClientConnections[0]->CreateChannel( CHTYPE_Actor, 1 );
-			if ( Channel )
+			DECLARE_SCOPE_CYCLE_COUNTER(TEXT("Replay actor level sorting time."), STAT_ReplayLevelSorting, STATGROUP_Net);
+
+			if (bPrioritizeActors)
 			{
-				UE_LOG( LogDemo, Verbose, TEXT( "TickDemoRecord creating destroy channel for NetGUID <%s,%s> Priority: %d" ), *DestructionInfo->NetGUID.ToString(), *DestructionInfo->PathName, ActorPriority.Priority );
+				UE_LOG(LogDemo, Verbose, TEXT("bPrioritizeActors and HasLevelStreamingFixes are both enabled. This will undo some prioritization work."));
+			}
 
-				// Send a close bunch on the new channel
-				Channel->SetChannelActorForDestroy( DestructionInfo );
+			
+			for (FActorPriority& ActorPriority : PrioritizedActors)
+			{
+				ULevel* Level = (ActorPriority.DestructionInfo != nullptr) ? ActorPriority.DestructionInfo->Level.Get() :
+																				ActorPriority.ActorInfo->Actor->GetLevel();
 
-				// Remove from connection's to-be-destroyed list (close bunch is reliable, so it will make it there)
-				ClientConnections[0]->DestroyedStartupOrDormantActors.Remove( DestructionInfo->NetGUID );
+				TArray<FActorPriority*>& ActorPriorityList = LevelToActorPriority.FindOrAdd(Level);
+				ActorPriorityList.Add(&ActorPriority);
 			}
 		}
-		else if (ActorInfo != nullptr && DestructionInfo == nullptr)
+
+		// NOTE: This has the potential to starve out levels in high traffic conditions.
+		// Iteration order of TMaps isn't guaranteed. Currently, it is insertion order based, so this should work OK
+		// (because levels will be seen in order of first actor seen, which will be priority order).
+		// However, if the map only has a few (e.g. one) high priority actors then it could prevent other maps
+		// with a higher number of high priority actors from being replicated (even across multiple frames).
+		// Further, the iteration order is just an implementation detail, so could change.
+		// If it's an actual problem, consider ranking the levels based on something like average actor priority.
+		bool bContinue = true;
+		for (TMap<ULevel*, TArray<FActorPriority*>>::TIterator It = LevelToActorPriority.CreateIterator(); It && bContinue; ++It)
 		{
-			AActor* Actor = ActorInfo->Actor;
+			TUniquePtr<FScopedPacketManager> PacketManager(ConditionallyCreatePacketManager(*It.Key()));
+			bContinue = ReplicatePrioritizedActors(It.Value(), Params);
+		}
+	}
+	else
+	{
+		ReplicatePrioritizedActors(PrioritizedActors, Params);
+	}
 
-			// Use NetUpdateFrequency for this actor, but clamp it to RECORD_HZ.
-			const float ClampedNetUpdateFrequency = FMath::Clamp(Actor->NetUpdateFrequency, MIN_RECORD_HZ, MAX_RECORD_HZ);
-			const double NetUpdateDelay = 1.0 / ClampedNetUpdateFrequency;
+	FlushNetChecked(*ClientConnection);
 
-			// Set defaults if this actor is replicating for first time
-			if ( ActorInfo->LastNetReplicateTime == 0 )
-			{
-				ActorInfo->LastNetReplicateTime		= DemoCurrentTime;
-				ActorInfo->OptimalNetUpdateDelta	= NetUpdateDelay;
-			}
+	WriteDemoFrameFromQueuedDemoPackets(*FileAr, ClientConnection->QueuedDemoPackets);
+
+	// Save a checkpoint if it's time
+	if ( CVarEnableCheckpoints.GetValueOnAnyThread() == 1 )
+	{
+		check( !PendingCheckpointActors.Num() );		// We early out above, so this shouldn't be possible
+
+		if (ShouldSaveCheckpoint())
+		{
+			SaveCheckpoint();
+			LastCheckpointTime = DemoCurrentTime;
+		}
+	}
+}
 
-			const float LastReplicateDelta = static_cast<float>( DemoCurrentTime - ActorInfo->LastNetReplicateTime );
+bool UDemoNetDriver::ReplicatePrioritizedActor(TArray<AActor*>& ActorsToGoDormant, const FActorPriority& ActorPriority, const class FRepActorsParams& Params)
+{
+	FNetworkObjectInfo* ActorInfo = ActorPriority.ActorInfo;
+	FActorDestructionInfo* DestructionInfo = ActorPriority.DestructionInfo;
 
-			if ( Actor->MinNetUpdateFrequency == 0.0f )
-			{
-				Actor->MinNetUpdateFrequency = 2.0f;
-			}
+	const double RecordStartTimeSeconds = FPlatformTime::Seconds();
 
-			// Calculate min delta (max rate actor will update), and max delta (slowest rate actor will update)
-			const float MinOptimalDelta = NetUpdateDelay;										// Don't go faster than NetUpdateFrequency
-			const float MaxOptimalDelta = FMath::Max( 1.0f / Actor->MinNetUpdateFrequency, MinOptimalDelta );	// Don't go slower than MinNetUpdateFrequency (or NetUpdateFrequency if it's slower)
+	// Deletion entry
+	if (ActorInfo == nullptr && DestructionInfo != nullptr)
+	{
+		UActorChannel* Channel = (UActorChannel*)ClientConnections[0]->CreateChannel(CHTYPE_Actor, 1);
+		if (Channel)
+		{
+			UE_LOG(LogDemo, Verbose, TEXT("TickDemoRecord creating destroy channel for NetGUID <%s,%s> Priority: %d"), *DestructionInfo->NetGUID.ToString(), *DestructionInfo->PathName, ActorPriority.Priority);
 
-			const float ScaleDownStartTime = 2.0f;
-			const float ScaleDownTimeRange = 5.0f;				
+			// Send a close bunch on the new channel
+			Channel->SetChannelActorForDestroy(DestructionInfo);
 
-			if ( LastReplicateDelta > ScaleDownStartTime )
-			{
-				// Interpolate between MinOptimalDelta/MaxOptimalDelta based on how long it's been since this actor actually sent anything
-				const float Alpha = FMath::Clamp( ( LastReplicateDelta - ScaleDownStartTime ) / ScaleDownTimeRange, 0.0f, 1.0f );
-				ActorInfo->OptimalNetUpdateDelta = FMath::Lerp( MinOptimalDelta, MaxOptimalDelta, Alpha );
-			}
+			// Remove from connection's to-be-destroyed list (close bunch is reliable, so it will make it there)
+			ClientConnections[0]->GetDestroyedStartupOrDormantActorGUIDs().Remove(DestructionInfo->NetGUID);
+		}
+	}
+	else if (ActorInfo != nullptr && DestructionInfo == nullptr)
+	{
+		AActor* Actor = ActorInfo->Actor;
 
-			const double NextUpdateDelta = bUseAdapativeNetFrequency ? ActorInfo->OptimalNetUpdateDelta : NetUpdateDelay;
+		// Use NetUpdateFrequency for this actor, but clamp it to RECORD_HZ.
+		const float ClampedNetUpdateFrequency = FMath::Clamp(Actor->NetUpdateFrequency, Params.MinRecordHz, Params.MaxRecordHz);
+		const double NetUpdateDelay = 1.0 / ClampedNetUpdateFrequency;
 
-			// Account for being fractionally into the next frame
-			// But don't be more than a fraction of a frame behind either (we don't want to do catch-up frames when there is a long delay)
-			const double ExtraTime = DemoCurrentTime - ActorInfo->NextUpdateTime;
-			const double ClampedExtraTime = FMath::Clamp(ExtraTime, 0.0, NetUpdateDelay);
+		// Set defaults if this actor is replicating for first time
+		if (ActorInfo->LastNetReplicateTime == 0)
+		{
+			ActorInfo->LastNetReplicateTime = DemoCurrentTime;
+			ActorInfo->OptimalNetUpdateDelta = NetUpdateDelay;
+		}
 
-			// Try to spread the updates across multiple frames to smooth out spikes.
-			ActorInfo->NextUpdateTime = ( DemoCurrentTime + NextUpdateDelta - ClampedExtraTime + ( ( FMath::SRand() - 0.5 ) * ServerTickTime ) );
+		const float LastReplicateDelta = static_cast<float>(DemoCurrentTime - ActorInfo->LastNetReplicateTime);
 
-			Actor->CallPreReplication( this );
+		if (Actor->MinNetUpdateFrequency == 0.0f)
+		{
+			Actor->MinNetUpdateFrequency = 2.0f;
+		}
 
-			const bool bDidReplicateActor = DemoReplicateActor( Actor, ClientConnections[0], false );
+		// Calculate min delta (max rate actor will update), and max delta (slowest rate actor will update)
+		const float MinOptimalDelta = NetUpdateDelay;										// Don't go faster than NetUpdateFrequency
+		const float MaxOptimalDelta = FMath::Max(1.0f / Actor->MinNetUpdateFrequency, MinOptimalDelta);	// Don't go slower than MinNetUpdateFrequency (or NetUpdateFrequency if it's slower)
 
-			if ( bDidReplicateActor && Actor->NetDormancy == DORM_DormantAll )
-			{
-				// If we've replicated this object at least once, and it wants to go dormant, make it dormant now
-				ActorsToGoDormant.Add( Actor );
-			}
+		const float ScaleDownStartTime = 2.0f;
+		const float ScaleDownTimeRange = 5.0f;
 
-			TSharedPtr<FRepChangedPropertyTracker> PropertyTracker = FindOrCreateRepChangedPropertyTracker( Actor );
+		if (LastReplicateDelta > ScaleDownStartTime)
+		{
+			// Interpolate between MinOptimalDelta/MaxOptimalDelta based on how long it's been since this actor actually sent anything
+			const float Alpha = FMath::Clamp((LastReplicateDelta - ScaleDownStartTime) / ScaleDownTimeRange, 0.0f, 1.0f);
+			ActorInfo->OptimalNetUpdateDelta = FMath::Lerp(MinOptimalDelta, MaxOptimalDelta, Alpha);
+		}
 
-			if ( !GuidCache->NetGUIDLookup.Contains( Actor ) )
-			{
-				// Clear external data if the actor has never replicated yet (and doesn't have a net guid)
-				PropertyTracker->ExternalData.Empty();
-				PropertyTracker->ExternalDataNumBits = 0;
-			}
+		const double NextUpdateDelta = Params.bUseAdapativeNetFrequency ? ActorInfo->OptimalNetUpdateDelta : NetUpdateDelay;
 
-			const bool bUpdatedExternalData = PropertyTracker->ExternalData.Num() > 0;
+		// Account for being fractionally into the next frame
+		// But don't be more than a fraction of a frame behind either (we don't want to do catch-up frames when there is a long delay)
+		const double ExtraTime = DemoCurrentTime - ActorInfo->NextUpdateTime;
+		const double ClampedExtraTime = FMath::Clamp(ExtraTime, 0.0, NetUpdateDelay);
 
-			if ( bDidReplicateActor || bUpdatedExternalData )
-			{
-				// Choose an optimal time, we choose 70% of the actual rate to allow frequency to go up if needed
-				ActorInfo->OptimalNetUpdateDelta = FMath::Clamp( LastReplicateDelta * 0.7f, MinOptimalDelta, MaxOptimalDelta );
-				ActorInfo->LastNetReplicateTime = DemoCurrentTime;
-			}
+		// Try to spread the updates across multiple frames to smooth out spikes.
+		ActorInfo->NextUpdateTime = (DemoCurrentTime + NextUpdateDelta - ClampedExtraTime + ((FMath::SRand() - 0.5) * Params.ServerTickTime));
+
+		Actor->CallPreReplication(this);
+
+		const bool bDidReplicateActor = DemoReplicateActor(Actor, ClientConnections[0], false);
+
+		if (bDidReplicateActor && Actor->NetDormancy == DORM_DormantAll)
+		{
+			// If we've replicated this object at least once, and it wants to go dormant, make it dormant now
+			ActorsToGoDormant.Add(Actor);
 		}
-		else
+
+		TSharedPtr<FRepChangedPropertyTracker> PropertyTracker = FindOrCreateRepChangedPropertyTracker(Actor);
+
+		if (!GuidCache->NetGUIDLookup.Contains(Actor))
 		{
-			UE_LOG(LogDemo, Warning, TEXT( "TickDemoRecord: prioritized actor entry should have either an actor or a destruction info" ));
+			// Clear external data if the actor has never replicated yet (and doesn't have a net guid)
+			PropertyTracker->ExternalData.Empty();
+			PropertyTracker->ExternalDataNumBits = 0;
 		}
 
-		const double RecordEndTimeSeconds = FPlatformTime::Seconds();
+		const bool bUpdatedExternalData = PropertyTracker->ExternalData.Num() > 0;
 
-		// Make sure we're under the desired recording time quota, if any.
-		if ( MaxDesiredRecordTimeMS > 0.0f )
+		if (bDidReplicateActor || bUpdatedExternalData)
 		{
-			const double RecordTimeMS = ( RecordEndTimeSeconds - RecordStartTimeSeconds ) * 1000.0;
-				
-			if ( (ActorInfo && ActorInfo->Actor) && (RecordTimeMS > (MaxDesiredRecordTimeMS * 0.95f )) )
-			{
-				UE_LOG(LogDemo, Verbose, TEXT("Actor %s took more than 95%% of maximum desired recording time. Actor: %.3fms. Max: %.3fms."),
-					*ActorInfo->Actor->GetName(), RecordTimeMS, MaxDesiredRecordTimeMS);
-			}
+			// Choose an optimal time, we choose 70% of the actual rate to allow frequency to go up if needed
+			ActorInfo->OptimalNetUpdateDelta = FMath::Clamp(LastReplicateDelta * 0.7f, MinOptimalDelta, MaxOptimalDelta);
+			ActorInfo->LastNetReplicateTime = DemoCurrentTime;
+		}
+	}
+	else
+	{
+		UE_LOG(LogDemo, Warning, TEXT("TickDemoRecord: prioritized actor entry should have either an actor or a destruction info"));
+	}
 
-			const double TotalRecordTimeMS = ( RecordEndTimeSeconds - ReplicationStartTimeSeconds ) * 1000.0;
-			
-			if ( TotalRecordTimeMS > MaxDesiredRecordTimeMS )
-			{
-				break;
-			}
+	const double RecordEndTimeSeconds = FPlatformTime::Seconds();
+
+	// Make sure we're under the desired recording time quota, if any.
+	if (MaxDesiredRecordTimeMS > 0.0f)
+	{
+		const double RecordTimeMS = (RecordEndTimeSeconds - RecordStartTimeSeconds) * 1000.0;
+
+		if ((ActorInfo && ActorInfo->Actor) && (RecordTimeMS > (MaxDesiredRecordTimeMS * 0.95f)))
+		{
+			UE_LOG(LogDemo, Verbose, TEXT("Actor %s took more than 95%% of maximum desired recording time. Actor: %.3fms. Max: %.3fms."),
+				*ActorInfo->Actor->GetName(), RecordTimeMS, MaxDesiredRecordTimeMS);
+		}
+
+		const double TotalRecordTimeMS = (RecordEndTimeSeconds - Params.ReplicationStartTimeSeconds) * 1000.0;
+
+		if (TotalRecordTimeMS > MaxDesiredRecordTimeMS)
+		{
+			return false;
+		}
+	}
+
+	return true;
+}
+
+bool UDemoNetDriver::ReplicatePrioritizedActors(const TArray<FActorPriority>& ToReplicate, const class FRepActorsParams& Params)
+{
+	TArray<AActor*> ActorsToGoDormant;
+	bool bTimeRemaining = true;
+
+	for (const FActorPriority& ActorPriority : ToReplicate)
+	{
+		bTimeRemaining = ReplicatePrioritizedActor(ActorsToGoDormant, ActorPriority, Params);
+		if (!bTimeRemaining)
+		{
+			break;
 		}
 	}
 
 	for (int32 Idx = 0; Idx < ActorsToGoDormant.Num(); Idx++)
 	{
-		GetNetworkObjectList().MarkDormant( ActorsToGoDormant[Idx], ClientConnections[0], 1, NetDriverName);
+		GetNetworkObjectList().MarkDormant(ActorsToGoDormant[Idx], ClientConnections[0], 1, NetDriverName);
 	}
 
-	// Make sure nothing is left over
-	ClientConnections[0]->FlushNet();
-	check( ClientConnections[0]->SendBuffer.GetNumBits() == 0 );
+	return bTimeRemaining;
+}
 
-	WriteDemoFrameFromQueuedDemoPackets( *FileAr, CastChecked< UDemoNetConnection >( ClientConnections[0] )->QueuedDemoPackets );
+bool UDemoNetDriver::ReplicatePrioritizedActors(const TArray<FActorPriority*>& ToReplicate, const class FRepActorsParams& Params)
+{
+	TArray<AActor*> ActorsToGoDormant;
+	bool bTimeRemaining = true;
 
-	// Save a checkpoint if it's time
-	if ( CVarEnableCheckpoints.GetValueOnAnyThread() == 1 )
+	for (const FActorPriority* ActorPriority : ToReplicate)
 	{
-		check( !PendingCheckpointActors.Num() );		// We early out above, so this shouldn't be possible
-
-		if (ShouldSaveCheckpoint())
+		bTimeRemaining = ReplicatePrioritizedActor(ActorsToGoDormant, *ActorPriority, Params);
+		if (!bTimeRemaining)
 		{
-			SaveCheckpoint();
-			LastCheckpointTime = DemoCurrentTime;
+			break;
 		}
 	}
+
+	for (int32 Idx = 0; Idx < ActorsToGoDormant.Num(); Idx++)
+	{
+		GetNetworkObjectList().MarkDormant(ActorsToGoDormant[Idx], ClientConnections[0], 1, NetDriverName);
+	}
+
+	return bTimeRemaining;
 }
 
 bool UDemoNetDriver::ShouldSaveCheckpoint()
@@ -1844,7 +2377,7 @@ void UDemoNetDriver::PauseChannels( const bool bPause )
 
 		UActorChannel* ActorChannel = Cast< UActorChannel >( OpenChannel );
 
-		if ( ActorChannel == NULL )
+		if ( ActorChannel == nullptr )
 		{
 			continue;
 		}
@@ -1856,7 +2389,7 @@ void UDemoNetDriver::PauseChannels( const bool bPause )
 			continue;
 		}
 
-		if ( ActorChannel->GetActor() == NULL )
+		if ( ActorChannel->GetActor() == nullptr )
 		{
 			continue;
 		}
@@ -1868,9 +2401,12 @@ void UDemoNetDriver::PauseChannels( const bool bPause )
 	bChannelsArePaused = bPause;
 }
 
-bool UDemoNetDriver::ReadDemoFrameIntoPlaybackPackets( FArchive& Ar )
+bool UDemoNetDriver::ReadDemoFrameIntoPlaybackPackets( FArchive& Ar, TArray<FPlaybackPacket>& InPlaybackPackets, const bool bForLevelFastForward, float* OutTime )
 {
 	SCOPED_NAMED_EVENT(UDemoNetDriver_ReadDemoFrameIntoPlaybackPackets, FColor::Purple);
+
+	check(!bForLevelFastForward || HasLevelStreamingFixes());
+
 	if ( Ar.IsError() )
 	{
 		StopDemo();
@@ -1905,113 +2441,180 @@ bool UDemoNetDriver::ReadDemoFrameIntoPlaybackPackets( FArchive& Ar )
 
 	Ar << TimeSeconds;
 
-	// Read any new streaming levels this frame
-	uint32 NumStreamingLevels = 0;
-	Ar.SerializeIntPacked( NumStreamingLevels );
+	if (OutTime)
+	{
+		*OutTime = TimeSeconds;
+	}
 
-	for ( uint32 i = 0; i < NumStreamingLevels; ++i )
+	if (PlaybackDemoHeader.Version >= HISTORY_LEVEL_STREAMING_FIXES)
 	{
-		FString PackageName;
-		FString PackageNameToLoad;
-		FTransform LevelTransform;
+		((UPackageMapClient*)ServerConnection->PackageMap)->ReceiveExportData(Ar);
+	}
+
+	// Check to see if we can skip adding these packets.
+	// This may happen if the archive isn't set to a proper position due to level fast forwarding.
+	const bool bAppendPackets = bIsLoadingCheckpoint || bForLevelFastForward || LatestReadFrameTime < TimeSeconds;
+	LatestReadFrameTime = FMath::Max(LatestReadFrameTime, TimeSeconds);
 
-		Ar << PackageName;
-		Ar << PackageNameToLoad;
-		Ar << LevelTransform;
+	if (HasLevelStreamingFixes())
+	{
+		uint32 NumStreamingLevels = 0;
+		Ar.SerializeIntPacked(NumStreamingLevels);
 
-		// Don't add if already exists
-		bool bFound = false;
+		// We want to avoid adding the same levels to the Seen list multiple times.
+		// This can occur if the Archive is "double read" due to a level fast forward.
+		const bool bAddToSeenList = bAppendPackets && !bForLevelFastForward;
 
-		for (ULevelStreaming* StreamingLevel : World->GetStreamingLevels())
+		FString NameTemp;
+		for (uint32 i = 0; i < NumStreamingLevels; i++)
 		{
-			FString SrcPackageName			= StreamingLevel->GetWorldAssetPackageName();
-			FString SrcPackageNameToLoad	= StreamingLevel->PackageNameToLoad.ToString();
+			Ar << NameTemp;
 
-			if ( SrcPackageName == PackageName && SrcPackageNameToLoad == PackageNameToLoad )
+			if (bAddToSeenList)
 			{
-				bFound = true;
-				break;
+				// Add this level to the seen list, but don't actually mark it as being seen.
+				// It will be marked when we have processed packets for it.
+				const FLevelStatus& LevelStatus = FindOrAddLevelStatus(NameTemp);
+				SeenLevelStatuses.Add(LevelStatus.LevelIndex);
 			}
 		}
+	}
+	else
+	{
+		// Read any new streaming levels this frame
+		uint32 NumStreamingLevels = 0;
+		Ar.SerializeIntPacked(NumStreamingLevels);
 
-		if ( bFound )
+		for (uint32 i = 0; i < NumStreamingLevels; ++i)
 		{
-			continue;
-		}
+			FString PackageName;
+			FString PackageNameToLoad;
+			FTransform LevelTransform;
+
+			Ar << PackageName;
+			Ar << PackageNameToLoad;
+			Ar << LevelTransform;
+
+			// Don't add if already exists
+			bool bFound = false;
+
+			for (ULevelStreaming* StreamingLevel : World->GetStreamingLevels())
+			{
+				FString SrcPackageName = StreamingLevel->GetWorldAssetPackageName();
+				FString SrcPackageNameToLoad = StreamingLevel->PackageNameToLoad.ToString();
+
+				if (SrcPackageName == PackageName && SrcPackageNameToLoad == PackageNameToLoad)
+				{
+					bFound = true;
+					break;
+				}
+			}
+
+			if (bFound)
+			{
+				continue;
+			}
 
-		ULevelStreamingKismet* StreamingLevel = NewObject<ULevelStreamingKismet>( World, NAME_None, RF_NoFlags, nullptr );
+			ULevelStreamingKismet* StreamingLevel = NewObject<ULevelStreamingKismet>(World, NAME_None, RF_NoFlags, nullptr);
 
-		StreamingLevel->SetShouldBeLoaded(true);
-		StreamingLevel->SetShouldBeVisible(true);
-		StreamingLevel->bShouldBlockOnLoad	= false;
-		StreamingLevel->bInitiallyLoaded	= true;
-		StreamingLevel->bInitiallyVisible	= true;
-		StreamingLevel->LevelTransform		= LevelTransform;
+			StreamingLevel->SetShouldBeLoaded(true);
+			StreamingLevel->SetShouldBeVisible(true);
+			StreamingLevel->bShouldBlockOnLoad = false;
+			StreamingLevel->bInitiallyLoaded = true;
+			StreamingLevel->bInitiallyVisible = true;
+			StreamingLevel->LevelTransform = LevelTransform;
 
-		StreamingLevel->PackageNameToLoad = FName( *PackageNameToLoad );
-		StreamingLevel->SetWorldAssetByPackageName( FName( *PackageName ) );
+			StreamingLevel->PackageNameToLoad = FName(*PackageNameToLoad);
+			StreamingLevel->SetWorldAssetByPackageName(FName(*PackageName));
 
-		World->AddStreamingLevel(StreamingLevel);
+			World->AddStreamingLevel(StreamingLevel);
 
-		UE_LOG( LogDemo, Log, TEXT( "ReadDemoFrameIntoPlaybackPackets: Loading streamingLevel: %s, %s" ), *PackageName, *PackageNameToLoad );
+			UE_LOG(LogDemo, Log, TEXT("ReadDemoFrameIntoPlaybackPackets: Loading streamingLevel: %s, %s"), *PackageName, *PackageNameToLoad);
+		}
 	}
+	
 
 #if DEMO_CHECKSUMS == 1
 	{
 		uint32 ServerDeltaTimeCheksum = 0;
 		Ar << ServerDeltaTimeCheksum;
 
-		const uint32 DeltaTimeChecksum = FCrc::MemCrc32( &TimeSeconds, sizeof( TimeSeconds ), 0 );
+		const uint32 DeltaTimeChecksum = FCrc::MemCrc32(&TimeSeconds, sizeof(TimeSeconds), 0);
 
-		if ( DeltaTimeChecksum != ServerDeltaTimeCheksum )
+		if (DeltaTimeChecksum != ServerDeltaTimeCheksum)
 		{
-			UE_LOG( LogDemo, Error, TEXT( "UDemoNetDriver::ReadDemoFrameIntoPlaybackPackets: DeltaTimeChecksum != ServerDeltaTimeCheksum" ) );
+			UE_LOG(LogDemo, Error, TEXT("UDemoNetDriver::ReadDemoFrameIntoPlaybackPackets: DeltaTimeChecksum != ServerDeltaTimeCheksum"));
 			StopDemo();
 			return false;
 		}
 	}
 #endif
 
-	if ( Ar.IsError() )
+	if (Ar.IsError())
 	{
-		UE_LOG( LogDemo, Error, TEXT( "UDemoNetDriver::ReadDemoFrameIntoPlaybackPackets: Failed to read demo ServerDeltaTime" ) );
+		UE_LOG(LogDemo, Error, TEXT("UDemoNetDriver::ReadDemoFrameIntoPlaybackPackets: Failed to read demo ServerDeltaTime"));
 		StopDemo();
 		return false;
 	}
 
-	// Load any custom external data in this frame
-	LoadExternalData( Ar, TimeSeconds );
+	FArchivePos SkipExternalOffset = 0;
+	if (HasLevelStreamingFixes())
+	{
+		Ar << SkipExternalOffset;
+	}
+
+	if (!bForLevelFastForward)
+	{
+		// Load any custom external data in this frame
+		LoadExternalData(Ar, TimeSeconds);
+	}
+	else
+	{
+		Ar.Seek(Ar.Tell() + SkipExternalOffset);
+	}
 
 	// Buffer any packets in this frame
-	while ( true )
+	uint32 SeenLevelIndex = 0;
+
+	while (true)
 	{
+		if (HasLevelStreamingFixes())
+		{
+			Ar.SerializeIntPacked(SeenLevelIndex);
+		}
+
 		int32 PacketBytes = 0;
 		uint8 ReadBuffer[MAX_DEMO_READ_WRITE_BUFFER];
 
-		if ( !ReadPacket( Ar, ReadBuffer, PacketBytes, sizeof( ReadBuffer ) ) )
+		if (!ReadPacket(Ar, ReadBuffer, PacketBytes, sizeof(ReadBuffer)))
 		{
-			UE_LOG( LogDemo, Error, TEXT( "UDemoNetDriver::ReadDemoFrameIntoPlaybackPackets: ReadPacket failed." ) );
+			UE_LOG(LogDemo, Error, TEXT("UDemoNetDriver::ReadDemoFrameIntoPlaybackPackets: ReadPacket failed."));
 
 			StopDemo();
 
-			if ( World != NULL && World->GetGameInstance() != NULL )
+			if (World != nullptr && World->GetGameInstance() != nullptr)
 			{
-				World->GetGameInstance()->HandleDemoPlaybackFailure( EDemoPlayFailure::Generic, FString( TEXT( "UDemoNetDriver::ReadDemoFrameIntoPlaybackPackets: PacketBytes > sizeof( ReadBuffer )" ) ) );
+				World->GetGameInstance()->HandleDemoPlaybackFailure(EDemoPlayFailure::Generic, FString(TEXT("UDemoNetDriver::ReadDemoFrameIntoPlaybackPackets: PacketBytes > sizeof( ReadBuffer )")));
 			}
 
 			return false;
 		}
 
-		if ( PacketBytes == 0 )
+		if (PacketBytes == 0)
 		{
 			break;
 		}
+		if (!bAppendPackets)
+		{
+			continue;
+		}
 
-		FPlaybackPacket& Packet = *( new( PlaybackPackets )FPlaybackPacket );
-		Packet.Data.AddUninitialized( PacketBytes );
+		FPlaybackPacket& Packet = *(new(InPlaybackPackets)FPlaybackPacket);
+		Packet.Data.AddUninitialized(PacketBytes);
 		Packet.TimeSeconds = TimeSeconds;
 		Packet.LevelIndex = ReadCurrentLevelIndex;
-		FMemory::Memcpy( Packet.Data.GetData(), ReadBuffer, PacketBytes );
+		Packet.SeenLevelIndex = SeenLevelIndex;
+		FMemory::Memcpy(Packet.Data.GetData(), ReadBuffer, PacketBytes);
 	}
 
 	return true;
@@ -2087,7 +2690,8 @@ bool UDemoNetDriver::ConditionallyReadDemoFrameIntoPlaybackPackets( FArchive& Ar
 	{
 		const float MAX_PLAYBACK_BUFFER_SECONDS = 5.0f;
 
-		if ( PlaybackPackets[PlaybackPackets.Num() - 1].TimeSeconds > DemoCurrentTime && PlaybackPackets[PlaybackPackets.Num() - 1].TimeSeconds - DemoCurrentTime > MAX_PLAYBACK_BUFFER_SECONDS )
+		const FPlaybackPacket& LastPacket = PlaybackPackets.Last();
+		if ( LastPacket.TimeSeconds > DemoCurrentTime && LastPacket.TimeSeconds - DemoCurrentTime > MAX_PLAYBACK_BUFFER_SECONDS )
 		{
 			return false;	// Don't buffer more than MAX_PLAYBACK_BUFFER_SECONDS worth of frames
 		}
@@ -2151,6 +2755,20 @@ bool UDemoNetDriver::ReadPacket( FArchive& Archive, uint8* OutReadBuffer, int32&
 	return true;
 }
 
+bool UDemoNetDriver::ShouldSkipPlaybackPacket(const FPlaybackPacket& Packet)
+{
+	if (HasLevelStreamingFixes() && Packet.SeenLevelIndex != 0)
+	{
+		// Flag the status as being seen, since we're potentially going to process it.
+		// We need to skip processing if it's not ready (in that case, we'll do a fast-forward).
+		FLevelStatus& LevelStatus = GetLevelStatus(Packet.SeenLevelIndex);
+		LevelStatus.bHasBeenSeen = true;
+		return !LevelStatus.bIsReady;
+	}
+
+	return false;
+}
+
 bool UDemoNetDriver::ConditionallyProcessPlaybackPackets()
 {
 	if ( PlaybackPackets.Num() == 0 )
@@ -2159,22 +2777,22 @@ bool UDemoNetDriver::ConditionallyProcessPlaybackPackets()
 		return false;
 	}
 
-	if ( DemoCurrentTime < PlaybackPackets[0].TimeSeconds )
+	const FPlaybackPacket& CurPacket = PlaybackPackets[0];
+	if ( DemoCurrentTime < CurPacket.TimeSeconds )
 	{
 		// Not enough time has passed to read another frame
 		return false;
 	}
 
-	if (PlaybackPackets[0].LevelIndex != CurrentLevelIndex)
+	if (CurPacket.LevelIndex != CurrentLevelIndex)
 	{
 		GetWorld()->GetGameInstance()->OnSeamlessTravelDuringReplay();
-		CurrentLevelIndex = PlaybackPackets[0].LevelIndex;
+		CurrentLevelIndex = CurPacket.LevelIndex;
 		ProcessSeamlessTravel(CurrentLevelIndex);
 		return false;
 	}
 
-	const bool Result = ProcessPacket( PlaybackPackets[0].Data.GetData(), PlaybackPackets[0].Data.Num() );
-
+	const bool Result = ProcessPacket(CurPacket);
 	PlaybackPackets.RemoveAt( 0 );
 
 	return Result;
@@ -2182,32 +2800,32 @@ bool UDemoNetDriver::ConditionallyProcessPlaybackPackets()
 
 void UDemoNetDriver::ProcessAllPlaybackPackets()
 {
-	for ( int32 i = 0; i < PlaybackPackets.Num(); i++ )
+	for (const FPlaybackPacket& PlaybackPacket : PlaybackPackets)
 	{
-		ProcessPacket( PlaybackPackets[i].Data.GetData(), PlaybackPackets[i].Data.Num() );
+		ProcessPacket(PlaybackPacket);
 	}
-
 	PlaybackPackets.Empty();
 }
 
-bool UDemoNetDriver::ProcessPacket( uint8* Data, int32 Count )
+bool UDemoNetDriver::ProcessPacket( const uint8* Data, int32 Count )
 {
 	PauseChannels( false );
 
-	if ( ServerConnection != NULL )
+	if ( ServerConnection != nullptr )
 	{
 		// Process incoming packet.
-		ServerConnection->ReceivedRawPacket( Data, Count );
+		// ReceivedRawPacket shouldn't change any data, so const_cast should be safe.
+		ServerConnection->ReceivedRawPacket( const_cast<uint8*>(Data), Count );
 	}
 
-	if ( ServerConnection == NULL || ServerConnection->State == USOCK_Closed )
+	if ( ServerConnection == nullptr || ServerConnection->State == USOCK_Closed )
 	{
 		// Something we received resulted in the demo being stopped
 		UE_LOG( LogDemo, Error, TEXT( "UDemoNetDriver::ProcessPacket: ReceivedRawPacket closed connection" ) );
 
 		StopDemo();
 
-		if ( World != NULL && World->GetGameInstance() != NULL )
+		if ( World != nullptr && World->GetGameInstance() != nullptr )
 		{
 			World->GetGameInstance()->HandleDemoPlaybackFailure( EDemoPlayFailure::Generic, FString( TEXT( "UDemoNetDriver::ProcessPacket: PacketBytes > sizeof( ReadBuffer )" ) ) );
 		}
@@ -2225,37 +2843,70 @@ void UDemoNetDriver::WriteDemoFrameFromQueuedDemoPackets( FArchive& Ar, TArray<F
 	// Save total absolute demo time in seconds
 	Ar << DemoCurrentTime;
 
-	// Save any new streaming levels
-	uint32 NumStreamingLevels = NewStreamingLevelsThisFrame.Num();
-	Ar.SerializeIntPacked( NumStreamingLevels );
+	((UPackageMapClient*)ClientConnections[0]->PackageMap)->AppendExportData(Ar);
+
+	if (HasLevelStreamingFixes())
+	{
+		uint32 NumStreamingLevels = AllLevelStatuses.Num();
+		Ar.SerializeIntPacked(NumLevelsAddedThisFrame);
+
+		for (uint32 i = NumStreamingLevels - NumLevelsAddedThisFrame; i < NumStreamingLevels; i++)
+		{
+			Ar << AllLevelStatuses[i].LevelName;
+		}
 
-	for ( uint32 i = 0; i < NumStreamingLevels; i++ )
+		NumLevelsAddedThisFrame = 0;
+	}
+	else
 	{
-		ULevelStreaming* StreamingLevel = World->GetStreamingLevels()[i];
+		// Save any new streaming levels
+		uint32 NumStreamingLevels = NewStreamingLevelsThisFrame.Num();
+		Ar.SerializeIntPacked(NumStreamingLevels);
+
+		for (uint32 i = 0; i < NumStreamingLevels; i++)
+		{
+			ULevelStreaming* StreamingLevel = World->GetStreamingLevels()[i];
 
-		// TODO: StreamingLevel could be null, but since we've already written out the integer count, skipping entries could cause an issue, so leaving as is for now
-		FString PackageName = StreamingLevel->GetWorldAssetPackageName();
-		FString PackageNameToLoad = StreamingLevel->PackageNameToLoad.ToString();
+			// TODO: StreamingLevel could be null, but since we've already written out the integer count, skipping entries could cause an issue, so leaving as is for now
+			FString PackageName = StreamingLevel->GetWorldAssetPackageName();
+			FString PackageNameToLoad = StreamingLevel->PackageNameToLoad.ToString();
 
-		Ar << PackageName;
-		Ar << PackageNameToLoad;
-		Ar << StreamingLevel->LevelTransform;
+			Ar << PackageName;
+			Ar << PackageNameToLoad;
+			Ar << StreamingLevel->LevelTransform;
 
-		UE_LOG( LogDemo, Log, TEXT( "WriteDemoFrameFromQueuedDemoPackets: StreamingLevel: %s, %s" ), *PackageName, *PackageNameToLoad );
+			UE_LOG(LogDemo, Log, TEXT("WriteDemoFrameFromQueuedDemoPackets: StreamingLevel: %s, %s"), *PackageName, *PackageNameToLoad);
+		}
+
+		NewStreamingLevelsThisFrame.Empty();
 	}
 
-	NewStreamingLevelsThisFrame.Empty();
 
-	// Save external data
-	SaveExternalData( Ar );
+	{
+		TUniquePtr<FScopedStoreArchiveOffset> ScopedOffset(HasLevelStreamingFixes() ? new FScopedStoreArchiveOffset(Ar) : nullptr);
+
+		// Save external data
+		SaveExternalData( Ar );
+	}
 
-	for ( int32 i = 0; i < QueuedPackets.Num(); i++ )
+	for (FQueuedDemoPacket& DemoPacket : QueuedPackets)
 	{
-		WritePacket( Ar, QueuedPackets[i].Data.GetData(), QueuedPackets[i].Data.Num() );
+		if (HasLevelStreamingFixes())
+		{
+			Ar.SerializeIntPacked(DemoPacket.SeenLevelIndex);
+		}
+
+		WritePacket(Ar, DemoPacket.Data.GetData(), DemoPacket.Data.Num());
 	}
 
 	QueuedPackets.Empty();
 
+	if (HasLevelStreamingFixes())
+	{
+		uint32 EndCountUnsigned = 0;
+		Ar.SerializeIntPacked(EndCountUnsigned);
+	}
+
 	// Write a count of 0 to signal the end of the frame
 	int32 EndCount = 0;
 	Ar << EndCount;
@@ -2335,6 +2986,28 @@ void UDemoNetDriver::AddUserToReplay( const FString& UserString )
 	}
 }
 
+#if CSV_PROFILER
+struct FCsvDemoSettings
+{
+	bool bCaptureCsv;
+	int32 StartTime;
+	int32 EndTime;
+};
+
+static FCsvDemoSettings GetCsvDemoSettings()
+{
+	FCsvDemoSettings Settings = {};
+	Settings.bCaptureCsv = FParse::Value(FCommandLine::Get(), TEXT("-csvdemostarttime="), Settings.StartTime);
+	if (Settings.bCaptureCsv)
+	{
+		if (!FParse::Value(FCommandLine::Get(), TEXT("-csvdemoendtime="), Settings.EndTime))
+			Settings.EndTime = -1.0;
+	}
+
+	return Settings;
+}
+#endif // CSV_PROFILER
+
 void UDemoNetDriver::TickDemoPlayback( float DeltaSeconds )
 {
 	SCOPED_NAMED_EVENT(UDemoNetDriver_TickDemoPlayback, FColor::Purple);
@@ -2343,6 +3016,31 @@ void UDemoNetDriver::TickDemoPlayback( float DeltaSeconds )
 		return;
 	}
 
+#if CSV_PROFILER
+	{
+		static FCsvDemoSettings CsvDemoSettings = GetCsvDemoSettings();
+		FCsvProfiler* CsvProfiler = FCsvProfiler::Get();
+		if (CsvDemoSettings.bCaptureCsv)
+		{
+			bool bDoCapture = IsPlaying()
+				&& DemoCurrentTime >= CsvDemoSettings.StartTime
+				&& ((DemoCurrentTime <= CsvDemoSettings.EndTime) || (CsvDemoSettings.EndTime < 0));
+
+			if (bDoCapture != CsvProfiler->IsCapturing())
+			{
+				if (bDoCapture)
+				{
+					CsvProfiler->BeginCapture();
+				}
+				else
+				{
+					CsvProfiler->EndCapture();
+				}
+			}
+		}
+	}
+#endif // CSV_PROFILER
+
 	if ( !IsPlaying() )
 	{
 		return;
@@ -2377,6 +3075,10 @@ void UDemoNetDriver::TickDemoPlayback( float DeltaSeconds )
 		CVarDemoSkipTime.AsVariable()->Set(TEXT("0"), ECVF_SetByConsole);
 	}
 
+	// Before we update tasks or move the demo time forward, see if there are any new sublevels that
+	// need to be fast forwarded.
+	PrepFastForwardLevels();
+
 	// Update total demo time
 	if ( ReplayStreamer->GetTotalDemoTime() > 0 )
 	{
@@ -2410,23 +3112,23 @@ void UDemoNetDriver::TickDemoPlayback( float DeltaSeconds )
 		}
 	}
 
-	// Make sure there is data available to read
-	// If we're at the end of the demo, just pause channels and return
-	if ( bDemoPlaybackDone || ( !PlaybackPackets.Num() && !ReplayStreamer->IsDataAvailable() ) )
+	// Advance demo time by seconds passed if we're not paused
+	if ( World->GetWorldSettings()->Pauser == nullptr )
 	{
-		PauseChannels( true );
-		return;
-	}
-
-	// Advance demo time by seconds passed if we're not paused
-	if ( World->GetWorldSettings()->Pauser == NULL )
-	{
-		DemoCurrentTime += DeltaSeconds;
+		DemoCurrentTime += DeltaSeconds;
 	}
 
 	// Clamp time
 	DemoCurrentTime = FMath::Clamp( DemoCurrentTime, 0.0f, DemoTotalTime - 0.01f );
 
+	// Make sure there is data available to read
+	// If we're at the end of the demo, just pause channels and return
+	if (bDemoPlaybackDone || (!PlaybackPackets.Num() && !ReplayStreamer->IsDataAvailable()))
+	{
+		PauseChannels(true);
+		return;
+	}
+
 	// Speculatively grab seconds now in case we need it to get the time it took to fast forward
 	const double FastForwardStartSeconds = FPlatformTime::Seconds();
 
@@ -2455,6 +3157,12 @@ void UDemoNetDriver::FinalizeFastForward( const float StartTime )
 
 	AGameStateBase* const GameState = World != nullptr ? World->GetGameState() : nullptr;
 
+	// Make sure that we delete any Rewind actors that aren't valid anymore.
+	if (bIsFastForwardingForCheckpoint)
+	{
+		CleanupOutstandingRewindActors();
+	}	
+
 	// Correct server world time for fast-forwarding after a checkpoint
 	if (GameState != nullptr)
 	{
@@ -2490,7 +3198,7 @@ void UDemoNetDriver::FinalizeFastForward( const float StartTime )
 			const FObjectReplicator* const ActorReplicator = ActorChannel->ActorReplicator;
 			if (Actor->IsNetStartupActor() && ActorReplicator)
 			{
-				ActorReplicator->RepLayout->DiffProperties(&(ActorReplicator->RepState->RepNotifies), ActorReplicator->RepState->StaticBuffer.GetData(), Actor, true);
+				ActorReplicator->RepLayout->DiffProperties(&(ActorReplicator->RepState->RepNotifies), ActorReplicator->RepState->StaticBuffer.GetData(), Actor, EDiffPropertiesFlags::Sync);
 			}
 		}
 	}
@@ -2498,9 +3206,9 @@ void UDemoNetDriver::FinalizeFastForward( const float StartTime )
 	// Flush all pending RepNotifies that were built up during the fast-forward.
 	if ( ServerConnection != nullptr )
 	{
-		for ( auto& ChannelPair : ServerConnection->ActorChannels )
+		for ( auto& ChannelPair : ServerConnection->ActorChannelMap() )
 		{
-			if ( ChannelPair.Value != NULL )
+			if ( ChannelPair.Value != nullptr )
 			{
 				for ( auto& ReplicatorPair : ChannelPair.Value->ReplicationMap )
 				{
@@ -2532,7 +3240,7 @@ void UDemoNetDriver::SpawnDemoRecSpectator( UNetConnection* Connection, const FU
 		return;
 	}
 
-	check( Connection != NULL );
+	check( Connection != nullptr );
 
 	// Get the replay spectator controller class from the default game mode object,
 	// since the game mode instance isn't replicated to clients of live games.
@@ -2554,7 +3262,7 @@ void UDemoNetDriver::SpawnDemoRecSpectator( UNetConnection* Connection, const FU
 	AGameModeBase* DefaultGameMode = DefaultGameModeClass.GetDefaultObject();
 	UClass* C = DefaultGameMode != nullptr ? DefaultGameMode->ReplaySpectatorPlayerControllerClass : nullptr;
 
-	if ( C == NULL )
+	if ( C == nullptr )
 	{
 		UE_LOG( LogDemo, Error, TEXT( "UDemoNetDriver::SpawnDemoRecSpectator: Failed to load demo spectator class." ) );
 		return;
@@ -2564,7 +3272,7 @@ void UDemoNetDriver::SpawnDemoRecSpectator( UNetConnection* Connection, const FU
 	SpawnInfo.ObjectFlags |= RF_Transient;	// We never want these to save into a map
 	SpectatorController = World->SpawnActor<APlayerController>( C, SpawnInfo );
 
-	if ( SpectatorController == NULL )
+	if ( SpectatorController == nullptr )
 	{
 		UE_LOG( LogDemo, Error, TEXT( "UDemoNetDriver::SpawnDemoRecSpectator: Failed to spawn demo spectator." ) );
 		return;
@@ -2690,6 +3398,14 @@ void UDemoNetDriver::RespawnNecessaryNetStartupActors()
 
 		FRollbackNetStartupActorInfo& RollbackActor = It.Value();
 
+		AActor* ExistingActor = FindObjectFast<AActor>(RollbackActor.Level, RollbackActor.Name);
+		if (ExistingActor)
+		{
+			check(ExistingActor->IsPendingKillOrUnreachable());
+
+			ExistingActor->Rename(nullptr, GetTransientPackage(), REN_DontCreateRedirectors | REN_ForceNoResetLoaders);
+		}
+
 		FActorSpawnParameters SpawnInfo;
 
 		SpawnInfo.Template							= CastChecked<AActor>( RollbackActor.Archetype );
@@ -2697,33 +3413,397 @@ void UDemoNetDriver::RespawnNecessaryNetStartupActors()
 		SpawnInfo.bNoFail							= true;
 		SpawnInfo.Name								= RollbackActor.Name;
 		SpawnInfo.OverrideLevel						= RollbackActor.Level;
+		SpawnInfo.bDeferConstruction				= true;
 
-		AActor* Actor = GetWorld()->SpawnActorAbsolute( RollbackActor.Archetype->GetClass(), FTransform( RollbackActor.Rotation, RollbackActor.Location ), SpawnInfo );
+		const FTransform SpawnTransform = FTransform( RollbackActor.Rotation, RollbackActor.Location );
 
-		if ( !ensure( Actor->GetFullName() == It.Key() ) )
+		AActor* Actor = GetWorld()->SpawnActorAbsolute( RollbackActor.Archetype->GetClass(), SpawnTransform, SpawnInfo );
+		if (Actor)
 		{
-			UE_LOG( LogDemo, Log, TEXT( "RespawnNecessaryNetStartupActors: NetStartupRollbackActor name doesn't match original: %s, %s" ), *Actor->GetFullName(), *It.Key() );
-		}
+			if ( !ensure( Actor->GetFullName() == It.Key() ) )
+			{
+				UE_LOG( LogDemo, Log, TEXT( "RespawnNecessaryNetStartupActors: NetStartupRollbackActor name doesn't match original: %s, %s" ), *Actor->GetFullName(), *It.Key() );
+			}
+
+			TSharedPtr<FRepLayout> RepLayout = GetObjectClassRepLayout(Actor->GetClass());
+			if (RepLayout.IsValid())
+			{
+				if (RepLayout.IsValid() && RollbackActor.RepState.IsValid())
+				{
+					const ENetRole SavedRole = Actor->Role;
+
+					RepLayout->DiffStableProperties(&RollbackActor.RepState->RepNotifies, nullptr, Actor, RollbackActor.RepState->StaticBuffer.GetData());
+
+					Actor->Role = SavedRole;
+				}
+			}
 
-		Actor->bNetStartup = true;
-		Actor->SwapRolesForReplay();
+			check(Actor->GetRemoteRole() != ROLE_Authority);
 
-		check( Actor->GetRemoteRole() == ROLE_Authority );
+			Actor->bNetStartup = true;
+
+			UGameplayStatics::FinishSpawningActor(Actor, SpawnTransform);
+
+			if (Actor->Role == ROLE_Authority)
+			{
+				Actor->SwapRolesForReplay();
+			}
+
+			if (RepLayout.IsValid() && RollbackActor.RepState.IsValid())
+			{
+				if (RollbackActor.RepState->RepNotifies.Num() > 0)
+				{
+					RepLayout->CallRepNotifies(RollbackActor.RepState.Get(), Actor);
+				}
+			}
+
+			check( Actor->GetRemoteRole() == ROLE_Authority );
+		}
 
 		It.RemoveCurrent();
 	}
 }
 
+void UDemoNetDriver::PrepFastForwardLevels()
+{
+	if (!HasLevelStreamingFixes() || NewStreamingLevelsThisFrame.Num() == 0)
+	{
+		return;
+	}
+
+	check(!bIsFastForwarding);
+	check(!bIsLoadingCheckpoint);
+
+	// Do a quick pass to double check everything is still valid, and that we have data for the levels.
+	UWorld* LocalWorld = GetWorld();
+	for (TWeakObjectPtr<UObject> &WeakLevel : NewStreamingLevelsThisFrame)
+	{
+		// For playback, we should only ever see ULevels in this list.
+		if (ULevel* Level = CastChecked<ULevel>(WeakLevel.Get()))
+		{
+			if (!ensure(!LevelsPendingFastForward.Contains(Level)))
+			{
+				UE_LOG(LogDemo, Warning, TEXT("FastForwardLevels - NewStreamingLevel found in Pending list! %s"), *GetFullName(Level));
+				continue;
+			}
+
+			TSet<TWeakObjectPtr<AActor>> LevelActors;
+			for (AActor* Actor : Level->Actors)
+			{
+				if (Actor == nullptr || !Actor->IsNetStartupActor())
+				{
+					continue;
+				}
+				else if (DeletedNetStartupActors.Contains(Actor->GetFullName()))
+				{
+					// Put this actor on the rollback list so we can undelete it during future scrubbing,
+					// then delete it.
+					QueueNetStartupActorForRollbackViaDeletion(Actor);
+					LocalWorld->DestroyActor(Actor, true);
+				}
+				else
+				{
+					LevelActors.Add(Actor);
+				}
+			}
+
+			if (LevelActors.Num() > 0)
+			{
+				LevelsPendingFastForward.Emplace(Level, MoveTemp(LevelActors));
+			}
+		}
+	}
+	NewStreamingLevelsThisFrame.Empty();
+
+	if (LevelsPendingFastForward.Num() == 0 ||
+		LastProcessedPacketTime == 0.f ||
+		// If there's already a FastForwardLevelsTask or GotoTimeTask, then we don't need
+		// to add another (as the levels will get picked up by either of those).
+		IsNamedTaskInQueue(ReplayTaskNames::GotoTimeInSecondsTask) ||
+		IsNamedTaskInQueue(ReplayTaskNames::FastForwardLevelsTask))
+	{
+		return;
+	}
+
+	AddReplayTask(new FFastForwardLevelsTask(this));
+}
+
+bool UDemoNetDriver::FastForwardLevels(FArchive* CheckpointArchive, int64 ExtraTime)
+{
+	DECLARE_SCOPE_CYCLE_COUNTER(TEXT("FastForwardLevels time"), STAT_FastForwardLevelTime, STATGROUP_Net);
+
+	PauseChannels(false);
+
+	// We can skip processing the checkpoint here, because Goto will load one up for us later.
+	// We only want to check the very next task, though. Otherwise, we could end processing other
+	// tasks in an invalid state.
+	if (GetNextQueuedTaskName() == ReplayTaskNames::GotoTimeInSecondsTask)
+	{
+		// This is a bit hacky, but we don't want to do *any* processing this frame.
+		// Therefore, we'll reset the ActiveReplayTask and return false.
+		// This will cause us to early out, and then handle the Goto task next frame.
+		ActiveReplayTask.Reset();
+		return false;
+	}
+
+	// Generate the list of level names, and an uber list of the startup actors.
+	// We manually track whenever a level is added and removed from the world, so these should always be valid.
+	TSet<int32> LevelIndices;
+	TSet<TWeakObjectPtr<AActor>> StartupActors;
+	TSet<ULevel*> LocalLevels;
+
+	// Reserve some default space, and just assume a minimum of at least 4 actors per level (super low estimate).
+	LevelIndices.Reserve(LevelsPendingFastForward.Num());
+	StartupActors.Reserve(LevelsPendingFastForward.Num() * 4);
+
+	for (auto It = LevelsPendingFastForward.CreateIterator(); It; ++It)
+	{
+		// Track the appropriate level, and mark it as ready.
+		FLevelStatus& LevelStatus = GetLevelStatus(GetLevelPackageName(*It.Key()));
+		LevelIndices.Add(LevelStatus.LevelIndex);
+		LevelStatus.bIsReady = true;
+
+		// Quick sanity check to make sure the actors are still valid
+		// NOTE: The only way any of these should not be valid is if the level was unloaded,
+		//			or something in the demo caused the actor to be destroyed *before*
+		//			the level was ready. Either case seems bad if we've made it this far.
+		TSet<TWeakObjectPtr<AActor>>& LevelStartupActors = It.Value();
+		for (auto ActorIt = LevelStartupActors.CreateIterator(); ActorIt; ++ActorIt)
+		{
+			if (!ensure((*ActorIt).IsValid()))
+			{
+				ActorIt.RemoveCurrent();
+			}
+		}
+
+		LocalLevels.Add(It.Key());
+		StartupActors.Append(MoveTemp(LevelStartupActors));
+	}
+	LevelsPendingFastForward.Reset();
+
+	// At this point, assume that we have all the data necessary to fast forward.
+	// Further, the archives should be at the appropriate positions to read.
+	FArchive* StreamingAr = ReplayStreamer->GetStreamingArchive();
+	FArchive* UsingAr = CheckpointArchive;
+
+	// If there's no checkpoint data, then we'll just move straight into the streaming archive.
+	// This can happen if we haven't passed a checkpoint yet.
+	if (CheckpointArchive->TotalSize() == 0)
+	{
+		UsingAr = StreamingAr;
+	}
+	else
+	{
+		// Skip forward to the Packet data.
+		// ReadDemoFrameIntoPlaybackPackets will handle skipping passed any additional data
+		FArchivePos PacketOffset = 0;
+		*CheckpointArchive << PacketOffset;
+		CheckpointArchive->Seek(PacketOffset + CheckpointArchive->Tell());
+	}
+
+	// We only want to process packets that are before anything we've currently processed.
+	// Further, we want to make sure that we leave the archive in a good state for later use.
+	int32 NumPackets = 0;
+	float LastReadTime = 0;
+	FArchivePos PreFramePos = 0;
+
+	// Load up all of our playback packets, from the demo and the stream.
+	TArray<FPlaybackPacket> LocalPlaybackPackets;
+	while (!UsingAr->AtEnd())
+	{
+		// Grab the packets, and make sure the stream is OK.
+		PreFramePos = UsingAr->Tell();
+		NumPackets = LocalPlaybackPackets.Num();
+		if (!ReadDemoFrameIntoPlaybackPackets(*UsingAr, LocalPlaybackPackets, true, &LastReadTime))
+		{
+			StopDemo();
+			return false;
+		}
+
+		// In case the archive had more data than we needed, we'll try to leave it where we left off
+		// before the level fast forward.
+		if (LastReadTime > LastProcessedPacketTime)
+		{
+			UsingAr->Seek(PreFramePos);
+			if (ensure(NumPackets != 0))
+			{
+				LocalPlaybackPackets.RemoveAt(NumPackets, LocalPlaybackPackets.Num() - NumPackets);
+			}	
+			break;
+		}
+
+		// Switch to the streaming archive.
+		UsingAr = StreamingAr;
+	}
+
+	// If we've gotten this far, it means we should have something to process.
+	check(LocalPlaybackPackets.Num() > 0);
+
+	// It's possible that the level we're streaming in may spawn Dynamic Actors.
+	// In that case, we want to make sure we track them so we can process them below.
+	// We only care about the actors if they're outered to the Level.
+	struct FDynamicActorTracker
+	{
+		FDynamicActorTracker(UWorld* InTrackWorld, TSet<ULevel*>& InCareAboutLevels, TSet<TWeakObjectPtr<AActor>>& InActorSet) :
+			TrackWorld(InTrackWorld),
+			CareAboutLevels(InCareAboutLevels),
+			ActorSet(InActorSet)
+		{
+			FOnActorSpawned::FDelegate TrackActorDelegate = FOnActorSpawned::FDelegate::CreateRaw(this, &FDynamicActorTracker::TrackActor);
+			TrackActorHandle = TrackWorld->AddOnActorSpawnedHandler(TrackActorDelegate);
+		}
+
+		~FDynamicActorTracker()
+		{
+			TrackWorld->RemoveOnActorSpawnedHandler(TrackActorHandle);
+		}
+
+	private:
+
+		void TrackActor(AActor* Actor)
+		{
+			if (Actor && CareAboutLevels.Contains(Actor->GetLevel()))
+			{
+				UE_LOG(LogDemo, Verbose, TEXT("FastForwardLevels - Sublevel spawned dynamic actor."));
+				ActorSet.Add(Actor);
+			}
+		}
+
+		UWorld* TrackWorld;
+		TSet<ULevel*> CareAboutLevels;
+		TSet<TWeakObjectPtr<AActor>>& ActorSet;
+		FDelegateHandle TrackActorHandle;
+
+	} ActorTracker(World, LocalLevels, StartupActors);
+
+	{
+		TGuardValue<bool> FastForward(bIsFastForwarding, true);
+		struct FScopedIgnoreChannels
+		{
+			FScopedIgnoreChannels(UNetConnection* InConnection):
+				Connection(InConnection)
+			{
+				if (Connection.IsValid())
+				{
+					Connection->SetIgnoreAlreadyOpenedChannels(true);
+				}	
+			}
+
+			~FScopedIgnoreChannels()
+			{
+				if (Connection.IsValid())
+				{
+					Connection->SetIgnoreAlreadyOpenedChannels(false);
+				}	
+			}
+
+		private:
+			TWeakObjectPtr<UNetConnection> Connection;
+		} ScopedIgnoreChannels(ServerConnection);
+
+		// Process all the packets we need.
+		for (FPlaybackPacket& Packet : LocalPlaybackPackets)
+		{
+			// Skip packets that aren't associated with levels.
+			if (Packet.SeenLevelIndex == 0)
+			{
+				continue;
+			}
+
+			// Don't attempt to go beyond the current demo time.
+			if (Packet.TimeSeconds > DemoCurrentTime)
+			{
+				break;
+			}
+
+			const FLevelStatus& LevelStatus = GetLevelStatus(Packet.SeenLevelIndex);
+			const bool bCareAboutLevel = LevelIndices.Contains(LevelStatus.LevelIndex);
+
+			if (bCareAboutLevel)
+			{
+				// If we tried to process the packet, but failed, then the replay will be in a broken state.
+				// ProcessPacket will have called StopDemo.
+				if (!ProcessPacket(Packet.Data.GetData(), Packet.Data.Num()))
+				{
+					UE_LOG(LogDemo, Warning, TEXT("FastForwardLevel failed to process packet"));
+					return false;
+				}
+			}
+		}
+	}
+
+	if (ensure(ServerConnection != nullptr))
+	{
+		// Make a pass at OnReps for startup actors, since they were skipped during checkpoint loading.
+		// At this point the shadow state of these actors should be the actual state from before the checkpoint,
+		// and the current state is the CDO state evolved by any changes that occurred during checkpoint loading and fast-forwarding.
+
+		TArray<UActorChannel*> ChannelsToUpdate;
+		ChannelsToUpdate.Reserve(StartupActors.Num());
+
+		for (UChannel* Channel : ServerConnection->OpenChannels)
+		{
+			// Skip non-actor channels.
+			if (Channel == nullptr || Channel->ChType != CHTYPE_Actor)
+			{
+				continue;
+			}
+
+			// Since we know this is an actor channel, should be safe to do a static_cast.
+			UActorChannel* const ActorChannel = static_cast<UActorChannel*>(Channel);
+			AActor* Actor = ActorChannel->GetActor();
+
+			// We only need to consider startup actors, or dynamic that were spawned and outered
+			// to one of our sublevels.
+			if (Actor == nullptr || !StartupActors.Contains(Actor))
+			{
+				continue;
+			}
+
+			ChannelsToUpdate.Add(ActorChannel);
+			if (const FObjectReplicator* const ActorReplicator = ActorChannel->ActorReplicator)
+			{
+				ActorReplicator->RepLayout->DiffProperties(&(ActorReplicator->RepState->RepNotifies), ActorReplicator->RepState->StaticBuffer.GetData(), Actor, EDiffPropertiesFlags::Sync);
+			}
+		}
+
+		for (UActorChannel* Channel : ChannelsToUpdate)
+		{
+			for (auto& ReplicatorPair : Channel->ReplicationMap)
+			{
+				ReplicatorPair.Value->CallRepNotifies(true);
+			}
+		}
+	}
+
+	return true;
+}
+
 bool UDemoNetDriver::LoadCheckpoint( FArchive* GotoCheckpointArchive, int64 GotoCheckpointSkipExtraTimeInMS )
 {
 	DECLARE_SCOPE_CYCLE_COUNTER(TEXT("LoadCheckpoint time"), STAT_ReplayCheckpointLoadTime, STATGROUP_Net);
 
-	check( GotoCheckpointArchive != NULL );
+	check( GotoCheckpointArchive != nullptr );
 	check( !bIsFastForwardingForCheckpoint );
 	check( !bIsFastForwarding );
 
 	int32 LevelForCheckpoint = 0;
 
+	if (HasLevelStreamingFixes())
+	{
+		// Make sure to read the packet offset, even though we won't use it here.
+		if (GotoCheckpointArchive->TotalSize() > 0)
+		{
+			FArchivePos PacketOffset = 0;
+			*GotoCheckpointArchive << PacketOffset;
+		}
+
+		ResetLevelStatuses();
+	}
+
+	LastProcessedPacketTime = 0.f;
+	LatestReadFrameTime = 0.f;
+
 	if (PlaybackDemoHeader.Version >= HISTORY_MULTIPLE_LEVELS)
 	{
 		if (GotoCheckpointArchive->TotalSize() > 0)
@@ -2776,8 +3856,8 @@ bool UDemoNetDriver::LoadCheckpoint( FArchive* GotoCheckpointArchive, int64 Goto
 	}
 
 	// Save off the current spectator position
-	// Check for NULL, which can be the case if we haven't played any of the demo yet but want to fast forward (joining live game for example)
-	if ( SpectatorController != NULL )
+	// Check for nullptr, which can be the case if we haven't played any of the demo yet but want to fast forward (joining live game for example)
+	if ( SpectatorController != nullptr )
 	{
 		// Save off the SpectatorController's GUID so that we know not to queue his bunches
 		AddNonQueuedActorForScrubbing( SpectatorController );
@@ -2815,7 +3895,14 @@ bool UDemoNetDriver::LoadCheckpoint( FArchive* GotoCheckpointArchive, int64 Goto
 	// For the spectator controller, this allows the state and position to persist.
 	TArray<FPreservedNetworkGUIDEntry> NetGUIDsToPreserve;
 
+	if (!ensureMsgf(TrackedRewindActorsByGUID.Num() == 0, TEXT("LoadCheckpoint: TrackedRewindAcotrsByGUID list not empty!")))
+	{
+		TrackedRewindActorsByGUID.Empty();
+	}
+
 #if 1
+	TSet<const AActor*> KeepAliveActors;
+
 	// Destroy all non startup actors. They will get restored with the checkpoint
 	for ( FActorIterator It( GetWorld() ); It; ++It )
 	{
@@ -2829,57 +3916,72 @@ bool UDemoNetDriver::LoadCheckpoint( FArchive* GotoCheckpointArchive, int64 Goto
 			AddNonQueuedActorForScrubbing(*It);
 		}
 		
-		if ( SpectatorController != nullptr )
+		const bool bShouldPreserveForPlayerController = (SpectatorController != nullptr && (*It == SpectatorController || *It == SpectatorController->GetSpectatorPawn() || It->GetOwner() == SpectatorController));
+		const bool bShouldPreserveForRewindability = (It->bReplayRewindable && !It->IsNetStartupActor());										
+
+		if (bShouldPreserveForPlayerController || bShouldPreserveForRewindability)
 		{
-			if ( *It == SpectatorController || *It == SpectatorController->GetSpectatorPawn() || It->GetOwner() == SpectatorController )
+			// If an non-startup actor that we don't destroy has an entry in the GuidCache, preserve that entry so
+			// that the object will be re-used after loading the checkpoint. Otherwise, a new copy
+			// of the object will be created each time a checkpoint is loaded, causing a leak.
+			const FNetworkGUID FoundGUID = GuidCache->NetGUIDLookup.FindRef( *It );
+				
+			if (FoundGUID.IsValid())
 			{
-				// If an non-startup actor that we don't destroy has an entry in the GuidCache, preserve that entry so
-				// that the object will be re-used after loading the checkpoint. Otherwise, a new copy
-				// of the object will be created each time a checkpoint is loaded, causing a leak.
-				const FNetworkGUID FoundGUID = GuidCache->NetGUIDLookup.FindRef( *It );
+				NetGUIDsToPreserve.Emplace(FoundGUID, *It);
 				
-				if ( FoundGUID.IsValid() )
+				if (bShouldPreserveForRewindability)
 				{
-					NetGUIDsToPreserve.Emplace( FoundGUID, *It );
+					TrackedRewindActorsByGUID.Add(FoundGUID);
 				}
-				continue;
 			}
+
+			KeepAliveActors.Add(*It);
+			continue;
 		}
 
+		// Prevent NetStartupActors from being destroyed.
+		// NetStartupActors that can't have properties directly re-applied should use QueueNetStartupActorForRollbackViaDeletion.
 		if ( It->IsNetStartupActor() )
 		{
+			// Go ahead and rewind this now, since we won't be destroying it later.
+			if (It->bReplayRewindable)
+			{
+				It->RewindForReplay();
+			}
+			KeepAliveActors.Add(*It);
 			continue;
 		}
 
 		GetWorld()->DestroyActor( *It, true );
 	}
 
-	// Find the SpectatorController on the channels, and make sure shutting down the connection doesn't destroy this actor
+	// Remove references to our KeepAlive actors so that cleaning up the channels won't destroy them.
 	for ( int32 i = ServerConnection->OpenChannels.Num() - 1; i >= 0; i-- )
 	{
 		UChannel* OpenChannel = ServerConnection->OpenChannels[i];
-		if ( OpenChannel != NULL )
+		if ( OpenChannel != nullptr )
 		{
 			UActorChannel* ActorChannel = Cast< UActorChannel >( OpenChannel );
-			if ( ActorChannel != NULL && ActorChannel->Actor == SpectatorController )
+			if ( ActorChannel != nullptr && KeepAliveActors.Contains(ActorChannel->Actor))
 			{
-				ActorChannel->Actor = NULL;
+				ActorChannel->Actor = nullptr;
 			}
 		}
 	}
 
 	if ( ServerConnection->OwningActor == SpectatorController )
 	{
-		ServerConnection->OwningActor = NULL;
+		ServerConnection->OwningActor = nullptr;
 	}
 #else
 	for ( int32 i = ServerConnection->OpenChannels.Num() - 1; i >= 0; i-- )
 	{
 		UChannel* OpenChannel = ServerConnection->OpenChannels[i];
-		if ( OpenChannel != NULL )
+		if ( OpenChannel != nullptr )
 		{
 			UActorChannel* ActorChannel = Cast< UActorChannel >( OpenChannel );
-			if ( ActorChannel != NULL && ActorChannel->GetActor() != NULL && !ActorChannel->GetActor()->IsNetStartupActor() )
+			if ( ActorChannel != nullptr && ActorChannel->GetActor() != nullptr && !ActorChannel->GetActor()->IsNetStartupActor() )
 			{
 				GetWorld()->DestroyActor( ActorChannel->GetActor(), true );
 			}
@@ -2947,9 +4049,9 @@ bool UDemoNetDriver::LoadCheckpoint( FArchive* GotoCheckpointArchive, int64 Goto
 		FNetGuidCacheObject& CacheObject = GuidCache->ObjectLookup.FindOrAdd( PreservedEntry.NetGUID );
 
 		CacheObject.Object = MakeWeakObjectPtr(const_cast<AActor*>(PreservedEntry.Actor));
-		check( CacheObject.Object != NULL );
+		check( CacheObject.Object != nullptr );
 		CacheObject.bNoLoad = true;
-		GuidCache->NetGUIDLookup.Add( MakeWeakObjectPtr( const_cast<AActor*>( PreservedEntry.Actor ) ), PreservedEntry.NetGUID );
+		GuidCache->NetGUIDLookup.Add( CacheObject.Object, PreservedEntry.NetGUID );
 	}
 
 	if ( GotoCheckpointArchive->TotalSize() == 0 || GotoCheckpointArchive->TotalSize() == INDEX_NONE )
@@ -2969,6 +4071,12 @@ bool UDemoNetDriver::LoadCheckpoint( FArchive* GotoCheckpointArchive, int64 Goto
 		{
 			SkipTimeInternal( ( float )GotoCheckpointSkipExtraTimeInMS / 1000.0f, true, true );
 		}
+		else
+		{
+			// Make sure that we delete any Rewind actors that aren't valid anymore.
+			// If there's more data to stream in, we will handle this in FinalizeFastForward.
+			CleanupOutstandingRewindActors();
+		}
 
 		return true;
 	}
@@ -2982,11 +4090,22 @@ bool UDemoNetDriver::LoadCheckpoint( FArchive* GotoCheckpointArchive, int64 Goto
 	// Destroy startup actors that shouldn't exist past this checkpoint
 	for ( FActorIterator It( GetWorld() ); It; ++It )
 	{
-		if ( DeletedNetStartupActors.Contains( It->GetFullName() ) )
+		const FString FullName = It->GetFullName();
+		if ( DeletedNetStartupActors.Contains( FullName ) )
 		{
+			if (It->bReplayRewindable)
+			{
+				// Log and skip. We can't queue Rewindable actors and we can't destroy them.
+				// This actor may still get destroyed during cleanup.
+				UE_LOG(LogDemo, Warning, TEXT("Replay Rewindable Actor found in the DeletedNetStartupActors. Replay may show artifacts (%s)"), *FullName);
+				continue;
+			}
+
 			// Put this actor on the rollback list so we can undelete it during future scrubbing
 			QueueNetStartupActorForRollbackViaDeletion( *It );
 
+			UE_LOG(LogDemo, Verbose, TEXT("LoadCheckpoint: deleting startup actor %s"), *FullName);
+
 			// Delete the actor
 			GetWorld()->DestroyActor( *It, true );
 		}
@@ -3033,7 +4152,7 @@ bool UDemoNetDriver::LoadCheckpoint( FArchive* GotoCheckpointArchive, int64 Goto
 
 	if ( PlaybackPackets.Num() > 0 )
 	{
-		DemoCurrentTime = PlaybackPackets[PlaybackPackets.Num() - 1].TimeSeconds;
+		DemoCurrentTime = PlaybackPackets.Last().TimeSeconds;
 	}
 	else
 	{
@@ -3045,6 +4164,12 @@ bool UDemoNetDriver::LoadCheckpoint( FArchive* GotoCheckpointArchive, int64 Goto
 		// If we need to skip more time for fine scrubbing, set that up now
 		SkipTimeInternal( ( float )GotoCheckpointSkipExtraTimeInMS / 1000.0f, true, true );
 	}
+	else
+	{
+		// Make sure that we delete any Rewind actors that aren't valid anymore.
+		// If there's more data to stream in, we will handle this in FinalizeFastForward.
+		CleanupOutstandingRewindActors();
+	}
 
 	ProcessAllPlaybackPackets();
 
@@ -3074,6 +4199,20 @@ bool UDemoNetDriver::LoadCheckpoint( FArchive* GotoCheckpointArchive, int64 Goto
 	return true;
 }
 
+bool UDemoNetDriver::IsSavingCheckpoint() const
+{
+	if (ClientConnections.Num() > 0)
+	{
+		UNetConnection* const NetConnection = ClientConnections[0];
+		if (NetConnection)
+		{
+			return NetConnection->bResendAllDataSinceOpen;
+		}
+	}
+
+	return false;
+}
+
 bool UDemoNetDriver::ShouldQueueBunchesForActorGUID(FNetworkGUID InGUID) const
 {
 	if ( CVarDemoQueueCheckpointChannels.GetValueOnGameThread() == 0)
@@ -3090,6 +4229,11 @@ bool UDemoNetDriver::ShouldQueueBunchesForActorGUID(FNetworkGUID InGUID) const
 	return false;
 }
 
+bool UDemoNetDriver::ShouldIgnoreRPCs() const
+{
+	return (CVarDemoFastForwardIgnoreRPCs.GetValueOnAnyThread() && (bIsLoadingCheckpoint || bIsFastForwarding));
+}
+
 FNetworkGUID UDemoNetDriver::GetGUIDForActor(const AActor* InActor) const
 {
 	UNetConnection* Connection = ServerConnection;
@@ -3146,7 +4290,7 @@ bool UDemoNetDriver::ShouldReceiveRepNotifiesForObject(UObject* Object) const
 
 void UDemoNetDriver::AddNonQueuedActorForScrubbing(AActor const* Actor)
 {
-	UActorChannel const* const* const FoundChannel = ServerConnection->ActorChannels.Find(MakeWeakObjectPtr(const_cast<AActor*>(Actor)));
+	UActorChannel const* const* const FoundChannel = ServerConnection->FindActorChannel(MakeWeakObjectPtr(const_cast<AActor*>(Actor)));
 	if (FoundChannel != nullptr && *FoundChannel != nullptr)
 	{
 		FNetworkGUID const ActorGUID = (*FoundChannel)->ActorNetGUID;
@@ -3162,6 +4306,67 @@ void UDemoNetDriver::AddNonQueuedGUIDForScrubbing(FNetworkGUID InGUID)
 	}
 }
 
+FDemoSavedPropertyState UDemoNetDriver::SavePropertyState() const
+{
+	FDemoSavedPropertyState State;
+
+	if (IsRecording())
+	{
+		const UNetConnection* const RecordingConnection = ClientConnections[0];
+		for (auto ChannelPair = RecordingConnection->ActorChannelConstIterator(); ChannelPair; ++ChannelPair)
+		{
+			const UActorChannel* const Channel = ChannelPair.Value();
+			if (Channel)
+			{
+				for (const auto& ReplicatorPair : Channel->ReplicationMap)
+				{
+					const UObject* const RepObject = ReplicatorPair.Key.Get();
+					if (RepObject)
+					{
+						FDemoSavedRepObjectState& SavedObject = State.Emplace_GetRef();
+						SavedObject.Object = RepObject;
+						SavedObject.RepLayout = ReplicatorPair.Value->RepLayout;
+
+						SavedObject.RepLayout->InitShadowData(SavedObject.PropertyData, RepObject->GetClass(), reinterpret_cast<const uint8* const>(RepObject));
+
+						// Store the properties in the new RepState
+						SavedObject.RepLayout->DiffProperties(nullptr, SavedObject.PropertyData.GetData(), RepObject, EDiffPropertiesFlags::Sync | EDiffPropertiesFlags::IncludeConditionalProperties);
+					}
+				}
+			}
+		}
+	}
+
+	return State;
+}
+
+bool UDemoNetDriver::ComparePropertyState(const FDemoSavedPropertyState& State) const
+{
+	bool bWasDifferent = false;
+
+	if (IsRecording())
+	{
+		for (const FDemoSavedRepObjectState& ObjectState : State)
+		{
+			const UObject* const RepObject = ObjectState.Object.Get();
+			if (RepObject)
+			{
+				if (ObjectState.RepLayout->DiffProperties(nullptr, const_cast<UObject* const>(RepObject), ObjectState.PropertyData.GetData(), EDiffPropertiesFlags::IncludeConditionalProperties))
+				{
+					bWasDifferent = true;
+				}
+			}
+			else
+			{
+				UE_LOG(LogDemo, Warning, TEXT("A replicated object was destroyed or marked pending kill since its state was saved!"));
+				bWasDifferent = true;
+			}
+		}
+	}
+
+	return bWasDifferent;
+}
+
 /*-----------------------------------------------------------------------------
 	UDemoNetConnection.
 -----------------------------------------------------------------------------*/
@@ -3238,7 +4443,7 @@ void UDemoNetConnection::FlushNet( bool bIgnoreSimulation )
 {
 	// in playback, there is no data to send except
 	// channel closing if an error occurs.
-	if ( GetDriver()->ServerConnection != NULL )
+	if ( GetDriver()->ServerConnection != nullptr )
 	{
 		InitSendBuffer();
 	}
@@ -3265,7 +4470,7 @@ void UDemoNetConnection::HandleClientPlayer( APlayerController* PC, UNetConnecti
 		return;
 	}
 
-	ULocalPlayer* LocalPlayer = NULL;
+	ULocalPlayer* LocalPlayer = nullptr;
 	for (FLocalPlayerIterator It(GEngine, Driver->GetWorld()); It; ++It)
 	{
 		LocalPlayer = *It;
@@ -3320,7 +4525,7 @@ TSharedPtr<FObjectReplicator> UDemoNetConnection::CreateReplicatorForNewActorCha
 	// Now that the shadow state is initialized, copy the CDO state into the actor state.
 	if (bIsCheckpointStartupActor && NewReplicator->RepLayout.IsValid() && Object->GetClass())
 	{
-		NewReplicator->RepLayout->DiffProperties(nullptr, Object, Object->GetClass()->GetDefaultObject(), true);
+		NewReplicator->RepLayout->DiffProperties(nullptr, Object, Object->GetClass()->GetDefaultObject(), EDiffPropertiesFlags::Sync);
 
 		// Need to swap roles for the startup actor since in the CDO they aren't swapped, and the CDO just
 		// overwrote the actor state.
@@ -3331,7 +4536,7 @@ TSharedPtr<FObjectReplicator> UDemoNetConnection::CreateReplicatorForNewActorCha
 	}
 
 	// Handle rewinding initially dormant startup actors that were changed on the client
-	const bool bIsInitialDormantStartupActor = NetDriver && NetDriver->IsPlayingClientReplay() && Actor && (Actor->NetDormancy == DORM_Initial) && Actor->IsNetStartupActor();
+	const bool bIsInitialDormantStartupActor = NetDriver && NetDriver->IsPlayingClientReplay() && Actor && (Actor->NetDormancy == DORM_Initial) && Actor->IsNetStartupActor() && !Actor->bReplayRewindable;
 	if (bIsInitialDormantStartupActor)
 	{
 		NetDriver->QueueNetStartupActorForRollbackViaDeletion(Actor);
@@ -3391,14 +4596,133 @@ void UDemoNetDriver::NotifyActorDestroyed( AActor* Actor, bool IsSeamlessTravel
 {
 	check( Actor != nullptr );
 
-	if ( IsRecording() && Actor->IsNetStartupActor() )
+	const bool bIsRecording = IsRecording();
+	const bool bNetStartup = Actor->IsNetStartupActor();
+	const bool bActorRewindable = Actor->bReplayRewindable;
+
+	if (bActorRewindable && !IsSeamlessTravel && !bIsRecording)
 	{
-		DeletedNetStartupActors.Add( Actor->GetFullName() );		// This is a TSet, so it will only happen once
+		if (bNetStartup || !TrackedRewindActorsByGUID.Contains(GuidCache->NetGUIDLookup.FindRef(Actor)))
+		{
+			// This may happen during playback due to new versions of code playing captures with old versions.
+			// but this should never happen during recording (otherwise it's likely a game code bug). 
+			// We catch that case below.
+			UE_LOG(LogDemo, Warning, TEXT("Replay Rewindable Actor destroyed during playback. Replay may show artifacts (%s)"), *Actor->GetFullName());
+		}
 	}
 
+	if (bIsRecording && bNetStartup)
+	{
+		// We don't want to send any destruction info in this case, because the actor should stick around.
+		// The Replay will manage deleting this when it performs streaming or travel behavior.
+		if (IsSeamlessTravel)
+		{
+			// This is a stripped down version of UNetDriver::NotifyActorDestroy and UActorChannel::Close
+			// combined, and should be kept up to date with those methods.
+		
+			// Remove the actor from the property tracker map
+			RepChangedPropertyTrackerMap.Remove(Actor);
+
+			if (UNetConnection* Connection = ClientConnections[0])
+			{
+				if (Actor->bNetTemporary)
+				{
+					Connection->SentTemporaries.Remove(Actor);
+				}
+
+				if (UActorChannel* Channel = Connection->FindActorChannelRef(Actor))
+				{
+					check(Channel->OpenedLocally);
+					Channel->bClearRecentActorRefs = false;
+					Channel->SetClosingFlag();
+					Channel->Actor = nullptr;
+					Channel->CleanupReplicators(false);
+				}
+	
+				Connection->DormantReplicatorMap.Remove(Actor);
+			}
+	
+			GetNetworkObjectList().Remove(Actor);
+			RenamedStartupActors.Remove(Actor->GetFName());
+			return;
+		}
+		else
+		{
+			// This was deleted due to a game interaction, which isn't supported for Rewindable actors (while recording).
+			// However, since the actor is going to be deleted imminently, we need to track it.
+			ensureMsgf(!bActorRewindable, TEXT("Replay Rewindable Actors should not be destroyed during recording (%s)"), *Actor->GetFullName());
+
+			UE_LOG(LogDemo, VeryVerbose, TEXT("NotifyActyorDestroyed: adding actor to deleted startup list: %s"), *Actor->GetFullName());
+			DeletedNetStartupActors.Add( Actor->GetFullName() );
+		}
+	}
+
+	TUniquePtr<FScopedPacketManager> PacketManager(ConditionallyCreatePacketManager(*Actor->GetLevel()));
 	Super::NotifyActorDestroyed( Actor, IsSeamlessTravel );
 }
 
+void UDemoNetDriver::CleanupOutstandingRewindActors()
+{
+	UWorld* LocalWorld = GetWorld();
+
+	for (const FNetworkGUID& NetGUID : TrackedRewindActorsByGUID)
+	{
+		if (FNetGuidCacheObject* CacheObject = GuidCache->ObjectLookup.Find(NetGUID))
+		{
+			if (AActor* Actor = Cast<AActor>(CacheObject->Object))
+			{
+				// Destroy the actor before removing entries from the GuidCache so its entries are still valid in NotifyActorDestroyed.
+				LocalWorld->DestroyActor(Actor);
+
+				ensureMsgf(GuidCache->NetGUIDLookup.Remove(CacheObject->Object) > 0, TEXT("CleanupOutstandingRewindActors: No entry found for %d in NetGUIDLookup"), NetGUID.Value);
+				GuidCache->ObjectLookup.Remove(NetGUID);
+				CacheObject->bNoLoad = false;
+			}
+			else
+			{
+				UE_LOG(LogDemo, Warning, TEXT("CleanupOutstandingRewindActors - Invalid object for %d, skipping."), NetGUID.Value);
+				continue;
+			}
+		}	
+		else
+		{
+			UE_LOG(LogDemo, Warning, TEXT("CleanupOutstandingRewindActors - CacheObject not found for %s"), NetGUID.Value);
+		}
+	}
+
+	TrackedRewindActorsByGUID.Empty();
+}
+
+void UDemoNetDriver::PreNotifyActorChannelOpen(UActorChannel* Channel, AActor* Actor)
+{
+	const bool bValidChannel = ensureMsgf(Channel, TEXT("PreNotifyActorChannelOpen called with invalid channel"));
+	const bool bValidActor = ensureMsgf(Actor, TEXT("PreNotifyActorChannelOpen called with invalid actor"));
+	
+	// Rewind the actor if necessary.
+	// This should be called before any other notifications / data reach the Actor.
+	if (bValidChannel && bValidActor && TrackedRewindActorsByGUID.Remove(Channel->ActorNetGUID) > 0)
+	{
+		Actor->RewindForReplay();
+	}	
+}
+
+void UDemoNetDriver::NotifyActorLevelUnloaded(AActor* Actor)
+{
+	if (ServerConnection != nullptr)
+	{
+		// This is a combination of the Client and Server logic for destroying a channel,
+		// since we won't actually be sending data back and forth.
+		if (UActorChannel* ActorChannel = ServerConnection->FindActorChannelRef(Actor))
+		{
+			ServerConnection->RemoveActorChannel(Actor);
+			ActorChannel->Actor = nullptr;
+			ActorChannel->ConditionalCleanUp();
+		}
+	}
+
+	Super::NotifyActorLevelUnloaded(Actor);
+}
+
 void UDemoNetDriver::QueueNetStartupActorForRollbackViaDeletion( AActor* Actor )
 {
 	check( Actor != nullptr );
@@ -3413,11 +4737,18 @@ void UDemoNetDriver::QueueNetStartupActorForRollbackViaDeletion( AActor* Actor )
 		return;		// We should only be doing this at runtime while playing a replay
 	}
 
-	if ( RollbackNetStartupActors.Contains( Actor->GetFullName() ) )
+	const FString ActorFullName = Actor->GetFullName();
+	if ( RollbackNetStartupActors.Contains( ActorFullName ) )
 	{
 		return;		// This actor is already queued up
 	}
 
+	if ( Actor->bReplayRewindable )
+	{
+		UE_LOG(LogDemo, Warning, TEXT("Attempted to queue a Replay Rewindable Actor for Rollback Via Deletion. Replay may have artifacts (%s)"), *ActorFullName);
+		return;
+	}
+
 	FRollbackNetStartupActorInfo RollbackActor;
 
 	RollbackActor.Name		= Actor->GetFName();
@@ -3426,7 +4757,24 @@ void UDemoNetDriver::QueueNetStartupActorForRollbackViaDeletion( AActor* Actor )
 	RollbackActor.Rotation	= Actor->GetActorRotation();
 	RollbackActor.Level		= Actor->GetLevel();
 
-	RollbackNetStartupActors.Add( Actor->GetFullName(), RollbackActor );
+	if (CVarDemoSaveRollbackActorState.GetValueOnAnyThread())
+	{
+		TSharedPtr<FObjectReplicator> NewReplicator = MakeShared<FObjectReplicator>();
+		if (NewReplicator.IsValid())
+		{
+			NewReplicator->InitWithObject(Actor->GetArchetype(), ServerConnection, false);
+
+			if (NewReplicator->RepLayout.IsValid() && NewReplicator->RepState.IsValid())
+			{
+				if (NewReplicator->RepLayout->DiffStableProperties(nullptr, &RollbackActor.ObjReferences, NewReplicator->RepState->StaticBuffer.GetData(), Actor))
+				{
+					RollbackActor.RepState = NewReplicator->RepState;
+				}
+			}
+		}
+	}
+
+	RollbackNetStartupActors.Add( ActorFullName, RollbackActor );
 }
 
 /*-----------------------------------------------------------------------------
@@ -3466,12 +4814,12 @@ void UDemoPendingNetGame::LoadMapCompleted(UEngine* Engine, FWorldContext& Conte
 		// If we don't have a world that means we failed loading the new world.
 		// Since there is no world, we must free the net driver ourselves
 		// Technically the pending net game should handle it, but things aren't quite setup properly to handle that either
-		if ( Context.World() == NULL )
+		if ( Context.World() == nullptr )
 		{
 			GEngine->DestroyNamedNetDriver( Context.PendingNetGame, DemoNetDriver->NetDriverName );
 		}
 
-		Context.PendingNetGame = NULL;
+		Context.PendingNetGame = nullptr;
 
 		GEngine->BrowseToDefaultMap( Context );
 
