Copying //UE4/Dev-Networking to //UE4/Dev-Main (Source: //UE4/Dev-Networking @ 4051526)

diff --git a/Engine/Source/Runtime/Engine/Private/HUD.cpp b/Engine/Source/Runtime/Engine/Private/HUD.cpp
--- a/Engine/Source/Runtime/Engine/Private/HUD.cpp
+++ b/Engine/Source/Runtime/Engine/Private/HUD.cpp
@@ -38,6 +38,7 @@ DEFINE_LOG_CATEGORY_STATIC(LogHUD, Log, All);
 #define LOCTEXT_NAMESPACE "HUD"
 
 FOnShowDebugInfo AHUD::OnShowDebugInfo;
+FOnHUDPostRender AHUD::OnHUDPostRender;
 
 // Should we visualize the safe zone? (and if so, title or action?)
 TAutoConsoleVariable<int32> GSafeZoneVisualizationModeCVar(
@@ -231,6 +232,8 @@ void AHUD::PostRender()
 
 	DrawSafeZoneOverlay();
 
+	OnHUDPostRender.Broadcast(this, DebugCanvas);
+
 	LastHUDRenderTime = GetWorld()->TimeSeconds;
 }
 
diff --git a/Engine/Source/Runtime/Engine/Private/LevelActor.cpp b/Engine/Source/Runtime/Engine/Private/LevelActor.cpp
--- a/Engine/Source/Runtime/Engine/Private/LevelActor.cpp
+++ b/Engine/Source/Runtime/Engine/Private/LevelActor.cpp
@@ -447,9 +447,6 @@ AActor* UWorld::SpawnActor( UClass* Class, FTransform const* UserTransformPtr, c
 	LevelToSpawnIn->Actors.Add( Actor );
 	LevelToSpawnIn->ActorsForGC.Add(Actor);
 
-	// Add this newly spawned actor to the network actor list
-	AddNetworkActor( Actor );
-
 #if PERF_SHOW_MULTI_PAWN_SPAWN_FRAMES
 	if( Cast<APawn>(Actor) )
 	{
@@ -485,6 +482,9 @@ AActor* UWorld::SpawnActor( UClass* Class, FTransform const* UserTransformPtr, c
 	}
 #endif
 
+	// Add this newly spawned actor to the network actor list. Do this after PostSpawnInitialize so that actor has "finished" spawning.
+	AddNetworkActor( Actor );
+
 	return Actor;
 }
 
diff --git a/Engine/Source/Runtime/Engine/Private/LevelTick.cpp b/Engine/Source/Runtime/Engine/Private/LevelTick.cpp
--- a/Engine/Source/Runtime/Engine/Private/LevelTick.cpp
+++ b/Engine/Source/Runtime/Engine/Private/LevelTick.cpp
@@ -113,6 +113,7 @@ DEFINE_STAT(STAT_RuntimeMovieSceneTickTime);
 DEFINE_STAT(STAT_FinishAsyncTraceTickTime);
 DEFINE_STAT(STAT_NetBroadcastTickTime);
 DEFINE_STAT(STAT_NetServerRepActorsTime);
+DEFINE_STAT(STAT_NetServerGatherPrioritizeRepActorsTime);
 DEFINE_STAT(STAT_NetConsiderActorsTime);
 DEFINE_STAT(STAT_NetUpdateUnmappedObjectsTime);
 DEFINE_STAT(STAT_NetInitialDormantCheckTime);
diff --git a/Engine/Source/Runtime/Engine/Private/NetConnection.cpp b/Engine/Source/Runtime/Engine/Private/NetConnection.cpp
--- a/Engine/Source/Runtime/Engine/Private/NetConnection.cpp
+++ b/Engine/Source/Runtime/Engine/Private/NetConnection.cpp
@@ -30,6 +30,8 @@
 #include "GameDelegates.h"
 #include "Misc/PackageName.h"
 #include "UObject/LinkerLoad.h"
+#include "UObject/ObjectKey.h"
+#include "UObject/UObjectIterator.h"
 
 #if !UE_BUILD_SHIPPING
 static TAutoConsoleVariable<int32> CVarPingExcludeFrameTime( TEXT( "net.PingExcludeFrameTime" ), 0, TEXT( "Calculate RTT time between NIC's of server and client." ) );
@@ -497,7 +499,7 @@ void UNetConnection::CleanUp()
 		else
 		{
 			check(Driver->ServerConnection == NULL);
-			verify(Driver->ClientConnections.Remove(this) == 1);
+			Driver->RemoveClientConnection(this);
 
 #if USE_SERVER_PERF_COUNTERS
 			if (IPerfCountersModule::IsAvailable())
@@ -755,19 +757,21 @@ void UNetConnection::UpdateLevelVisibility(const FName& PackageName, bool bIsVis
 
 			// Any destroyed actors that were destroyed prior to the streaming level being unloaded for the client will not be in the connections
 			// destroyed actors list when the level is reloaded, so seek them out and add in
-			for (const TPair<FNetworkGUID, FActorDestructionInfo>& DestroyedPair : Driver->DestroyedStartupOrDormantActors)
+			for (const auto& DestroyedPair : Driver->DestroyedStartupOrDormantActors)
 			{
-				if (DestroyedPair.Value.StreamingLevelName == PackageName)
+				if (DestroyedPair.Value->StreamingLevelName == PackageName)
 				{
-					DestroyedStartupOrDormantActors.Add(DestroyedPair.Key);
+					AddDestructionInfo(DestroyedPair.Value.Get());
 				}
 			}
 
 			// Any dormant actor that has changes flushed or made before going dormant needs to be updated on the client 
 			// when the streaming level is loaded, so mark them active for this connection
+			UWorld* LevelWorld = nullptr;
 			if (TempPkg)
 			{
-				if (UWorld* LevelWorld = (UWorld*)FindObjectWithOuter(TempPkg, UWorld::StaticClass()))
+				LevelWorld = (UWorld*)FindObjectWithOuter(TempPkg, UWorld::StaticClass());
+				if (LevelWorld)
 				{
 					if (LevelWorld->PersistentLevel)
 					{
@@ -785,6 +789,12 @@ void UNetConnection::UpdateLevelVisibility(const FName& PackageName, bool bIsVis
 					}
 				}
 			}
+
+			if (ReplicationConnectionDriver)
+			{
+				ReplicationConnectionDriver->NotifyClientVisibleLevelNamesAdd(PackageName, LevelWorld);
+			}
+
 		}
 		else
 		{
@@ -796,6 +806,10 @@ void UNetConnection::UpdateLevelVisibility(const FName& PackageName, bool bIsVis
 	{
 		ClientVisibleLevelNames.Remove(PackageName);
 		UE_LOG( LogPlayerController, Verbose, TEXT("ServerUpdateLevelVisibility() Removed '%s'"), *PackageName.ToString() );
+		if (ReplicationConnectionDriver)
+		{
+			ReplicationConnectionDriver->NotifyClientVisibleLevelNamesRemove(PackageName);
+		}
 			
 		// Close any channels now that have actors that were apart of the level the client just unloaded
 		for ( auto It = ActorChannels.CreateIterator(); It; ++It )
@@ -957,6 +971,7 @@ void UNetConnection::ReceivedRawPacket( void* InData, int32 Count )
 	}
 }
 
+uint32 GNetOutBytes = 0;
 
 void UNetConnection::FlushNet(bool bIgnoreSimulation)
 {
@@ -1092,6 +1107,7 @@ void UNetConnection::FlushNet(bool bIgnoreSimulation)
 
 		OutBytes += PacketBytes;
 		Driver->OutBytes += PacketBytes;
+		GNetOutBytes += PacketBytes;
 		InitSendBuffer();
 	}
 
@@ -1183,6 +1199,7 @@ void UNetConnection::ReceivedPacket( FBitReader& Reader )
 		return;
 	}
 
+	const bool bIgnoreRPCs = Driver->ShouldIgnoreRPCs();
 
 	bool bSkipAck = false;
 
@@ -1483,6 +1500,50 @@ void UNetConnection::ReceivedPacket( FBitReader& Reader )
 			// Receiving data.
 			UChannel* Channel = Channels[Bunch.ChIndex];
 
+			// We're on a 100% reliable connection and we are rolling back some data.
+			// In that case, we can generally ignore these bunches.
+			if (InternalAck && Channel && bIgnoreAlreadyOpenedChannels)
+			{
+				// This was an open bunch for a channel that's already opened.
+				// We can ignore future bunches from this channel.
+				const bool bNewlyOpened = Bunch.bOpen && (!Bunch.bPartial || Bunch.bPartialInitial);
+				if (bNewlyOpened || IgnoringChannels.Contains(Bunch.ChIndex))
+				{
+					if (bNewlyOpened)
+					{
+						// NOTE: This could break if there are sufficient enough MustBeMapped GUIDs to have
+						//			prevented the actor GUID from being written.
+						//			Don't think that's currently handle by the normal system either, though.
+						if (Bunch.bHasMustBeMappedGUIDs)
+						{
+							// Skip passed any must be mapped guids.
+							uint16 NumMustBeMappedGUIDs = 0;
+							Bunch << NumMustBeMappedGUIDs;
+							for (int32 i = 0; i < NumMustBeMappedGUIDs; i++)
+							{
+								FNetworkGUID TempNetGUID;
+								Bunch << TempNetGUID;
+							}
+						}
+
+						FNetworkGUID ActorGUID;
+						Bunch << ActorGUID;
+						IgnoringChannels.Add(Bunch.ChIndex, ActorGUID);
+					}
+
+					if (Bunch.bClose && (!Bunch.bPartial || Bunch.bPartialFinal))
+					{
+						FNetworkGUID ActorGUID = IgnoringChannels.FindAndRemoveChecked(Bunch.ChIndex);
+						if (ActorGUID.IsStatic())
+						{
+							AActor* StaticActor = Cast<AActor>(Driver->GuidCache->GetObjectFromNetGUID(ActorGUID, false));
+							Driver->World->DestroyActor(StaticActor);
+						}
+					}
+					continue;
+				}
+			}
+
 			// Ignore if reliable packet has already been processed.
 			if ( Bunch.bReliable && Bunch.ChSequence <= InReliable[Bunch.ChIndex] )
 			{
@@ -1566,6 +1627,8 @@ void UNetConnection::ReceivedPacket( FBitReader& Reader )
 				}
 			}
 
+			Bunch.bIgnoreRPCs = bIgnoreRPCs;
+
 			// Dispatch the raw, unsequenced bunch to the channel.
 			bool bLocalSkipAck = false;
 			Channel->ReceivedRawBunch( Bunch, bLocalSkipAck ); //warning: May destroy channel.
@@ -1596,6 +1659,13 @@ void UNetConnection::ReceivedPacket( FBitReader& Reader )
 	}
 }
 
+void UNetConnection::SetIgnoreAlreadyOpenedChannels(bool bInIgnoreAlreadyOpenedChannels)
+{
+	check(InternalAck);
+	bIgnoreAlreadyOpenedChannels = bInIgnoreAlreadyOpenedChannels;
+	IgnoringChannels.Reset();
+}
+
 int32 UNetConnection::WriteBitsToSendBuffer( 
 	const uint8 *	Bits, 
 	const int32		SizeInBits, 
@@ -2434,7 +2504,7 @@ bool UNetConnection::ShouldReplicateVoicePacketFrom(const FUniqueNetId& Sender)
 void UNetConnection::ResetGameWorldState()
 {
 	//Clear out references and do whatever else so that nothing holds onto references that it doesn't need to.
-	DestroyedStartupOrDormantActors.Empty();
+	ResetDestructionInfos();
 	ClientVisibleLevelNames.Empty();
 	KeepProcessingActorChannelBunchesMap.Empty();
 	DormantReplicatorMap.Empty();
@@ -2469,7 +2539,7 @@ void UNetConnection::FlushDormancy(class AActor* Actor)
 	// here, the server will not add the actor to the dormancy list when he closes the channel 
 	// after he gets the client ack. The result is the channel will close but be open again
 	// right away
-	UActorChannel* Ch = ActorChannels.FindRef(Actor);
+	UActorChannel* Ch = FindActorChannelRef(Actor);
 
 	if ( Ch != nullptr )
 	{
@@ -2483,7 +2553,7 @@ void UNetConnection::FlushDormancy(class AActor* Actor)
 
 void UNetConnection::ForcePropertyCompare( AActor* Actor )
 {
-	UActorChannel* Ch = ActorChannels.FindRef( Actor );
+	UActorChannel* Ch = FindActorChannelRef( Actor );
 
 	if ( Ch != nullptr )
 	{
@@ -2666,3 +2736,179 @@ void UNetConnection::SetPlayerOnlinePlatformName(const FName InPlayerOnlinePlatf
 {
 	PlayerOnlinePlatformName = InPlayerOnlinePlatformName;
 }
+
+/*-----------------------------------------------------------------------------
+	USimulatedClientNetConnection.
+-----------------------------------------------------------------------------*/
+
+USimulatedClientNetConnection::USimulatedClientNetConnection( const FObjectInitializer& ObjectInitializer ) : Super( ObjectInitializer )
+{
+	InternalAck = true;
+}
+
+void USimulatedClientNetConnection::HandleClientPlayer( class APlayerController* PC, class UNetConnection* NetConnection )
+{
+	State = USOCK_Open;
+	PlayerController = PC;
+	OwningActor = PC;
+}
+
+// ----------------------------------------------------------------
+
+static void	AddSimulatedNetConnections(const TArray<FString>& Args, UWorld* World)
+{
+	int32 ConnectionCount = 99;
+	if (Args.Num() > 0)
+	{
+		Lex::FromString(ConnectionCount, *Args[0]);
+	}
+
+	// Search for server game net driver. Do it this way so we can cheat in PIE
+	UNetDriver* BestNetDriver = nullptr;
+	for (TObjectIterator<UNetDriver> NetDriverIt; NetDriverIt; ++NetDriverIt)
+	{
+		if (NetDriverIt->NetDriverName == NAME_GameNetDriver && NetDriverIt->IsServer())
+		{
+			BestNetDriver = *NetDriverIt;
+			break;
+		}
+	}
+
+	if (!BestNetDriver)
+	{
+		return;
+	}
+
+	AActor* DefaultViewTarget = nullptr;
+	APlayerController* PC = nullptr;
+	for (auto Iterator = BestNetDriver->GetWorld()->GetPlayerControllerIterator(); Iterator; ++Iterator)
+	{
+		PC = Cast<APlayerController>(*Iterator);
+		if (PC)
+		{
+			DefaultViewTarget = PC->GetViewTarget();
+			break;
+		}
+	}
+	
+
+	UE_LOG(LogNet, Display, TEXT("Adding %d Simulated Connections..."), ConnectionCount);
+	while(ConnectionCount-- > 0)
+	{
+		USimulatedClientNetConnection* Connection = NewObject<USimulatedClientNetConnection>();
+		Connection->InitConnection( BestNetDriver, USOCK_Open, BestNetDriver->GetWorld()->URL, 1000000 );
+		Connection->InitSendBuffer();
+		BestNetDriver->AddClientConnection( Connection );
+		Connection->HandleClientPlayer(PC, Connection);
+		Connection->SetClientWorldPackageName(BestNetDriver->GetWorldPackage()->GetFName());
+	}	
+}
+
+FAutoConsoleCommandWithWorldAndArgs AddimulatedConnectionsCmd(TEXT("net.SimulateConnections"), TEXT("Starts a Simulated Net Driver"),	FConsoleCommandWithWorldAndArgsDelegate::CreateStatic(AddSimulatedNetConnections) );
+
+// ----------------------------------------------------------------
+
+
+static void	PrintActorReportFunc(const TArray<FString>& Args, UWorld* InWorld)
+{
+	// Search for server game net driver. Do it this way so we can cheat in PIE
+	UNetDriver* BestNetDriver = nullptr;
+	for (TObjectIterator<UNetDriver> NetDriverIt; NetDriverIt; ++NetDriverIt)
+	{
+		if (NetDriverIt->NetDriverName == NAME_GameNetDriver && NetDriverIt->IsServer())
+		{
+			BestNetDriver = *NetDriverIt;
+			break;
+		}
+	}
+
+	int32 TotalCount = 0;
+	
+	TMap<UClass*, int32> ClassCount;
+	TMap<UClass*, int32> ActualClassCount;
+	TMap<ENetDormancy, int32> DormancyCount;
+	FBox BoundingBox;
+
+	TMap<AActor*, int32> RawActorPtrMap;
+	TMap<TWeakObjectPtr<AActor>, int32> WeakPtrMap;
+	TMap<FObjectKey, int32> ObjKeyMap;
+
+	UWorld* World = BestNetDriver ? BestNetDriver->GetWorld() : InWorld;
+	if (!World)
+		return;
+
+	
+	for (TActorIterator<AActor> It(World); It; ++It)
+	{
+		AActor* Actor = *It;
+		if (Actor->GetIsReplicated() == false)
+		{
+			continue;
+		}
+
+		TotalCount++;
+		DormancyCount.FindOrAdd(Actor->NetDormancy)++;
+
+		BoundingBox += Actor->GetActorLocation();
+
+		UClass* CurrentClass = Actor->GetClass();
+
+		ActualClassCount.FindOrAdd(CurrentClass)++;
+
+		while(CurrentClass)
+		{
+			ClassCount.FindOrAdd(CurrentClass)++;
+			CurrentClass = CurrentClass->GetSuperClass();
+		}
+
+		RawActorPtrMap.Add(Actor) = FMath::Rand();
+		WeakPtrMap.Add(Actor) = FMath::Rand();
+		ObjKeyMap.Add(FObjectKey(Actor)) = FMath::Rand();
+	}
+
+	ClassCount.ValueSort(TGreater<int32>());
+	ActualClassCount.ValueSort(TGreater<int32>());
+
+	UE_LOG(LogNet, Display, TEXT("Class Count (includes inheritance)"));
+	for (auto MapIt : ClassCount)
+	{
+		UE_LOG(LogNet, Display, TEXT("%s - %d"), *GetNameSafe(MapIt.Key), MapIt.Value);
+	}
+
+
+	UE_LOG(LogNet, Display, TEXT(""));
+	UE_LOG(LogNet, Display, TEXT("Class Count (actual clases)"));
+	for (auto MapIt : ActualClassCount)
+	{
+		UE_LOG(LogNet, Display, TEXT("%s - %d"), *GetNameSafe(MapIt.Key), MapIt.Value);
+	}
+
+	UE_LOG(LogNet, Display, TEXT(""));
+	UE_LOG(LogNet, Display, TEXT("Complete Bounding Box: %s"), *BoundingBox.ToString());
+	UE_LOG(LogNet, Display, TEXT("                 Size: %s"), *BoundingBox.GetSize().ToString());
+
+	UE_LOG(LogNet, Display, TEXT(""));
+
+	for (auto MapIt : DormancyCount)
+	{
+		UE_LOG(LogNet, Display, TEXT("%s - %d"), *UEnum::GetValueAsString(TEXT("/Script/Engine.ENetDormancy"), MapIt.Key), MapIt.Value);
+	}
+
+	UE_LOG(LogNet, Display, TEXT(""));
+	UE_LOG(LogNet, Display, TEXT("Total Replicated Actor Count: %d"), TotalCount);
+
+
+	UE_LOG(LogNet, Display, TEXT(""));
+	UE_LOG(LogNet, Display, TEXT("Raw Actor Map: "));
+	RawActorPtrMap.Dump(*GLog);
+
+	UE_LOG(LogNet, Display, TEXT(""));
+	UE_LOG(LogNet, Display, TEXT("Weak Ptr Map: "));
+	WeakPtrMap.Dump(*GLog);
+
+	UE_LOG(LogNet, Display, TEXT(""));
+	UE_LOG(LogNet, Display, TEXT("ObjectKey Map: "));
+	ObjKeyMap.Dump(*GLog);
+}
+
+FAutoConsoleCommandWithWorldAndArgs PrintActorReportCmd(TEXT("net.ActorReport"), TEXT(""),	FConsoleCommandWithWorldAndArgsDelegate::CreateStatic(PrintActorReportFunc) );
\ No newline at end of file
diff --git a/Engine/Source/Runtime/Engine/Private/NetworkDriver.cpp b/Engine/Source/Runtime/Engine/Private/NetworkDriver.cpp
--- a/Engine/Source/Runtime/Engine/Private/NetworkDriver.cpp
+++ b/Engine/Source/Runtime/Engine/Private/NetworkDriver.cpp
@@ -58,7 +58,8 @@
 #include "GameFramework/PlayerState.h"
 #include "Net/PerfCountersHelpers.h"
 #include "Stats/StatsMisc.h"
-
+#include "Engine/ReplicationDriver.h"
+#include "ProfilingDebugging/CsvProfiler.h"
 
 #if USE_SERVER_PERF_COUNTERS
 #include "PerfCountersModule.h"
@@ -68,6 +69,7 @@
 #include "Editor.h"
 #endif
 
+
 // Default net driver stats
 DEFINE_STAT(STAT_Ping);
 DEFINE_STAT(STAT_Channels);
@@ -96,10 +98,9 @@ DEFINE_STAT(STAT_OutLoss);
 DEFINE_STAT(STAT_InLoss);
 DEFINE_STAT(STAT_NumConsideredActors);
 DEFINE_STAT(STAT_PrioritizedActors);
-DEFINE_STAT(STAT_NumRelevantActors);
-DEFINE_STAT(STAT_NumRelevantDeletedActors);
-DEFINE_STAT(STAT_NumReplicatedActorAttempts);
 DEFINE_STAT(STAT_NumReplicatedActors);
+DEFINE_STAT(STAT_NumReplicatedActorBytes);
+DEFINE_STAT(STAT_NumRelevantDeletedActors);
 DEFINE_STAT(STAT_NumActorChannels);
 DEFINE_STAT(STAT_NumActors);
 DEFINE_STAT(STAT_NumNetActors);
@@ -141,6 +142,10 @@ DECLARE_CYCLE_STAT(TEXT("NetDriver TickFlush"), STAT_NetTickFlush, STATGROUP_Gam
 DECLARE_CYCLE_STAT(TEXT("NetDriver TickFlush GatherStats"), STAT_NetTickFlushGatherStats, STATGROUP_Game);
 DECLARE_CYCLE_STAT(TEXT("NetDriver TickFlush GatherStatsPerfCounters"), STAT_NetTickFlushGatherStatsPerfCounters, STATGROUP_Game);
 
+int32 GNumSaturatedConnections; // Counter for how many connections are skipped/early out due to bandwidth saturation
+int32 GNumSharedSerializationHit;
+int32 GNumSharedSerializationMiss;
+
 #if UE_BUILD_SHIPPING
 #define DEBUG_REMOTEFUNCTION(Format, ...)
 #else
@@ -276,6 +281,7 @@ void UNetDriver::PostInitProperties()
 #endif // WITH_EDITOR
 	
 		OnLevelRemovedFromWorldHandle = FWorldDelegates::LevelRemovedFromWorld.AddUObject(this, &UNetDriver::OnLevelRemovedFromWorld);
+		OnLevelAddedToWorldHandle = FWorldDelegates::LevelAddedToWorld.AddUObject(this, &UNetDriver::OnLevelAddedToWorld);
 	}
 }
 
@@ -289,10 +295,25 @@ void UNetDriver::AssertValid()
 	return bUseAdapativeNetFrequency;
 }
 
+void UNetDriver::AddNetworkActor(AActor* Actor)
+{
+	GetNetworkObjectList().FindOrAdd(Actor, NetDriverName);
+	if (ReplicationDriver)
+	{
+		ReplicationDriver->AddNetworkActor(Actor);
+	}
+}
+
 FNetworkObjectInfo* UNetDriver::FindOrAddNetworkObjectInfo(const AActor* InActor)
 {
-	if (TSharedPtr<FNetworkObjectInfo>* InfoPtr = GetNetworkObjectList().FindOrAdd(const_cast<AActor*>(InActor), NetDriverName))
+	bool bWasAdded = false;
+	if (TSharedPtr<FNetworkObjectInfo>* InfoPtr = GetNetworkObjectList().FindOrAdd(const_cast<AActor*>(InActor), NetDriverName, &bWasAdded))
 	{
+		if (bWasAdded && ReplicationDriver)
+		{
+			ReplicationDriver->AddNetworkActor(const_cast<AActor*>(InActor));
+		}
+
 		return InfoPtr->Get();
 	}
 
@@ -528,7 +549,7 @@ void UNetDriver::TickFlush(float DeltaSeconds)
 				// Handle these stats to do this differently for client & server
 				if (ServerConnection != nullptr)
 				{
-					NumActorChannels = ServerConnection->ActorChannels.Num();
+					NumActorChannels = ServerConnection->ActorChannelsNum();
 					NumDormantActors = GetNetworkObjectList().GetNumDormantActorsForConnection(ServerConnection);
 #if STATS
 					ServerConnection->PackageMap->GetNetGUIDStats(AckCount, UnAckCount, PendingCount);
@@ -540,7 +561,7 @@ void UNetDriver::TickFlush(float DeltaSeconds)
 					int32 SharedDormantActors = GetNetworkObjectList().GetDormantObjectsOnAllConnections().Num();
 					for (int32 i = 0; i < ClientConnections.Num(); i++)
 					{
-						NumActorChannels += ClientConnections[i]->ActorChannels.Num();
+						NumActorChannels += ClientConnections[i]->ActorChannelsNum();
 						NumDormantActors += GetNetworkObjectList().GetNumDormantActorsForConnection(ClientConnections[i]);
 #if STATS
 						int32 ClientAckCount = 0;
@@ -1183,6 +1204,20 @@ bool UNetDriver::InitConnectionClass(void)
 	return NetConnectionClass != NULL;
 }
 
+bool UNetDriver::InitReplicationDriverClass()
+{
+	if (ReplicationDriverClass == nullptr && !ReplicationDriverClassName.IsEmpty())
+	{
+		ReplicationDriverClass = LoadClass<UReplicationDriver>(nullptr, *ReplicationDriverClassName, nullptr, LOAD_None, nullptr);
+		if (ReplicationDriverClass == nullptr)
+		{
+			UE_LOG(LogNet, Error,TEXT("Failed to load class '%s'"), *ReplicationDriverClassName);
+		}
+	}
+
+	return ReplicationDriverClass != nullptr;
+}
+
 bool UNetDriver::InitBase(bool bInitAsClient, FNetworkNotify* InNotify, const FURL& URL, bool bReuseAddressAndPort, FString& Error)
 {
 	// Read any timeout overrides from the URL
@@ -1215,8 +1250,10 @@ bool UNetDriver::InitBase(bool bInitAsClient, FNetworkNotify* InNotify, const FU
 	if (!bInitAsClient)
 	{
 		ConnectionlessHandler.Reset(nullptr);
-
 		InitDestroyedStartupActors();
+
+		InitReplicationDriverClass();
+		SetReplicationDriver(UReplicationDriver::CreateReplicationDriver(this, URL, GetWorld()));
 	}
 
 	Notify = InNotify;
@@ -1355,6 +1392,8 @@ void UNetDriver::Shutdown()
 #if DO_ENABLE_NET_TEST
 	PacketSimulationSettings.UnregisterCommands();
 #endif
+
+	SetReplicationDriver(nullptr);
 }
 
 bool UNetDriver::IsServer() const
@@ -1478,7 +1517,7 @@ void UNetDriver::InternalProcessRemoteFunction
 	}
 		
 	// Get the actor channel.
-	UActorChannel* Ch = Connection->ActorChannels.FindRef(Actor);
+	UActorChannel* Ch = Connection->FindActorChannelRef(Actor);
 	if( !Ch )
 	{
 		if( IsServer )
@@ -1509,19 +1548,33 @@ void UNetDriver::InternalProcessRemoteFunction
 		}	
 	}
 
+	ProcessRemoteFunctionForChannel(Ch, ClassCache, FieldCache, TargetObj, Connection, Function, Parms, OutParms, Stack, IsServer);
+}
+
+void UNetDriver::ProcessRemoteFunctionForChannel(UActorChannel* Ch, const FClassNetCache* ClassCache, const FFieldNetCache* FieldCache, UObject* TargetObj, UNetConnection* Connection, UFunction* Function, void* Parms, FOutParmRec* OutParms, FFrame* Stack, const bool IsServer, const ERemoteFunctionSendPolicy SendPolicy)
+{
+	// ---------------------------------------------------------------------------------------------------------------------------------------------------------------------------------------------------------------------------
+	// This function should be kept fast! Assume this is getting called multiple times at once. Don't look things up/recalc them if they do not change per connection/actor.
+	// ---------------------------------------------------------------------------------------------------------------------------------------------------------------------------------------------------------------------------
+
+	if (Ch->Closing)
+	{
+		return;
+	}
+
 	// Make sure initial channel-opening replication has taken place.
 	if( Ch->OpenPacketId.First==INDEX_NONE )
 	{
 		if (!IsServer)
 		{
-			DEBUG_REMOTEFUNCTION(TEXT("Initial channel replication has not occurred, not calling %s::%s"), *Actor->GetName(), *Function->GetName());
+			DEBUG_REMOTEFUNCTION(TEXT("Initial channel replication has not occurred, not calling %s::%s"), *TargetObj->GetFullName(), *Function->GetName());
 			return;
 		}
 
 		// triggering replication of an Actor while already in the middle of replication can result in invalid data being sent and is therefore illegal
 		if (Ch->bIsReplicatingActor)
 		{
-			FString Error(FString::Printf(TEXT("Attempt to replicate function '%s' on Actor '%s' while it is in the middle of variable replication!"), *Function->GetName(), *Actor->GetName()));
+			FString Error(FString::Printf(TEXT("Attempt to replicate function '%s' on Actor '%s' while it is in the middle of variable replication!"), *Function->GetName(), *TargetObj->GetFullName()));
 			UE_LOG(LogScript, Error, TEXT("%s"), *Error);
 			ensureMsgf(false, TEXT("%s"), *Error);
 			return;
@@ -1558,12 +1611,12 @@ void UNetDriver::InternalProcessRemoteFunction
 		if ( !Bunch.bReliable )
 		{
 			// Not reliable, so not fatal. This can happen a lot in debug builds at startup if client is slow to get in game
-			UE_LOG( LogNet, Warning, TEXT( "Can't send function '%s' on '%s': Reliable buffer overflow. FieldCache->FieldNetIndex: %d Max %d. Ch MaxPacket: %d" ), *Function->GetName(), *Actor->GetName(), FieldCache->FieldNetIndex, ClassCache->GetMaxIndex(), Ch->Connection->MaxPacket );
+			UE_LOG( LogNet, Warning, TEXT( "Can't send function '%s' on '%s': Reliable buffer overflow. FieldCache->FieldNetIndex: %d Max %d. Ch MaxPacket: %d" ), *Function->GetName(), *TargetObj->GetFullName(), FieldCache->FieldNetIndex, ClassCache->GetMaxIndex(), Ch->Connection->MaxPacket );
 		}
 		else
 		{
 			// The connection has overflowed the reliable buffer. We cannot recover from this. Disconnect this user.
-			UE_LOG( LogNet, Warning, TEXT( "Closing connection. Can't send function '%s' on '%s': Reliable buffer overflow. FieldCache->FieldNetIndex: %d Max %d. Ch MaxPacket: %d." ), *Function->GetName(), *Actor->GetName(), FieldCache->FieldNetIndex, ClassCache->GetMaxIndex(), Ch->Connection->MaxPacket );
+			UE_LOG( LogNet, Warning, TEXT( "Closing connection. Can't send function '%s' on '%s': Reliable buffer overflow. FieldCache->FieldNetIndex: %d Max %d. Ch MaxPacket: %d." ), *Function->GetName(), *TargetObj->GetFullName(), FieldCache->FieldNetIndex, ClassCache->GetMaxIndex(), Ch->Connection->MaxPacket );
 
 			FString ErrorMsg = NSLOCTEXT( "NetworkErrors", "ClientReliableBufferOverflow", "Outgoing reliable buffer overflow" ).ToString();
 			FNetControlMessage<NMT_Failure>::Send( Connection, ErrorMsg );
@@ -1581,7 +1634,7 @@ void UNetDriver::InternalProcessRemoteFunction
 
 	if ( CVarNetReliableDebug.GetValueOnAnyThread() > 0 )
 	{
-		Bunch.DebugString = FString::Printf( TEXT( "%.2f RPC: %s - %s" ), Connection->Driver->Time, *Actor->GetName(), *Function->GetName() );
+		Bunch.DebugString = FString::Printf( TEXT( "%.2f RPC: %s - %s" ), Connection->Driver->Time, *TargetObj->GetFullName(), *Function->GetName() );
 	}
 #endif
 
@@ -1597,7 +1650,7 @@ void UNetDriver::InternalProcessRemoteFunction
 			{
 				if ( OutParms == NULL )
 				{
-					UE_LOG( LogNet, Warning, TEXT( "Missing OutParms. Property: %s, Function: %s, Actor: %s" ), *It->GetName(), *Function->GetName(), *Actor->GetName() );
+					UE_LOG( LogNet, Warning, TEXT( "Missing OutParms. Property: %s, Function: %s, Actor: %s" ), *It->GetName(), *Function->GetName(), *TargetObj->GetFullName() );
 					continue;
 				}
 
@@ -1640,7 +1693,7 @@ void UNetDriver::InternalProcessRemoteFunction
 
 	// Use the replication layout to send the rpc parameter values
 	TSharedPtr<FRepLayout> RepLayout = GetFunctionRepLayout( Function );
-	RepLayout->SendPropertiesForRPC( Actor, Function, Ch, TempWriter, Parms );
+	RepLayout->SendPropertiesForRPC( Function, Ch, TempWriter, Parms );
 
 	if (TempWriter.IsError())
 	{
@@ -1660,9 +1713,20 @@ void UNetDriver::InternalProcessRemoteFunction
 
 		int32 HeaderBits	= 0;
 		int32 ParameterBits	= 0;
-
-		// Queue unreliable multicast 
-		const bool QueueBunch = ( !Bunch.bReliable && Function->FunctionFlags & FUNC_NetMulticast );
+		
+		bool QueueBunch = false;
+		switch (SendPolicy)
+		{
+			case ERemoteFunctionSendPolicy::Default:
+				QueueBunch = ( !Bunch.bReliable && Function->FunctionFlags & FUNC_NetMulticast );
+				break;
+			case ERemoteFunctionSendPolicy::ForceQueue:
+				QueueBunch = true;
+				break;
+			case ERemoteFunctionSendPolicy::ForceSend:
+				QueueBunch = false;
+				break;
+		}
 
 		if ( QueueBunch )
 		{
@@ -1704,28 +1768,28 @@ void UNetDriver::InternalProcessRemoteFunction
 				// Unreliable multicast functions are queued and sent out during property replication
 				if (LogAsWarning)
 				{
-					UE_LOG(LogNetTraffic, Warning,	TEXT("      Queing unreliable multicast RPC: %s::%s [%.1f bytes]"), *Actor->GetName(), *Function->GetName(), Bunch.GetNumBits() / 8.f );
+					UE_LOG(LogNetTraffic, Warning,	TEXT("      Queing unreliable multicast RPC: %s::%s [%.1f bytes]"), *TargetObj->GetFullName(), *Function->GetName(), Bunch.GetNumBits() / 8.f );
 				}
 				else
 				{
-					UE_LOG(LogNetTraffic, Log,		TEXT("      Queing unreliable multicast RPC: %s::%s [%.1f bytes]"), *Actor->GetName(), *Function->GetName(), Bunch.GetNumBits() / 8.f );
+					UE_LOG(LogNetTraffic, Log,		TEXT("      Queing unreliable multicast RPC: %s::%s [%.1f bytes]"), *TargetObj->GetFullName(), *Function->GetName(), Bunch.GetNumBits() / 8.f );
 				}
 
-				NETWORK_PROFILER(GNetworkProfiler.TrackQueuedRPC(Connection, TargetObj, Actor, Function, HeaderBits, ParameterBits, 0));
+				NETWORK_PROFILER(GNetworkProfiler.TrackQueuedRPC(Connection, TargetObj, Ch->Actor, Function, HeaderBits, ParameterBits, 0));
 				Ch->QueueRemoteFunctionBunch(TargetObj, Function, Bunch);
 			}
 			else
 			{
 				if (LogAsWarning)
 				{
-					UE_LOG(LogNetTraffic, Warning,	TEXT("      Sent RPC: %s::%s [%.1f bytes]"), *Actor->GetName(), *Function->GetName(), Bunch.GetNumBits() / 8.f );
+					UE_LOG(LogNetTraffic, Warning,	TEXT("      Sent RPC: %s::%s [%.1f bytes]"), *TargetObj->GetFullName(), *Function->GetName(), Bunch.GetNumBits() / 8.f );
 				}
 				else
 				{
-					UE_LOG(LogNetTraffic, Log,		TEXT("      Sent RPC: %s::%s [%.1f bytes]"), *Actor->GetName(), *Function->GetName(), Bunch.GetNumBits() / 8.f );
+					UE_LOG(LogNetTraffic, Log,		TEXT("      Sent RPC: %s::%s [%.1f bytes]"), *TargetObj->GetFullName(), *Function->GetName(), Bunch.GetNumBits() / 8.f );
 				}
 
-				NETWORK_PROFILER(GNetworkProfiler.TrackSendRPC(Actor, Function, HeaderBits, ParameterBits, 0, Connection));
+				NETWORK_PROFILER(GNetworkProfiler.TrackSendRPC(Ch->Actor, Function, HeaderBits, ParameterBits, 0, Connection));
 				Ch->SendBunch( &Bunch, 1 );
 			}
 		}
@@ -1868,6 +1932,7 @@ void UNetDriver::FinishDestroy()
 		GuidCache.Reset();
 
 		FWorldDelegates::LevelRemovedFromWorld.Remove(OnLevelRemovedFromWorldHandle);
+		FWorldDelegates::LevelAddedToWorld.Remove(OnLevelAddedToWorldHandle);
 	}
 	else
 	{
@@ -2231,7 +2296,13 @@ FActorDestructionInfo* UNetDriver::CreateDestructionInfo( UNetDriver* NetDriver,
 
 	FNetworkGUID NetGUID = NetDriver->GuidCache->GetOrAssignNetGUID( ThisActor );
 
-	FActorDestructionInfo& NewInfo = NetDriver->DestroyedStartupOrDormantActors.FindOrAdd( NetGUID );
+	TUniquePtr<FActorDestructionInfo>& NewInfoPtr = NetDriver->DestroyedStartupOrDormantActors.FindOrAdd( NetGUID );
+	if (NewInfoPtr.IsValid() == false)
+	{
+		NewInfoPtr = MakeUnique<FActorDestructionInfo>();
+	}
+
+	FActorDestructionInfo &NewInfo = *NewInfoPtr;
 	NewInfo.DestroyedPosition = ThisActor->GetActorLocation();
 	NewInfo.NetGUID = NetGUID;
 	NewInfo.Level = ThisActor->GetLevel();
@@ -2277,12 +2348,14 @@ void UNetDriver::NotifyActorDestroyed( AActor* ThisActor, bool IsSeamlessTravel
 			DestructionInfo = CreateDestructionInfo( this, ThisActor, DestructionInfo);
 		}
 
+		const FNetworkObjectInfo* NetworkObjectInfo = GetNetworkObjectList().Find( ThisActor ).Get();
+
 		for( int32 i=ClientConnections.Num()-1; i>=0; i-- )
 		{
 			UNetConnection* Connection = ClientConnections[i];
 			if( ThisActor->bNetTemporary )
 				Connection->SentTemporaries.Remove( ThisActor );
-			UActorChannel* Channel = Connection->ActorChannels.FindRef(ThisActor);
+			UActorChannel* Channel = Connection->FindActorChannelRef(ThisActor);
 			if( Channel )
 			{
 				check(Channel->OpenedLocally);
@@ -2291,7 +2364,6 @@ void UNetDriver::NotifyActorDestroyed( AActor* ThisActor, bool IsSeamlessTravel
 			}
 			else
 			{
-				const FNetworkObjectInfo* NetworkObjectInfo = GetNetworkObjectList().Find( ThisActor ).Get();
 				const bool bDormantOrRecentlyDormant = NetworkObjectInfo && (NetworkObjectInfo->DormantConnections.Contains(Connection) || NetworkObjectInfo->RecentlyDormantConnections.Contains(Connection));
 
 				if (bShouldCreateDestructionInfo || bDormantOrRecentlyDormant)
@@ -2299,9 +2371,10 @@ void UNetDriver::NotifyActorDestroyed( AActor* ThisActor, bool IsSeamlessTravel
 					// Make a new destruction info if necessary. It is necessary if the actor is dormant or recently dormant because
 					// even though the client knew about the actor at some point, it doesn't have a channel to handle destruction.
 					DestructionInfo = CreateDestructionInfo(this, ThisActor, DestructionInfo);
-					Connection->DestroyedStartupOrDormantActors.Add(DestructionInfo->NetGUID);
+					Connection->AddDestructionInfo(DestructionInfo);
 				}
 			}
+						
 
 			// Remove it from any dormancy lists				
 			Connection->DormantReplicatorMap.Remove( ThisActor );
@@ -2309,10 +2382,21 @@ void UNetDriver::NotifyActorDestroyed( AActor* ThisActor, bool IsSeamlessTravel
 	}
 
 	// Remove this actor from the network object list
-	GetNetworkObjectList().Remove( ThisActor );
+	RemoveNetworkActor( ThisActor );
+}
+
+
+void UNetDriver::RemoveNetworkActor(AActor* Actor)
+{
+	GetNetworkObjectList().Remove(Actor);
 
 	// Remove from renamed list if destroyed
-	RenamedStartupActors.Remove(ThisActor->GetFName());
+	RenamedStartupActors.Remove(Actor->GetFName());
+
+	if (ReplicationDriver)
+	{
+		ReplicationDriver->RemoveNetworkActor(Actor);
+	}
 }
 
 void UNetDriver::NotifyActorRenamed(AActor* ThisActor, FName PreviousName)
@@ -2350,7 +2434,7 @@ void UNetDriver::NotifyStreamingLevelUnload( ULevel* Level)
 
 		if (Level->LevelScriptActor)
 		{
-			UActorChannel * Channel = ServerConnection->ActorChannels.FindRef((AActor*)Level->LevelScriptActor);
+			UActorChannel * Channel = ServerConnection->FindActorChannelRef((AActor*)Level->LevelScriptActor);
 			if (Channel)
 			{
 				UE_LOG(LogNet, Log, TEXT("NotifyStreamingLevelUnload: BREAKING"));
@@ -2386,10 +2470,10 @@ void UNetDriver::NotifyActorLevelUnloaded( AActor* TheActor )
 	if (ServerConnection != NULL)
 	{
 		// we can't kill the channel until the server says so, so just clear the actor ref and break the channel
-		UActorChannel* Channel = ServerConnection->ActorChannels.FindRef(TheActor);
+		UActorChannel* Channel = ServerConnection->FindActorChannelRef(TheActor);
 		if (Channel != NULL)
 		{
-			ServerConnection->ActorChannels.Remove(TheActor);
+			ServerConnection->RemoveActorChannel(TheActor);
 			Channel->Actor = NULL;
 			Channel->Broken = true;
 			Channel->CleanupReplicators();
@@ -2397,13 +2481,37 @@ void UNetDriver::NotifyActorLevelUnloaded( AActor* TheActor )
 	}
 }
 
+void UNetDriver::NotifyActorTearOff(AActor* Actor)
+{
+	if (ReplicationDriver)
+	{
+		ReplicationDriver->NotifyActorTearOff(Actor);
+	}
+}
+
+void UNetDriver::ForceNetUpdate(AActor* Actor)
+{
+	// Let Replication Driver handle it if its exists
+	if (ReplicationDriver)
+	{
+		ReplicationDriver->ForceNetUpdate(Actor);
+		return;
+	}
+	
+	// Legacy implementation
+	if ( FNetworkObjectInfo* NetActor = FindNetworkObjectInfo(Actor) )
+	{
+		NetActor->NextUpdateTime = World->TimeSeconds - 0.01f;
+	}
+}
+
 /** UNetDriver::FlushActorDormancy(AActor* Actor)
  *	 Flushes the actor from the NetDriver's dormant list and/or cancels pending dormancy on the actor channel.
  *
  *	 This does not change the Actor's actual NetDormant state. If a dormant actor is Flushed, it will net update at least one more
  *	 time, and then go back to dormant.
  */
-void UNetDriver::FlushActorDormancy(AActor* Actor)
+void UNetDriver::FlushActorDormancy(AActor* Actor, bool bWasDormInitial)
 {
 	// Note: Going into dormancy is completely handled in ServerReplicateActor. We want to avoid
 	// event-based handling of going into dormancy, because we have to deal with connections joining in progress.
@@ -2421,6 +2529,29 @@ void UNetDriver::FlushActorDormancy(AActor* Actor)
 	check(Actor);
 	check(ServerConnection == NULL);
 
+	if (ReplicationDriver)
+	{
+		ReplicationDriver->FlushNetDormancy(Actor, bWasDormInitial);
+	}
+
+	FlushActorDormancyInternal(Actor);
+	
+#endif // WITH_SERVER_CODE
+}
+
+void UNetDriver::NotifyActorDormancyChange(AActor* Actor, ENetDormancy OldDormancyState)
+{
+	if (ReplicationDriver)
+	{
+		ReplicationDriver->NotifyActorDormancyChange(Actor, OldDormancyState);
+	}
+
+	FlushActorDormancyInternal(Actor);
+}
+
+void UNetDriver::FlushActorDormancyInternal(AActor *Actor)
+{
+#if WITH_SERVER_CODE
 	// Go through each connection and remove the actor from the dormancy list
 	for (int32 i=0; i < ClientConnections.Num(); ++i)
 	{
@@ -2430,7 +2561,7 @@ void UNetDriver::FlushActorDormancy(AActor* Actor)
 			NetConnection->FlushDormancy(Actor);
 		}
 	}
-#endif // WITH_SERVER_CODE
+#endif
 }
 
 void UNetDriver::ForcePropertyCompare( AActor* Actor )
@@ -2660,13 +2791,13 @@ bool FPacketSimulationSettings::ParseSettings(const TCHAR* Cmd, const TCHAR* Opt
 	if( ParseHelper(Cmd, TEXT("PktLoss="), PktLoss, OptionalQualifier) )
 	{
 		bParsed = true;
-		FMath::Clamp<int32>( PktLoss, 0, 100 );
+		PktLoss = FMath::Clamp<int32>( PktLoss, 0, 100 );
 		UE_LOG(LogNet, Log, TEXT("PktLoss set to %d"), PktLoss);
 	}
 	if( ParseHelper(Cmd, TEXT("PktOrder="), PktOrder, OptionalQualifier) )
 	{
 		bParsed = true;
-		FMath::Clamp<int32>( PktOrder, 0, 1 );
+		PktOrder = FMath::Clamp<int32>( PktOrder, 0, 1 );
 		UE_LOG(LogNet, Log, TEXT("PktOrder set to %d"), PktOrder);
 	}
 	if( ParseHelper(Cmd, TEXT("PktLag="), PktLag, OptionalQualifier) )
@@ -2677,13 +2808,13 @@ bool FPacketSimulationSettings::ParseSettings(const TCHAR* Cmd, const TCHAR* Opt
 	if( ParseHelper(Cmd, TEXT("PktDup="), PktDup, OptionalQualifier) )
 	{
 		bParsed = true;
-		FMath::Clamp<int32>( PktDup, 0, 100 );
+		PktDup = FMath::Clamp<int32>( PktDup, 0, 100 );
 		UE_LOG(LogNet, Log, TEXT("PktDup set to %d"), PktDup);
 	}	
 	if (ParseHelper(Cmd, TEXT("PktLagVariance="), PktLagVariance, OptionalQualifier))
 	{
 		bParsed = true;
-		FMath::Clamp<int32>( PktLagVariance, 0, 100 );
+		PktLagVariance = FMath::Clamp<int32>( PktLagVariance, 0, 100 );
 		UE_LOG(LogNet, Log, TEXT("PktLagVariance set to %d"), PktLagVariance);
 	}
 	return bParsed;
@@ -2993,7 +3124,7 @@ void UNetDriver::ServerReplicateActors_BuildConsiderList( TArray<FNetworkObjectI
 
 	for ( AActor* Actor : ActorsToRemove )
 	{
-		GetNetworkObjectList().Remove( Actor );
+		RemoveNetworkActor( Actor );
 	}
 
 	// Update stats
@@ -3111,7 +3242,7 @@ int32 UNetDriver::ServerReplicateActors_PrioritizeActors( UNetConnection* Connec
 		{
 			AActor* Actor = ActorInfo->Actor;
 
-			UActorChannel* Channel = Connection->ActorChannels.FindRef( ActorInfo->WeakActor );
+			UActorChannel* Channel = Connection->FindActorChannelRef( ActorInfo->WeakActor );
 
 			// Skip actor if not relevant and theres no channel already.
 			// Historically Relevancy checks were deferred until after prioritization because they were expensive (line traces).
@@ -3192,9 +3323,9 @@ int32 UNetDriver::ServerReplicateActors_PrioritizeActors( UNetConnection* Connec
 		}
 
 		// Add in deleted actors
-		for ( auto It = Connection->DestroyedStartupOrDormantActors.CreateIterator(); It; ++It )
+		for ( auto It = Connection->GetDestroyedStartupOrDormantActorGUIDs().CreateConstIterator(); It; ++It )
 		{
-			FActorDestructionInfo& DInfo = DestroyedStartupOrDormantActors.FindChecked( *It );
+			FActorDestructionInfo& DInfo = *DestroyedStartupOrDormantActors.FindChecked( *It );
 			OutPriorityList[FinalSortedCount] = FActorPriority( Connection, &DInfo, ConnectionViewers );
 			OutPriorityActors[FinalSortedCount] = OutPriorityList + FinalSortedCount;
 			FinalSortedCount++;
@@ -3224,13 +3355,8 @@ int32 UNetDriver::ServerReplicateActors_ProcessPrioritizedActors( UNetConnection
 
 	if ( !Connection->IsNetReady( 0 ) )
 	{
+		GNumSaturatedConnections++;
 		// Connection saturated, don't process any actors
-
-		// Update stats even though there was no processing.
-		SET_DWORD_STAT( STAT_NumReplicatedActorAttempts, ActorUpdatesThisConnection );
-		SET_DWORD_STAT( STAT_NumReplicatedActors, ActorUpdatesThisConnectionSent );
-		SET_DWORD_STAT( STAT_NumRelevantActors, FinalRelevantCount );
-
 		return 0;
 	}
 
@@ -3254,8 +3380,8 @@ int32 UNetDriver::ServerReplicateActors_ProcessPrioritizedActors( UNetConnection
 				FinalRelevantCount++;
 				UE_LOG( LogNetTraffic, Log, TEXT( "Server replicate actor creating destroy channel for NetGUID <%s,%s> Priority: %d" ), *PriorityActors[j]->DestructionInfo->NetGUID.ToString(), *PriorityActors[j]->DestructionInfo->PathName, PriorityActors[j]->Priority );
 
-				Channel->SetChannelActorForDestroy( PriorityActors[j]->DestructionInfo );						   // Send a close bunch on the new channel
-				Connection->DestroyedStartupOrDormantActors.Remove( PriorityActors[j]->DestructionInfo->NetGUID ); // Remove from connections to-be-destroyed list (close bunch of reliable, so it will make it there)
+				Channel->SetChannelActorForDestroy( PriorityActors[j]->DestructionInfo );		// Send a close bunch on the new channel
+				Connection->RemoveDestructionInfo( PriorityActors[j]->DestructionInfo );		// Remove from connections to-be-destroyed list (close bunch of reliable, so it will make it there)
 			}
 			continue;
 		}
@@ -3285,7 +3411,7 @@ int32 UNetDriver::ServerReplicateActors_ProcessPrioritizedActors( UNetConnection
 			// bTearOff actors should never be checked
 			if ( bLevelInitializedForActor )
 			{
-				if ( !Actor->bTearOff && ( !Channel || Time - Channel->RelevantTime > 1.f ) )
+				if ( !Actor->GetTearOff() && ( !Channel || Time - Channel->RelevantTime > 1.f ) )
 				{
 					if ( IsActorRelevantToConnection( Actor, ConnectionViewers ) )
 					{
@@ -3396,16 +3522,14 @@ int32 UNetDriver::ServerReplicateActors_ProcessPrioritizedActors( UNetConnection
 					if ( !Connection->IsNetReady( 0 ) )
 					{
 						// We can bail out now since this connection is saturated, we'll return how far we got though
-						SET_DWORD_STAT( STAT_NumReplicatedActorAttempts, ActorUpdatesThisConnection );
-						SET_DWORD_STAT( STAT_NumReplicatedActors, ActorUpdatesThisConnectionSent );
-						SET_DWORD_STAT( STAT_NumRelevantActors, FinalRelevantCount );
+						GNumSaturatedConnections++;
 						return j;
 					}
 				}
 			}
 
 			// If the actor wasn't recently relevant, or if it was torn off, close the actor channel if it exists for this connection
-			if ( ( !bIsRecentlyRelevant || Actor->bTearOff ) && Channel != NULL )
+			if ( ( !bIsRecentlyRelevant || Actor->GetTearOff() ) && Channel != NULL )
 			{
 				// Non startup (map) actors have their channels closed immediately, which destroys them.
 				// Startup actors get to keep their channels open.
@@ -3420,14 +3544,146 @@ int32 UNetDriver::ServerReplicateActors_ProcessPrioritizedActors( UNetConnection
 		}
 	}
 
-	SET_DWORD_STAT( STAT_NumReplicatedActorAttempts, ActorUpdatesThisConnection );
-	SET_DWORD_STAT( STAT_NumReplicatedActors, ActorUpdatesThisConnectionSent );
-	SET_DWORD_STAT( STAT_NumRelevantActors, FinalRelevantCount );
-
 	return FinalSortedCount;
 }
 #endif
 
+// -------------------------------------------------------------------------------------------------------------------------
+//	Replication profiling (server cpu) helpers.
+//		-CSV stat cateogry "Replication". Not enabled by default. (use -csvCategories=Replication to enable from cmd line)
+//		-Auto CSV capturing can be enabled with -ReplicationCSVCaptureFrames=X. This will start a CSV capture once there
+//			is one client connected and will auto terminate the process once X frames have elapsed.
+// -------------------------------------------------------------------------------------------------------------------------
+
+
+// Replication CSV category is enabled by default in server builds
+CSV_DEFINE_CATEGORY(Replication, WITH_SERVER_CODE);
+
+extern int32 GNumReplicateActorCalls;
+extern double GReplicateActorTimeSeconds;
+extern uint32 GNetOutBytes;
+
+double GReplicationGatherPrioritizeTimeSeconds;
+double GServerReplicateActorTimeSeconds;
+
+int32 GNumClientConnections;
+
+DECLARE_DWORD_COUNTER_STAT(TEXT("Num Saturated Connections"), STAT_NumSaturatedConnections, STATGROUP_Net);
+
+DECLARE_DWORD_COUNTER_STAT(TEXT("SharedSerialization RPC Hit"), STAT_SharedSerializationRPCHit, STATGROUP_Net);
+DECLARE_DWORD_COUNTER_STAT(TEXT("SharedSerialization RPC Miss"), STAT_SharedSerializationRPCMiss, STATGROUP_Net);
+
+DECLARE_DWORD_COUNTER_STAT(TEXT("SharedSerialization Property Hit"), STAT_SharedSerializationPropertyHit, STATGROUP_Net);
+DECLARE_DWORD_COUNTER_STAT(TEXT("SharedSerialization Property Miss"), STAT_SharedSerializationPropertyMiss, STATGROUP_Net);
+
+struct FReplicationAutoCapture
+{
+	int32 CaptureFrames=-1;
+	int32 KillProcessFrames = -1;
+	bool StartedCapture = false;
+
+	void DoFrame()
+	{
+#if CSV_PROFILER
+		if (CaptureFrames == -1)
+		{
+			// First time see if we want to auto capture
+			CaptureFrames = 0;
+			FParse::Value(FCommandLine::Get(), TEXT("ReplicationCSVCaptureFrames="), CaptureFrames);
+		}
+		
+		if (CaptureFrames > 0)
+		{			
+			if (!StartedCapture)
+			{
+				StartedCapture = true;
+				FCsvProfiler::Get()->EnableCategoryByString(TEXT("Replication"));
+				FCsvProfiler::Get()->BeginCapture();
+			}
+				
+			if (--CaptureFrames <= 0)
+			{
+				FCsvProfiler::Get()->EndCapture();
+				KillProcessFrames = 60;
+			}
+
+		}
+		else if (KillProcessFrames > 0)
+		{
+			// Kill process when finished
+			--KillProcessFrames;
+			if (KillProcessFrames == 0)
+			{
+				GLog->PanicFlushThreadedLogs();
+				FPlatformMisc::RequestExit(1);
+			}
+		}
+#endif
+	}
+};
+FReplicationAutoCapture GReplicationAutoCapture;
+
+struct FScopedNetDriverStats
+{
+	FScopedNetDriverStats(uint32& OutBytesRef) : OutBytes(OutBytesRef)
+	{
+		GReplicateActorTimeSeconds = 0;
+		GNumReplicateActorCalls = 0;
+		GNumSaturatedConnections = 0;
+
+		GReplicationGatherPrioritizeTimeSeconds = 0.f;
+		GServerReplicateActorTimeSeconds = 0.f;
+
+		// Whatever these values currently are were (mostly) set by RPCs (technically something else could have force ReplicateActor to be called but this is rare).
+		SET_DWORD_STAT(STAT_SharedSerializationRPCHit, GNumSharedSerializationHit);
+		SET_DWORD_STAT(STAT_SharedSerializationRPCMiss, GNumSharedSerializationMiss);
+
+		StartTime = FPlatformTime::Seconds();
+		StartOutBytes = GNetOutBytes;
+
+		GReplicationAutoCapture.DoFrame();
+	}
+
+	~FScopedNetDriverStats()
+	{
+		const double TotalTime = FPlatformTime::Seconds() - StartTime;
+
+		GServerReplicateActorTimeSeconds = TotalTime;
+		GReplicationGatherPrioritizeTimeSeconds = TotalTime - GReplicateActorTimeSeconds;
+
+		uint32 FrameOutBytes = GNetOutBytes - StartOutBytes;
+
+		SET_FLOAT_STAT(STAT_NetServerGatherPrioritizeRepActorsTime, GReplicationGatherPrioritizeTimeSeconds * 1000.0);
+		SET_DWORD_STAT(STAT_NumReplicatedActors, GNumReplicateActorCalls);
+		SET_DWORD_STAT(STAT_NumSaturatedConnections, GNumSaturatedConnections);
+
+		CSV_CUSTOM_STAT(Replication, ServerReplicateActorTimeMS, (float)(GServerReplicateActorTimeSeconds * 1000.0), ECsvCustomStatOp::Set );
+		CSV_CUSTOM_STAT(Replication, GatherPrioritizeTimeMS, (float)(GReplicationGatherPrioritizeTimeSeconds * 1000.0), ECsvCustomStatOp::Set );
+		CSV_CUSTOM_STAT(Replication, ReplicateActorTimeMS, (float)(GReplicateActorTimeSeconds * 1000.0), ECsvCustomStatOp::Set );
+		CSV_CUSTOM_STAT(Replication, NumReplicateActorCallsPerConAvg, ((float)GNumReplicateActorCalls)/(float)GNumClientConnections, ECsvCustomStatOp::Set );
+		CSV_CUSTOM_STAT(Replication, Connections, (float)GNumClientConnections, ECsvCustomStatOp::Set );
+		CSV_CUSTOM_STAT(Replication, SatConnections, (float)GNumSaturatedConnections, ECsvCustomStatOp::Set );
+		CSV_CUSTOM_STAT(Replication, OutKBytes, ((float)FrameOutBytes) / 1000.f, ECsvCustomStatOp::Set );
+				
+
+		SET_DWORD_STAT(STAT_SharedSerializationPropertyHit, GNumSharedSerializationHit);
+		SET_DWORD_STAT(STAT_SharedSerializationPropertyMiss, GNumSharedSerializationMiss);
+
+		// Note: we want to reset this at the end of the frame since the RPC stats are incremented at the top (recv)
+		GNumSharedSerializationHit = 0;
+		GNumSharedSerializationMiss = 0;
+	}
+	
+	uint32& OutBytes;
+	double StartTime;
+	uint32 StartOutBytes;
+};
+
+
+// -------------------------------------------------------------------------------------------------------------------------
+//	ServerReplicateActors: this is main function to replicate actors to client connections. It can be "outsourced" to a Replication Driver.
+// -------------------------------------------------------------------------------------------------------------------------
+
 int32 UNetDriver::ServerReplicateActors(float DeltaSeconds)
 {
 	SCOPE_CYCLE_COUNTER(STAT_NetServerRepActorsTime);
@@ -3438,13 +3694,26 @@ int32 UNetDriver::ServerReplicateActors(float DeltaSeconds)
 		return 0;
 	}
 
-	check( World );
+	SET_DWORD_STAT(STAT_NumReplicatedActors, 0);
+	SET_DWORD_STAT(STAT_NumReplicatedActorBytes, 0);
 
-	int32 Updated = 0;
+#if !UE_BUILD_SHIPPING
+	FScopedNetDriverStats NetDriverStats(OutBytes);
+	GNumClientConnections = ClientConnections.Num();
+#endif
+	
+	if (ReplicationDriver)
+	{
+		return ReplicationDriver->ServerReplicateActors(DeltaSeconds);
+	}
+
+	check( World );
 
 	// Bump the ReplicationFrame value to invalidate any properties marked as "unchanged" for this frame.
 	ReplicationFrame++;
 
+	int32 Updated = 0;
+
 	const int32 NumClientsToTick = ServerReplicateActors_PrepConnections( DeltaSeconds );
 
 	if ( NumClientsToTick == 0 )
@@ -3506,7 +3775,7 @@ int32 UNetDriver::ServerReplicateActors(float DeltaSeconds)
 				if (Actor != NULL && !ConsiderList[ConsiderIdx]->bPendingNetUpdate)
 				{
 					// find the channel
-					UActorChannel *Channel = Connection->ActorChannels.FindRef(ConsiderList[ConsiderIdx]->WeakActor);
+					UActorChannel *Channel = Connection->FindActorChannelRef(ConsiderList[ConsiderIdx]->WeakActor);
 					// and if the channel last update time doesn't match the last net update time for the actor
 					if (Channel != NULL && Channel->LastUpdateTime < ConsiderList[ConsiderIdx]->LastNetUpdateTime)
 					{
@@ -3784,7 +4053,7 @@ void UNetDriver::DrawNetDriverDebug()
 		{
 			DrawColor = FColor::Blue;
 		}
-		else if (Connection->ActorChannels.FindRef(*It) != NULL)
+		else if (Connection->FindActorChannelRef(*It) != NULL)
 		{
 			DrawColor = FColor::Green;
 		}
@@ -3832,6 +4101,11 @@ void UNetDriver::AddClientConnection(UNetConnection * NewConnection)
 
 	ClientConnections.Add(NewConnection);
 
+	if (ReplicationDriver)
+	{
+		ReplicationDriver->AddClientConnection(NewConnection);
+	}
+
 #if USE_SERVER_PERF_COUNTERS
 	PerfCountersIncrement(TEXT("AddedConnections"));
 #endif
@@ -3845,7 +4119,7 @@ void UNetDriver::AddClientConnection(UNetConnection * NewConnection)
 		if (It.Key().IsStatic())
 		{
 			UE_LOG(LogNet, VeryVerbose, TEXT("Adding actor NetGUID <%s> to new connection's destroy list"), *It.Key().ToString());
-			NewConnection->DestroyedStartupOrDormantActors.Add(It.Key());
+			NewConnection->AddDestructionInfo(It.Value().Get());
 		}
 	}
 }
@@ -3857,8 +4131,14 @@ void UNetDriver::CreateReplicatedStaticActorDestructionInfo(UNetDriver* NetDrive
 	FNetworkGUID NetGUID = NetDriver->GuidCache->AssignNewNetGUIDFromPath_Server( Info.PathName.ToString(), Info.ObjOuter.Get(), Info.ObjClass );
 
 	UE_LOG(LogNet, Log, TEXT("CreateReplicatedStaticActorDestructionInfo %s %s %s %s"), *NetDriver->GetName(), *Level->GetName(), *Info.PathName.ToString(), *NetGUID.ToString());
-
-	FActorDestructionInfo &NewInfo = NetDriver->DestroyedStartupOrDormantActors.FindOrAdd( NetGUID );
+	
+	TUniquePtr<FActorDestructionInfo>& NewInfoPtr = NetDriver->DestroyedStartupOrDormantActors.FindOrAdd( NetGUID );
+	if (NewInfoPtr.IsValid() == false)
+	{
+		NewInfoPtr = MakeUnique<FActorDestructionInfo>();
+	}
+	
+	FActorDestructionInfo &NewInfo = *NewInfoPtr;
 	NewInfo.DestroyedPosition = Info.DestroyedPosition;
 	NewInfo.NetGUID = NetGUID;
 	NewInfo.Level = Level;
@@ -3902,6 +4182,24 @@ void UNetDriver::InitDestroyedStartupActors()
 	}
 }
 
+void UNetDriver::NotifyActorFullyDormantForConnection(AActor* Actor, UNetConnection* Connection)
+{
+	GetNetworkObjectList().MarkDormant(Actor, Connection, IsServer() ? ClientConnections.Num() : 1, NetDriverName);
+	if (ReplicationDriver)
+	{
+		ReplicationDriver->NotifyActorFullyDormantForConnection(Actor, Connection);
+	}
+}
+
+void UNetDriver::RemoveClientConnection(UNetConnection* ClientConnectionToRemove)
+{
+	verify(ClientConnections.Remove(ClientConnectionToRemove) == 1);
+	if (ReplicationDriver)
+	{
+		ReplicationDriver->RemoveClientConnection(ClientConnectionToRemove);
+	}
+}
+
 void UNetDriver::SetWorld(class UWorld* InWorld)
 {
 	if (World)
@@ -3925,6 +4223,11 @@ void UNetDriver::SetWorld(class UWorld* InWorld)
 
 		GetNetworkObjectList().AddInitialObjects(InWorld, NetDriverName);
 	}
+
+	if (ReplicationDriver)
+	{
+		ReplicationDriver->SetWorld(InWorld);
+	}
 }
 
 void UNetDriver::ResetGameWorldState()
@@ -3967,6 +4270,21 @@ void UNetDriver::PostSeamlessTravelGarbageCollect()
 	CleanPackageMaps();
 }
 
+void UNetDriver::SetReplicationDriver(UReplicationDriver* NewReplicationDriver)
+{
+	if (ReplicationDriver)
+	{
+		ReplicationDriver->MarkPendingKill();
+	}
+
+	ReplicationDriver = NewReplicationDriver;
+	if (ReplicationDriver)
+	{
+		ReplicationDriver->InitForNetDriver(this);
+		ReplicationDriver->SetWorld( GetWorld() );
+	}
+}
+
 static void	DumpRelevantActors( UWorld* InWorld )
 {
 	UNetDriver *NetDriver = InWorld->NetDriver;
@@ -4060,33 +4378,92 @@ void UNetDriver::OnLevelRemovedFromWorld(class ULevel* InLevel, class UWorld* In
 			if (Actor)
 			{
 				NotifyActorLevelUnloaded(Actor);
-				GetNetworkObjectList().Remove(Actor);
+				RemoveNetworkActor(Actor);
 			}
 		}
 
 		TArray<FNetworkGUID> RemovedGUIDs;
 		for (auto It = DestroyedStartupOrDormantActors.CreateIterator(); It; ++It)
 		{
-			if (It->Value.Level == InLevel)
+			if (It->Value->Level == InLevel)
 			{
+				// Connections must be updated before we remove from our map
+				FActorDestructionInfo* DestructInfo = It->Value.Get();
+				for (UNetConnection* Connection : ClientConnections)
+				{
+					Connection->RemoveDestructionInfo(DestructInfo);
+				}
+
 				RemovedGUIDs.Add(It->Key);
 				It.RemoveCurrent();
 			}
 		}
+	}
+}
+
+// --------------------------------------------------------------------------------------------------------------------------------------------
+// --------------------------------------------------------------------------------------------------------------------------------------------
+// --------------------------------------------------------------------------------------------------------------------------------------------
+
+int32 CVar_RepDriver_Enable = 1;
+static FAutoConsoleVariableRef CVarRepManagerEnable(TEXT("Net.RepDriver.Enable"), CVar_RepDriver_Enable, TEXT("Enables Replication Driver. 0 will fallback to legacy NetDriver implementation."), ECVF_Default );
+
+UReplicationDriver::UReplicationDriver()
+{
+}
+
+UReplicationDriver* UReplicationDriver::CreateReplicationDriver(UNetDriver* NetDriver, const FURL& URL, UWorld* World)
+{
+	if (CreateReplicationDriverDelegate().IsBound())
+	{
+		return CreateReplicationDriverDelegate().Execute(NetDriver, URL, World);
+	}
 
-		if (RemovedGUIDs.Num())
+	// If we initialize the server map from the initial commandline, this will run before the "-execcmds" option has been handled. So explicitly check for force cmd line options
+	static bool FirstTime = true;
+	if (FirstTime)
+	{
+		FirstTime = false;
+		if (FParse::Param(FCommandLine::Get(), TEXT("RepDriverEnable")))
 		{
-			for (UNetConnection* Connection : ClientConnections)
-			{
-				for (const FNetworkGUID& GUIDToRemove : RemovedGUIDs)
-				{
-					Connection->DestroyedStartupOrDormantActors.Remove(GUIDToRemove);
-				}
-			}
+			CVar_RepDriver_Enable = 1;
+		}
+		else if (FParse::Param(FCommandLine::Get(), TEXT("RepDriverDisable")))
+		{
+			CVar_RepDriver_Enable = 0;
 		}
 	}
+
+	if (CVar_RepDriver_Enable == 0)
+	{
+		return nullptr;
+	}
+
+	UClass* ReplicationDriverClass = NetDriver ? NetDriver->ReplicationDriverClass : nullptr;
+	if (ReplicationDriverClass == nullptr)
+	{
+		UE_LOG(LogNet, Log, TEXT("ReplicationDriverClass is null! Not using ReplicationDriver."));
+		return nullptr;
+	}
+
+	return NewObject<UReplicationDriver>(GetTransientPackage(), ReplicationDriverClass);
+}
+
+FCreateReplicationDriver& UReplicationDriver::CreateReplicationDriverDelegate()
+{
+	static FCreateReplicationDriver Delegate;
+	return Delegate;
+}
+
+void UNetDriver::OnLevelAddedToWorld(ULevel* InLevel, UWorld* InWorld)
+{
 }
 
+// --------------------------------------------------------------------------------------------------------------------------------------------
+// --------------------------------------------------------------------------------------------------------------------------------------------
+// --------------------------------------------------------------------------------------------------------------------------------------------
+
+
 void UNetDriver::ProcessRemoteFunction(class AActor* Actor, UFunction* Function, void* Parameters, FOutParmRec* OutParms, FFrame* Stack, class UObject* SubObject )
 {
 #if !UE_BUILD_SHIPPING
@@ -4100,6 +4477,13 @@ void UNetDriver::ProcessRemoteFunction(class AActor* Actor, UFunction* Function,
 	if (!bBlockSendRPC)
 #endif
 	{
+		// Forward to replication Driver if there is one
+		if (ReplicationDriver && ReplicationDriver->ProcessRemoteFunction(Actor, Function, Parameters, OutParms, Stack, SubObject))
+		{
+			return;
+		}
+		
+		// RepDriver didn't handle it, default implementation
 		bool bIsServer = IsServer();
 
 		UNetConnection* Connection = NULL;
diff --git a/Engine/Source/Runtime/Engine/Private/NetworkObjectList.cpp b/Engine/Source/Runtime/Engine/Private/NetworkObjectList.cpp
--- a/Engine/Source/Runtime/Engine/Private/NetworkObjectList.cpp
+++ b/Engine/Source/Runtime/Engine/Private/NetworkObjectList.cpp
@@ -36,7 +36,7 @@ TSharedPtr<FNetworkObjectInfo> FNetworkObjectList::Find(AActor* const Actor)
 	return TSharedPtr<FNetworkObjectInfo>();
 }
 
-TSharedPtr<FNetworkObjectInfo>* FNetworkObjectList::FindOrAdd(AActor* const Actor, const FName NetDriverName)
+TSharedPtr<FNetworkObjectInfo>* FNetworkObjectList::FindOrAdd(AActor* const Actor, const FName NetDriverName, bool* OutWasAdded)
 {
 	if (Actor == nullptr || Actor->IsPendingKill() ||
 
@@ -59,11 +59,20 @@ TSharedPtr<FNetworkObjectInfo>* FNetworkObjectList::FindOrAdd(AActor* const Acto
 			ActiveNetworkObjects.Add(*NetworkObjectInfo);
 
 			UE_LOG(LogNetDormancy, VeryVerbose, TEXT("FNetworkObjectList::Add: Adding actor. Actor: %s, Total: %i, Active: %i, NetDriverName: %s"), *Actor->GetName(), AllNetworkObjects.Num(), ActiveNetworkObjects.Num(), *NetDriverName.ToString());
+
+			if (OutWasAdded)
+			{
+				*OutWasAdded = true;
+			}
 		}
 	}
 	else
 	{
 		UE_LOG(LogNetDormancy, VeryVerbose, TEXT("FNetworkObjectList::Add: Already contained. Actor: %s, Total: %i, Active: %i, NetDriverName: %s"), *Actor->GetName(), AllNetworkObjects.Num(), ActiveNetworkObjects.Num(), *NetDriverName.ToString());
+		if (OutWasAdded)
+		{
+			*OutWasAdded = false;
+		}
 	}
 	
 	check((ActiveNetworkObjects.Num() + ObjectsDormantOnAllConnections.Num()) == AllNetworkObjects.Num());
diff --git a/Engine/Source/Runtime/Engine/Private/PackageMapClient.cpp b/Engine/Source/Runtime/Engine/Private/PackageMapClient.cpp
--- a/Engine/Source/Runtime/Engine/Private/PackageMapClient.cpp
+++ b/Engine/Source/Runtime/Engine/Private/PackageMapClient.cpp
@@ -22,6 +22,8 @@
 #include "GameFramework/GameStateBase.h"
 #include "HAL/LowLevelMemTracker.h"
 #include "Engine/DemoNetDriver.h"
+#include "Serialization/MemoryReader.h"
+#include "Serialization/MemoryWriter.h"
 
 #if WITH_EDITOR
 #include "UObject/ObjectRedirector.h"
@@ -442,7 +444,7 @@ bool UPackageMapClient::SerializeNewActor(FArchive& Ar, class UActorChannel *Cha
 			if ( Ar.IsSaving() )
 			{
 				FObjectReplicator * RepData = &Channel->GetActorReplicationData();
-				uint8* Recent = RepData && RepData->RepState != NULL && RepData->RepState->StaticBuffer.Num() ? RepData->RepState->StaticBuffer.GetData() : NULL;
+				uint8* Recent = RepData && RepData->RepState.IsValid() && RepData->RepState->StaticBuffer.Num() ? RepData->RepState->StaticBuffer.GetData() : NULL;
 				if ( Recent )
 				{
 					((AActor*)Recent)->ReplicatedMovement.Location = LocalLocation;
@@ -967,6 +969,31 @@ UObject* UPackageMapClient::ResolvePathAndAssignNetGUID( const FNetworkGUID& Net
 //
 //--------------------------------------------------------------------
 
+// MAX_BUNCH_SIZE is in bits, so shift to get bytes.
+// TODO: This limit might not actually need to be enforced anymore.
+static const int32 MAX_GUID_MEMORY = MAX_BUNCH_SIZE >> 3;
+static const int32 MAX_GUID_COUNT = 2048;
+
+bool UPackageMapClient::ExportNetGUIDForReplay(FNetworkGUID& NetGUID, const UObject* Object, FString& PathName, UObject* ObjOuter)
+{
+	{
+		TGuardValue<bool> ExportingGUID(GuidCache->IsExportingNetGUIDBunch, true);
+
+		TArray<uint8>& GUIDMemory = ExportGUIDArchives.Emplace_GetRef();
+		GUIDMemory.Reserve(MAX_GUID_MEMORY);
+
+		FMemoryWriter Writer(GUIDMemory);
+		InternalWriteObject(Writer, NetGUID, Object, PathName, ObjOuter);
+
+		check(!Writer.IsError());
+		ensureMsgf(GUIDMemory.Num() <= MAX_GUID_MEMORY, TEXT("ExportNetGUIDForReplay exceeded MAX_GUID_MEMORY. Max=%l Count=%l"), MAX_GUID_MEMORY, GUIDMemory.Num());
+	}
+
+	CurrentExportNetGUIDs.Empty();
+	ExportNetGUIDCount = 0;
+	return true;
+}
+
 /** Exports the NetGUID and paths needed to the CurrentExportBunch */
 bool UPackageMapClient::ExportNetGUID( FNetworkGUID NetGUID, const UObject* Object, FString PathName, UObject* ObjOuter )
 {
@@ -975,6 +1002,11 @@ bool UPackageMapClient::ExportNetGUID( FNetworkGUID NetGUID, const UObject* Obje
 	check( !NetGUID.IsDefault() );
 	check( Object == NULL || ShouldSendFullPath( Object, NetGUID ) );
 
+	if (Connection->InternalAck)
+	{
+		return ExportNetGUIDForReplay(NetGUID, Object, PathName, ObjOuter);
+	}
+
 	// Two passes are used to export this net guid:
 	// 1. Attempt to append this net guid to current bunch
 	// 2. If step 1 fails, append to fresh new bunch
@@ -1112,22 +1144,22 @@ void UPackageMapClient::ReceiveNetGUIDBunch( FInBunch &InBunch )
 
 	if ( bHasRepLayoutExport )
 	{
+		// We need to keep this around to ensure we don't break backwards compatability.
+		PRAGMA_DISABLE_DEPRECATION_WARNINGS
 		ReceiveNetFieldExports( InBunch );
+		PRAGMA_ENABLE_DEPRECATION_WARNINGS
 		return;
 	}
 
-	GuidCache->IsExportingNetGUIDBunch = true;
+	TGuardValue<bool> IsExportingGuard(GuidCache->IsExportingNetGUIDBunch, true);
 
 	int32 NumGUIDsInBunch = 0;
 	InBunch << NumGUIDsInBunch;
 
-	static const int32 MAX_GUID_COUNT = 2048;
-
 	if ( NumGUIDsInBunch > MAX_GUID_COUNT )
 	{
 		UE_LOG( LogNetPackageMap, Error, TEXT( "UPackageMapClient::ReceiveNetGUIDBunch: NumGUIDsInBunch > MAX_GUID_COUNT (%i)" ), NumGUIDsInBunch );
 		InBunch.SetError();
-		GuidCache->IsExportingNetGUIDBunch = false;
 		return;
 	}
 
@@ -1144,14 +1176,12 @@ void UPackageMapClient::ReceiveNetGUIDBunch( FInBunch &InBunch )
 		if ( InBunch.IsError() )
 		{
 			UE_LOG( LogNetPackageMap, Error, TEXT( "UPackageMapClient::ReceiveNetGUIDBunch: InBunch.IsError() after InternalLoadObject" ) );
-			GuidCache->IsExportingNetGUIDBunch = false;
 			return;
 		}
 		NumGUIDsRead++;
 	}
 
 	UE_LOG(LogNetPackageMap, Log, TEXT("UPackageMapClient::ReceiveNetGUIDBunch end. BitPos: %d"), InBunch.GetPosBits() );
-	GuidCache->IsExportingNetGUIDBunch = false;
 }
 
 TSharedPtr< FNetFieldExportGroup > UPackageMapClient::GetNetFieldExportGroup( const FString& PathName )
@@ -1194,10 +1224,15 @@ TSharedPtr< FNetFieldExportGroup > UPackageMapClient::GetNetFieldExportGroupChec
 	return GuidCache->NetFieldExportGroupMap.FindChecked( PathName );
 }
 
-void UPackageMapClient::SerializeNetFieldExportGroupMap( FArchive& Ar )
+void UPackageMapClient::SerializeNetFieldExportGroupMap( FArchive& Ar, bool bClearPendingExports )
 {
 	if ( Ar.IsSaving() )
 	{
+		if (bClearPendingExports)
+		{
+			NetFieldExports.Empty();
+		}
+
 		// Save the number of layouts
 		uint32 NumNetFieldExportGroups = GuidCache->NetFieldExportGroupMap.Num();
 		Ar << NumNetFieldExportGroups;
@@ -1356,6 +1391,158 @@ void UPackageMapClient::AppendNetFieldExports( TArray<FOutBunch *>& OutgoingBunc
 	NetFieldExports.Empty();
 }
 
+void UPackageMapClient::AppendExportData(FArchive& Archive)
+{
+	check(Connection->InternalAck);
+
+	AppendNetFieldExports(Archive);
+	AppendNetExportGUIDs(Archive);
+}
+
+void UPackageMapClient::ReceiveExportData(FArchive& Archive)
+{
+	ReceiveNetFieldExports(Archive);
+	ReceiveNetExportGUIDs(Archive);
+}
+
+void UPackageMapClient::AppendNetFieldExports(FArchive& Archive)
+{
+	check(Connection->InternalAck);
+
+	uint32 NetFieldCount = NetFieldExports.Num();
+	Archive.SerializeIntPacked(NetFieldCount);
+
+	if (0 == NetFieldCount)
+	{
+		return;
+	}
+
+	for (const uint64 FieldExport : NetFieldExports)
+	{
+		// Parse the path name index and cmd index out of the uint64
+		uint32 PathNameIndex = FieldExport >> 32;
+		uint32 NetFieldExportHandle = FieldExport & (((uint64)1 << 32) - 1);
+
+		check(PathNameIndex != 0);
+
+		FString PathName = GuidCache->NetFieldExportGroupIndexToPath.FindChecked(PathNameIndex);
+		TSharedPtr< FNetFieldExportGroup > NetFieldExportGroup = GuidCache->NetFieldExportGroupMap.FindChecked(PathName);
+		check(NetFieldExportHandle == NetFieldExportGroup->NetFieldExports[NetFieldExportHandle].Handle);
+
+		// Export the path if we need to
+		uint32 NeedsExport = OverrideAckState->NetFieldExportGroupPathAcked.Contains(PathNameIndex) ? 0 : 1;
+
+		Archive.SerializeIntPacked(PathNameIndex);
+		Archive.SerializeIntPacked(NeedsExport);
+
+		if (NeedsExport)
+		{
+			uint32 NumExports = NetFieldExportGroup->NetFieldExports.Num();
+
+			Archive << PathName;
+			Archive.SerializeIntPacked(NumExports);
+		}
+
+		Archive << NetFieldExportGroup->NetFieldExports[NetFieldExportHandle];
+	}
+
+	NetFieldExports.Empty();
+}
+
+void UPackageMapClient::ReceiveNetFieldExports(FArchive& Archive)
+{
+	check(Connection->InternalAck);
+
+	// Read number of net field exports
+	uint32 NumLayoutCmdExports = 0;
+	Archive.SerializeIntPacked(NumLayoutCmdExports);
+
+	for (int32 i = 0; i < (int32)NumLayoutCmdExports; i++)
+	{
+		uint32 PathNameIndex = 0;
+		uint32 WasExported = 0;
+
+		Archive.SerializeIntPacked(PathNameIndex);
+		Archive.SerializeIntPacked(WasExported);
+
+		TSharedPtr<FNetFieldExportGroup> NetFieldExportGroup;
+		if (!!WasExported)
+		{
+			FString PathName;
+			uint32 NumExports = 0;
+
+			Archive << PathName;
+			Archive.SerializeIntPacked(NumExports);
+
+			GEngine->NetworkRemapPath(Connection->Driver, PathName, true);
+			GuidCache->NetFieldExportGroupIndexToPath.Add(PathNameIndex, PathName);
+			GuidCache->NetFieldExportGroupPathToIndex.Add(PathName, PathNameIndex);
+
+			NetFieldExportGroup = GuidCache->NetFieldExportGroupMap.FindRef(PathName);
+			if (!NetFieldExportGroup.IsValid())
+			{
+				NetFieldExportGroup = TSharedPtr<FNetFieldExportGroup>(new FNetFieldExportGroup());
+				NetFieldExportGroup->PathName = PathName;
+				NetFieldExportGroup->PathNameIndex = PathNameIndex;
+				NetFieldExportGroup->NetFieldExports.SetNum(NumExports);
+				GuidCache->NetFieldExportGroupMap.Add(PathName, NetFieldExportGroup);
+			}
+		}
+		else
+		{
+			const FString& PathName = GuidCache->NetFieldExportGroupIndexToPath.FindChecked(PathNameIndex);
+			NetFieldExportGroup = GuidCache->NetFieldExportGroupMap.FindRef(PathName);
+		}
+
+		TArray<FNetFieldExport>& Exports = NetFieldExportGroup->NetFieldExports;
+		FNetFieldExport Export;
+		Archive << Export;
+
+		if (Exports.IsValidIndex(Export.Handle))
+		{
+			Exports[Export.Handle] = Export;
+		}
+		else
+		{
+			UE_LOG(LogNetPackageMap, Error, TEXT("ReceiveNetFieldExports: Invalid NetFieldExportHandle '%i', Max '%i'"), Export.Handle, Exports.Num());
+		}
+	}
+}
+
+void UPackageMapClient::AppendNetExportGUIDs(FArchive& Archive)
+{
+	check(Connection->InternalAck);
+
+	uint32 NumGUIDs = ExportGUIDArchives.Num();
+	Archive.SerializeIntPacked(NumGUIDs);
+
+	for (TArray<uint8>& GUIDData : ExportGUIDArchives)
+	{
+		Archive << GUIDData;
+	}
+
+	ExportGUIDArchives.Empty();
+}
+
+void UPackageMapClient::ReceiveNetExportGUIDs(FArchive& Archive)
+{
+	check(Connection->InternalAck);
+	TGuardValue<bool> IsExportingGuard(GuidCache->IsExportingNetGUIDBunch, true);
+
+	uint32 NumGUIDs = 0;
+	Archive.SerializeIntPacked(NumGUIDs);
+
+	for (uint32 i = 0; i < NumGUIDs; i++)
+	{
+		TArray<uint8> GUIDData;
+		Archive << GUIDData;
+
+		FMemoryReader Reader(GUIDData);
+		UObject* Object = nullptr;
+		InternalLoadObject(Reader, Object, 0);
+	}
+}
+
 void UPackageMapClient::ReceiveNetFieldExports( FInBunch &InBunch )
 {
 	// WARNING: If this code path is enabled for use beyond replay, it will need a security audit/rewrite
@@ -1452,11 +1639,8 @@ void UPackageMapClient::ReceiveNetFieldExports( FInBunch &InBunch )
 
 void UPackageMapClient::AppendExportBunches(TArray<FOutBunch *>& OutgoingBunches)
 {
-	// If we have rep layouts to export, handle those now
-	if ( NetFieldExports.Num() > 0 )
-	{
-		AppendNetFieldExports( OutgoingBunches );
-	}
+	check(!Connection->InternalAck);
+	check(NetFieldExports.Num() == 0);
 
 	// Finish current in progress bunch if necessary
 	if (ExportNetGUIDCount > 0)
diff --git a/Engine/Source/Runtime/Engine/Private/PlayerController.cpp b/Engine/Source/Runtime/Engine/Private/PlayerController.cpp
--- a/Engine/Source/Runtime/Engine/Private/PlayerController.cpp
+++ b/Engine/Source/Runtime/Engine/Private/PlayerController.cpp
@@ -612,7 +612,7 @@ void APlayerController::ForceSingleNetUpdateFor(AActor* Target)
 				Conn = ((UChildConnection*)Conn)->Parent;
 				checkSlow(Conn != NULL);
 			}
-			UActorChannel* Channel = Conn->ActorChannels.FindRef(Target);
+			UActorChannel* Channel = Conn->FindActorChannelRef(Target);
 			if (Channel != NULL)
 			{
 				FNetworkObjectInfo* NetActor = Target->FindOrAddNetworkObjectInfo();
@@ -712,7 +712,7 @@ void APlayerController::ClientRestart_Implementation(APawn* NewPawn)
 	AcknowledgedPawn = NULL;
 
 	SetPawn(NewPawn);
-	if ( (GetPawn() != NULL) && GetPawn()->bTearOff )
+	if ( (GetPawn() != NULL) && GetPawn()->GetTearOff() )
 	{
 		UnPossess();
 		SetPawn(NULL);
diff --git a/Engine/Source/Runtime/Engine/Private/RepLayout.cpp b/Engine/Source/Runtime/Engine/Private/RepLayout.cpp
--- a/Engine/Source/Runtime/Engine/Private/RepLayout.cpp
+++ b/Engine/Source/Runtime/Engine/Private/RepLayout.cpp
@@ -34,6 +34,9 @@ static FAutoConsoleVariable CVarLogSkippedRepNotifies(TEXT("Net.LogSkippedRepNot
 int32 MaxRepArraySize = UNetworkSettings::DefaultMaxRepArraySize;
 int32 MaxRepArrayMemory = UNetworkSettings::DefaultMaxRepArrayMemory;
 
+extern int32 GNumSharedSerializationHit;
+extern int32 GNumSharedSerializationMiss;
+
 FConsoleVariableSinkHandle CreateMaxArraySizeCVarAndRegisterSink()
 {
 	static FAutoConsoleVariable CVarMaxArraySize(TEXT("net.MaxRepArraySize"), MaxRepArraySize, TEXT("Maximum allowable size for replicated dynamic arrays (in number of elements). Value must be between 1 and 65535."));
@@ -1372,6 +1375,7 @@ void FRepLayout::SendProperties_r(
 		// Use shared serialization if was found
 		if (SharedPropInfo)
 		{
+			GNumSharedSerializationHit++;
 #if !(UE_BUILD_SHIPPING || UE_BUILD_TEST)
 			if (CVarNetVerifyShareSerializedData->GetInt() != 0)
 			{
@@ -1409,6 +1413,7 @@ void FRepLayout::SendProperties_r(
 		}
 		else
 		{
+			GNumSharedSerializationMiss++;
 			WritePropertyHandle( Writer, HandleIterator.Handle, bDoChecksum );
 
 			const int32 NumStartBits = Writer.GetNumBits();
@@ -2771,12 +2776,20 @@ void FRepLayout::SanityCheckChangeList( const uint8* RESTRICT Data, TArray< uint
 class FDiffPropertiesImpl : public FRepLayoutCmdIterator< FDiffPropertiesImpl, FCmdIteratorBaseStackState >
 {
 public:
-	FDiffPropertiesImpl( const bool bInSync, TArray< UProperty * >*	InRepNotifies, const TArray< FRepParentCmd >& InParents, const TArray< FRepLayoutCmd >& InCmds ) : 
+	FDiffPropertiesImpl( const EDiffPropertiesFlags InFlags, TArray< UProperty * >*	InRepNotifies, const TArray< FRepParentCmd >& InParents, const TArray< FRepLayoutCmd >& InCmds ) : 
 		FRepLayoutCmdIterator( InParents, InCmds ),
-		bSync( bInSync ),
+		Flags( InFlags ),
+		ParentPropertyFlags( ERepParentFlags::IsLifetime ),
 		RepNotifies( InRepNotifies ),
 		bDifferent( false )
-	{}
+	{
+		// Currently, only lifetime properties init from their defaults, so default to that,
+		// but also diff conditional properties if requested.
+		if ( (Flags & EDiffPropertiesFlags::IncludeConditionalProperties) != EDiffPropertiesFlags::None )
+		{
+			ParentPropertyFlags |= ERepParentFlags::IsConditional;
+		}
+	}
 
 	INIT_STACK( FCmdIteratorBaseStackState ) { }
 
@@ -2791,15 +2804,14 @@ public:
 		{
 			bDifferent = true;
 
-			if ( !bSync )
-			{			
+			if ( (Flags & EDiffPropertiesFlags::Sync) == EDiffPropertiesFlags::None )
+			{
 				UE_LOG( LogRep, Warning, TEXT( "FDiffPropertiesImpl: Array sizes different: %s %i / %i" ), *Cmd.Property->GetFullName(), StackState.DataArray->Num(), StackState.ShadowArray->Num() );
 				return;
 			}
 
-			if ( (Parents[Cmd.ParentIndex].Flags & ERepParentFlags::IsLifetime) == ERepParentFlags::None )
+			if ( (Parents[Cmd.ParentIndex].Flags & ParentPropertyFlags) == ERepParentFlags::None )
 			{
-				// Currently, only lifetime properties init from their defaults
 				return;
 			}
 
@@ -2820,19 +2832,18 @@ public:
 		const FRepParentCmd& Parent = Parents[Cmd.ParentIndex];
 
 		// Make the shadow state match the actual state at the time of send
-		if ( Parent.RepNotifyCondition == REPNOTIFY_Always || !PropertiesAreIdentical( Cmd, (const void*)( Data + Cmd.Offset ), (const void*)( ShadowData + Cmd.Offset ) ) )
+		if ( (RepNotifies && Parent.RepNotifyCondition == REPNOTIFY_Always) || !PropertiesAreIdentical( Cmd, (const void*)( Data + Cmd.Offset ), (const void*)( ShadowData + Cmd.Offset ) ) )
 		{
 			bDifferent = true;
 
-			if ( !bSync )
+			if ( (Flags & EDiffPropertiesFlags::Sync) == EDiffPropertiesFlags::None )
 			{			
 				UE_LOG( LogRep, Warning, TEXT( "FDiffPropertiesImpl: Property different: %s" ), *Cmd.Property->GetFullName() );
 				return;
 			}
 
-			if ( (Parent.Flags & ERepParentFlags::IsLifetime) == ERepParentFlags::None )
+			if ( (Parents[Cmd.ParentIndex].Flags & ParentPropertyFlags) == ERepParentFlags::None )
 			{
-				// Currently, only lifetime properties init from their defaults
 				return;
 			}
 
@@ -2849,14 +2860,127 @@ public:
 		}
 	}
 
-	bool					bSync;
+	bool IsDifferent() const { return bDifferent; }
+
+private:
+	EDiffPropertiesFlags	Flags;
+	ERepParentFlags			ParentPropertyFlags;
 	TArray< UProperty * >*	RepNotifies;
 	bool					bDifferent;
 };
 
 bool FRepLayout::DiffProperties( TArray<UProperty*>* RepNotifies, void* RESTRICT Destination, const void* RESTRICT Source, const bool bSync ) const
-{	
-	FDiffPropertiesImpl DiffPropertiesImpl( bSync, RepNotifies, Parents, Cmds );
+{
+	return DiffProperties(RepNotifies, Destination, Source, bSync ? EDiffPropertiesFlags::Sync : EDiffPropertiesFlags::None);
+}
+
+bool FRepLayout::DiffProperties( TArray<UProperty*>* RepNotifies, void* RESTRICT Destination, const void* RESTRICT Source, const EDiffPropertiesFlags Flags ) const
+{
+	FDiffPropertiesImpl DiffPropertiesImpl( Flags, RepNotifies, Parents, Cmds );
+
+	DiffPropertiesImpl.ProcessCmds( (uint8*)Destination, (uint8*)Source );
+
+	return DiffPropertiesImpl.IsDifferent();
+}
+
+class FDiffStablePropertiesImpl : public FRepLayoutCmdIterator< FDiffStablePropertiesImpl, FCmdIteratorBaseStackState >
+{
+public:
+	FDiffStablePropertiesImpl( TArray<UProperty*>* InRepNotifies, TArray<UObject*>* InObjReferences, const TArray< FRepParentCmd >& InParents, const TArray< FRepLayoutCmd >& InCmds ) : 
+		FRepLayoutCmdIterator( InParents, InCmds ),
+		RepNotifies(InRepNotifies),
+		ObjReferences(InObjReferences),
+		bDifferent( false )
+	{}
+
+	INIT_STACK( FCmdIteratorBaseStackState ) { }
+
+	SHOULD_PROCESS_NEXT_CMD() 
+	{ 
+		return true;
+	}
+
+	PROCESS_ARRAY_CMD( FCmdIteratorBaseStackState ) 
+	{
+		if ( StackState.DataArray->Num() != StackState.ShadowArray->Num() )
+		{
+			bDifferent = true;
+
+			if ( (Parents[Cmd.ParentIndex].Flags & ERepParentFlags::IsLifetime) == ERepParentFlags::None )
+			{
+				// Currently, only lifetime properties init from their defaults
+				return;
+			}
+
+			// Make the shadow state match the actual state
+			FScriptArrayHelper ShadowArrayHelper( (UArrayProperty *)Cmd.Property, ShadowData );
+			ShadowArrayHelper.Resize( StackState.DataArray->Num() );
+		}
+
+		StackState.BaseData			= (uint8*)StackState.DataArray->GetData();
+		StackState.ShadowBaseData	= (uint8*)StackState.ShadowArray->GetData();
+
+		// Loop over array
+		ProcessDataArrayElements_r( StackState, Cmd );
+	}
+
+	PROCESS_CMD( FCmdIteratorBaseStackState ) 
+	{
+		const FRepParentCmd& Parent = Parents[Cmd.ParentIndex];
+
+		// Make the shadow state match the actual state at the time of send
+		if ( !PropertiesAreIdentical( Cmd, (const void*)( Data + Cmd.Offset ), (const void*)( ShadowData + Cmd.Offset ) ) )
+		{
+			bDifferent = true;
+
+			if ( (Parent.Flags & ERepParentFlags::IsLifetime) == ERepParentFlags::None )
+			{
+				// Currently, only lifetime properties init from their defaults
+				return;
+			}
+
+			if (Cmd.Property->HasAnyPropertyFlags(CPF_Transient))
+			{
+				// skip transient properties
+				return;
+			}
+
+			if (Cmd.Type == ERepLayoutCmdType::PropertyObject)
+			{
+				UObjectPropertyBase * ObjProperty = CastChecked< UObjectPropertyBase>( Cmd.Property );
+				if (ObjProperty)
+				{
+					UObject* ObjValue = ObjProperty->GetObjectPropertyValue( (const void*)( ShadowData + Cmd.Offset ) );
+					if (ObjValue && !ObjValue->IsNameStableForNetworking())
+					{
+						// skip object references without a stable name
+						return;
+					}
+
+					if (ObjValue && ObjReferences)
+					{
+						ObjReferences->AddUnique(ObjValue);
+					}
+				}
+			}
+
+			StoreProperty( Cmd, (void*)( Data + Cmd.Offset ), (const void*)( ShadowData + Cmd.Offset ) );
+
+			if ( RepNotifies && Parent.Property->HasAnyPropertyFlags( CPF_RepNotify ) )
+			{
+				RepNotifies->AddUnique( Parent.Property );
+			}
+		}
+	}
+
+	TArray<UProperty*>* RepNotifies;
+	TArray<UObject*>* ObjReferences;
+	bool bDifferent;
+};
+
+bool FRepLayout::DiffStableProperties( TArray<UProperty*>* RepNotifies, TArray<UObject*>* ObjReferences, void* RESTRICT Destination, const void* RESTRICT Source ) const
+{
+	FDiffStablePropertiesImpl DiffPropertiesImpl( RepNotifies, ObjReferences, Parents, Cmds );
 
 	DiffPropertiesImpl.ProcessCmds( (uint8*)Destination, (uint8*)Source );
 
@@ -3387,6 +3511,7 @@ void FRepLayout::SerializeProperties_r(
 		// Not concerned with unmapped guids because object references can't be shared
 		if (SharedPropInfo)
 		{
+			GNumSharedSerializationHit++;
 #if !(UE_BUILD_SHIPPING || UE_BUILD_TEST)
 			if (CVarNetVerifyShareSerializedData->GetInt() != 0 && Ar.IsSaving())
 			{
@@ -3418,6 +3543,7 @@ void FRepLayout::SerializeProperties_r(
 		}
 		else
 		{
+			GNumSharedSerializationMiss++;
 			if ( !Cmd.Property->NetSerializeItem( Ar, Map, (void*)( (uint8*)Data + Cmd.Offset ) ) )
 			{
 				bHasUnmapped = true;
@@ -3623,7 +3749,7 @@ void FRepLayout::ClearSharedSerializationForRPC()
 	SharedInfoRPCParentsChanged.Reset();
 }
 
-void FRepLayout::SendPropertiesForRPC( UObject* Object, UFunction * Function, UActorChannel * Channel, FNetBitWriter & Writer, void* Data ) const
+void FRepLayout::SendPropertiesForRPC( UFunction * Function, UActorChannel * Channel, FNetBitWriter & Writer, void* Data ) const
 {
 	check( Function == Owner );
 
@@ -3836,7 +3962,7 @@ void FRepLayout::InitChangedTracker(FRepChangedPropertyTracker * ChangedTracker)
 void FRepLayout::InitShadowData(
 	TArray< uint8, TAlignedHeapAllocator<16> >&	ShadowData,
 	UClass *									InObjectClass,
-	uint8 *										Src ) const
+	const uint8 * const							Src ) const
 {
 	ShadowData.Empty();
 	ShadowData.AddZeroed( InObjectClass->GetDefaultsCount() );
@@ -3851,7 +3977,7 @@ void FRepLayout::InitShadowData(
 void FRepLayout::InitRepState(
 	FRepState*									RepState,
 	UClass *									InObjectClass, 
-	uint8 *										Src, 
+	const uint8 * const							Src, 
 	TSharedPtr< FRepChangedPropertyTracker > &	InRepChangedPropertyTracker ) const
 {
 	InitShadowData( RepState->StaticBuffer, InObjectClass, Src );
@@ -3883,7 +4009,7 @@ void FRepLayout::ConstructProperties( TArray< uint8, TAlignedHeapAllocator<16> >
 	}
 }
 
-void FRepLayout::InitProperties( TArray< uint8, TAlignedHeapAllocator<16> >& ShadowData, uint8* Src ) const
+void FRepLayout::InitProperties( TArray< uint8, TAlignedHeapAllocator<16> >& ShadowData, const uint8* const Src ) const
 {
 	LLM_SCOPE(ELLMTag::Networking);
 
diff --git a/Engine/Source/Runtime/Engine/Private/UnrealClient.cpp b/Engine/Source/Runtime/Engine/Private/UnrealClient.cpp
--- a/Engine/Source/Runtime/Engine/Private/UnrealClient.cpp
+++ b/Engine/Source/Runtime/Engine/Private/UnrealClient.cpp
@@ -1035,7 +1035,7 @@ FViewport::FViewport(FViewportClient* InViewportClient):
 	if (GIsEditor) 
 	{
 		GConfig->GetInt( TEXT("UnrealEd.HitProxy"), TEXT("HitProxySize"), (int32&)HitProxySize, GEditorIni );
-		FMath::Clamp( HitProxySize, (uint32)1, (uint32)MAX_HITPROXYSIZE );
+		HitProxySize = FMath::Clamp( HitProxySize, (uint32)1, (uint32)MAX_HITPROXYSIZE );
 	}
 
 	// Cache the viewport client's hit proxy storage requirement.
diff --git a/Engine/Source/Runtime/Engine/Private/UserInterface/Canvas.cpp b/Engine/Source/Runtime/Engine/Private/UserInterface/Canvas.cpp
--- a/Engine/Source/Runtime/Engine/Private/UserInterface/Canvas.cpp
+++ b/Engine/Source/Runtime/Engine/Private/UserInterface/Canvas.cpp
@@ -2044,12 +2044,13 @@ void UCanvas::K2_DrawBorder(UTexture* BorderTexture, UTexture* BackgroundTexture
 	}
 }
 
-void UCanvas::K2_DrawBox(FVector2D ScreenPosition, FVector2D ScreenSize, float Thickness)
+void UCanvas::K2_DrawBox(FVector2D ScreenPosition, FVector2D ScreenSize, float Thickness, FLinearColor RenderColor)
 {
 	if (ScreenSize.X > 0.0f && ScreenSize.Y > 0.0f && Canvas)
 	{
 		FCanvasBoxItem BoxItem(ScreenPosition, ScreenSize);
 		BoxItem.LineThickness = Thickness;
+		BoxItem.SetColor(RenderColor);
 		DrawItem(BoxItem);
 	}
 }
diff --git a/Engine/Source/Runtime/Engine/Private/World.cpp b/Engine/Source/Runtime/Engine/Private/World.cpp
--- a/Engine/Source/Runtime/Engine/Private/World.cpp
+++ b/Engine/Source/Runtime/Engine/Private/World.cpp
@@ -315,6 +315,8 @@ FWorldDelegates::FOnWorldPostActorTick FWorldDelegates::OnWorldPostActorTick;
 FWorldDelegates::FRefreshLevelScriptActionsEvent FWorldDelegates::RefreshLevelScriptActions;
 #endif // WITH_EDITOR
 
+UWorld::FOnWorldInitializedActors FWorldDelegates::OnWorldInitializedActors;
+
 UWorld::UWorld( const FObjectInitializer& ObjectInitializer )
 : UObject(ObjectInitializer)
 , ActiveLevelCollectionIndex(INDEX_NONE)
@@ -3614,6 +3616,14 @@ void UWorld::InitializeActorsForPlay(const FURL& InURL, bool bResetTime)
 		ViewportConsole->BuildRuntimeAutoCompleteList();
 	}
 
+	// Let others know the actors are initialized. There are two versions here for convenience.
+	FActorsInitializedParams OnActorInitParams(this, bResetTime);
+
+	OnActorsInitialized.Broadcast(OnActorInitParams);
+	FWorldDelegates::OnWorldInitializedActors.Broadcast(OnActorInitParams); // Global notification
+
+	// FIXME: Nav and AI system should now use above delegate
+
 	// let all subsystems/managers know:
 	// @note if UWorld starts to host more of these it might a be a good idea to create a generic IManagerInterface of some sort
 	if (NavigationSystem != nullptr)
@@ -3913,8 +3923,7 @@ void UWorld::AddNetworkActor( AActor* Actor )
 	{
 		if (Driver != nullptr)
 		{
-			// Special case the demo net driver, since actors currently only have one associated NetDriverName.
-			Driver->GetNetworkObjectList().FindOrAdd(Actor, Driver->NetDriverName);
+			Driver->AddNetworkActor(Actor);
 		}
 	});
 }
@@ -3927,7 +3936,7 @@ void UWorld::RemoveNetworkActor( AActor* Actor )
 		{
 			if (Driver != nullptr)
 			{
-				Driver->GetNetworkObjectList().Remove(Actor);
+				Driver->RemoveNetworkActor(Actor);
 			}
 		});
 	}
diff --git a/Engine/Source/Runtime/Engine/Public/EngineStats.h b/Engine/Source/Runtime/Engine/Public/EngineStats.h
--- a/Engine/Source/Runtime/Engine/Public/EngineStats.h
+++ b/Engine/Source/Runtime/Engine/Public/EngineStats.h
@@ -72,6 +72,7 @@ DECLARE_CYCLE_STAT_EXTERN(TEXT("Runtime Movie Tick Time"),STAT_RuntimeMovieScene
 DECLARE_CYCLE_STAT_EXTERN(TEXT("Finish Async Trace Time"),STAT_FinishAsyncTraceTickTime,STATGROUP_Game, );
 DECLARE_CYCLE_STAT_EXTERN(TEXT("Net Broadcast Tick Time"),STAT_NetBroadcastTickTime,STATGROUP_Game, );
 DECLARE_CYCLE_STAT_EXTERN(TEXT("ServerReplicateActors Time"),STAT_NetServerRepActorsTime,STATGROUP_Game, );
+DECLARE_FLOAT_COUNTER_STAT_EXTERN(TEXT("ServerReplicateActors Gather Prioritze Time"),STAT_NetServerGatherPrioritizeRepActorsTime,STATGROUP_Game, );
 DECLARE_CYCLE_STAT_EXTERN(TEXT("Consider Actors Time"),STAT_NetConsiderActorsTime,STATGROUP_Game, );
 DECLARE_CYCLE_STAT_EXTERN(TEXT("Update Unmapped Objects Time"),STAT_NetUpdateUnmappedObjectsTime,STATGROUP_Game, );
 DECLARE_CYCLE_STAT_EXTERN(TEXT("Inital Dormant Time"),STAT_NetInitialDormantCheckTime,STATGROUP_Game, );
@@ -181,8 +182,6 @@ DECLARE_DWORD_ACCUMULATOR_STAT_EXTERN(TEXT("Num Considered Actors"),STAT_NumCons
 DECLARE_DWORD_ACCUMULATOR_STAT_EXTERN(TEXT("Num Prioritized Actors"),STAT_PrioritizedActors,STATGROUP_Net, );
 DECLARE_DWORD_ACCUMULATOR_STAT_EXTERN(TEXT("Num Relevant Actors"),STAT_NumRelevantActors,STATGROUP_Net, );
 DECLARE_DWORD_ACCUMULATOR_STAT_EXTERN(TEXT("Num Relevant Deleted Actors"),STAT_NumRelevantDeletedActors,STATGROUP_Net, );
-DECLARE_DWORD_ACCUMULATOR_STAT_EXTERN(TEXT("Num Replicated Actor Attempts"),STAT_NumReplicatedActorAttempts,STATGROUP_Net, );
-DECLARE_DWORD_ACCUMULATOR_STAT_EXTERN(TEXT("Num Replicated Actors Sent"),STAT_NumReplicatedActors,STATGROUP_Net, );
 DECLARE_DWORD_ACCUMULATOR_STAT_EXTERN(TEXT("Num Actors"),STAT_NumActors,STATGROUP_Net, );
 DECLARE_DWORD_ACCUMULATOR_STAT_EXTERN(TEXT("Num Network Actors"),STAT_NumNetActors,STATGROUP_Net, );
 DECLARE_DWORD_ACCUMULATOR_STAT_EXTERN(TEXT("Num Dormant Actors"),STAT_NumDormantActors,STATGROUP_Net, );
@@ -197,6 +196,8 @@ DECLARE_DWORD_ACCUMULATOR_STAT_EXTERN(TEXT("Saturated"),STAT_NetSaturated,STATGR
 DECLARE_DWORD_ACCUMULATOR_STAT_EXTERN(TEXT("Imported NetGuids"),STAT_ImportedNetGuids,STATGROUP_Net, );
 DECLARE_DWORD_ACCUMULATOR_STAT_EXTERN(TEXT("Pending Outer NetGuids"),STAT_PendingOuterNetGuids,STATGROUP_Net, );
 DECLARE_DWORD_ACCUMULATOR_STAT_EXTERN(TEXT("Unmapped Replicators"),STAT_UnmappedReplicators,STATGROUP_Net, );
+DECLARE_DWORD_ACCUMULATOR_STAT_EXTERN(TEXT("Num Replicated Actors Sent"),STAT_NumReplicatedActors,STATGROUP_Net, );
+DECLARE_DWORD_ACCUMULATOR_STAT_EXTERN(TEXT("Num Replicated Actor Bytes Sent"),STAT_NumReplicatedActorBytes,STATGROUP_Net, );
 
 #if !UE_BUILD_SHIPPING
 /**
diff --git a/Engine/Source/Runtime/Engine/Public/Net/DataBunch.h b/Engine/Source/Runtime/Engine/Public/Net/DataBunch.h
--- a/Engine/Source/Runtime/Engine/Public/Net/DataBunch.h
+++ b/Engine/Source/Runtime/Engine/Public/Net/DataBunch.h
@@ -13,6 +13,8 @@
 class UChannel;
 class UNetConnection;
 
+extern const int32 MAX_BUNCH_SIZE;
+
 //
 // A bunch of data to send.
 //
@@ -116,6 +118,7 @@ public:
 	uint8				bPartialFinal;			// The final bunch of a partial bunch
 	uint8				bHasPackageMapExports;	// This bunch has networkGUID name/id pairs
 	uint8				bHasMustBeMappedGUIDs;	// This bunch has guids that must be mapped before we can process this bunch
+	uint8				bIgnoreRPCs;
 
 	FString	ToString()
 	{
@@ -132,7 +135,9 @@ public:
 		Str += FString::Printf(TEXT("bIsReplicationPaused: %d "), bIsReplicationPaused);
 		Str += FString::Printf(TEXT("bReliable: %d "), bReliable);
 		Str += FString::Printf(TEXT("bPartial: %d//%d//%d "), bPartial, bPartialInitial, bPartialFinal);
-		Str += FString::Printf( TEXT( "bHasPackageMapExports: %d " ), bHasPackageMapExports );
+		Str += FString::Printf(TEXT("bHasPackageMapExports: %d "), bHasPackageMapExports );
+		Str += FString::Printf(TEXT("bHasMustBeMappedGUIDs: %d "), bHasMustBeMappedGUIDs );
+		Str += FString::Printf(TEXT("bIgnoreRPCs: %d "), bIgnoreRPCs );
 #else
 		FString Str = FString::Printf(TEXT("Channel[%d]. Seq %d. PacketId: %d"), ChIndex, ChSequence, PacketId);
 #endif
diff --git a/Engine/Source/Runtime/Engine/Public/Net/DataReplication.h b/Engine/Source/Runtime/Engine/Public/Net/DataReplication.h
--- a/Engine/Source/Runtime/Engine/Public/Net/DataReplication.h
+++ b/Engine/Source/Runtime/Engine/Public/Net/DataReplication.h
@@ -133,7 +133,7 @@ public:
 	TMap< UProperty*, TArray<uint8> >				RepNotifyMetaData;
 
 	TSharedPtr< FRepLayout >						RepLayout;
-	FRepState *										RepState;
+	TSharedPtr< FRepState > 						RepState;
 
 	TSet< FNetworkGUID >							ReferencedGuids;
 	int32											TrackedGuidMemoryBytes;
diff --git a/Engine/Source/Runtime/Engine/Public/Net/RepLayout.h b/Engine/Source/Runtime/Engine/Public/Net/RepLayout.h
--- a/Engine/Source/Runtime/Engine/Public/Net/RepLayout.h
+++ b/Engine/Source/Runtime/Engine/Public/Net/RepLayout.h
@@ -22,6 +22,14 @@ class UPackageMapClient;
 // Properties will be copied in here so memory needs aligned to largest type
 typedef TArray< uint8, TAlignedHeapAllocator<16> > FRepStateStaticBuffer;
 
+enum class EDiffPropertiesFlags : uint32
+{
+	None							= 0,
+	Sync							= ( 1 << 0 ),
+	IncludeConditionalProperties	= ( 1 << 1 )
+};
+
+ENUM_CLASS_FLAGS(EDiffPropertiesFlags);
 
 class FRepChangedParent
 {
@@ -460,12 +468,12 @@ public:
 	void InitShadowData(
 		TArray< uint8, TAlignedHeapAllocator<16> >&	ShadowData,
 		UClass *									InObjectClass,
-		uint8 *										Src ) const;
+		const uint8 * const							Src ) const;
 
 	void InitRepState( 
 		FRepState *									RepState, 
 		UClass *									InObjectClass, 
-		uint8 *										Src, 
+		const uint8 * const							Src, 
 		TSharedPtr< FRepChangedPropertyTracker > &	InRepChangedPropertyTracker ) const;
 
 	void InitChangedTracker( FRepChangedPropertyTracker * ChangedTracker ) const;
@@ -512,20 +520,24 @@ public:
 
 	void MergeChangeList( const uint8* RESTRICT Data, const TArray< uint16 >& Dirty1, const TArray< uint16 >& Dirty2, TArray< uint16 >& MergedDirty ) const;
 
+	DEPRECATED(4.20, "Use the DiffProperties overload with the EDiffPropertiesFlags parameter")
 	bool DiffProperties( TArray<UProperty*>* RepNotifies, void* RESTRICT Destination, const void* RESTRICT Source, const bool bSync ) const;
 
+	bool DiffProperties( TArray<UProperty*>* RepNotifies, void* RESTRICT Destination, const void* RESTRICT Source, const EDiffPropertiesFlags Flags ) const;
+	bool DiffStableProperties( TArray<UProperty*>* RepNotifies, TArray<UObject*>* ObjReferences, void* RESTRICT Destination, const void* RESTRICT Source ) const;
+
 	void GetLifetimeCustomDeltaProperties(TArray< int32 > & OutCustom, TArray< ELifetimeCondition >	& OutConditions);
 
 	// RPC support
 	void InitFromFunction( UFunction * InFunction, const UNetConnection* ServerConnection = nullptr );
-	void SendPropertiesForRPC( UObject* Object, UFunction * Function, UActorChannel * Channel, FNetBitWriter & Writer, void* Data ) const;
+	void SendPropertiesForRPC( UFunction * Function, UActorChannel * Channel, FNetBitWriter & Writer, void* Data ) const;
 	void ReceivePropertiesForRPC( UObject* Object, UFunction * Function, UActorChannel * Channel, FNetBitReader & Reader, void* Data, TSet<FNetworkGUID>& UnmappedGuids) const;
 
 	// Builds shared serialization state for a multicast rpc
-	void BuildSharedSerializationForRPC( void* Data );
+	void ENGINE_API BuildSharedSerializationForRPC( void* Data );
 
 	// Clears shared serialization state for a multicast rpc
-	void ClearSharedSerializationForRPC();
+	void ENGINE_API ClearSharedSerializationForRPC();
 
 	// Struct support
 	ENGINE_API void SerializePropertiesForStruct( UStruct * Struct, FBitArchive & Ar, UPackageMap	* Map, void* Data, bool & bHasUnmapped ) const;	
@@ -725,7 +737,7 @@ private:
 	void BuildHandleToCmdIndexTable_r( const int32 CmdStart, const int32 CmdEnd, TArray< FHandleToCmdIndex >& HandleToCmdIndex );
 
 	void ConstructProperties( TArray< uint8, TAlignedHeapAllocator<16> >& ShadowData ) const;
-	void InitProperties( TArray< uint8, TAlignedHeapAllocator<16> >& ShadowData, uint8* Src ) const;
+	void InitProperties( TArray< uint8, TAlignedHeapAllocator<16> >& ShadowData, const uint8* const Src ) const;
 	void DestructProperties( FRepStateStaticBuffer& RepStateStaticBuffer ) const;
 
 	TArray< FRepParentCmd >		Parents;
diff --git a/Engine/Source/Runtime/Sockets/Private/Android/SocketSubsystemAndroid.cpp b/Engine/Source/Runtime/Sockets/Private/Android/SocketSubsystemAndroid.cpp
--- a/Engine/Source/Runtime/Sockets/Private/Android/SocketSubsystemAndroid.cpp
+++ b/Engine/Source/Runtime/Sockets/Private/Android/SocketSubsystemAndroid.cpp
@@ -133,7 +133,7 @@ TSharedRef<FInternetAddr> FSocketSubsystemAndroid::GetLocalHostAddr(FOutputDevic
 		if (bIsValid)
 		{
 			UE_LOG(LogSockets, Warning, TEXT("FSocketSubsystemAndroid::GetLocalHostAddr Using MULTIHOME"));
-			return Addr;
+			return TempAddr;
 		}
 	}
 
diff --git a/Engine/Source/Runtime/Sockets/Private/BSDIPv6Sockets/SocketsBSDIPv6.cpp b/Engine/Source/Runtime/Sockets/Private/BSDIPv6Sockets/SocketsBSDIPv6.cpp
--- a/Engine/Source/Runtime/Sockets/Private/BSDIPv6Sockets/SocketsBSDIPv6.cpp
+++ b/Engine/Source/Runtime/Sockets/Private/BSDIPv6Sockets/SocketsBSDIPv6.cpp
@@ -183,50 +183,56 @@ bool FSocketBSDIPv6::Send(const uint8* Data, int32 Count, int32& BytesSent)
 
 bool FSocketBSDIPv6::RecvFrom(uint8* Data, int32 BufferSize, int32& BytesRead, FInternetAddr& Source, ESocketReceiveFlags::Type Flags)
 {
+	bool bSuccess = false;
+	const bool bStreamSocket = (SocketType == SOCKTYPE_Streaming);
+	const int TranslatedFlags = TranslateFlags(Flags);
 	SOCKLEN Size = sizeof(sockaddr_in6);
 	sockaddr& Addr = *(FInternetAddrBSDIPv6&)Source;
 
-	const int TranslatedFlags = TranslateFlags(Flags);
-
 	// Read into the buffer and set the source address
 	BytesRead = recvfrom(Socket, (char*)Data, BufferSize, TranslatedFlags, &Addr, &Size);
 
 //	NETWORK_PROFILER(FSocket::RecvFrom(Data,BufferSize,BytesRead,Source));
 
-	if (BytesRead < 0 && SocketSubsystem->TranslateErrorCode(BytesRead) == SE_EWOULDBLOCK)
+	if (BytesRead >= 0)
 	{
-		// EWOULDBLOCK is not an error condition
-		BytesRead = 0;
+		// For Streaming sockets, 0 indicates a graceful failure
+		bSuccess = !bStreamSocket || (BytesRead > 0);
 	}
-	else if (BytesRead <= 0) // 0 means gracefully closed
+	else
 	{
+		// For Streaming sockets, don't treat SE_EWOULDBLOCK as an error
+		bSuccess = bStreamSocket && (SocketSubsystem->TranslateErrorCode(BytesRead) == SE_EWOULDBLOCK);
 		BytesRead = 0;
-		return false;
 	}
 
-	return true;
+	return bSuccess;
 }
 
 
 bool FSocketBSDIPv6::Recv(uint8* Data, int32 BufferSize, int32& BytesRead, ESocketReceiveFlags::Type Flags)
 {
+	bool bSuccess = false;
+	const bool bStreamSocket = (SocketType == SOCKTYPE_Streaming);
 	const int TranslatedFlags = TranslateFlags(Flags);
+
 	BytesRead = recv(Socket, (char*)Data, BufferSize, TranslatedFlags);
 
 //	NETWORK_PROFILER(FSocket::Recv(Data,BufferSize,BytesRead));
 
-	if (BytesRead < 0 && SocketSubsystem->TranslateErrorCode(BytesRead) == SE_EWOULDBLOCK)
+	if (BytesRead >= 0)
 	{
-		// EWOULDBLOCK is not an error condition
-		BytesRead = 0;
+		// For Streaming sockets, 0 indicates a graceful failure
+		bSuccess = !bStreamSocket || (BytesRead > 0);
 	}
-	else if (BytesRead <= 0) // 0 means gracefully closed
+	else
 	{
+		// For Streaming sockets, don't treat SE_EWOULDBLOCK as an error
+		bSuccess = bStreamSocket && (SocketSubsystem->TranslateErrorCode(BytesRead) == SE_EWOULDBLOCK);
 		BytesRead = 0;
-		return false;
 	}
 
-	return true;
+	return bSuccess;
 }
 
 bool FSocketBSDIPv6::Wait(ESocketWaitConditions::Type Condition, FTimespan WaitTime)
@@ -427,9 +433,11 @@ int32 FSocketBSDIPv6::GetPortNo(void)
 	{
 		check(SocketSubsystem);
 		UE_LOG(LogSockets, Error, TEXT("Failed to read address for socket (%s)"), SocketSubsystem->GetSocketError());
+		
+		return 0;
 	}
 
-	// Read the port number
+	// Convert big endian port to native endian port.
 	return ntohs(Addr.sin6_port);
 }
 
diff --git a/Engine/Source/Runtime/Sockets/Private/BSDSockets/SocketsBSD.cpp b/Engine/Source/Runtime/Sockets/Private/BSDSockets/SocketsBSD.cpp
--- a/Engine/Source/Runtime/Sockets/Private/BSDSockets/SocketsBSD.cpp
+++ b/Engine/Source/Runtime/Sockets/Private/BSDSockets/SocketsBSD.cpp
@@ -158,53 +158,64 @@ bool FSocketBSD::Send(const uint8* Data, int32 Count, int32& BytesSent)
 
 bool FSocketBSD::RecvFrom(uint8* Data, int32 BufferSize, int32& BytesRead, FInternetAddr& Source, ESocketReceiveFlags::Type Flags)
 {
+	bool bSuccess = false;
+	const bool bStreamSocket = (SocketType == SOCKTYPE_Streaming);
+	const int TranslatedFlags = TranslateFlags(Flags);
 	SOCKLEN Size = sizeof(sockaddr_in);
 	sockaddr& Addr = *(FInternetAddrBSD&)Source;
 
-	const int TranslatedFlags = TranslateFlags(Flags);
-
 	// Read into the buffer and set the source address
 	BytesRead = recvfrom(Socket, (char*)Data, BufferSize, TranslatedFlags, &Addr, &Size);
 //	NETWORK_PROFILER(FSocket::RecvFrom(Data,BufferSize,BytesRead,Source));
 
-	if (BytesRead < 0 && SocketSubsystem->TranslateErrorCode(BytesRead) == SE_EWOULDBLOCK)
+	if (BytesRead >= 0)
 	{
-		// EWOULDBLOCK is not an error condition
-		BytesRead = 0;
+		// For Streaming sockets, 0 indicates a graceful failure
+		bSuccess = !bStreamSocket || (BytesRead > 0);
 	}
-	else if (BytesRead <= 0) // 0 means gracefully closed
+	else
 	{
+		// For Streaming sockets, don't treat SE_EWOULDBLOCK as an error
+		bSuccess = bStreamSocket && (SocketSubsystem->TranslateErrorCode(BytesRead) == SE_EWOULDBLOCK);
 		BytesRead = 0;
-		return false;
 	}
 
-	LastActivityTime = FDateTime::UtcNow();
+	if (bSuccess)
+	{
+		LastActivityTime = FDateTime::UtcNow();
+	}
 
-	return true;
+	return bSuccess;
 }
 
 
 bool FSocketBSD::Recv(uint8* Data, int32 BufferSize, int32& BytesRead, ESocketReceiveFlags::Type Flags)
 {
+	bool bSuccess = false;
+	const bool bStreamSocket = (SocketType == SOCKTYPE_Streaming);
 	const int TranslatedFlags = TranslateFlags(Flags);
 	BytesRead = recv(Socket, (char*)Data, BufferSize, TranslatedFlags);
 
 //	NETWORK_PROFILER(FSocket::Recv(Data,BufferSize,BytesRead));
 
-	if (BytesRead < 0 && SocketSubsystem->TranslateErrorCode(BytesRead) == SE_EWOULDBLOCK)
+	if (BytesRead >= 0)
 	{
-		// EWOULDBLOCK is not an error condition
-		BytesRead = 0;
+		// For Streaming sockets, 0 indicates a graceful failure
+		bSuccess = !bStreamSocket || (BytesRead > 0);
 	}
-	else if (BytesRead <= 0) // 0 means gracefully closed
+	else
 	{
+		// For Streaming sockets, don't treat SE_EWOULDBLOCK as an error
+		bSuccess = bStreamSocket && (SocketSubsystem->TranslateErrorCode(BytesRead) == SE_EWOULDBLOCK);
 		BytesRead = 0;
-		return false;
 	}
 
-	LastActivityTime = FDateTime::UtcNow();
+	if (bSuccess)
+	{
+		LastActivityTime = FDateTime::UtcNow();
+	}
 
-	return true;
+	return bSuccess;
 }
 
 
@@ -422,9 +433,11 @@ int32 FSocketBSD::GetPortNo(void)
 	{
 		check(SocketSubsystem);
 		UE_LOG(LogSockets, Error, TEXT("Failed to read address for socket (%s)"), SocketSubsystem->GetSocketError());
+		
+		return 0;
 	}
 
-	// Read the port number
+	// Convert big endian port to native endian port.
 	return ntohs(Addr.sin_port);
 }
 
diff --git a/Engine/Source/Runtime/Sockets/Private/Unix/SocketSubsystemUnix.cpp b/Engine/Source/Runtime/Sockets/Private/Unix/SocketSubsystemUnix.cpp
--- a/Engine/Source/Runtime/Sockets/Private/Unix/SocketSubsystemUnix.cpp
+++ b/Engine/Source/Runtime/Sockets/Private/Unix/SocketSubsystemUnix.cpp
@@ -111,7 +111,7 @@ TSharedRef<FInternetAddr> FSocketSubsystemUnix::GetLocalHostAddr(FOutputDevice&
 		TempAddr->SetIp(Home, bIsValid);
 		if (bIsValid)
 		{
-			return Addr;
+			return TempAddr;
 		}
 	}
 
