Copying //UE4/Dev-Framework to //UE4/Dev-Main (Source: //UE4/Dev-Framework @ 4058146)

diff --git a/Engine/Source/Editor/UnrealEd/Private/Kismet2/BlueprintEditorUtils.cpp b/Engine/Source/Editor/UnrealEd/Private/Kismet2/BlueprintEditorUtils.cpp
--- a/Engine/Source/Editor/UnrealEd/Private/Kismet2/BlueprintEditorUtils.cpp
+++ b/Engine/Source/Editor/UnrealEd/Private/Kismet2/BlueprintEditorUtils.cpp
@@ -835,6 +835,33 @@ UEdGraphPin* FBlueprintEditorUtils::FindFirstCompilerRelevantLinkedPin(UEdGraphP
 	return RelevantNodeLinks.Num() > 0 ? RelevantNodeLinks[0].LinkedPin : nullptr;
 }
 
+void FBlueprintEditorUtils::RemoveAllLocalBookmarks(const UBlueprint* ForBlueprint)
+{
+	if (ForBlueprint)
+	{
+		bool bSaveConfig = false;
+		const FString BPPackageName = ForBlueprint->GetOutermost()->GetName();
+		UBlueprintEditorSettings* LocalSettings = GetMutableDefault<UBlueprintEditorSettings>();
+		for (int32 i = 0; i < LocalSettings->BookmarkNodes.Num(); ++i)
+		{
+			const FBPEditorBookmarkNode& BookmarkNode = LocalSettings->BookmarkNodes[i];
+			const FEditedDocumentInfo* BookmarkInfo = LocalSettings->Bookmarks.Find(BookmarkNode.NodeGuid);
+			if (BookmarkInfo != nullptr && BookmarkInfo->EditedObjectPath.GetLongPackageName() == BPPackageName)
+			{
+				bSaveConfig = true;
+
+				LocalSettings->BookmarkNodes.RemoveAt(i--);
+				LocalSettings->Bookmarks.Remove(BookmarkNode.NodeGuid);
+			}
+		}
+
+		if (bSaveConfig)
+		{
+			LocalSettings->SaveConfig();
+		}
+	}
+}
+
 /** 
  * Check FKismetCompilerContext::SetCanEverTickForActor
  */
diff --git a/Engine/Source/Editor/UnrealEd/Private/Kismet2/DebuggerCommands.cpp b/Engine/Source/Editor/UnrealEd/Private/Kismet2/DebuggerCommands.cpp
--- a/Engine/Source/Editor/UnrealEd/Private/Kismet2/DebuggerCommands.cpp
+++ b/Engine/Source/Editor/UnrealEd/Private/Kismet2/DebuggerCommands.cpp
@@ -232,14 +232,17 @@ static void LeaveDebuggingMode()
 		GUnrealEd->PlayWorld->bDebugPauseExecution = false;
 	}
 
-	if( FSlateApplication::Get().InKismetDebuggingMode() )
+	// Determine whether or not we are resuming play.
+	const bool bIsResumingPlay = !FKismetDebugUtilities::IsSingleStepping() && !GEditor->ShouldEndPlayMap();
+
+	if( FSlateApplication::Get().InKismetDebuggingMode() && bIsResumingPlay )
 	{
-		// Focus the game view port when resuming from debugging
+		// Focus the game view port when resuming from debugging.
 		FModuleManager::GetModuleChecked<FLevelEditorModule>("LevelEditor").FocusPIEViewport();
 	}
 
-	// Tell the application to stop ticking in this stack frame
-	FSlateApplication::Get().LeaveDebuggingMode( FKismetDebugUtilities::IsSingleStepping() );
+	// Tell the application to stop ticking in this stack frame. The parameter controls whether or not to recapture the mouse to the game viewport.
+	FSlateApplication::Get().LeaveDebuggingMode( !bIsResumingPlay );
 }
 
 
diff --git a/Engine/Source/Editor/UnrealEd/Private/Kismet2/Kismet2NameValidators.cpp b/Engine/Source/Editor/UnrealEd/Private/Kismet2/Kismet2NameValidators.cpp
--- a/Engine/Source/Editor/UnrealEd/Private/Kismet2/Kismet2NameValidators.cpp
+++ b/Engine/Source/Editor/UnrealEd/Private/Kismet2/Kismet2NameValidators.cpp
@@ -47,7 +47,7 @@ FText INameValidatorInterface::GetErrorText(const FString& Name, EValidatorResul
 		ErrorText = LOCTEXT("NameTooLong_Error", "Names must have fewer than 100 characters!");
 		break;
 	case EValidatorResult::LocallyInUse:
-		ErrorText  = LOCTEXT("LocallyInUse_Error", "Conflicts with another another object in the same scope!");
+		ErrorText  = LOCTEXT("LocallyInUse_Error", "Conflicts with another object in the same scope!");
 		break;
 	case EValidatorResult::Ok:
 		break;
diff --git a/Engine/Source/Editor/UnrealEd/Private/Kismet2/KismetDebugUtilities.cpp b/Engine/Source/Editor/UnrealEd/Private/Kismet2/KismetDebugUtilities.cpp
--- a/Engine/Source/Editor/UnrealEd/Private/Kismet2/KismetDebugUtilities.cpp
+++ b/Engine/Source/Editor/UnrealEd/Private/Kismet2/KismetDebugUtilities.cpp
@@ -568,6 +568,10 @@ void FKismetDebugUtilities::CheckBreakConditions(UEdGraphNode* NodeStoppedAt, bo
 				}
 			}
 		}
+		else if (NodeStoppedAt != Data.MostRecentStoppedNode)
+		{
+			Data.MostRecentStoppedNode = nullptr;
+		}
 	}
 	
 	if (InOutBreakExecution)
@@ -688,10 +692,11 @@ void FKismetDebugUtilities::AttemptToBreakExecution(UBlueprint* BlueprintObj, co
 	// Now enter within-the-frame debugging mode
 	if (bShouldInStackDebug)
 	{
+		const TArray<const FFrame*>& ScriptStack = FBlueprintExceptionTracker::Get().ScriptStack;
 		Data.LastExceptionMessage = Info.GetDescription();
 		FKismetEditorUtilities::BringKismetToFocusAttentionOnObject(NodeStoppedAt);
-		CallStackViewer::UpdateDisplayedCallstack(FBlueprintExceptionTracker::Get().ScriptStack);
-		WatchViewer::UpdateDisplayedWatches(BlueprintObj);
+		CallStackViewer::UpdateDisplayedCallstack(ScriptStack);
+		WatchViewer::UpdateDisplayedWatches(ScriptStack);
 		FSlateApplication::Get().EnterDebuggingMode();
 	}
 }
@@ -736,8 +741,11 @@ void FKismetDebugUtilities::NotifyDebuggerOfEndOfGameFrame(UWorld* CurrentWorld)
 }
 
 bool FKismetDebugUtilities::IsSingleStepping()
-{ 
-	return FKismetDebugUtilitiesData::Get().bIsSingleStepping; 
+{
+	const FKismetDebugUtilitiesData& Data = FKismetDebugUtilitiesData::Get();
+	return Data.bIsSingleStepping
+		|| Data.bIsSteppingOut
+		|| Data.TargetGraphStackDepth != INDEX_NONE; 
 }
 
 //////////////////////////////////////////////////////////////////////////
@@ -1151,7 +1159,7 @@ FKismetDebugUtilities::EWatchTextResult FKismetDebugUtilities::FindDebuggingData
 				return EWTR_NoDebugObject;
 			}
 
-			void* PropertyBase = NULL;
+			void* PropertyBase = nullptr;
 
 			// Walk up the stack frame to see if we can find a function scope that contains the property as a local
 			for (const FFrame* TestFrame = Data.StackFrameAtIntraframeDebugging; TestFrame != NULL; TestFrame = TestFrame->PreviousFrame)
@@ -1233,6 +1241,35 @@ FKismetDebugUtilities::EWatchTextResult FKismetDebugUtilities::FindDebuggingData
 				}
 			}
 
+			// If we still haven't found a result, try changing the active object to whatever is passed into the self pin.
+			if (!PropertyBase)
+			{
+				UEdGraphNode* WatchNode = WatchPin->GetOwningNode();
+
+				if (WatchNode)
+				{
+					UEdGraphPin* SelfPin = WatchNode->FindPin(TEXT("self"));
+					if (SelfPin && SelfPin != WatchPin)
+					{
+						UProperty* SelfPinProperty = nullptr;
+						void* SelfPinData = nullptr;
+						void* SelfPinDelta = nullptr;
+						UObject* SelfPinParent = nullptr;
+						FKismetDebugUtilities::EWatchTextResult Result = FindDebuggingData(Blueprint, ActiveObject, SelfPin, SelfPinProperty, SelfPinData, SelfPinDelta, SelfPinParent);
+						UObjectPropertyBase* SelfPinPropertyBase = Cast<UObjectPropertyBase>(SelfPinProperty);
+						if (Result == EWTR_Valid && SelfPinPropertyBase != nullptr)
+						{
+							void* PropertyValue = SelfPinProperty->ContainerPtrToValuePtr<void>(SelfPinData);
+							UObject* TempActiveObject = SelfPinPropertyBase->GetObjectPropertyValue(PropertyValue);
+							if (TempActiveObject)
+							{
+								return FindDebuggingData(Blueprint, TempActiveObject, WatchPin, OutProperty, OutData, OutDelta, OutParent);
+							}
+						}
+					}
+				}
+			}
+
 			// Now either print out the variable value, or that it was out-of-scope
 			if (PropertyBase != nullptr)
 			{
diff --git a/Engine/Source/Editor/UnrealEd/Private/ObjectTools.cpp b/Engine/Source/Editor/UnrealEd/Private/ObjectTools.cpp
--- a/Engine/Source/Editor/UnrealEd/Private/ObjectTools.cpp
+++ b/Engine/Source/Editor/UnrealEd/Private/ObjectTools.cpp
@@ -911,11 +911,93 @@ namespace ObjectTools
 			FForceReplaceInfo ReplaceInfo;
 			FForceReplaceInfo GeneratedClassReplaceInfo;
 
+			bool bNeedsGarbageCollection = false;
+
 			// Scope the reregister context below to complete after object deletion and before garbage collection
 			{
 				// Replacing references inside already loaded objects could cause rendering issues, so globally detach all components from their scenes for now
 				FGlobalComponentRecreateRenderStateContext ReregisterContext;
-				
+
+				// First, make sure that the class we're consolidating to has its hierarchy fixed so
+				// that we don't create a cycle (e.g. directly or indirectly parent it to itself):
+				UClass* ClassToConsolidateTo = nullptr;
+				if (ObjectToConsolidateTo)
+				{
+					if (UBlueprint* BlueprintObject = Cast<UBlueprint>(ObjectToConsolidateTo))
+					{
+						ClassToConsolidateTo = BlueprintObject->GeneratedClass;
+
+						if (!ClassToConsolidateTo)
+						{
+							ClassToConsolidateTo = UObject::StaticClass();
+						}
+
+						// Don't parent a blueprint to itself, instead fall back to the part of the 
+						// hierarchy that is not being consolidated. Worst case, fall back to 
+						// UObject::StaticClass():
+						UClass* NewParent = BlueprintObject->ParentClass;
+						UClass* ParentIter = NewParent;
+						while (ParentIter)
+						{
+							if (ObjectsToConsolidate.Contains(ParentIter->ClassGeneratedBy))
+							{
+								NewParent = ParentIter->GetSuperClass();
+							}
+							ParentIter = ParentIter->GetSuperClass();
+						}
+
+						if (!NewParent || ObjectsToConsolidate.Contains(NewParent->ClassGeneratedBy))
+						{
+							NewParent = UObject::StaticClass();
+						}
+
+						BlueprintObject->ParentClass = NewParent;
+
+						// Recompile the child blueprint to fix up the generated class
+						FKismetEditorUtilities::CompileBlueprint(BlueprintObject, EBlueprintCompileOptions::SkipGarbageCollection);
+					}
+				}
+
+				// Then reparent any direct children to the class we're consolidating to:
+				for (UObject* Object : ObjectsToConsolidate)
+				{
+					if (UBlueprint* BlueprintObject = Cast<UBlueprint>(Object))
+					{
+						if (BlueprintObject->ParentClass != nullptr && BlueprintObject->GeneratedClass)
+						{
+							TArray<UClass*> ChildClasses;
+							GetDerivedClasses(BlueprintObject->GeneratedClass, ChildClasses, false);
+							for(UClass* ChildClass : ChildClasses)
+							{
+								UBlueprint* ChildBlueprint = Cast<UBlueprint>(ChildClass->ClassGeneratedBy);
+								if (ChildBlueprint != nullptr && !ChildClass->HasAnyClassFlags(CLASS_NewerVersionExists))
+								{
+									// Do not reparent and recompile a Blueprint that is going to be deleted.
+									if (ObjectsToConsolidate.Find(ChildBlueprint) == INDEX_NONE)
+									{
+										ChildBlueprint->Modify();
+
+										UClass* NewParent = ClassToConsolidateTo;
+
+										if (!NewParent)
+										{
+											NewParent = UObject::StaticClass();
+										}
+
+										ChildBlueprint->ParentClass = NewParent;
+
+										// Recompile the child blueprint to fix up the generated class
+										FKismetEditorUtilities::CompileBlueprint(ChildBlueprint, EBlueprintCompileOptions::SkipGarbageCollection);
+
+										// Defer garbage collection until after we're done processing the list of objects
+										bNeedsGarbageCollection = true;
+									}
+								}
+							}
+						}
+					}
+				}
+
 				ForceReplaceReferences(ObjectToConsolidateTo, ObjectsToConsolidate, ReplaceInfo);
 
 				if (UBlueprint* ObjectToConsolidateTo_BP = Cast<UBlueprint>(ObjectToConsolidateTo))
@@ -923,6 +1005,7 @@ namespace ObjectTools
 					// Replace all UClass/TSubClassOf properties of generated class.
 					TArray<UObject*> ObjectsToConsolidate_BP;
 					TArray<UClass*> OldGeneratedClasses;
+					TMap<UClass*, UClass*> OldChildClassToOldParentClass;
 					ObjectsToConsolidate_BP.Reserve(ObjectsToConsolidate.Num());
 					OldGeneratedClasses.Reserve(ObjectsToConsolidate.Num());
 					for (UObject* ObjectToConsolidate : ObjectsToConsolidate)
@@ -930,6 +1013,13 @@ namespace ObjectTools
 						UClass* OldGeneratedClass = Cast<UBlueprint>(ObjectToConsolidate)->GeneratedClass;
 						ObjectsToConsolidate_BP.Add(OldGeneratedClass);
 						OldGeneratedClasses.Add(OldGeneratedClass);
+
+						TArray<UClass*> OldChildClasses;
+						GetDerivedClasses(OldGeneratedClass, OldChildClasses, false);
+						for (UClass* OldChildClass : OldChildClasses)
+						{
+							OldChildClassToOldParentClass.Add(OldChildClass, OldGeneratedClass);
+						}
 					}
 
 					ForceReplaceReferences(ObjectToConsolidateTo_BP->GeneratedClass, ObjectsToConsolidate_BP, GeneratedClassReplaceInfo);
@@ -940,6 +1030,12 @@ namespace ObjectTools
 						Cast<UBlueprint>(ObjectsToConsolidate[Index])->GeneratedClass = OldGeneratedClasses[Index];
 					}
 
+					// repair superstruct references:
+					for (const TPair<UClass*,UClass*>& OldChild : OldChildClassToOldParentClass)
+					{
+						OldChild.Key->SetSuperStruct(OldChild.Value);
+					}
+
 					ReplaceInfo.AppendUnique(GeneratedClassReplaceInfo);
 				}
 				DirtiedPackages.Append( ReplaceInfo.DirtiedPackages );
@@ -965,8 +1061,12 @@ namespace ObjectTools
 					}
 				}
 
-				// Clean up the actors we replaced
-				CollectGarbage( GARBAGE_COLLECTION_KEEPFLAGS );
+				bNeedsGarbageCollection = true;
+			}
+
+			if (bNeedsGarbageCollection)
+			{
+				CollectGarbage(GARBAGE_COLLECTION_KEEPFLAGS);
 			}
 
 			FEditorDelegates::OnAssetsPreDelete.Broadcast(ReplaceInfo.ReplaceableObjects);
diff --git a/Engine/Source/Editor/UnrealEd/Private/Settings/SettingsClasses.cpp b/Engine/Source/Editor/UnrealEd/Private/Settings/SettingsClasses.cpp
--- a/Engine/Source/Editor/UnrealEd/Private/Settings/SettingsClasses.cpp
+++ b/Engine/Source/Editor/UnrealEd/Private/Settings/SettingsClasses.cpp
@@ -570,9 +570,30 @@ void UProjectPackagingSettings::PostInitProperties()
 	// Cache the current set of Blueprint assets selected for nativization.
 	CachedNativizeBlueprintAssets = NativizeBlueprintAssets;
 
+	FixCookingPaths();
+
 	Super::PostInitProperties();
 }
 
+void UProjectPackagingSettings::FixCookingPaths()
+{
+	// Fix AlwaysCook/NeverCook paths to use content root
+	for (FDirectoryPath& PathToFix : DirectoriesToAlwaysCook)
+	{
+		if (!PathToFix.Path.IsEmpty() && !PathToFix.Path.StartsWith(TEXT("/"), ESearchCase::CaseSensitive))
+		{
+			PathToFix.Path = FString::Printf(TEXT("/Game/%s"), *PathToFix.Path);
+		}
+	}
+
+	for (FDirectoryPath& PathToFix : DirectoriesToNeverCook)
+	{
+		if (!PathToFix.Path.IsEmpty() && !PathToFix.Path.StartsWith(TEXT("/"), ESearchCase::CaseSensitive))
+		{
+			PathToFix.Path = FString::Printf(TEXT("/Game/%s"), *PathToFix.Path);
+		}
+	}
+}
 
 void UProjectPackagingSettings::PostEditChangeProperty( FPropertyChangedEvent& PropertyChangedEvent )
 {
@@ -582,15 +603,10 @@ void UProjectPackagingSettings::PostEditChangeProperty( FPropertyChangedEvent& P
 		? PropertyChangedEvent.MemberProperty->GetFName()
 		: NAME_None;
 
-	if (Name == FName((TEXT("DirectoriesToAlwaysCook"))))
+	if (Name == FName(TEXT("DirectoriesToAlwaysCook")) || Name == FName(TEXT("DirectoriesToNeverCook")) || Name == NAME_None)
 	{
-		// fix up paths
-		for(int32 PathIndex = 0; PathIndex < DirectoriesToAlwaysCook.Num(); PathIndex++)
-		{
-			FString Path = DirectoriesToAlwaysCook[PathIndex].Path;
-			FPaths::MakePathRelativeTo(Path, FPlatformProcess::BaseDir());
-			DirectoriesToAlwaysCook[PathIndex].Path = Path;
-		}
+		// We need to fix paths for no name updates to catch the reloadconfig call
+		FixCookingPaths();
 	}
 	else if (Name == FName((TEXT("StagingDirectory"))))
 	{
diff --git a/Engine/Source/Editor/UnrealEd/Public/Kismet2/BlueprintEditorUtils.h b/Engine/Source/Editor/UnrealEd/Public/Kismet2/BlueprintEditorUtils.h
--- a/Engine/Source/Editor/UnrealEd/Public/Kismet2/BlueprintEditorUtils.h
+++ b/Engine/Source/Editor/UnrealEd/Public/Kismet2/BlueprintEditorUtils.h
@@ -659,6 +659,13 @@ public:
 	 */
 	static UEdGraphPin* FindFirstCompilerRelevantLinkedPin(UEdGraphPin* FromPin);
 
+	/**
+	 * Removes all local bookmarks that reference the given Blueprint asset.
+	 *
+	 * @param			ForBlueprint	The Blueprint asset for which to remove local Bookmarks.
+	 */
+	static void RemoveAllLocalBookmarks(const UBlueprint* ForBlueprint);
+
 	//////////////////////////////////////////////////////////////////////////
 	// Functions
 
diff --git a/Engine/Source/Programs/UnrealBuildTool/Platform/Windows/PVSToolChain.cs b/Engine/Source/Programs/UnrealBuildTool/Platform/Windows/PVSToolChain.cs
--- a/Engine/Source/Programs/UnrealBuildTool/Platform/Windows/PVSToolChain.cs
+++ b/Engine/Source/Programs/UnrealBuildTool/Platform/Windows/PVSToolChain.cs
@@ -1,4 +1,4 @@
-// Copyright 1998-2018 Epic Games, Inc. All Rights Reserved.
+﻿// Copyright 1998-2018 Epic Games, Inc. All Rights Reserved.
 
 using System;
 using System.Collections.Generic;
diff --git a/Engine/Source/Programs/UnrealHeaderTool/Private/CodeGenerator.cpp b/Engine/Source/Programs/UnrealHeaderTool/Private/CodeGenerator.cpp
--- a/Engine/Source/Programs/UnrealHeaderTool/Private/CodeGenerator.cpp
+++ b/Engine/Source/Programs/UnrealHeaderTool/Private/CodeGenerator.cpp
@@ -4030,14 +4030,15 @@ void FNativeClassHeaderGenerator::ExportNativeFunctionHeader(
 		}
 	}
 
-	if (UProperty* Return = Function->GetReturnProperty())
+	UProperty* ReturnProperty = Function->GetReturnProperty();
+	if (ReturnProperty != nullptr)
 	{
 		FString ExtendedReturnType;
-		FString ReturnType = Return->GetCPPType(&ExtendedReturnType, (FunctionHeaderStyle == EExportFunctionHeaderStyle::Definition && (FunctionType != EExportFunctionType::Interface) ? CPPF_Implementation : 0) | CPPF_ArgumentOrReturnValue);
-		OutFwdDecls.Add(Return->GetCPPTypeForwardDeclaration());
+		FString ReturnType = ReturnProperty->GetCPPType(&ExtendedReturnType, (FunctionHeaderStyle == EExportFunctionHeaderStyle::Definition && (FunctionType != EExportFunctionType::Interface) ? CPPF_Implementation : 0) | CPPF_ArgumentOrReturnValue);
+		OutFwdDecls.Add(ReturnProperty->GetCPPTypeForwardDeclaration());
 		FUHTStringBuilder ReplacementText;
 		ReplacementText += ReturnType;
-		ApplyAlternatePropertyExportText(Return, ReplacementText, EExportingState::Normal);
+		ApplyAlternatePropertyExportText(ReturnProperty, ReplacementText, EExportingState::Normal);
 		Out.Logf(TEXT("%s%s"), *ReplacementText, *ExtendedReturnType);
 	}
 	else
@@ -4109,7 +4110,29 @@ void FNativeClassHeaderGenerator::ExportNativeFunctionHeader(
 		if (bIsInterface && FunctionHeaderStyle == EExportFunctionHeaderStyle::Declaration)
 		{
 			// all methods in interface classes are pure virtuals
-			Out.Log(TEXT("=0"));
+			if (bIsK2Override)
+			{
+				// For BlueprintNativeEvent methods we use the PURE_VIRTUAL() expansion. This allows Blueprints that implement the interface class to be nativized.
+				FString ReturnValue;
+				if (ReturnProperty != nullptr)
+				{
+					UByteProperty* ByteProperty = Cast<UByteProperty>(ReturnProperty);
+					if (ByteProperty != nullptr && ByteProperty->Enum != nullptr && ByteProperty->Enum->GetCppForm() != UEnum::ECppForm::EnumClass)
+					{
+						ReturnValue = FString::Printf(TEXT(" return TEnumAsByte<%s>(%s);"), *ByteProperty->Enum->CppType, *GetNullParameterValue(ReturnProperty, false));
+					}
+					else
+					{
+						ReturnValue = FString::Printf(TEXT(" return %s;"), *GetNullParameterValue(ReturnProperty, false));
+					}
+				}
+
+				Out.Logf(TEXT(" PURE_VIRTUAL(%s::%s,%s)"), NameLookupCPP.GetNameCPP(CastChecked<UClass>(Function->GetOuter()), true), *FunctionName, *ReturnValue);
+			}
+			else
+			{
+				Out.Log(TEXT("=0"));
+			}
 		}
 	}
 }
diff --git a/Engine/Source/Programs/UnrealHeaderTool/Private/HeaderParser.cpp b/Engine/Source/Programs/UnrealHeaderTool/Private/HeaderParser.cpp
--- a/Engine/Source/Programs/UnrealHeaderTool/Private/HeaderParser.cpp
+++ b/Engine/Source/Programs/UnrealHeaderTool/Private/HeaderParser.cpp
@@ -6704,7 +6704,7 @@ void FHeaderParser::CompileFunctionDeclaration(FClasses& AllClasses)
 		}
 	}
 
-	// Check to see if there is a function in the super class with the same name but a different signature
+	// Check to see if there is a function in the super class with the same name
 	UStruct* SuperStruct = GetCurrentClass();
 	if (SuperStruct)
 	{
@@ -6714,10 +6714,8 @@ void FHeaderParser::CompileFunctionDeclaration(FClasses& AllClasses)
 	{
 		if (UFunction* OverriddenFunction = ::FindField<UFunction>(SuperStruct, FuncInfo.Function.Identifier))
 		{
-			if (!AreFunctionSignaturesEqual(TopFunction, OverriddenFunction))
-			{
-				FError::Throwf(TEXT("Function '%s' has a different signature from the one defined in base class '%s'"), FuncInfo.Function.Identifier, *OverriddenFunction->GetOuter()->GetName());
-			}
+			// Native function overrides should be done in CPP text, not in a UFUNCTION() declaration (you can't change flags, and it'd otherwise be a burden to keep them identical)
+			UE_LOG_ERROR_UHT(TEXT("%s: Override of UFUNCTION in parent class (%s) cannot have a UFUNCTION() declaration above it; it will use the same parameters as the original declaration."), FuncInfo.Function.Identifier, *OverriddenFunction->GetOuter()->GetName());
 		}
 	}
 
@@ -6804,114 +6802,6 @@ void FHeaderParser::CompileFunctionDeclaration(FClasses& AllClasses)
 	TopFunction->FunctionFlags |= FuncInfo.FunctionFlags;
 	StoredFuncData->UpdateFunctionData(FuncInfo);
 
-	// Verify parameter list and return type compatibility within the
-	// function, if any, that it overrides.
-	auto FunctionIterator = GetCurrentScope()->GetTypeIterator<UFunction>();
-	while (FunctionIterator.MoveNext())
-	{
-		UFunction* Function = *FunctionIterator;
-		if (Function->GetFName() != TopFunction->GetFName() || Function == TopFunction)
-			continue;
-
-		// Don't allow private functions to be redefined.
-		if (Function->FunctionFlags & FUNC_Private)
-			FError::Throwf(TEXT("Can't override private function '%s'"), FuncInfo.Function.Identifier);
-
-		// see if they both either have a return value or don't
-		if ((TopFunction->GetReturnProperty() != NULL) != (Function->GetReturnProperty() != NULL))
-		{
-			ReturnToLocation(FuncNameRetry);
-			FError::Throwf(TEXT("Redefinition of '%s %s' differs from original: return value mismatch"), TypeOfFunction, FuncInfo.Function.Identifier );
-		}
-
-		// See if all parameters match.
-		if (TopFunction->NumParms!=Function->NumParms)
-		{
-			ReturnToLocation(FuncNameRetry);
-			FError::Throwf(TEXT("Redefinition of '%s %s' differs from original; different number of parameters"), TypeOfFunction, FuncInfo.Function.Identifier );
-		}
-
-		// Check all individual parameters.
-		int32 Count=0;
-		for( TFieldIterator<UProperty> CurrentFuncParam(TopFunction),SuperFuncParam(Function); Count<Function->NumParms; ++CurrentFuncParam,++SuperFuncParam,++Count )
-		{
-			if( !FPropertyBase(*CurrentFuncParam).MatchesType(FPropertyBase(*SuperFuncParam), 1) )
-			{
-				if( CurrentFuncParam->PropertyFlags & CPF_ReturnParm )
-				{
-					ReturnToLocation(FuncNameRetry);
-					FError::Throwf(TEXT("Redefinition of %s %s differs only by return type"), TypeOfFunction, FuncInfo.Function.Identifier );
-				}
-				else
-				{
-					ReturnToLocation(FuncNameRetry);
-					FError::Throwf(TEXT("Redefinition of '%s %s' differs from original"), TypeOfFunction, FuncInfo.Function.Identifier );
-				}
-				break;
-			}
-			else if ( CurrentFuncParam->HasAnyPropertyFlags(CPF_OutParm) != SuperFuncParam->HasAnyPropertyFlags(CPF_OutParm) )
-			{
-				ReturnToLocation(FuncNameRetry);
-				FError::Throwf(TEXT("Redefinition of '%s %s' differs from original - 'out' mismatch on parameter %i"), TypeOfFunction, FuncInfo.Function.Identifier, Count + 1);
-			}
-			else if ( CurrentFuncParam->HasAnyPropertyFlags(CPF_ReferenceParm) != SuperFuncParam->HasAnyPropertyFlags(CPF_ReferenceParm) )
-			{
-				ReturnToLocation(FuncNameRetry);
-				FError::Throwf(TEXT("Redefinition of '%s %s' differs from original - 'ref' mismatch on parameter %i"), TypeOfFunction, FuncInfo.Function.Identifier, Count + 1);
-			}
-		}
-
-		if( Count<TopFunction->NumParms )
-		{
-			continue;
-		}
-
-		// if super version is event, overridden version must be defined as event (check before inheriting FUNC_Event)
-		if ( (Function->FunctionFlags & FUNC_Event) && !(FuncInfo.FunctionFlags & FUNC_Event) )
-		{
-			FError::Throwf(TEXT("Superclass version is defined as an event so '%s' should be!"), FuncInfo.Function.Identifier);
-		}
-		// Function flags to copy from parent.
-		FuncInfo.FunctionFlags |= (Function->FunctionFlags & FUNC_FuncInherit);
-
-		// Make sure the replication conditions aren't being redefined
-		if ((FuncInfo.FunctionFlags & FUNC_NetFuncFlags) != (Function->FunctionFlags & FUNC_NetFuncFlags))
-		{
-			FError::Throwf(TEXT("Redefinition of replication conditions for function '%s'"), FuncInfo.Function.Identifier);
-		}
-		FuncInfo.FunctionFlags |= (Function->FunctionFlags & FUNC_NetFuncFlags);
-
-		// Are we overriding a function?
-		if (TopFunction == Function->GetOuter())
-		{
-			// Duplicate.
-			ReturnToLocation( FuncNameRetry );
-			FError::Throwf(TEXT("Duplicate function '%s'"), *Function->GetName() );
-		}
-		// Overriding an existing function.
-		else if( Function->FunctionFlags & FUNC_Final )
-		{
-			ReturnToLocation(FuncNameRetry);
-			FError::Throwf(TEXT("%s: Can't override a 'final' function"), *Function->GetName() );
-		}
-		// Native function overrides should be done in CPP text, not in a UFUNCTION() declaration (you can't change flags, and it'd otherwise be a burden to keep them identical)
-		else if( Cast<UClass>(TopFunction->GetOuter()) != NULL )
-		{
-			//ReturnToLocation(FuncNameRetry);
-			FError::Throwf(TEXT("%s: An override of a function cannot have a UFUNCTION() declaration above it; it will use the same parameters as the original base declaration."), *Function->GetName() );
-		}
-
-		// Balk if required specifiers differ.
-		if ((Function->FunctionFlags & FUNC_FuncOverrideMatch) != (FuncInfo.FunctionFlags & FUNC_FuncOverrideMatch))
-		{
-			FError::Throwf(TEXT("Function '%s' specifiers differ from original"), *Function->GetName());
-		}
-
-		// Here we have found the original.
-		TopFunction->SetSuperStruct(Function);
-		break;
-	}
-
 	// Bind the function.
 	TopFunction->Bind();
 	
diff --git a/Engine/Source/Runtime/Core/Public/UObject/FrameworkObjectVersion.h b/Engine/Source/Runtime/Core/Public/UObject/FrameworkObjectVersion.h
--- a/Engine/Source/Runtime/Core/Public/UObject/FrameworkObjectVersion.h
+++ b/Engine/Source/Runtime/Core/Public/UObject/FrameworkObjectVersion.h
@@ -116,6 +116,9 @@ struct CORE_API FFrameworkObjectVersion
 		// Function terminator nodes serialize an FMemberReference rather than a name/class pair
 		FunctionTerminatorNodesUseMemberReference,
 
+		// Custom event and non-native interface event implementations add 'const' to reference parameters
+		EditableEventsUseConstRefParameters,
+
 		// -----<new versions can be added above this line>-------------------------------------------------
 		VersionPlusOne,
 		LatestVersion = VersionPlusOne - 1
diff --git a/Engine/Source/Runtime/CoreUObject/Private/Blueprint/BlueprintSupport.cpp b/Engine/Source/Runtime/CoreUObject/Private/Blueprint/BlueprintSupport.cpp
--- a/Engine/Source/Runtime/CoreUObject/Private/Blueprint/BlueprintSupport.cpp
+++ b/Engine/Source/Runtime/CoreUObject/Private/Blueprint/BlueprintSupport.cpp
@@ -872,7 +872,7 @@ bool FLinkerLoad::DeferPotentialCircularImport(const int32 Index)
 			ResolveDependencyPlaceholder(ImportPlaceholder, LoadClass);
 
 #if USE_DEFERRED_DEPENDENCY_CHECK_VERIFICATION_TESTS
-			const bool bIsStillPlaceholder = Import.XObject->IsA<ULinkerPlaceholderClass>() || Import.XObject->IsA<ULinkerPlaceholderFunction>();
+			const bool bIsStillPlaceholder = Import.XObject && (Import.XObject->IsA<ULinkerPlaceholderClass>() || Import.XObject->IsA<ULinkerPlaceholderFunction>());
 			DEFERRED_DEPENDENCY_CHECK(!bIsStillPlaceholder);
 			return bIsStillPlaceholder;
 #else 
diff --git a/Engine/Source/Runtime/CoreUObject/Private/UObject/LinkerLoad.cpp b/Engine/Source/Runtime/CoreUObject/Private/UObject/LinkerLoad.cpp
--- a/Engine/Source/Runtime/CoreUObject/Private/UObject/LinkerLoad.cpp
+++ b/Engine/Source/Runtime/CoreUObject/Private/UObject/LinkerLoad.cpp
@@ -3625,10 +3625,14 @@ void FLinkerLoad::Preload( UObject* Object )
 				}
 			}
 		}
-		else if( Object->GetLinker() )
+		else if (FLinkerLoad* Linker = Object->GetLinker())
 		{
+#if USE_CIRCULAR_DEPENDENCY_LOAD_DEFERRING
+			uint32 const DeferredLoadFlag = (LoadFlags & LOAD_DeferDependencyLoads);
+			TGuardValue<uint32> LoadFlagsGuard(Linker->LoadFlags, Linker->LoadFlags | DeferredLoadFlag);
+#endif
 			// Send to the object's linker.
-			Object->GetLinker()->Preload( Object );
+			Linker->Preload(Object);
 		}
 	}
 }
diff --git a/Engine/Source/Runtime/CoreUObject/Private/UObject/PropertyBaseObject.cpp b/Engine/Source/Runtime/CoreUObject/Private/UObject/PropertyBaseObject.cpp
--- a/Engine/Source/Runtime/CoreUObject/Private/UObject/PropertyBaseObject.cpp
+++ b/Engine/Source/Runtime/CoreUObject/Private/UObject/PropertyBaseObject.cpp
@@ -443,6 +443,11 @@ UObject* UObjectPropertyBase::FindImportedObject( const UProperty* Property, UOb
 #endif // USE_CIRCULAR_DEPENDENCY_LOAD_DEFERRING
 			{
 				uint32 LoadFlags = LOAD_NoWarn | LOAD_FindIfFail;
+				if (Linker)
+				{
+					LoadFlags |= Linker->LoadFlags & LOAD_DeferDependencyLoads;
+				}
+					
 
 				UE_LOG(LogProperty, Verbose, TEXT("FindImportedObject is attempting to import [%s] (class = %s) with StaticLoadObject"), Text, *GetFullNameSafe(ObjectClass));
 				Result = StaticLoadObject(ObjectClass, nullptr, Text, nullptr, LoadFlags, nullptr);
diff --git a/Engine/Source/Runtime/CoreUObject/Private/UObject/PropertyStruct.cpp b/Engine/Source/Runtime/CoreUObject/Private/UObject/PropertyStruct.cpp
--- a/Engine/Source/Runtime/CoreUObject/Private/UObject/PropertyStruct.cpp
+++ b/Engine/Source/Runtime/CoreUObject/Private/UObject/PropertyStruct.cpp
@@ -239,8 +239,31 @@ const TCHAR* UStructProperty::ImportText_Internal(const TCHAR* InBuffer, void* D
 {
 #if USE_CIRCULAR_DEPENDENCY_LOAD_DEFERRING
 	FScopedPlaceholderPropertyTracker ImportPropertyTracker(this);
+
+	uint32 PropagatedLoadFlags = 0;
+	if (FLinkerLoad* Linker = GetLinker())
+	{
+		PropagatedLoadFlags |= (Linker->LoadFlags & LOAD_DeferDependencyLoads);
+	}
+
+	uint32 OldFlags = 0;
+	FLinkerLoad* StructLinker = Struct->GetLinker();
+	if (StructLinker)
+	{
+		OldFlags = StructLinker->LoadFlags;
+		StructLinker->LoadFlags |= OldFlags | PropagatedLoadFlags;
+	}
 #endif 
-	return Struct->ImportText(InBuffer, Data, Parent, PortFlags, ErrorText, GetName(), true);
+	const TCHAR* Result = Struct->ImportText(InBuffer, Data, Parent, PortFlags, ErrorText, GetName(), true);
+
+#if USE_CIRCULAR_DEPENDENCY_LOAD_DEFERRING
+	if (StructLinker)
+	{
+		StructLinker->LoadFlags = OldFlags;
+	}
+#endif
+	
+	return Result;
 }
 
 const TCHAR* UStructProperty::ImportText_Static(UScriptStruct* InStruct, const FString& Name, const TCHAR* InBuffer, void* Data, int32 PortFlags, UObject* Parent, FOutputDevice* ErrorText)
diff --git a/Engine/Source/Runtime/CoreUObject/Private/UObject/UObjectGlobals.cpp b/Engine/Source/Runtime/CoreUObject/Private/UObject/UObjectGlobals.cpp
--- a/Engine/Source/Runtime/CoreUObject/Private/UObject/UObjectGlobals.cpp
+++ b/Engine/Source/Runtime/CoreUObject/Private/UObject/UObjectGlobals.cpp
@@ -880,7 +880,7 @@ UObject* StaticLoadObjectInternal(UClass* ObjectClass, UObject* InOuter, const T
 	const bool bContainsObjectName = !!FCString::Strstr(InName, TEXT("."));
 
 	// break up the name into packages, returning the innermost name and its outer
-	ResolveName(InOuter, StrName, true, true, LoadFlags & (LOAD_EditorOnly | LOAD_Quiet | LOAD_NoWarn));
+	ResolveName(InOuter, StrName, true, true, LoadFlags & (LOAD_EditorOnly | LOAD_Quiet | LOAD_NoWarn | LOAD_DeferDependencyLoads));
 	if (InOuter)
 	{
 		// If we have a full UObject name then attempt to find the object in memory first,
diff --git a/Engine/Source/Runtime/CoreUObject/Private/UObject/UnrealType.cpp b/Engine/Source/Runtime/CoreUObject/Private/UObject/UnrealType.cpp
--- a/Engine/Source/Runtime/CoreUObject/Private/UObject/UnrealType.cpp
+++ b/Engine/Source/Runtime/CoreUObject/Private/UObject/UnrealType.cpp
@@ -4,7 +4,7 @@
 
 DEFINE_LOG_CATEGORY(LogType);
 
-bool FPropertyValueIterator::NextValue(EPropertyValueIteratorFlags InRecursionFlags)
+bool FPropertyValueIterator::NextValue(EPropertyValueIteratorFlags InRecursionFlags, bool bReturningFromStruct)
 {
 	if (PropertyIteratorStack.Num() == 0)
 	{
@@ -14,6 +14,12 @@ bool FPropertyValueIterator::NextValue(EPropertyValueIteratorFlags InRecursionFl
 
 	FPropertyValueStackEntry& Entry = PropertyIteratorStack.Last();
 
+	// If we are returning from struct, increment value index as we delayed incrementing it when entering recursion
+	if (bReturningFromStruct)
+	{
+		Entry.ValueIndex++;
+	}
+
 	// If we have pending values, deal with them
 	if (Entry.ValueIndex < Entry.ValueArray.Num())
 	{
@@ -78,15 +84,18 @@ bool FPropertyValueIterator::NextValue(EPropertyValueIteratorFlags InRecursionFl
 				// And child to stack and rerun. This invalidates Entry. We leave ValueIndex the same so it can be used for recursive property lookup
 				PropertyIteratorStack.Emplace(StructProperty->Struct, PropertyValue, DeprecatedPropertyFlags);
 
-				return NextValue(InRecursionFlags);
+				return NextValue(InRecursionFlags, false);
 			}
 		}
 
 		// Else this is a normal property and has nothing to expand
 		// We don't expand enum properties because EnumProperty handles value wrapping for us
 
-		// We didn't recurse into a struct, so increment next value to check
-		Entry.ValueIndex++;
+		// We didn't recurse into a struct, so increment next value to check, unless we we just popped a struct off the stack
+		if (!bReturningFromStruct)
+		{
+			Entry.ValueIndex++;
+		}
 	}
 
 	// Out of pending values, try to add more
@@ -99,10 +108,7 @@ bool FPropertyValueIterator::NextValue(EPropertyValueIteratorFlags InRecursionFl
 
 			if (PropertyIteratorStack.Num() > 0)
 			{
-				// Increment value index as we delayed incrementing it when entering recursion
-				PropertyIteratorStack.Last().ValueIndex++;
-			
-				return NextValue(InRecursionFlags);
+				return NextValue(InRecursionFlags, true);
 			}
 
 			return false;
@@ -139,7 +145,7 @@ void FPropertyValueIterator::IterateToNext()
 		bSkipRecursionOnce = false;
 	}
 
-	while (NextValue(LocalRecursionFlags))
+	while (NextValue(LocalRecursionFlags, false))
 	{
 		// If this property is valid type, stop iteration
 		FPropertyValueStackEntry& Entry = PropertyIteratorStack.Last();
diff --git a/Engine/Source/Runtime/CoreUObject/Public/UObject/Script.h b/Engine/Source/Runtime/CoreUObject/Public/UObject/Script.h
--- a/Engine/Source/Runtime/CoreUObject/Public/UObject/Script.h
+++ b/Engine/Source/Runtime/CoreUObject/Public/UObject/Script.h
@@ -151,7 +151,7 @@ FORCEINLINE FArchive& operator<<(FArchive& Ar, EFunctionFlags& Flags)
 ENUM_CLASS_FLAGS(EFunctionFlags)
 
 // Combinations of flags.
-#define FUNC_FuncInherit       ((EFunctionFlags)(FUNC_Exec | FUNC_Event | FUNC_BlueprintCallable | FUNC_BlueprintEvent | FUNC_BlueprintAuthorityOnly | FUNC_BlueprintCosmetic))
+#define FUNC_FuncInherit       ((EFunctionFlags)(FUNC_Exec | FUNC_Event | FUNC_BlueprintCallable | FUNC_BlueprintEvent | FUNC_BlueprintAuthorityOnly | FUNC_BlueprintCosmetic | FUNC_Const))
 #define FUNC_FuncOverrideMatch ((EFunctionFlags)(FUNC_Exec | FUNC_Final | FUNC_Static | FUNC_Public | FUNC_Protected | FUNC_Private))
 #define FUNC_NetFuncFlags      ((EFunctionFlags)(FUNC_Net | FUNC_NetReliable | FUNC_NetServer | FUNC_NetClient | FUNC_NetMulticast))
 #define FUNC_AccessSpecifiers  ((EFunctionFlags)(FUNC_Public | FUNC_Private | FUNC_Protected))
diff --git a/Engine/Source/Runtime/CoreUObject/Public/UObject/UnrealType.h b/Engine/Source/Runtime/CoreUObject/Public/UObject/UnrealType.h
--- a/Engine/Source/Runtime/CoreUObject/Public/UObject/UnrealType.h
+++ b/Engine/Source/Runtime/CoreUObject/Public/UObject/UnrealType.h
@@ -4878,7 +4878,7 @@ private:
 	bool bSkipRecursionOnce;
 
 	/** Goes to the next Property/value pair. Returns true if next value is valid */
-	bool NextValue(EPropertyValueIteratorFlags RecursionFlags);
+	bool NextValue(EPropertyValueIteratorFlags RecursionFlags, bool bReturningFromStruct);
 
 	/** Iterates to next property being checked for or until reaching the end of the structure */
 	COREUOBJECT_API void IterateToNext();
diff --git a/Engine/Source/Runtime/Engine/Classes/Camera/PlayerCameraManager.h b/Engine/Source/Runtime/Engine/Classes/Camera/PlayerCameraManager.h
--- a/Engine/Source/Runtime/Engine/Classes/Camera/PlayerCameraManager.h
+++ b/Engine/Source/Runtime/Engine/Classes/Camera/PlayerCameraManager.h
@@ -405,6 +405,7 @@ public:
 	uint32 bEnableColorScaleInterp : 1;
 
 	/** True if clients are handling setting their own viewtarget and the server should not replicate it (e.g. during certain Matinee sequences) */
+	UPROPERTY(EditAnywhere, BlueprintReadWrite, Category=PlayerCameraManager)
 	uint32 bClientSimulatingViewTarget : 1;
 
 	/** True if server will use camera positions replicated from the client instead of calculating them locally. */
diff --git a/Engine/Source/Runtime/Engine/Classes/Engine/Blueprint.h b/Engine/Source/Runtime/Engine/Classes/Engine/Blueprint.h
--- a/Engine/Source/Runtime/Engine/Classes/Engine/Blueprint.h
+++ b/Engine/Source/Runtime/Engine/Classes/Engine/Blueprint.h
@@ -782,6 +782,7 @@ public:
 	virtual void BeginCacheForCookedPlatformData(const ITargetPlatform *TargetPlatform) override;
 	virtual bool IsCachedCookedPlatformDataLoaded(const ITargetPlatform* TargetPlatform) override;
 	virtual void ClearAllCachedCookedPlatformData() override;
+	virtual void BeginDestroy() override;
 	//~ End UObject Interface
 
 	/** Consigns the GeneratedClass and the SkeletonGeneratedClass to oblivion, and nulls their references */
diff --git a/Engine/Source/Runtime/Engine/Classes/Engine/GameInstance.h b/Engine/Source/Runtime/Engine/Classes/Engine/GameInstance.h
--- a/Engine/Source/Runtime/Engine/Classes/Engine/GameInstance.h
+++ b/Engine/Source/Runtime/Engine/Classes/Engine/GameInstance.h
@@ -125,13 +125,18 @@ protected:
 	/** Delegate for handling external console commands */
 	void OnConsoleInput(const FString& Command);
 
+	/** List of locally participating players in this game instance */
 	UPROPERTY()
-	TArray<ULocalPlayer*> LocalPlayers;		// List of locally participating players in this game instance
+	TArray<ULocalPlayer*> LocalPlayers;
 	
 	/** Class to manage online services */
 	UPROPERTY()
 	class UOnlineSession* OnlineSession;
 
+	/** List of objects that are being kept alive by this game instance. Stored as array for fast iteration, should not be modified every frame */
+	UPROPERTY()
+	TArray<UObject*> ReferencedObjects;
+
 	/** Listeners to PreClientTravel call */
 	FOnPreClientTravel NotifyPreClientTravelDelegates;
 
@@ -385,6 +390,12 @@ public:
 	 */
 	virtual void ReturnToMainMenu();
 
+	/** Registers an object to keep alive as long as this game instance is alive */
+	virtual void RegisterReferencedObject(UObject* ObjectToReference);
+
+	/** Remove a referenced object, this will allow it to GC out */
+	virtual void UnregisterReferencedObject(UObject* ObjectToReference);
+
 protected:
 	/** Called when the game instance is started either normally or through PIE. */
 	virtual void OnStart();
diff --git a/Engine/Source/Runtime/Engine/Classes/GameFramework/Character.h b/Engine/Source/Runtime/Engine/Classes/GameFramework/Character.h
--- a/Engine/Source/Runtime/Engine/Classes/GameFramework/Character.h
+++ b/Engine/Source/Runtime/Engine/Classes/GameFramework/Character.h
@@ -365,7 +365,6 @@ public:
 	virtual void OnRep_ReplicatedBasedMovement();
 
 	/** Set whether this actor's movement replicates to network clients. */
-	UFUNCTION(BlueprintCallable, Category=Replication)
 	virtual void SetReplicateMovement(bool bInReplicateMovement) override;
 
 protected:
diff --git a/Engine/Source/Runtime/Engine/Classes/GameFramework/CheatManager.h b/Engine/Source/Runtime/Engine/Classes/GameFramework/CheatManager.h
--- a/Engine/Source/Runtime/Engine/Classes/GameFramework/CheatManager.h
+++ b/Engine/Source/Runtime/Engine/Classes/GameFramework/CheatManager.h
@@ -308,11 +308,11 @@ class ENGINE_API UCheatManager : public UObject
 
 	/** Exec function to return the mouse sensitivity to its default value */
 	UFUNCTION(exec)
-	void SetMouseSensitivityToDefault();
+	virtual void SetMouseSensitivityToDefault();
 
 	/** Backwards compatibility exec function for people used to it instead of using InvertAxisKey */
 	UFUNCTION(exec)
-	void InvertMouse();
+	virtual void InvertMouse();
 
 	/** Executes commands listed in CheatScript.ScriptName ini section of DefaultGame.ini */
 	UFUNCTION(exec)
diff --git a/Engine/Source/Runtime/Engine/Classes/GameFramework/GameModeBase.h b/Engine/Source/Runtime/Engine/Classes/GameFramework/GameModeBase.h
--- a/Engine/Source/Runtime/Engine/Classes/GameFramework/GameModeBase.h
+++ b/Engine/Source/Runtime/Engine/Classes/GameFramework/GameModeBase.h
@@ -412,7 +412,7 @@ public:
 	AActor* K2_FindPlayerStart(AController* Player, const FString& IncomingName = TEXT(""));
 
 	/** Returns true if it's valid to call RestartPlayer. By default will call Player->CanRestartPlayer */
-	UFUNCTION(BlueprintNativeEvent, Category=Game)
+	UFUNCTION(BlueprintCallable, BlueprintNativeEvent, Category=Game)
 	bool PlayerCanRestart(APlayerController* Player);
 
 	/** Tries to spawn the player's pawn, at the location returned by FindPlayerStart */
diff --git a/Engine/Source/Runtime/Engine/Classes/GameFramework/Pawn.h b/Engine/Source/Runtime/Engine/Classes/GameFramework/Pawn.h
--- a/Engine/Source/Runtime/Engine/Classes/GameFramework/Pawn.h
+++ b/Engine/Source/Runtime/Engine/Classes/GameFramework/Pawn.h
@@ -130,7 +130,7 @@ public:
 	uint8 RemoteViewPitch;
 
 	/** Controller of the last Actor that caused us damage. */
-	UPROPERTY(transient)
+	UPROPERTY(BlueprintReadOnly, transient, Category="Pawn")
 	AController* LastHitBy;
 
 	/** Controller currently possessing this Actor */
diff --git a/Engine/Source/Runtime/Engine/Classes/GameFramework/PlayerController.h b/Engine/Source/Runtime/Engine/Classes/GameFramework/PlayerController.h
--- a/Engine/Source/Runtime/Engine/Classes/GameFramework/PlayerController.h
+++ b/Engine/Source/Runtime/Engine/Classes/GameFramework/PlayerController.h
@@ -5,6 +5,7 @@
 #include "CoreMinimal.h"
 #include "SlateFwd.h"
 #include "UObject/ObjectMacros.h"
+#include "Containers/SortedMap.h"
 #include "Misc/Guid.h"
 #include "InputCoreTypes.h"
 #include "Templates/SubclassOf.h"
@@ -61,6 +62,8 @@ namespace EDynamicForceFeedbackAction
 	};
 }
 
+typedef uint64 FDynamicForceFeedbackHandle;
+
 struct FDynamicForceFeedbackDetails
 {
 	uint32 bAffectsLeftLarge:1;
@@ -338,8 +341,39 @@ public:
 	TArray<FForceFeedbackEffectHistoryEntry> ForceFeedbackEffectHistoryEntries;
 #endif
 
-	/** Map of active force feedback effects */
-	TMap<int32, FDynamicForceFeedbackDetails> DynamicForceFeedbacks;
+private:
+
+	struct FDynamicForceFeedbackAction
+	{
+		/** Time over which interpolation should happen */
+		float TotalTime;
+		/** Time so far elapsed for the interpolation */
+		float TimeElapsed;
+		/** Current feedback intensity values */
+		FDynamicForceFeedbackDetails ForceFeedbackDetails;
+		/** Unique ID for the action */
+		FDynamicForceFeedbackHandle Handle;
+		/** Unique ID generation static */
+		static FDynamicForceFeedbackHandle HandleAllocator;
+
+		/** 
+		 * Updates Values with this action's details.
+		 * @param DeltaTime	Time since last update
+		 * @param Values	Values structure to update
+		 * @return Returns false if the elapsed time has exceeded the total time.
+		 */
+		bool Update(const float DeltaTime, FForceFeedbackValues& Values);
+	};
+
+	/** Map of dynamic force feedback effects invoked from native */
+	TSortedMap<FDynamicForceFeedbackHandle, FDynamicForceFeedbackAction> DynamicForceFeedbacks;
+
+	/** Map of dynamic force feedback effects invoked from blueprints */
+	TSortedMap<int32, FDynamicForceFeedbackDetails*> LatentDynamicForceFeedbacks;
+
+	friend class FLatentDynamicForceFeedbackAction;
+
+public:
 
 	/** Currently playing haptic effects for both the left and right hand */
 	TSharedPtr<struct FActiveHapticFeedbackEffect> ActiveHapticEffect_Left;
@@ -1016,6 +1050,7 @@ public:
 	UFUNCTION(reliable, client, BlueprintCallable, Category="Game|Feedback")
 	void ClientStopForceFeedback(class UForceFeedbackEffect* ForceFeedbackEffect, FName Tag);
 
+private:
 	/** 
 	 * Latent action that controls the playing of force feedback 
 	 * Begins playing when Start is called.  Calling Update or Stop if the feedback is not active will have no effect.
@@ -1031,6 +1066,24 @@ public:
 	UFUNCTION(BlueprintCallable, meta=(Latent, LatentInfo="LatentInfo", ExpandEnumAsExecs="Action", Duration="-1", bAffectsLeftLarge="true", bAffectsLeftSmall="true", bAffectsRightLarge="true", bAffectsRightSmall="true", AdvancedDisplay="bAffectsLeftLarge,bAffectsLeftSmall,bAffectsRightLarge,bAffectsRightSmall"), Category="Game|Feedback")
 	void PlayDynamicForceFeedback(float Intensity, float Duration, bool bAffectsLeftLarge, bool bAffectsLeftSmall, bool bAffectsRightLarge, bool bAffectsRightSmall, TEnumAsByte<EDynamicForceFeedbackAction::Type> Action, FLatentActionInfo LatentInfo);
 
+public:
+	/** 
+	 * Allows playing of a dynamic force feedback event from native code
+	 * Begins playing when Start is called.  Calling with Action set to Update or Stop if the feedback is not active will have no effect.
+	 * When Update is called the Intensity, Duration, and affect values will be updated
+	 * @param	Intensity				How strong the feedback should be.  Valid values are between 0.0 and 1.0
+	 * @param	Duration				How long the feedback should play for.  If the value is negative it will play until stopped
+	 * @param   bAffectsLeftLarge		Whether the intensity should be applied to the large left servo
+	 * @param   bAffectsLeftSmall		Whether the intensity should be applied to the small left servo
+	 * @param   bAffectsRightLarge		Whether the intensity should be applied to the large right servo
+	 * @param   bAffectsRightSmall		Whether the intensity should be applied to the small right servo
+	 * @param   Action					Whether to (re)start, update, or stop the action
+	 * @param	UniqueID				The ID returned by the start action when wanting to restart, update, or stop the action
+	 * @return  The index to pass in to the function to update the latent action in the future if needed. Returns 0 if the feedback was stopped
+	 *          or the specified UniqueID did not map to an action that can be started/updated.
+	 */
+	FDynamicForceFeedbackHandle PlayDynamicForceFeedback(float Intensity, float Duration, bool bAffectsLeftLarge, bool bAffectsLeftSmall, bool bAffectsRightLarge, bool bAffectsRightSmall, EDynamicForceFeedbackAction::Type Action = EDynamicForceFeedbackAction::Start, FDynamicForceFeedbackHandle ActionHandle = 0);
+
 	/**
 	 * Play a haptic feedback curve on the player's controller
 	 * @param	HapticEffect			The haptic effect to play
@@ -1632,6 +1685,7 @@ public:
 	virtual void ViewAPlayer(int32 dir);
 
 	/** @return true if this controller thinks it's able to restart. Called from GameModeBase::PlayerCanRestart */
+	UFUNCTION(BlueprintCallable, Category = "Game")
 	virtual bool CanRestartPlayer();
 
 	/**
diff --git a/Engine/Source/Runtime/Engine/Classes/GameFramework/WorldSettings.h b/Engine/Source/Runtime/Engine/Classes/GameFramework/WorldSettings.h
--- a/Engine/Source/Runtime/Engine/Classes/GameFramework/WorldSettings.h
+++ b/Engine/Source/Runtime/Engine/Classes/GameFramework/WorldSettings.h
@@ -627,19 +627,19 @@ public:
 	float DemoPlayTimeDilation;
 
 	/** Lowest acceptable global time dilation. */
-	UPROPERTY(config, EditAnywhere, Category = Tick, AdvancedDisplay)
+	UPROPERTY(config, EditAnywhere, Category = Tick, AdvancedDisplay, meta = (UIMin = "0", ClampMin = "0"))
 	float MinGlobalTimeDilation;
 	
 	/** Highest acceptable global time dilation. */
-	UPROPERTY(config, EditAnywhere, Category = Tick, AdvancedDisplay)
+	UPROPERTY(config, EditAnywhere, Category = Tick, AdvancedDisplay, meta = (UIMin = "0", ClampMin = "0"))
 	float MaxGlobalTimeDilation;
 
 	/** Smallest possible frametime, not considering dilation. Equiv to 1/FastestFPS. */
-	UPROPERTY(config, EditAnywhere, Category = Tick, AdvancedDisplay)
+	UPROPERTY(config, EditAnywhere, Category = Tick, AdvancedDisplay, meta = (UIMin = "0", ClampMin = "0"))
 	float MinUndilatedFrameTime;
 
 	/** Largest possible frametime, not considering dilation. Equiv to 1/SlowestFPS. */
-	UPROPERTY(config, EditAnywhere, Category = Tick, AdvancedDisplay)
+	UPROPERTY(config, EditAnywhere, Category = Tick, AdvancedDisplay, meta = (UIMin = "0", ClampMin = "0"))
 	float MaxUndilatedFrameTime;
 
 	// If paused, FName of person pausing the game.
diff --git a/Engine/Source/Runtime/Engine/Classes/Kismet/BlueprintAsyncActionBase.h b/Engine/Source/Runtime/Engine/Classes/Kismet/BlueprintAsyncActionBase.h
--- a/Engine/Source/Runtime/Engine/Classes/Kismet/BlueprintAsyncActionBase.h
+++ b/Engine/Source/Runtime/Engine/Classes/Kismet/BlueprintAsyncActionBase.h
@@ -5,8 +5,11 @@
 #include "UObject/ObjectMacros.h"
 #include "UObject/Object.h"
 #include "UObject/ScriptMacros.h"
+#include "UObject/WeakObjectPtr.h"
 #include "BlueprintAsyncActionBase.generated.h"
 
+class UGameInstance;
+
 UCLASS()
 class ENGINE_API UBlueprintAsyncActionBase : public UObject
 {
@@ -16,6 +19,16 @@ class ENGINE_API UBlueprintAsyncActionBase : public UObject
 	UFUNCTION(BlueprintCallable, meta=(BlueprintInternalUseOnly="true"))
 	virtual void Activate();
 
-	/** Call when the action is completely done, this makes the action free to delete */
+	/**
+	 * Call to globally register this object with a game instance, it will not be destroyed until SetReadyToDestroy is called
+	 * This allows having an action stay alive until SetReadyToDestroy is manually called, allowing it to be used inside loops or if the calling BP goes away
+	 */
+	virtual void RegisterWithGameInstance(UObject* WorldContextObject);
+	virtual void RegisterWithGameInstance(UGameInstance* GameInstance);
+
+	/** Call when the action is completely done, this makes the action free to delete, and will unregister it with the game instance */
 	virtual void SetReadyToDestroy();
+
+protected:
+	TWeakObjectPtr<UGameInstance> RegisteredWithGameInstance;
 };
diff --git a/Engine/Source/Runtime/Engine/Classes/Kismet/BlueprintMapLibrary.h b/Engine/Source/Runtime/Engine/Classes/Kismet/BlueprintMapLibrary.h
--- a/Engine/Source/Runtime/Engine/Classes/Kismet/BlueprintMapLibrary.h
+++ b/Engine/Source/Runtime/Engine/Classes/Kismet/BlueprintMapLibrary.h
@@ -23,7 +23,6 @@ public:
 	 * @param	TargetMap		The map to add the key and value to
 	 * @param	Key				The key that will be used to look the value up
 	 * @param	Value			The value to be retrieved later
-	 * @return	True if a Value was added, or False if the Key was already present and has been overwritten
 	 */
 	UFUNCTION(BlueprintCallable, CustomThunk, meta=(DisplayName = "Add", CompactNodeTitle = "ADD", MapParam = "TargetMap", MapKeyParam = "Key", MapValueParam = "Value", AutoCreateRefTerm = "Key, Value"), Category = "Utilities|Map")
 	static void Map_Add(const TMap<int32, int32>& TargetMap, const int32& Key, const int32& Value);
diff --git a/Engine/Source/Runtime/Engine/Classes/Kismet/BlueprintSetLibrary.h b/Engine/Source/Runtime/Engine/Classes/Kismet/BlueprintSetLibrary.h
--- a/Engine/Source/Runtime/Engine/Classes/Kismet/BlueprintSetLibrary.h
+++ b/Engine/Source/Runtime/Engine/Classes/Kismet/BlueprintSetLibrary.h
@@ -17,13 +17,10 @@ class ENGINE_API UBlueprintSetLibrary : public UBlueprintFunctionLibrary
 
 public:
 	/** 
-	 * Adds item to set. Output value indicates whether the item was successfully added, meaning an 
-	 * output of False indicates the item was already in the Set.
+	 * Adds item to set
 	 *
 	 * @param	TargetSet		The set to add item to
 	 * @param	NewItem			The item to add to the set
-	 * @return	True if NewItem was added to the set (False indicates an equivalent item was present)
-	 * 
 	 */
 	UFUNCTION(BlueprintCallable, CustomThunk, meta=(DisplayName = "Add", CompactNodeTitle = "ADD", SetParam = "TargetSet|NewItem", AutoCreateRefTerm = "NewItem"), Category="Utilities|Set")
 	static void Set_Add(const TSet<int32>& TargetSet, const int32& NewItem);
diff --git a/Engine/Source/Runtime/Engine/Private/AsyncActionLoadPrimaryAsset.cpp b/Engine/Source/Runtime/Engine/Private/AsyncActionLoadPrimaryAsset.cpp
--- a/Engine/Source/Runtime/Engine/Private/AsyncActionLoadPrimaryAsset.cpp
+++ b/Engine/Source/Runtime/Engine/Private/AsyncActionLoadPrimaryAsset.cpp
@@ -40,12 +40,13 @@ void UAsyncActionLoadPrimaryAssetBase::HandleLoadCompleted()
 	SetReadyToDestroy();
 }
 
-UAsyncActionLoadPrimaryAsset* UAsyncActionLoadPrimaryAsset::AsyncLoadPrimaryAsset(FPrimaryAssetId PrimaryAsset, const TArray<FName>& LoadBundles)
+UAsyncActionLoadPrimaryAsset* UAsyncActionLoadPrimaryAsset::AsyncLoadPrimaryAsset(UObject* WorldContextObject, FPrimaryAssetId PrimaryAsset, const TArray<FName>& LoadBundles)
 {
 	UAsyncActionLoadPrimaryAsset* Action = NewObject<UAsyncActionLoadPrimaryAsset>();
 	Action->AssetsToLoad.Add(PrimaryAsset);
 	Action->LoadBundles = LoadBundles;
 	Action->Operation = EAssetManagerOperation::Load;
+	Action->RegisterWithGameInstance(WorldContextObject);
 
 	return Action;
 }
@@ -62,12 +63,13 @@ void UAsyncActionLoadPrimaryAsset::HandleLoadCompleted()
 	Completed.Broadcast(AssetLoaded);
 }
 
-UAsyncActionLoadPrimaryAssetClass* UAsyncActionLoadPrimaryAssetClass::AsyncLoadPrimaryAssetClass(FPrimaryAssetId PrimaryAsset, const TArray<FName>& LoadBundles)
+UAsyncActionLoadPrimaryAssetClass* UAsyncActionLoadPrimaryAssetClass::AsyncLoadPrimaryAssetClass(UObject* WorldContextObject, FPrimaryAssetId PrimaryAsset, const TArray<FName>& LoadBundles)
 {
 	UAsyncActionLoadPrimaryAssetClass* Action = NewObject<UAsyncActionLoadPrimaryAssetClass>();
 	Action->AssetsToLoad.Add(PrimaryAsset);
 	Action->LoadBundles = LoadBundles;
 	Action->Operation = EAssetManagerOperation::Load;
+	Action->RegisterWithGameInstance(WorldContextObject);
 
 	return Action;
 }
@@ -84,12 +86,13 @@ void UAsyncActionLoadPrimaryAssetClass::HandleLoadCompleted()
 	Completed.Broadcast(AssetLoaded);
 }
 
-UAsyncActionLoadPrimaryAssetList* UAsyncActionLoadPrimaryAssetList::AsyncLoadPrimaryAssetList(const TArray<FPrimaryAssetId>& PrimaryAssetList, const TArray<FName>& LoadBundles)
+UAsyncActionLoadPrimaryAssetList* UAsyncActionLoadPrimaryAssetList::AsyncLoadPrimaryAssetList(UObject* WorldContextObject, const TArray<FPrimaryAssetId>& PrimaryAssetList, const TArray<FName>& LoadBundles)
 {
 	UAsyncActionLoadPrimaryAssetList* Action = NewObject<UAsyncActionLoadPrimaryAssetList>();
 	Action->AssetsToLoad = PrimaryAssetList;
 	Action->LoadBundles = LoadBundles;
 	Action->Operation = EAssetManagerOperation::Load;
+	Action->RegisterWithGameInstance(WorldContextObject);
 
 	return Action;
 }
@@ -107,12 +110,13 @@ void UAsyncActionLoadPrimaryAssetList::HandleLoadCompleted()
 	Completed.Broadcast(AssetList);
 }
 
-UAsyncActionLoadPrimaryAssetClassList* UAsyncActionLoadPrimaryAssetClassList::AsyncLoadPrimaryAssetClassList(const TArray<FPrimaryAssetId>& PrimaryAssetList, const TArray<FName>& LoadBundles)
+UAsyncActionLoadPrimaryAssetClassList* UAsyncActionLoadPrimaryAssetClassList::AsyncLoadPrimaryAssetClassList(UObject* WorldContextObject, const TArray<FPrimaryAssetId>& PrimaryAssetList, const TArray<FName>& LoadBundles)
 {
 	UAsyncActionLoadPrimaryAssetClassList* Action = NewObject<UAsyncActionLoadPrimaryAssetClassList>();
 	Action->AssetsToLoad = PrimaryAssetList;
 	Action->LoadBundles = LoadBundles;
 	Action->Operation = EAssetManagerOperation::Load;
+	Action->RegisterWithGameInstance(WorldContextObject);
 
 	return Action;
 }
@@ -141,23 +145,25 @@ void UAsyncActionLoadPrimaryAssetClassList::HandleLoadCompleted()
 	Completed.Broadcast(AssetClassList);
 }
 
-UAsyncActionChangePrimaryAssetBundles* UAsyncActionChangePrimaryAssetBundles::AsyncChangeBundleStateForMatchingPrimaryAssets(const TArray<FName>& NewBundles, const TArray<FName>& OldBundles)
+UAsyncActionChangePrimaryAssetBundles* UAsyncActionChangePrimaryAssetBundles::AsyncChangeBundleStateForMatchingPrimaryAssets(UObject* WorldContextObject, const TArray<FName>& NewBundles, const TArray<FName>& OldBundles)
 {
 	UAsyncActionChangePrimaryAssetBundles* Action = NewObject<UAsyncActionChangePrimaryAssetBundles>();
 	Action->LoadBundles = NewBundles;
 	Action->OldBundles = OldBundles;
 	Action->Operation = EAssetManagerOperation::ChangeBundleStateMatching;
+	Action->RegisterWithGameInstance(WorldContextObject);
 	
 	return Action;
 }
 
-UAsyncActionChangePrimaryAssetBundles* UAsyncActionChangePrimaryAssetBundles::AsyncChangeBundleStateForPrimaryAssetList(const TArray<FPrimaryAssetId>& PrimaryAssetList, const TArray<FName>& AddBundles, const TArray<FName>& RemoveBundles)
+UAsyncActionChangePrimaryAssetBundles* UAsyncActionChangePrimaryAssetBundles::AsyncChangeBundleStateForPrimaryAssetList(UObject* WorldContextObject, const TArray<FPrimaryAssetId>& PrimaryAssetList, const TArray<FName>& AddBundles, const TArray<FName>& RemoveBundles)
 {
 	UAsyncActionChangePrimaryAssetBundles* Action = NewObject<UAsyncActionChangePrimaryAssetBundles>();
 	Action->LoadBundles = AddBundles;
 	Action->OldBundles = RemoveBundles;
 	Action->AssetsToLoad = PrimaryAssetList;
 	Action->Operation = EAssetManagerOperation::ChangeBundleStateList;
+	Action->RegisterWithGameInstance(WorldContextObject);
 
 	return Action;
 }
diff --git a/Engine/Source/Runtime/Engine/Private/AsyncActionLoadPrimaryAsset.h b/Engine/Source/Runtime/Engine/Private/AsyncActionLoadPrimaryAsset.h
--- a/Engine/Source/Runtime/Engine/Private/AsyncActionLoadPrimaryAsset.h
+++ b/Engine/Source/Runtime/Engine/Private/AsyncActionLoadPrimaryAsset.h
@@ -60,8 +60,8 @@ public:
 	 * Load a primary asset into memory. The completed delegate will go off when the load succeeds or fails, you should cast the Loaded object to verify it is the correct type.
 	 * If LoadBundles is specified, those bundles are loaded along with the asset
 	 */
-	UFUNCTION(BlueprintCallable, meta=(BlueprintInternalUseOnly="true", Category = "AssetManager", AutoCreateRefTerm = "LoadBundles"))
-	static UAsyncActionLoadPrimaryAsset* AsyncLoadPrimaryAsset(FPrimaryAssetId PrimaryAsset, const TArray<FName>& LoadBundles);
+	UFUNCTION(BlueprintCallable, meta=(BlueprintInternalUseOnly="true", Category = "AssetManager", AutoCreateRefTerm = "LoadBundles", WorldContext = "WorldContextObject"))
+	static UAsyncActionLoadPrimaryAsset* AsyncLoadPrimaryAsset(UObject* WorldContextObject, FPrimaryAssetId PrimaryAsset, const TArray<FName>& LoadBundles);
 
 	UPROPERTY(BlueprintAssignable)
 	FOnPrimaryAssetLoaded Completed;
@@ -84,8 +84,8 @@ public:
 	 * Load a primary asset class into memory. The completed delegate will go off when the load succeeds or fails, you should cast the Loaded class to verify it is the correct type 
 	 * If LoadBundles is specified, those bundles are loaded along with the asset
 	 */
-	UFUNCTION(BlueprintCallable, meta = (BlueprintInternalUseOnly = "true", Category = "AssetManager", AutoCreateRefTerm = "LoadBundles"))
-	static UAsyncActionLoadPrimaryAssetClass* AsyncLoadPrimaryAssetClass(FPrimaryAssetId PrimaryAsset, const TArray<FName>& LoadBundles);
+	UFUNCTION(BlueprintCallable, meta = (BlueprintInternalUseOnly = "true", Category = "AssetManager", AutoCreateRefTerm = "LoadBundles", WorldContext = "WorldContextObject"))
+	static UAsyncActionLoadPrimaryAssetClass* AsyncLoadPrimaryAssetClass(UObject* WorldContextObject, FPrimaryAssetId PrimaryAsset, const TArray<FName>& LoadBundles);
 
 	UPROPERTY(BlueprintAssignable)
 	FOnPrimaryAssetClassLoaded Completed;
@@ -108,8 +108,8 @@ public:
 	 * Load a list primary assets into memory. The completed delegate will go off when the load succeeds or fails, you should cast the Loaded object list to verify it is the correct type 
 	 * If LoadBundles is specified, those bundles are loaded along with the asset list
 	 */
-	UFUNCTION(BlueprintCallable, meta=(BlueprintInternalUseOnly="true", Category = "AssetManager", AutoCreateRefTerm = "LoadBundles"))
-	static UAsyncActionLoadPrimaryAssetList* AsyncLoadPrimaryAssetList(const TArray<FPrimaryAssetId>& PrimaryAssetList, const TArray<FName>& LoadBundles);
+	UFUNCTION(BlueprintCallable, meta=(BlueprintInternalUseOnly="true", Category = "AssetManager", AutoCreateRefTerm = "LoadBundles", WorldContext = "WorldContextObject"))
+	static UAsyncActionLoadPrimaryAssetList* AsyncLoadPrimaryAssetList(UObject* WorldContextObject, const TArray<FPrimaryAssetId>& PrimaryAssetList, const TArray<FName>& LoadBundles);
 
 	UPROPERTY(BlueprintAssignable)
 	FOnPrimaryAssetListLoaded Completed;
@@ -132,8 +132,8 @@ public:
 	 * Load a list primary asset classes into memory. The completed delegate will go off when the load succeeds or fails, you should cast the Loaded object list to verify it is the correct type 
 	 * If LoadBundles is specified, those bundles are loaded along with the asset list
 	 */
-	UFUNCTION(BlueprintCallable, meta = (BlueprintInternalUseOnly = "true", Category = "AssetManager", AutoCreateRefTerm = "LoadBundles"))
-	static UAsyncActionLoadPrimaryAssetClassList* AsyncLoadPrimaryAssetClassList(const TArray<FPrimaryAssetId>& PrimaryAssetList, const TArray<FName>& LoadBundles);
+	UFUNCTION(BlueprintCallable, meta = (BlueprintInternalUseOnly = "true", Category = "AssetManager", AutoCreateRefTerm = "LoadBundles", WorldContext = "WorldContextObject"))
+	static UAsyncActionLoadPrimaryAssetClassList* AsyncLoadPrimaryAssetClassList(UObject* WorldContextObject, const TArray<FPrimaryAssetId>& PrimaryAssetList, const TArray<FName>& LoadBundles);
 
 	UPROPERTY(BlueprintAssignable)
 	FOnPrimaryAssetClassListLoaded Completed;
@@ -155,14 +155,14 @@ public:
 	/**
 	 * Change the bundle state of all assets that match OldBundles to instead contain NewBundles. 
 	 */
-	UFUNCTION(BlueprintCallable, meta = (BlueprintInternalUseOnly = "true", Category = "AssetManager"))
-	static UAsyncActionChangePrimaryAssetBundles* AsyncChangeBundleStateForMatchingPrimaryAssets(const TArray<FName>& NewBundles, const TArray<FName>& OldBundles);
+	UFUNCTION(BlueprintCallable, meta = (BlueprintInternalUseOnly = "true", Category = "AssetManager", WorldContext = "WorldContextObject"))
+	static UAsyncActionChangePrimaryAssetBundles* AsyncChangeBundleStateForMatchingPrimaryAssets(UObject* WorldContextObject, const TArray<FName>& NewBundles, const TArray<FName>& OldBundles);
 
 	/**
 	 * Change the bundle state of assets in PrimaryAssetList. AddBundles are added and RemoveBundles are removed, both must be filled in but an empty array is allowed
 	 */
-	UFUNCTION(BlueprintCallable, meta = (BlueprintInternalUseOnly = "true", Category = "AssetManager"))
-	static UAsyncActionChangePrimaryAssetBundles* AsyncChangeBundleStateForPrimaryAssetList(const TArray<FPrimaryAssetId>& PrimaryAssetList, const TArray<FName>& AddBundles, const TArray<FName>& RemoveBundles);
+	UFUNCTION(BlueprintCallable, meta = (BlueprintInternalUseOnly = "true", Category = "AssetManager", WorldContext = "WorldContextObject"))
+	static UAsyncActionChangePrimaryAssetBundles* AsyncChangeBundleStateForPrimaryAssetList(UObject* WorldContextObject, const TArray<FPrimaryAssetId>& PrimaryAssetList, const TArray<FName>& AddBundles, const TArray<FName>& RemoveBundles);
 
 	UPROPERTY(BlueprintAssignable)
 	FOnPrimaryAssetBundlesChanged Completed;
diff --git a/Engine/Source/Runtime/Engine/Private/Blueprint.cpp b/Engine/Source/Runtime/Engine/Private/Blueprint.cpp
--- a/Engine/Source/Runtime/Engine/Private/Blueprint.cpp
+++ b/Engine/Source/Runtime/Engine/Private/Blueprint.cpp
@@ -1379,6 +1379,13 @@ void UBlueprint::ClearAllCachedCookedPlatformData()
 	}
 }
 
+void UBlueprint::BeginDestroy()
+{
+	Super::BeginDestroy();
+
+	FBlueprintEditorUtils::RemoveAllLocalBookmarks(this);
+}
+
 #endif // WITH_EDITOR
 
 UBlueprint* UBlueprint::GetBlueprintFromClass(const UClass* InClass)
diff --git a/Engine/Source/Runtime/Engine/Private/ComponentInstanceDataCache.cpp b/Engine/Source/Runtime/Engine/Private/ComponentInstanceDataCache.cpp
--- a/Engine/Source/Runtime/Engine/Private/ComponentInstanceDataCache.cpp
+++ b/Engine/Source/Runtime/Engine/Private/ComponentInstanceDataCache.cpp
@@ -269,7 +269,10 @@ void FActorComponentInstanceData::ApplyToComponent(UActorComponent* Component, c
 
 		for (UObject* InstancedObject : InstancedObjects)
 		{
-			InstancedObject->Rename(nullptr, Component, REN_DontCreateRedirectors | REN_ForceNoResetLoaders);
+			if (InstancedObject)
+			{
+				InstancedObject->Rename(nullptr, Component, REN_DontCreateRedirectors | REN_ForceNoResetLoaders);
+			}
 		}
 
 		FComponentPropertyReader ComponentPropertyReader(Component, SavedProperties);
diff --git a/Engine/Source/Runtime/Engine/Private/EdGraph/EdGraphNode.cpp b/Engine/Source/Runtime/Engine/Private/EdGraph/EdGraphNode.cpp
--- a/Engine/Source/Runtime/Engine/Private/EdGraph/EdGraphNode.cpp
+++ b/Engine/Source/Runtime/Engine/Private/EdGraph/EdGraphNode.cpp
@@ -35,13 +35,27 @@ FArchive& operator<<(FArchive& Ar, FEdGraphTerminalType& T)
 
 	if (Ar.CustomVer(FFrameworkObjectVersion::GUID) >= FFrameworkObjectVersion::PinsStoreFName)
 	{
-	Ar << T.TerminalCategory;
-	Ar << T.TerminalSubCategory;
+		Ar << T.TerminalCategory;
+		Ar << T.TerminalSubCategory;
 	}
 	else
 	{
 		FString TerminalCategoryStr;
 		Ar << TerminalCategoryStr;
+
+		if (Ar.UE4Ver() < VER_UE4_ADDED_SOFT_OBJECT_PATH)
+		{
+			// Handle asset->soft object rename, this is here instead of BP code because this structure is embedded
+			if (TerminalCategoryStr == TEXT("asset"))
+			{
+				TerminalCategoryStr = TEXT("softobject");
+			}
+			else if (TerminalCategoryStr == TEXT("assetclass"))
+			{
+				TerminalCategoryStr = TEXT("softclass");
+			}
+		}
+
 		T.TerminalCategory = *TerminalCategoryStr;
 
 		FString TerminalSubCategoryStr;
diff --git a/Engine/Source/Runtime/Engine/Private/GameInstance.cpp b/Engine/Source/Runtime/Engine/Private/GameInstance.cpp
--- a/Engine/Source/Runtime/Engine/Private/GameInstance.cpp
+++ b/Engine/Source/Runtime/Engine/Private/GameInstance.cpp
@@ -1235,4 +1235,14 @@ AGameModeBase* UGameInstance::CreateGameModeForURL(FURL InURL)
 TSubclassOf<AGameModeBase> UGameInstance::OverrideGameModeClass(TSubclassOf<AGameModeBase> GameModeClass, const FString& MapName, const FString& Options, const FString& Portal) const
 {
 	 return GameModeClass;
+}
+
+void UGameInstance::RegisterReferencedObject(UObject* ObjectToReference)
+{
+	ReferencedObjects.AddUnique(ObjectToReference);
+}
+
+void UGameInstance::UnregisterReferencedObject(UObject* ObjectToReference)
+{
+	ReferencedObjects.RemoveSingleSwap(ObjectToReference);
 }
\ No newline at end of file
diff --git a/Engine/Source/Runtime/Engine/Private/GameUserSettings.cpp b/Engine/Source/Runtime/Engine/Private/GameUserSettings.cpp
--- a/Engine/Source/Runtime/Engine/Private/GameUserSettings.cpp
+++ b/Engine/Source/Runtime/Engine/Private/GameUserSettings.cpp
@@ -56,6 +56,7 @@ enum EGameUserSettingsVersion
 UGameUserSettings::UGameUserSettings(const FObjectInitializer& ObjectInitializer)
 :	Super(ObjectInitializer)
 {
+	// this will only call the base version of SetToDefaults but some constructors may rely on it being set
 	SetToDefaults();
 }
 
diff --git a/Engine/Source/Runtime/Engine/Private/GameplayStatics.cpp b/Engine/Source/Runtime/Engine/Private/GameplayStatics.cpp
--- a/Engine/Source/Runtime/Engine/Private/GameplayStatics.cpp
+++ b/Engine/Source/Runtime/Engine/Private/GameplayStatics.cpp
@@ -2355,7 +2355,7 @@ bool UGameplayStatics::ProjectWorldToScreen(APlayerController const* Player, con
 				ScreenPosition -= FVector2D(ProjectionData.GetConstrainedViewRect().Min);
 			}
 
-			bResult = Player->PostProcessWorldToScreen(WorldPosition, ScreenPosition, bPlayerViewportRelative);
+			bResult = bResult && Player->PostProcessWorldToScreen(WorldPosition, ScreenPosition, bPlayerViewportRelative);
 			return bResult;
 		}
 	}
diff --git a/Engine/Source/Runtime/Engine/Private/Kismet/BlueprintAsyncActionBase.cpp b/Engine/Source/Runtime/Engine/Private/Kismet/BlueprintAsyncActionBase.cpp
--- a/Engine/Source/Runtime/Engine/Private/Kismet/BlueprintAsyncActionBase.cpp
+++ b/Engine/Source/Runtime/Engine/Private/Kismet/BlueprintAsyncActionBase.cpp
@@ -1,6 +1,8 @@
 // Copyright 1998-2018 Epic Games, Inc. All Rights Reserved.
 
 #include "Kismet/BlueprintAsyncActionBase.h"
+#include "Engine/Engine.h"
+#include "Engine/GameInstance.h"
 
 //////////////////////////////////////////////////////////////////////////
 // UBlueprintAsyncActionBase
@@ -18,7 +20,38 @@ void UBlueprintAsyncActionBase::Activate()
 {
 }
 
+void UBlueprintAsyncActionBase::RegisterWithGameInstance(UObject* WorldContextObject)
+{
+	UWorld* FoundWorld = GEngine->GetWorldFromContextObject(WorldContextObject, EGetWorldErrorMode::LogAndReturnNull);
+
+	if (FoundWorld && FoundWorld->GetGameInstance())
+	{
+		RegisterWithGameInstance(FoundWorld->GetGameInstance());
+	}
+}
+
+void UBlueprintAsyncActionBase::RegisterWithGameInstance(UGameInstance* GameInstance)
+{
+	if (GameInstance)
+	{
+		UGameInstance* OldGameInstance = RegisteredWithGameInstance.Get();
+		if (OldGameInstance)
+		{
+			OldGameInstance->UnregisterReferencedObject(this);
+		}
+
+		GameInstance->RegisterReferencedObject(this);
+		RegisteredWithGameInstance = GameInstance;
+	}
+}
+
 void UBlueprintAsyncActionBase::SetReadyToDestroy()
 {
 	ClearFlags(RF_StrongRefOnFrame);
+
+	UGameInstance* OldGameInstance = RegisteredWithGameInstance.Get();
+	if (OldGameInstance)
+	{
+		OldGameInstance->UnregisterReferencedObject(this);
+	}
 }
diff --git a/Engine/Source/Runtime/Engine/Private/NetworkDriver.cpp b/Engine/Source/Runtime/Engine/Private/NetworkDriver.cpp
--- a/Engine/Source/Runtime/Engine/Private/NetworkDriver.cpp
+++ b/Engine/Source/Runtime/Engine/Private/NetworkDriver.cpp
@@ -2959,7 +2959,26 @@ int32 UNetDriver::ServerReplicateActors_PrepConnections( const float DeltaSecond
 			bFoundReadyConnection = true;
 
 			// the view target is what the player controller is looking at OR the owning actor itself when using beacons
-			Connection->ViewTarget = Connection->PlayerController ? Connection->PlayerController->GetViewTarget() : OwningActor;
+			AActor* DesiredViewTarget = OwningActor;
+			if (Connection->PlayerController)
+			{
+				if (AActor* ViewTarget = Connection->PlayerController->GetViewTarget())
+				{
+					if (ViewTarget->GetWorld())
+					{
+						// It is safe to use the player controller's view target.
+						DesiredViewTarget = ViewTarget;
+					}
+					else
+					{
+						// Log an error, since this means the view target for the player controller no longer has a valid world (this can happen
+						// if the player controller's view target was in a sublevel instance that has been unloaded).
+						UE_LOG(LogNet, Warning, TEXT("Player controller %s's view target (%s) no longer has a valid world! Was it unloaded as part a level instance?"),
+							*Connection->PlayerController->GetName(), *ViewTarget->GetName());
+					}
+				}
+			}
+			Connection->ViewTarget = DesiredViewTarget;
 
 			for ( int32 ChildIdx = 0; ChildIdx < Connection->Children.Num(); ChildIdx++ )
 			{
diff --git a/Engine/Source/Runtime/Engine/Private/Particles/ParticleComponents.cpp b/Engine/Source/Runtime/Engine/Private/Particles/ParticleComponents.cpp
--- a/Engine/Source/Runtime/Engine/Private/Particles/ParticleComponents.cpp
+++ b/Engine/Source/Runtime/Engine/Private/Particles/ParticleComponents.cpp
@@ -3597,7 +3597,7 @@ void UParticleSystemComponent::OnUnregister()
 		TEXT("OnUnregister %s Component=0x%p Scene=0x%p FXSystem=0x%p"),
 		Template != NULL ? *Template->GetName() : TEXT("NULL"), this, GetWorld()->Scene, FXSystem);
 
-	bWasActive = bIsActive;
+	bWasActive = bIsActive && !bWasDeactivated;
 
 	ResetParticles(true);
 	FXSystem = NULL;
diff --git a/Engine/Source/Runtime/Engine/Private/PlayerController.cpp b/Engine/Source/Runtime/Engine/Private/PlayerController.cpp
--- a/Engine/Source/Runtime/Engine/Private/PlayerController.cpp
+++ b/Engine/Source/Runtime/Engine/Private/PlayerController.cpp
@@ -3705,8 +3705,23 @@ void APlayerController::ClientStopForceFeedback_Implementation( UForceFeedbackEf
 
 /// @endcond
 
+uint64 APlayerController::FDynamicForceFeedbackAction::HandleAllocator = 0;
+
+bool APlayerController::FDynamicForceFeedbackAction::Update(const float DeltaTime, FForceFeedbackValues& Values)
+{
+	TimeElapsed += DeltaTime;
+
+	if (TotalTime >= 0.f && TimeElapsed >= TotalTime)
+	{
+		return false;
+	}
+
+	ForceFeedbackDetails.Update(Values);
+	return true;
+}
+
 /** Action that interpolates a component over time to a desired position */
-class FDynamicForceFeedbackAction : public FPendingLatentAction
+class FLatentDynamicForceFeedbackAction : public FPendingLatentAction
 {
 public:
 	/** Time over which interpolation should happen */
@@ -3714,8 +3729,10 @@ public:
 	/** Time so far elapsed for the interpolation */
 	float TimeElapsed;
 	/** If we are currently running. If false, update will complete */
-	uint32 bRunning:1;
-	
+	uint8 bRunning:1;
+	/** Whether the latent action is currently in the player controller's array */
+	uint8 bAddedToPlayerController:1;
+
 	TWeakObjectPtr<APlayerController> PlayerController;
 
 	FDynamicForceFeedbackDetails ForceFeedbackDetails;
@@ -3729,10 +3746,11 @@ public:
 	/** Object to call callback on upon completion */
 	FWeakObjectPtr CallbackTarget;
 
-	FDynamicForceFeedbackAction(APlayerController* InPlayerController, const float InDuration, const FLatentActionInfo& LatentInfo)
+	FLatentDynamicForceFeedbackAction(APlayerController* InPlayerController, const float InDuration, const FLatentActionInfo& LatentInfo)
 		: TotalTime(InDuration)
 		, TimeElapsed(0.f)
 		, bRunning(true)
+		, bAddedToPlayerController(false)
 		, PlayerController(InPlayerController)
 		, ExecutionFunction(LatentInfo.ExecutionFunction)
 		, OutputLink(LatentInfo.Linkage)
@@ -3741,6 +3759,17 @@ public:
 	{
 	}
 
+	~FLatentDynamicForceFeedbackAction()
+	{
+		if (bAddedToPlayerController)
+		{
+			if (APlayerController* PC = PlayerController.Get())
+			{
+				PC->LatentDynamicForceFeedbacks.Remove(LatentUUID);
+			}
+		}
+	}
+
 	virtual void UpdateOperation(FLatentResponse& Response) override
 	{
 		// Update elapsed time
@@ -3752,11 +3781,13 @@ public:
 		{
 			if (bComplete)
 			{
-				PC->DynamicForceFeedbacks.Remove(LatentUUID);
+				PC->LatentDynamicForceFeedbacks.Remove(LatentUUID);
+				bAddedToPlayerController = false;
 			}
 			else
 			{
-				PC->DynamicForceFeedbacks.Add(LatentUUID, ForceFeedbackDetails);
+				PC->LatentDynamicForceFeedbacks.Add(LatentUUID, &ForceFeedbackDetails);
+				bAddedToPlayerController = true;
 			}
 		}
 
@@ -3767,7 +3798,8 @@ public:
 	{
 		if (APlayerController* PC = PlayerController.Get())
 		{
-			PC->DynamicForceFeedbacks.Remove(LatentUUID);
+			PC->LatentDynamicForceFeedbacks.Remove(LatentUUID);
+			bAddedToPlayerController = false;
 		}
 	}
 
@@ -3775,7 +3807,8 @@ public:
 	{
 		if (APlayerController* PC = PlayerController.Get())
 		{
-			PC->DynamicForceFeedbacks.Remove(LatentUUID);
+			PC->LatentDynamicForceFeedbacks.Remove(LatentUUID);
+			bAddedToPlayerController = false;
 		}
 	}
 };
@@ -3783,7 +3816,7 @@ public:
 void APlayerController::PlayDynamicForceFeedback(float Intensity, float Duration, bool bAffectsLeftLarge, bool bAffectsLeftSmall, bool bAffectsRightLarge, bool bAffectsRightSmall, TEnumAsByte<EDynamicForceFeedbackAction::Type> Action, FLatentActionInfo LatentInfo)
 {
 	FLatentActionManager& LatentActionManager = GetWorld()->GetLatentActionManager();
-	FDynamicForceFeedbackAction* LatentAction = LatentActionManager.FindExistingAction<FDynamicForceFeedbackAction>(LatentInfo.CallbackTarget, LatentInfo.UUID);
+	FLatentDynamicForceFeedbackAction* LatentAction = LatentActionManager.FindExistingAction<FLatentDynamicForceFeedbackAction>(LatentInfo.CallbackTarget, LatentInfo.UUID);
 
 	if (LatentAction)
 	{
@@ -3809,7 +3842,7 @@ void APlayerController::PlayDynamicForceFeedback(float Intensity, float Duration
 	}
 	else if (Action == EDynamicForceFeedbackAction::Start)
 	{
-		LatentAction = new FDynamicForceFeedbackAction(this, Duration, LatentInfo);
+		LatentAction = new FLatentDynamicForceFeedbackAction(this, Duration, LatentInfo);
 
 		LatentAction->ForceFeedbackDetails.Intensity = Intensity;
 		LatentAction->ForceFeedbackDetails.bAffectsLeftLarge = bAffectsLeftLarge;
@@ -3821,6 +3854,64 @@ void APlayerController::PlayDynamicForceFeedback(float Intensity, float Duration
 	}
 }
 
+FDynamicForceFeedbackHandle APlayerController::PlayDynamicForceFeedback(float Intensity, float Duration, bool bAffectsLeftLarge, bool bAffectsLeftSmall, bool bAffectsRightLarge, bool bAffectsRightSmall, EDynamicForceFeedbackAction::Type Action, FDynamicForceFeedbackHandle ActionHandle)
+{
+	FDynamicForceFeedbackHandle FeedbackHandle = 0;
+
+	if (Action == EDynamicForceFeedbackAction::Stop)
+	{
+		if (ActionHandle > 0)
+		{
+			DynamicForceFeedbacks.Remove(ActionHandle);
+		}
+	}
+	else
+	{
+		FDynamicForceFeedbackAction* FeedbackAction = (ActionHandle > 0 ? DynamicForceFeedbacks.Find(ActionHandle) : nullptr);
+
+		if (FeedbackAction == nullptr && Action == EDynamicForceFeedbackAction::Start)
+		{
+			if (ActionHandle > 0)
+			{
+				if (ActionHandle <= FDynamicForceFeedbackAction::HandleAllocator)
+				{
+					// Restarting a stopped/finished index, this is fine
+					FeedbackAction = &DynamicForceFeedbacks.Add(ActionHandle);
+					FeedbackAction->Handle = ActionHandle;
+				}
+				else
+				{
+					UE_LOG(LogPlayerController, Error, TEXT("Specifying an ID to start a dynamic force feedback with that has not yet been assigned is unsafe. No action has been started."));
+				}
+			}
+			else
+			{
+				FeedbackAction = &DynamicForceFeedbacks.Add(++FDynamicForceFeedbackAction::HandleAllocator);
+				FeedbackAction->Handle = FDynamicForceFeedbackAction::HandleAllocator;
+			}
+		}
+
+		if (FeedbackAction)
+		{
+			if (Action == EDynamicForceFeedbackAction::Start)
+			{
+				FeedbackAction->TotalTime = Duration;
+				FeedbackAction->TimeElapsed = 0.f;
+			}
+
+			FeedbackAction->ForceFeedbackDetails.Intensity = Intensity;
+			FeedbackAction->ForceFeedbackDetails.bAffectsLeftLarge = bAffectsLeftLarge;
+			FeedbackAction->ForceFeedbackDetails.bAffectsLeftSmall = bAffectsLeftSmall;
+			FeedbackAction->ForceFeedbackDetails.bAffectsRightLarge = bAffectsRightLarge;
+			FeedbackAction->ForceFeedbackDetails.bAffectsRightSmall = bAffectsRightSmall;
+
+			FeedbackHandle = FeedbackAction->Handle;
+		}
+	}
+
+	return FeedbackHandle;
+}
+
 void APlayerController::PlayHapticEffect(UHapticFeedbackEffect_Base* HapticEffect, EControllerHand Hand, float Scale, bool bLoop)
 {
 	if (HapticEffect)
@@ -3962,9 +4053,17 @@ void APlayerController::ProcessForceFeedbackAndHaptics(const float DeltaTime, co
 			}
 		}
 
-		for (const TPair<int32, FDynamicForceFeedbackDetails>& DynamicEntry : DynamicForceFeedbacks)
+		for (TSortedMap<uint64, FDynamicForceFeedbackAction>::TIterator It(DynamicForceFeedbacks.CreateIterator()); It; ++It)
+		{
+			if (!It.Value().Update(DeltaTime, ForceFeedbackValues))
+			{
+				It.RemoveCurrent();
+			}
+		}
+
+		for (const TPair<int32, FDynamicForceFeedbackDetails*>& DynamicEntry : LatentDynamicForceFeedbacks)
 		{
-			DynamicEntry.Value.Update(ForceFeedbackValues);
+			DynamicEntry.Value->Update(ForceFeedbackValues);
 		}
 
 		if (FForceFeedbackManager* ForceFeedbackManager = FForceFeedbackManager::Get(World))
diff --git a/Engine/Source/Runtime/Engine/Private/TargetPoint.cpp b/Engine/Source/Runtime/Engine/Private/TargetPoint.cpp
--- a/Engine/Source/Runtime/Engine/Private/TargetPoint.cpp
+++ b/Engine/Source/Runtime/Engine/Private/TargetPoint.cpp
@@ -57,6 +57,10 @@ ATargetPoint::ATargetPoint(const FObjectInitializer& ObjectInitializer)
 			ArrowComponent->SpriteInfo.DisplayName = ConstructorStatics.NAME_TargetPoint;
 			ArrowComponent->SetupAttachment(SpriteComponent);
 			ArrowComponent->bIsScreenSizeScaled = true;
+
+			// Counteract the scaled down parent so that the arrow is large enough to see.
+			if (SpriteComponent)
+				ArrowComponent->SetRelativeScale3D((FVector::OneVector / SpriteComponent->RelativeScale3D));
 		}
 	}
 #endif // WITH_EDITORONLY_DATA
diff --git a/Engine/Source/Runtime/Engine/Private/UnrealEngine.cpp b/Engine/Source/Runtime/Engine/Private/UnrealEngine.cpp
--- a/Engine/Source/Runtime/Engine/Private/UnrealEngine.cpp
+++ b/Engine/Source/Runtime/Engine/Private/UnrealEngine.cpp
@@ -12810,6 +12810,7 @@ void UEngine::CreateGameUserSettings()
 {
 	UGameUserSettings::LoadConfigIni();
 	GameUserSettings = NewObject<UGameUserSettings>(GetTransientPackage(), GEngine->GameUserSettingsClass);
+	GameUserSettings->SetToDefaults();
 	GameUserSettings->LoadSettings();
 }
 
diff --git a/Engine/Source/Runtime/Engine/Private/UserInterface/PlayerInput.cpp b/Engine/Source/Runtime/Engine/Private/UserInterface/PlayerInput.cpp
--- a/Engine/Source/Runtime/Engine/Private/UserInterface/PlayerInput.cpp
+++ b/Engine/Source/Runtime/Engine/Private/UserInterface/PlayerInput.cpp
@@ -1,4 +1,4 @@
-// Copyright 1998-2018 Epic Games, Inc. All Rights Reserved.
+﻿// Copyright 1998-2018 Epic Games, Inc. All Rights Reserved.
 
 /*=============================================================================
 	PlayerInput.cpp: Unreal input system.
@@ -967,7 +967,7 @@ void UPlayerInput::ProcessInputStack(const TArray<UInputComponent*>& InputCompon
 	// must be called non-recursively and on the game thread
 	check(IsInGameThread() && !AxisDelegates.Num() && !VectorAxisDelegates.Num() && !NonAxisDelegates.Num() && !KeysToConsume.Num() && !FoundChords.Num() && !EventIndices.Num() && !KeysWithEvents.Num() && !PotentialActions.Num());
 
-	// @todo, if input is coming in asynchronously, we should buffer anyhting that comes in during playerinput() and not include it
+	// @todo, if input is coming in asynchronously, we should buffer anything that comes in during playerinput() and not include it
 	// in this processing batch
 
 	APlayerController* PlayerController = GetOuterAPlayerController();
diff --git a/Engine/Source/Runtime/Engine/Private/WorldSettings.cpp b/Engine/Source/Runtime/Engine/Private/WorldSettings.cpp
--- a/Engine/Source/Runtime/Engine/Private/WorldSettings.cpp
+++ b/Engine/Source/Runtime/Engine/Private/WorldSettings.cpp
@@ -104,6 +104,26 @@ void AWorldSettings::PostInitProperties()
 	{
 		NavigationSystemConfig = NewObject<UNavigationSystemConfig>(this);
 	}
+
+	if (MinGlobalTimeDilation < 0)
+	{
+		MinGlobalTimeDilation = 0;
+	}
+
+	if (MaxGlobalTimeDilation < 0)
+	{
+		MaxGlobalTimeDilation = 0;
+	}
+
+	if (MinUndilatedFrameTime < 0)
+	{
+		MinUndilatedFrameTime = 0;
+	}
+
+	if (MaxUndilatedFrameTime < 0)
+	{
+		MaxUndilatedFrameTime = 0;
+	}
 }
 
 void AWorldSettings::PreInitializeComponents()
diff --git a/Engine/Source/Runtime/GameplayTags/Classes/GameplayTagContainer.h b/Engine/Source/Runtime/GameplayTags/Classes/GameplayTagContainer.h
--- a/Engine/Source/Runtime/GameplayTags/Classes/GameplayTagContainer.h
+++ b/Engine/Source/Runtime/GameplayTags/Classes/GameplayTagContainer.h
@@ -240,7 +240,7 @@ protected:
 	explicit FGameplayTag(FName InTagName);
 
 	/** This Tags Name */
-	UPROPERTY(VisibleAnywhere, Category = GameplayTags)
+	UPROPERTY(VisibleAnywhere, Category = GameplayTags, SaveGame)
 	FName TagName;
 
 	friend class UGameplayTagsManager;
@@ -820,7 +820,7 @@ protected:
 	void FillParentTags();
 
 	/** Array of gameplay tags */
-	UPROPERTY(BlueprintReadWrite, Category=GameplayTags) // Change to VisibleAnywhere after fixing up games
+	UPROPERTY(BlueprintReadWrite, Category=GameplayTags, SaveGame) // Change to VisibleAnywhere after fixing up games
 	TArray<FGameplayTag> GameplayTags;
 
 	/** Array of expanded parent tags, in addition to GameplayTags. Used to accelerate parent searches. May contain duplicates in some cases */
@@ -1058,6 +1058,8 @@ public:
 	static FGameplayTagQuery MakeQuery_MatchAllTags(FGameplayTagContainer const& InTags);
 	static FGameplayTagQuery MakeQuery_MatchNoTags(FGameplayTagContainer const& InTags);
 
+	static FGameplayTagQuery MakeQuery_MatchTag(FGameplayTag const& InTag);
+
 	friend class FQueryEvaluator;
 };
 
diff --git a/Engine/Source/Runtime/GameplayTags/Private/GameplayTagContainer.cpp b/Engine/Source/Runtime/GameplayTags/Private/GameplayTagContainer.cpp
--- a/Engine/Source/Runtime/GameplayTags/Private/GameplayTagContainer.cpp
+++ b/Engine/Source/Runtime/GameplayTags/Private/GameplayTagContainer.cpp
@@ -1465,6 +1465,17 @@ FGameplayTagQuery FGameplayTagQuery::MakeQuery_MatchNoTags(FGameplayTagContainer
 		);
 }
 
+// static
+FGameplayTagQuery FGameplayTagQuery::MakeQuery_MatchTag(FGameplayTag const & InTag)
+{
+	return FGameplayTagQuery::BuildQuery
+	(
+		FGameplayTagQueryExpression()
+		.AllTagsMatch()
+		.AddTag(InTag)
+	);
+}
+
 
 #if WITH_EDITOR
 
diff --git a/Engine/Source/Runtime/MoviePlayer/Private/DefaultGameMoviePlayer.cpp b/Engine/Source/Runtime/MoviePlayer/Private/DefaultGameMoviePlayer.cpp
--- a/Engine/Source/Runtime/MoviePlayer/Private/DefaultGameMoviePlayer.cpp
+++ b/Engine/Source/Runtime/MoviePlayer/Private/DefaultGameMoviePlayer.cpp
@@ -379,7 +379,7 @@ void FDefaultGameMoviePlayer::StopMovie()
 	}
 }
 
-void FDefaultGameMoviePlayer::WaitForMovieToFinish()
+void FDefaultGameMoviePlayer::WaitForMovieToFinish(bool bAllowEngineTick)
 {
 	const bool bEnforceMinimumTime = LoadingScreenAttributes.MinimumLoadingScreenDisplayTime >= 0.0f;
 
@@ -445,6 +445,11 @@ void FDefaultGameMoviePlayer::WaitForMovieToFinish()
 
 				float DeltaTime = SlateApp.GetDeltaTime();				
 
+				if (GEngine && bAllowEngineTick && LoadingScreenAttributes.bAllowEngineTick)
+				{
+					GEngine->Tick(DeltaTime, false);
+				}
+
 				ENQUEUE_UNIQUE_RENDER_COMMAND_TWOPARAMETER(
 					BeginLoadingMovieFrameAndTickMovieStreamer,
 					FDefaultGameMoviePlayer*, MoviePlayer, this,
@@ -739,7 +744,11 @@ void FDefaultGameMoviePlayer::OnPreLoadMap(const FString& LevelName)
 
 void FDefaultGameMoviePlayer::OnPostLoadMap(UWorld* LoadedWorld)
 {
-	WaitForMovieToFinish();
+	if (!LoadingScreenAttributes.bAllowEngineTick)
+	{
+		// If engine tick is enabled, we don't want to tick here and instead want to run from the WaitForMovieToFinish call in LaunchEngineLoop
+		WaitForMovieToFinish();
+	}
 }
 
 void FDefaultGameMoviePlayer::SetSlateOverlayWidget(TSharedPtr<SWidget> NewOverlayWidget)
diff --git a/Engine/Source/Runtime/MoviePlayer/Private/DefaultGameMoviePlayer.h b/Engine/Source/Runtime/MoviePlayer/Private/DefaultGameMoviePlayer.h
--- a/Engine/Source/Runtime/MoviePlayer/Private/DefaultGameMoviePlayer.h
+++ b/Engine/Source/Runtime/MoviePlayer/Private/DefaultGameMoviePlayer.h
@@ -73,7 +73,7 @@ public:
 	virtual bool PlayEarlyStartupMovies() override;
 	virtual bool PlayMovie() override;
 	virtual void StopMovie() override;
-	virtual void WaitForMovieToFinish() override;
+	virtual void WaitForMovieToFinish(bool bAllowEngineTick = false) override;
 	virtual bool IsLoadingFinished() const override;
 	virtual bool IsMovieCurrentlyPlaying() const override;
 	virtual bool LoadingScreenIsPrepared() const override;
diff --git a/Engine/Source/Runtime/MoviePlayer/Private/NullMoviePlayer.h b/Engine/Source/Runtime/MoviePlayer/Private/NullMoviePlayer.h
--- a/Engine/Source/Runtime/MoviePlayer/Private/NullMoviePlayer.h
+++ b/Engine/Source/Runtime/MoviePlayer/Private/NullMoviePlayer.h
@@ -40,7 +40,7 @@ public:
 	virtual bool PlayEarlyStartupMovies() override { return false; }
 	virtual bool PlayMovie() override { return false; }
 	virtual void StopMovie() override {}
-	virtual void WaitForMovieToFinish() override {}
+	virtual void WaitForMovieToFinish(bool bAllowEngineTick = false) override {}
 	virtual bool IsLoadingFinished() const override {return true;}
 	virtual bool IsMovieCurrentlyPlaying() const override  {return false;}
 	virtual bool LoadingScreenIsPrepared() const override {return false;}
diff --git a/Engine/Source/Runtime/MoviePlayer/Public/MoviePlayer.h b/Engine/Source/Runtime/MoviePlayer/Public/MoviePlayer.h
--- a/Engine/Source/Runtime/MoviePlayer/Public/MoviePlayer.h
+++ b/Engine/Source/Runtime/MoviePlayer/Public/MoviePlayer.h
@@ -112,6 +112,7 @@ struct MOVIEPLAYER_API FLoadingScreenAttributes
 		, bMoviesAreSkippable(true)
 		, bWaitForManualStop(false)
 		, bAllowInEarlyStartup(false)
+		, bAllowEngineTick(false)
 		, PlaybackType(EMoviePlaybackType::MT_Normal) {}
 
 	/** The widget to be displayed on top of the movie or simply standalone if there is no movie. */
@@ -135,6 +136,9 @@ struct MOVIEPLAYER_API FLoadingScreenAttributes
 	/** If true loading screens here cannot have any uobjects of any kind or use any engine features at all. This will start the movies very early as a result on platforms that support it */
 	bool bAllowInEarlyStartup;
 
+	/** If true, this will call the engine tick while the game thread is stalled waiting for a loading movie to finish. This only works for post-startup load screens and is potentially unsafe */
+	bool bAllowEngineTick;
+
 	/** Should we just play back, loop, etc.  NOTE: if the playback type is MT_LoopLast, then bAutoCompleteWhenLoadingCompletes will be togged on when the last movie is hit*/
 	TEnumAsByte<EMoviePlaybackType> PlaybackType;
 
@@ -192,8 +196,8 @@ public:
 	 */
 	virtual void StopMovie() = 0;
 	
-	/** Call only on the game thread. Spins this thread until the movie stops. */
-	virtual void WaitForMovieToFinish() = 0;
+	/** Call only on the game thread. Spins this thread until the movie stops. If bAllowEngineTick is true, the engine tick may run inside here */
+	virtual void WaitForMovieToFinish(bool bAllowEngineTick = false) = 0;
 
 	/** Called from to check if the game thread is finished loading. */
 	virtual bool IsLoadingFinished() const = 0;
diff --git a/Engine/Source/Runtime/NavigationSystem/Private/NavMesh/RecastNavMeshGenerator.cpp b/Engine/Source/Runtime/NavigationSystem/Private/NavMesh/RecastNavMeshGenerator.cpp
--- a/Engine/Source/Runtime/NavigationSystem/Private/NavMesh/RecastNavMeshGenerator.cpp
+++ b/Engine/Source/Runtime/NavigationSystem/Private/NavMesh/RecastNavMeshGenerator.cpp
@@ -745,7 +745,7 @@ FORCEINLINE_DEBUGGABLE void ExportRigidBodyTriMesh(UBodySetup& BodySetup, TNavSt
 }
 
 void ExportRigidBodyBoxElements(const FKAggregateGeom& AggGeom, TNavStatArray<float>& VertexBuffer, TNavStatArray<int32>& IndexBuffer,
-								TNavStatArray<int32>& ShapeBuffer, FBox& UnrealBounds, const FTransform& LocalToWorld, int32 NumExistingVerts)
+								TNavStatArray<int32>& ShapeBuffer, FBox& UnrealBounds, const FTransform& LocalToWorld, const int32 NumExistingVerts = 0)
 {
 	for (int32 i = 0; i < AggGeom.BoxElems.Num(); i++)
 	{
@@ -753,7 +753,7 @@ void ExportRigidBodyBoxElements(const FKAggregateGeom& AggGeom, TNavStatArray<fl
 		const FMatrix ElemTM = BoxInfo.GetTransform().ToMatrixWithScale() * LocalToWorld.ToMatrixWithScale();
 		const FVector Extent(BoxInfo.X * 0.5f, BoxInfo.Y * 0.5f, BoxInfo.Z * 0.5f);
 
-		const int32 VertBase = (NumExistingVerts + VertexBuffer.Num()) / 3;
+		const int32 VertBase = NumExistingVerts + (VertexBuffer.Num() / 3);
 		
 		// Store index of first vertex in shape buffer
 		ShapeBuffer.Add(VertBase);
@@ -795,7 +795,7 @@ void ExportRigidBodyBoxElements(const FKAggregateGeom& AggGeom, TNavStatArray<fl
 }
 
 void ExportRigidBodySphylElements(const FKAggregateGeom& AggGeom, TNavStatArray<float>& VertexBuffer, TNavStatArray<int32>& IndexBuffer,
-								  TNavStatArray<int32>& ShapeBuffer, FBox& UnrealBounds, const FTransform& LocalToWorld, int32 NumExistingVerts)
+								  TNavStatArray<int32>& ShapeBuffer, FBox& UnrealBounds, const FTransform& LocalToWorld, const int32 NumExistingVerts = 0)
 {
 	TArray<FVector> ArcVerts;
 
@@ -804,7 +804,7 @@ void ExportRigidBodySphylElements(const FKAggregateGeom& AggGeom, TNavStatArray<
 		const FKSphylElem& SphylInfo = AggGeom.SphylElems[i];
 		const FMatrix ElemTM = SphylInfo.GetTransform().ToMatrixWithScale() * LocalToWorld.ToMatrixWithScale();
 
-		const int32 VertBase = (NumExistingVerts + VertexBuffer.Num()) / 3;
+		const int32 VertBase = NumExistingVerts + (VertexBuffer.Num() / 3);
 
 		// Store index of first vertex in shape buffer
 		ShapeBuffer.Add(VertBase);
@@ -875,7 +875,7 @@ void ExportRigidBodySphylElements(const FKAggregateGeom& AggGeom, TNavStatArray<
 }
 
 void ExportRigidBodySphereElements(const FKAggregateGeom& AggGeom, TNavStatArray<float>& VertexBuffer, TNavStatArray<int32>& IndexBuffer,
-								   TNavStatArray<int32>& ShapeBuffer, FBox& UnrealBounds, const FTransform& LocalToWorld, int32 NumExistingVerts)
+								   TNavStatArray<int32>& ShapeBuffer, FBox& UnrealBounds, const FTransform& LocalToWorld, const int32 NumExistingVerts = 0)
 {
 	TArray<FVector> ArcVerts;
 
@@ -884,7 +884,7 @@ void ExportRigidBodySphereElements(const FKAggregateGeom& AggGeom, TNavStatArray
 		const FKSphereElem& SphereInfo = AggGeom.SphereElems[i];
 		const FMatrix ElemTM = SphereInfo.GetTransform().ToMatrixWithScale() * LocalToWorld.ToMatrixWithScale();
 
-		const int32 VertBase = (NumExistingVerts + VertexBuffer.Num()) / 3;
+		const int32 VertBase = NumExistingVerts + (VertexBuffer.Num() / 3);
 
 		// Store index of first vertex in shape buffer
 		ShapeBuffer.Add(VertBase);
@@ -948,13 +948,12 @@ FORCEINLINE_DEBUGGABLE void ExportRigidBodySetup(UBodySetup& BodySetup, TNavStat
 	BodySetup.CreatePhysicsMeshes();
 
 	static TNavStatArray<int32> TemporaryShapeBuffer;
-	const int32 NumExistingsVerts = VertexBuffer.Num();
 
 	ExportRigidBodyTriMesh(BodySetup, VertexBuffer, IndexBuffer, UnrealBounds, LocalToWorld);
 	ExportRigidBodyConvexElements(BodySetup, VertexBuffer, IndexBuffer, TemporaryShapeBuffer, UnrealBounds, LocalToWorld);
-	ExportRigidBodyBoxElements(BodySetup.AggGeom, VertexBuffer, IndexBuffer, TemporaryShapeBuffer, UnrealBounds, LocalToWorld, NumExistingsVerts);
-	ExportRigidBodySphylElements(BodySetup.AggGeom, VertexBuffer, IndexBuffer, TemporaryShapeBuffer, UnrealBounds, LocalToWorld, NumExistingsVerts);
-	ExportRigidBodySphereElements(BodySetup.AggGeom, VertexBuffer, IndexBuffer, TemporaryShapeBuffer, UnrealBounds, LocalToWorld, NumExistingsVerts);
+	ExportRigidBodyBoxElements(BodySetup.AggGeom, VertexBuffer, IndexBuffer, TemporaryShapeBuffer, UnrealBounds, LocalToWorld);
+	ExportRigidBodySphylElements(BodySetup.AggGeom, VertexBuffer, IndexBuffer, TemporaryShapeBuffer, UnrealBounds, LocalToWorld);
+	ExportRigidBodySphereElements(BodySetup.AggGeom, VertexBuffer, IndexBuffer, TemporaryShapeBuffer, UnrealBounds, LocalToWorld);
 
 	TemporaryShapeBuffer.Reset();
 }
diff --git a/Engine/Source/Runtime/UMG/Public/Components/Widget.h b/Engine/Source/Runtime/UMG/Public/Components/Widget.h
--- a/Engine/Source/Runtime/UMG/Public/Components/Widget.h
+++ b/Engine/Source/Runtime/UMG/Public/Components/Widget.h
@@ -446,7 +446,7 @@ public:
 
 	/** @return true if the widget is currently being hovered by a pointer device */
 	UFUNCTION(BlueprintCallable, Category="Widget")
-	bool IsHovered() const;
+	virtual bool IsHovered() const;
 
 	/**
 	 * Checks to see if this widget currently has the keyboard focus
