Could not repro crash. Added some safety code and logging.

diff --git a/Engine/Source/Runtime/Engine/Private/Particles/ParticleComponents.cpp b/Engine/Source/Runtime/Engine/Private/Particles/ParticleComponents.cpp
--- a/Engine/Source/Runtime/Engine/Private/Particles/ParticleComponents.cpp
+++ b/Engine/Source/Runtime/Engine/Private/Particles/ParticleComponents.cpp
@@ -2250,6 +2250,12 @@ void UParticleSystem::PostLoad()
 
 		if (!bCookedOut)
 		{
+			if (Emitter->LODLevels.Num() == 0)
+			{
+				UE_LOG(LogParticles, Warning, TEXT("ParticleSystem contains emitter with no lod levels - %s - %s"), *GetFullName(), *Emitter->GetFullName());
+				continue;
+			}
+
 			UParticleLODLevel* LODLevel = Emitter->LODLevels[0];
 			check(LODLevel);
 
