Merging using Fortnight-Main<->Dev-Anim

- Rename MeshComponentUpdateFlag to VisibilityBasedAnimTickOption for better description
- Made it configurable from INI file so that you can change engine default for your project
- This doesn't include INI file change yet. I'm waiting for this to go through without issue, and then will check in the INI file change

diff --git a/Engine/Config/BaseEngine.ini b/Engine/Config/BaseEngine.ini
--- a/Engine/Config/BaseEngine.ini
+++ b/Engine/Config/BaseEngine.ini
@@ -564,6 +564,8 @@ bUseStreamingPause=false
 +ClassRedirects=(OldName="/Script/AssetScriptingUtilities.SkeletalMeshUtilitiesLibrary", NewName="/Script/EditorScriptingUtilities.EditorSkeletalMeshLibrary")
 +FunctionRedirects=(OldName="StaticMeshUtilitiesLibrary.GetLODScreenSizes",NewName="EditorStaticMeshLibrary.GetLodScreenSizes")
 
++EnumRedirects=(OldName="EMeshComponentUpdateFlag",NewName="/Script/Engine.EVisibilityBasedAnimTickOption")
++PropertyRedirects=(OldName="SkinnedMeshComponent.MeshComponentUpdateFlag", NewName="VisibilityBasedAnimTickOption")
 
 [CoreUObject.Metadata]
 ; Note: UnrealHeaderTool should be rerun after making changes to MetadataRedirects to catch any keys specified in class headers
diff --git a/Engine/Source/Editor/DetailCustomizations/Private/AnimMontageSegmentDetails.cpp b/Engine/Source/Editor/DetailCustomizations/Private/AnimMontageSegmentDetails.cpp
--- a/Engine/Source/Editor/DetailCustomizations/Private/AnimMontageSegmentDetails.cpp
+++ b/Engine/Source/Editor/DetailCustomizations/Private/AnimMontageSegmentDetails.cpp
@@ -252,7 +252,7 @@ void SAnimationSegmentViewport::Construct(const FArguments& InArgs)
 	ViewportWidget->SetViewportInterface( SceneViewport.ToSharedRef() );
 	
 	PreviewComponent = NewObject<UDebugSkelMeshComponent>();
-	PreviewComponent->bEnablePhysicsOnDedicatedServer = EMeshComponentUpdateFlag::AlwaysTickPoseAndRefreshBones;
+	PreviewComponent->VisibilityBasedAnimTickOption = EVisibilityBasedAnimTickOption::AlwaysTickPoseAndRefreshBones;
 	PreviewScene.AddComponent(PreviewComponent, FTransform::Identity);
 
 	InitSkeleton();
diff --git a/Engine/Source/Editor/DetailCustomizations/Private/AnimSequenceDetails.cpp b/Engine/Source/Editor/DetailCustomizations/Private/AnimSequenceDetails.cpp
--- a/Engine/Source/Editor/DetailCustomizations/Private/AnimSequenceDetails.cpp
+++ b/Engine/Source/Editor/DetailCustomizations/Private/AnimSequenceDetails.cpp
@@ -400,7 +400,7 @@ void SAnimationRefPoseViewport::Construct(const FArguments& InArgs)
 
 	// Create the preview component
 	PreviewComponent = NewObject<UDebugSkelMeshComponent>();
-	PreviewComponent->MeshComponentUpdateFlag = EMeshComponentUpdateFlag::AlwaysTickPoseAndRefreshBones;
+	PreviewComponent->VisibilityBasedAnimTickOption = EVisibilityBasedAnimTickOption::AlwaysTickPoseAndRefreshBones;
 	PreviewScene.AddComponent( PreviewComponent, FTransform::Identity );
 
 	this->ChildSlot
diff --git a/Engine/Source/Editor/Persona/Private/AnimationEditorPreviewScene.cpp b/Engine/Source/Editor/Persona/Private/AnimationEditorPreviewScene.cpp
--- a/Engine/Source/Editor/Persona/Private/AnimationEditorPreviewScene.cpp
+++ b/Engine/Source/Editor/Persona/Private/AnimationEditorPreviewScene.cpp
@@ -222,7 +222,7 @@ void FAnimationEditorPreviewScene::SetPreviewMeshInternal(USkeletalMesh* NewPrev
 
 		AddPreviewAttachedObjects();
 
-		SkeletalMeshComponent->MeshComponentUpdateFlag = EMeshComponentUpdateFlag::AlwaysTickPoseAndRefreshBones;
+		SkeletalMeshComponent->VisibilityBasedAnimTickOption = EVisibilityBasedAnimTickOption::AlwaysTickPoseAndRefreshBones;
 	}
 
 	for (auto Iter = AdditionalMeshes.CreateIterator(); Iter; ++Iter)
diff --git a/Engine/Source/Editor/Persona/Private/SAnimationSequenceBrowser.cpp b/Engine/Source/Editor/Persona/Private/SAnimationSequenceBrowser.cpp
--- a/Engine/Source/Editor/Persona/Private/SAnimationSequenceBrowser.cpp
+++ b/Engine/Source/Editor/Persona/Private/SAnimationSequenceBrowser.cpp
@@ -1147,7 +1147,7 @@ void SAnimationSequenceBrowser::CreateAssetTooltipResources()
 	ViewportWidget->SetViewportInterface(SceneViewport.ToSharedRef());
 
 	// Setup the preview component to ensure an animation will update when requested
-	PreviewComponent->MeshComponentUpdateFlag = EMeshComponentUpdateFlag::AlwaysTickPoseAndRefreshBones;
+	PreviewComponent->VisibilityBasedAnimTickOption = EVisibilityBasedAnimTickOption::AlwaysTickPoseAndRefreshBones;
 	PreviewScene.AddComponent(PreviewComponent, FTransform::Identity);
 
 	const USkeletalMeshEditorSettings* Options = GetDefault<USkeletalMeshEditorSettings>();
diff --git a/Engine/Source/Editor/SequenceRecorder/Private/ActorRecording.cpp b/Engine/Source/Editor/SequenceRecorder/Private/ActorRecording.cpp
--- a/Engine/Source/Editor/SequenceRecorder/Private/ActorRecording.cpp
+++ b/Engine/Source/Editor/SequenceRecorder/Private/ActorRecording.cpp
@@ -455,7 +455,7 @@ void UActorRecording::StartRecordingActorProperties(ULevelSequence* CurrentSeque
 					{
 						SkeletalMeshComponent->SetAnimationMode(EAnimationMode::AnimationSingleNode);
 						SkeletalMeshComponent->bEnableUpdateRateOptimizations = false;
-						SkeletalMeshComponent->MeshComponentUpdateFlag = EMeshComponentUpdateFlag::AlwaysTickPoseAndRefreshBones;
+						SkeletalMeshComponent->VisibilityBasedAnimTickOption = EVisibilityBasedAnimTickOption::AlwaysTickPoseAndRefreshBones;
 						SkeletalMeshComponent->ForcedLodModel = 1;
 					}
 
diff --git a/Engine/Source/Editor/SequenceRecorder/Private/AnimationRecorder.cpp b/Engine/Source/Editor/SequenceRecorder/Private/AnimationRecorder.cpp
--- a/Engine/Source/Editor/SequenceRecorder/Private/AnimationRecorder.cpp
+++ b/Engine/Source/Editor/SequenceRecorder/Private/AnimationRecorder.cpp
@@ -771,7 +771,7 @@ FAnimationRecorderManager& FAnimationRecorderManager::Get()
 FAnimRecorderInstance::FAnimRecorderInstance()
 	: SkelComp(nullptr)
 	, Recorder(nullptr)
-	, CachedMeshComponentUpdateFlag(EMeshComponentUpdateFlag::AlwaysTickPoseAndRefreshBones)
+	, CachedVisibilityBasedAnimTickOption(EVisibilityBasedAnimTickOption::AlwaysTickPoseAndRefreshBones)
 	, bCachedEnableUpdateRateOptimizations(false)
 {
 }
@@ -810,10 +810,10 @@ void FAnimRecorderInstance::InitInternal(USkeletalMeshComponent* InComponent, co
 
 		// turn off URO and make sure we always update even if out of view
 		bCachedEnableUpdateRateOptimizations = InComponent->bEnableUpdateRateOptimizations;
-		CachedMeshComponentUpdateFlag = InComponent->MeshComponentUpdateFlag;
+		CachedVisibilityBasedAnimTickOption = InComponent->VisibilityBasedAnimTickOption;
 
 		InComponent->bEnableUpdateRateOptimizations = false;
-		InComponent->MeshComponentUpdateFlag = EMeshComponentUpdateFlag::AlwaysTickPoseAndRefreshBones;
+		InComponent->VisibilityBasedAnimTickOption = EVisibilityBasedAnimTickOption::AlwaysTickPoseAndRefreshBones;
 	}
 }
 
@@ -862,7 +862,7 @@ void FAnimRecorderInstance::FinishRecording(bool bShowMessage)
 
 		// restore update flags
 		SkelComp->bEnableUpdateRateOptimizations = bCachedEnableUpdateRateOptimizations;
-		SkelComp->MeshComponentUpdateFlag = CachedMeshComponentUpdateFlag;
+		SkelComp->VisibilityBasedAnimTickOption = CachedVisibilityBasedAnimTickOption;
 	}
 }
 
diff --git a/Engine/Source/Editor/SequenceRecorder/Private/AnimationRecorder.h b/Engine/Source/Editor/SequenceRecorder/Private/AnimationRecorder.h
--- a/Engine/Source/Editor/SequenceRecorder/Private/AnimationRecorder.h
+++ b/Engine/Source/Editor/SequenceRecorder/Private/AnimationRecorder.h
@@ -127,7 +127,7 @@ public:
 	TSharedPtr<FAnimationRecorder> Recorder;
 
 	/** Used to store/restore update flag when recording */
-	EMeshComponentUpdateFlag::Type CachedMeshComponentUpdateFlag;
+	EVisibilityBasedAnimTickOption CachedVisibilityBasedAnimTickOption;
 
 	/** Used to store/restore URO when recording */
 	bool bCachedEnableUpdateRateOptimizations;
diff --git a/Engine/Source/Editor/SequenceRecorder/Private/Sections/MovieSceneAnimationSectionRecorder.cpp b/Engine/Source/Editor/SequenceRecorder/Private/Sections/MovieSceneAnimationSectionRecorder.cpp
--- a/Engine/Source/Editor/SequenceRecorder/Private/Sections/MovieSceneAnimationSectionRecorder.cpp
+++ b/Engine/Source/Editor/SequenceRecorder/Private/Sections/MovieSceneAnimationSectionRecorder.cpp
@@ -170,6 +170,6 @@ void FMovieSceneAnimationSectionRecorder::Record(float CurrentTime)
 	{
 		// re-force updates on as gameplay can sometimes turn these back off!
 		SkeletalMeshComponent->bEnableUpdateRateOptimizations = false;
-		SkeletalMeshComponent->MeshComponentUpdateFlag = EMeshComponentUpdateFlag::AlwaysTickPoseAndRefreshBones;
+		SkeletalMeshComponent->VisibilityBasedAnimTickOption = EVisibilityBasedAnimTickOption::AlwaysTickPoseAndRefreshBones;
 	}
 }
diff --git a/Engine/Source/Editor/UnrealEd/Private/SSkeletonWidget.cpp b/Engine/Source/Editor/UnrealEd/Private/SSkeletonWidget.cpp
--- a/Engine/Source/Editor/UnrealEd/Private/SSkeletonWidget.cpp
+++ b/Engine/Source/Editor/UnrealEd/Private/SSkeletonWidget.cpp
@@ -1373,7 +1373,7 @@ void SBasePoseViewport::Construct(const FArguments& InArgs)
 	SEditorViewport::Construct(SEditorViewport::FArguments());
 
 	PreviewComponent = NewObject<UDebugSkelMeshComponent>();
-	PreviewComponent->MeshComponentUpdateFlag = EMeshComponentUpdateFlag::AlwaysTickPoseAndRefreshBones;
+	PreviewComponent->VisibilityBasedAnimTickOption = EVisibilityBasedAnimTickOption::AlwaysTickPoseAndRefreshBones;
 	PreviewScene.AddComponent(PreviewComponent, FTransform::Identity);
 
 	SetSkeleton(InArgs._Skeleton);
diff --git a/Engine/Source/Runtime/AnimGraphRuntime/Public/AnimCustomInstance.h b/Engine/Source/Runtime/AnimGraphRuntime/Public/AnimCustomInstance.h
--- a/Engine/Source/Runtime/AnimGraphRuntime/Public/AnimCustomInstance.h
+++ b/Engine/Source/Runtime/AnimGraphRuntime/Public/AnimCustomInstance.h
@@ -25,7 +25,7 @@ class ANIMGRAPHRUNTIME_API UAnimCustomInstance : public UAnimInstance
 	{
 		// make sure to tick and refresh all the time when ticks
 		// @TODO: this needs restoring post-binding
-		InSkeletalMeshComponent->MeshComponentUpdateFlag = EMeshComponentUpdateFlag::AlwaysTickPoseAndRefreshBones;
+		InSkeletalMeshComponent->VisibilityBasedAnimTickOption = EVisibilityBasedAnimTickOption::AlwaysTickPoseAndRefreshBones;
 #if WITH_EDITOR
 		InSkeletalMeshComponent->SetUpdateAnimationInEditor(true);
 #endif
diff --git a/Engine/Source/Runtime/CoreUObject/Private/UObject/CoreRedirects.cpp b/Engine/Source/Runtime/CoreUObject/Private/UObject/CoreRedirects.cpp
--- a/Engine/Source/Runtime/CoreUObject/Private/UObject/CoreRedirects.cpp
+++ b/Engine/Source/Runtime/CoreUObject/Private/UObject/CoreRedirects.cpp
@@ -1305,7 +1305,6 @@ static void RegisterNativeRedirects40(TArray<FCoreRedirect>& Redirects)
 	PROPERTY_REDIRECT("SceneComponent.RelativeTranslation", "SceneComponent.RelativeLocation");
 	PROPERTY_REDIRECT("SceneComponent.bAbsoluteTranslation", "SceneComponent.bAbsoluteLocation");
 	PROPERTY_REDIRECT("SkeletalMeshComponent.AnimationBlueprint", "SkeletalMeshComponent.AnimBlueprintGeneratedClass");
-	PROPERTY_REDIRECT("SkinnedMeshComponent.SkinnedMeshUpdateFlag", "SkinnedMeshComponent.MeshComponentUpdateFlag");
 	PROPERTY_REDIRECT("SlateBrush.TextureName", "SlateBrush.ResourceName");
 	PROPERTY_REDIRECT("SlateBrush.TextureObject", "SlateBrush.ResourceObject");
 	PROPERTY_REDIRECT("WorldSettings.DefaultGameType", "WorldSettings.DefaultGameMode");
@@ -1374,10 +1373,6 @@ static void RegisterNativeRedirects40(TArray<FCoreRedirect>& Redirects)
 	FCoreRedirect* EMaxConcurrentResolutionRule = ENUM_REDIRECT("EMaxConcurrentResolutionRule", "EMaxConcurrentResolutionRule");
 	EMaxConcurrentResolutionRule->ValueChanges.Add(TEXT("EMaxConcurrentResolutionRule::StopFarthest"), TEXT("EMaxConcurrentResolutionRule::StopFarthestThenPreventNew"));
 
-	FCoreRedirect* EMeshComponentUpdateFlag = ENUM_REDIRECT("EMeshComponentUpdateFlag", "EMeshComponentUpdateFlag");
-	EMeshComponentUpdateFlag->ValueChanges.Add(TEXT("SMU_AlwaysTickPose"), TEXT("EMeshComponentUpdateFlag::AlwaysTickPose"));
-	EMeshComponentUpdateFlag->ValueChanges.Add(TEXT("SMU_AlwaysTickPoseAndRefreshBones"), TEXT("EMeshComponentUpdateFlag::AlwaysTickPoseAndRefreshBones"));
-	EMeshComponentUpdateFlag->ValueChanges.Add(TEXT("SMU_OnlyTickPoseWhenRendered"), TEXT("EMeshComponentUpdateFlag::OnlyTickPoseWhenRendered"));
 
 	FCoreRedirect* EParticleEventType = ENUM_REDIRECT("EParticleEventType", "EParticleEventType");
 	EParticleEventType->ValueChanges.Add(TEXT("EPET_Kismet"), TEXT("EPET_Blueprint"));
diff --git a/Engine/Source/Runtime/Engine/Classes/Components/SkeletalMeshComponent.h b/Engine/Source/Runtime/Engine/Classes/Components/SkeletalMeshComponent.h
--- a/Engine/Source/Runtime/Engine/Classes/Components/SkeletalMeshComponent.h
+++ b/Engine/Source/Runtime/Engine/Classes/Components/SkeletalMeshComponent.h
@@ -1134,7 +1134,7 @@ public:
 	FOnAnimInitialized OnAnimInitialized;
 
 	/**
-		If MeshComponentUpdateFlag == EMeshComponentUpdateFlag::OnlyTickMontagesWhenNotRendered
+		If VisibilityBasedAnimTickOption == EVisibilityBasedAnimTickOption::OnlyTickMontagesWhenNotRendered
 		Should we tick Montages only?
 	*/
 	bool ShouldOnlyTickMontages(const float DeltaTime) const;
diff --git a/Engine/Source/Runtime/Engine/Classes/Components/SkinnedMeshComponent.h b/Engine/Source/Runtime/Engine/Classes/Components/SkinnedMeshComponent.h
--- a/Engine/Source/Runtime/Engine/Classes/Components/SkinnedMeshComponent.h
+++ b/Engine/Source/Runtime/Engine/Classes/Components/SkinnedMeshComponent.h
@@ -50,26 +50,38 @@ enum EPhysBodyOp
 	PBO_MAX,
 };
 
-/** Skinned Mesh Update Flag based on rendered or not. */
+/** Skinned Mesh Animation Tick option based on rendered or not. This dictates "TickPose and RefreshBoneTransforms" */
 UENUM()
+enum class EVisibilityBasedAnimTickOption
+{
+	/** Always Tick and Refresh BoneTransforms whether rendered or not. */
+	AlwaysTickPoseAndRefreshBones,
+	/** Always Tick, but Refresh BoneTransforms only when rendered. */
+	AlwaysTickPose,
+	/**
+		When rendered Tick Pose and Refresh Bone Transforms,
+		otherwise, just update montages and skip everything else.
+		(AnimBP graph will not be updated).
+	*/
+	OnlyTickMontagesWhenNotRendered,
+	/** Tick only when rendered, and it will only RefreshBoneTransforms when rendered. */
+	OnlyTickPoseWhenRendered,
+};
+
+/** Previous deprecated animation tick option. */
 namespace EMeshComponentUpdateFlag
 {
-	enum Type
-	{
-		/** Always Tick and Refresh BoneTransforms whether rendered or not. */
-		AlwaysTickPoseAndRefreshBones,
-		/** Always Tick, but Refresh BoneTransforms only when rendered. */
-		AlwaysTickPose,
-		/**
-			When rendered Tick Pose and Refresh Bone Transforms,
-			otherwise, just update montages and skip everything else.
-			(AnimBP graph will not be updated).
-		*/
-		OnlyTickMontagesWhenNotRendered,
-		/** Tick only when rendered, and it will only RefreshBoneTransforms when rendered. */
-		OnlyTickPoseWhenRendered,
-	};
-}
+	DEPRECATED(4.20, "EMeshComponentUpdateFlag has been deprecated use EVisibilityBasedAnimTickOption instead")
+	typedef int32 Type;
+	DEPRECATED(4.20, "EMeshComponentUpdateFlag has been deprecated use EVisibilityBasedAnimTickOption instead")
+	static const uint8 AlwaysTickPoseAndRefreshBones = 0;
+	DEPRECATED(4.20, "EMeshComponentUpdateFlag has been deprecated use EVisibilityBasedAnimTickOption instead")
+	static const uint8 AlwaysTickPose = 1;
+	DEPRECATED(4.20, "EMeshComponentUpdateFlag has been deprecated use EVisibilityBasedAnimTickOption instead")
+	static const uint8 OnlyTickMontagesWhenNotRendered = 2;
+	DEPRECATED(4.20, "EMeshComponentUpdateFlag has been deprecated use EVisibilityBasedAnimTickOption instead")
+	static const uint8 OnlyTickPoseWhenRendered = 3;
+};
 
 /** Values for specifying bone space. */
 UENUM()
@@ -351,15 +363,13 @@ public:
 	/** Array of bone visibilities (containing one of the values in EBoneVisibilityStatus for each bone).  A bone is only visible if it is *exactly* 1 (BVS_Visible) */
 	TArray<uint8> BoneVisibilityStates;
 
-	/** This is update frequency flag even when our Owner has not been rendered recently
-	 * 
-	 * SMU_AlwaysTickPoseAndRefreshBones,			// Always Tick and Refresh BoneTransforms whether rendered or not
-	 * SMU_AlwaysTickPose,							// Always Tick, but Refresh BoneTransforms only when rendered
-	 * SMU_OnlyTickPoseWhenRendered,				// Tick only when rendered, and it will only RefreshBoneTransforms when rendered
-	 * 
+	/*
+	 * This is tick animation frequency option based on this component rendered or not or using montage
+	 *  You can change this default value in the INI file 
+	 * Mostly related with performance
 	 */
-	UPROPERTY(EditAnywhere, AdvancedDisplay, BlueprintReadWrite, Category=SkeletalMesh)
-	TEnumAsByte<EMeshComponentUpdateFlag::Type> MeshComponentUpdateFlag;
+	UPROPERTY(EditAnywhere, AdvancedDisplay, BlueprintReadWrite, Config, Category=Optimization)
+	EVisibilityBasedAnimTickOption VisibilityBasedAnimTickOption;
 
 	/** 
 	 * When true, we will just using the bounds from our MasterPoseComponent.  This is useful for when we have a Mesh Parented
@@ -395,11 +405,11 @@ public:
 	//
 	
 	/** When true, skip using the physics asset etc. and always use the fixed bounds defined in the SkeletalMesh. */
-	UPROPERTY(EditAnywhere, AdvancedDisplay, BlueprintReadWrite, Category=SkeletalMesh)
+	UPROPERTY(EditAnywhere, AdvancedDisplay, BlueprintReadWrite, Category=Optimization)
 	uint8 bComponentUseFixedSkelBounds:1;
 
 	/** If true, when updating bounds from a PhysicsAsset, consider _all_ BodySetups, not just those flagged with bConsiderForBounds. */
-	UPROPERTY(EditAnywhere, AdvancedDisplay, BlueprintReadWrite, Category=SkeletalMesh)
+	UPROPERTY(EditAnywhere, AdvancedDisplay, BlueprintReadWrite, Category=Optimization)
 	uint8 bConsiderAllBodiesForBounds:1;
 
 	/** If true, this component uses its parents LOD when attached if available
diff --git a/Engine/Source/Runtime/Engine/Private/Animation/AnimInstanceProxy.cpp b/Engine/Source/Runtime/Engine/Private/Animation/AnimInstanceProxy.cpp
--- a/Engine/Source/Runtime/Engine/Private/Animation/AnimInstanceProxy.cpp
+++ b/Engine/Source/Runtime/Engine/Private/Animation/AnimInstanceProxy.cpp
@@ -1114,7 +1114,7 @@ void FAnimInstanceProxy::SlotEvaluatePose(const FName& SlotNodeName, const FComp
 	{
 		// If MontageEvaluationData is not valid anymore, pass-through AnimSlot.
 		// This can happen if InitAnim pushes a RefreshBoneTransforms when not rendered,
-		// with EMeshComponentUpdateFlag::OnlyTickMontagesWhenNotRendered set.
+		// with EVisibilityBasedAnimTickOption::OnlyTickMontagesWhenNotRendered set.
 		if (!EvalState.Montage.IsValid())
 		{
 			BlendedPose = SourcePose;
diff --git a/Engine/Source/Runtime/Engine/Private/Animation/SkeletalMeshActor.cpp b/Engine/Source/Runtime/Engine/Private/Animation/SkeletalMeshActor.cpp
--- a/Engine/Source/Runtime/Engine/Private/Animation/SkeletalMeshActor.cpp
+++ b/Engine/Source/Runtime/Engine/Private/Animation/SkeletalMeshActor.cpp
@@ -18,7 +18,7 @@ ASkeletalMeshActor::ASkeletalMeshActor(const FObjectInitializer& ObjectInitializ
 {
 
 	SkeletalMeshComponent = CreateDefaultSubobject<USkeletalMeshComponent>(TEXT("SkeletalMeshComponent0"));
-	SkeletalMeshComponent->MeshComponentUpdateFlag = EMeshComponentUpdateFlag::AlwaysTickPose;
+	SkeletalMeshComponent->VisibilityBasedAnimTickOption = EVisibilityBasedAnimTickOption::AlwaysTickPose;
 	// check BaseEngine.ini for profile setup
 	SkeletalMeshComponent->SetCollisionProfileName(UCollisionProfile::PhysicsActor_ProfileName);
 	RootComponent = SkeletalMeshComponent;
diff --git a/Engine/Source/Runtime/Engine/Private/Character.cpp b/Engine/Source/Runtime/Engine/Private/Character.cpp
--- a/Engine/Source/Runtime/Engine/Private/Character.cpp
+++ b/Engine/Source/Runtime/Engine/Private/Character.cpp
@@ -94,7 +94,7 @@ ACharacter::ACharacter(const FObjectInitializer& ObjectInitializer)
 		Mesh->AlwaysLoadOnClient = true;
 		Mesh->AlwaysLoadOnServer = true;
 		Mesh->bOwnerNoSee = false;
-		Mesh->MeshComponentUpdateFlag = EMeshComponentUpdateFlag::AlwaysTickPose;
+		Mesh->VisibilityBasedAnimTickOption = EVisibilityBasedAnimTickOption::AlwaysTickPose;
 		Mesh->bCastDynamicShadow = true;
 		Mesh->bAffectDynamicIndirectLighting = true;
 		Mesh->PrimaryComponentTick.TickGroup = TG_PrePhysics;
diff --git a/Engine/Source/Runtime/Engine/Private/Components/SkeletalMeshComponent.cpp b/Engine/Source/Runtime/Engine/Private/Components/SkeletalMeshComponent.cpp
--- a/Engine/Source/Runtime/Engine/Private/Components/SkeletalMeshComponent.cpp
+++ b/Engine/Source/Runtime/Engine/Private/Components/SkeletalMeshComponent.cpp
@@ -160,7 +160,7 @@ USkeletalMeshComponent::USkeletalMeshComponent(const FObjectInitializer& ObjectI
 	bWantsInitializeComponent = true;
 	GlobalAnimRateScale = 1.0f;
 	bNoSkeletonUpdate = false;
-	MeshComponentUpdateFlag = EMeshComponentUpdateFlag::AlwaysTickPoseAndRefreshBones;
+	VisibilityBasedAnimTickOption = EVisibilityBasedAnimTickOption::AlwaysTickPoseAndRefreshBones;
 	KinematicBonesUpdateType = EKinematicBonesUpdateToPhysics::SkipSimulatingBones;
 	PhysicsTransformUpdateMode = EPhysicsTransformUpdateMode::SimulationUpatesComponentTransform;
 	SetGenerateOverlapEvents(false);
@@ -515,7 +515,7 @@ void USkeletalMeshComponent::OnRegister()
 	// so we need to force initialize them before we begin to tick.
 	InitAnim(true);
 
-	if (bRenderStatic || (MeshComponentUpdateFlag == EMeshComponentUpdateFlag::OnlyTickPoseWhenRendered && !FApp::CanEverRender()))
+	if (bRenderStatic || (VisibilityBasedAnimTickOption == EVisibilityBasedAnimTickOption::OnlyTickPoseWhenRendered && !FApp::CanEverRender()))
 	{
 		SetComponentTickEnabled(false);
 	}
@@ -971,7 +971,7 @@ bool USkeletalMeshComponent::ShouldOnlyTickMontages(const float DeltaTime) const
 {
 	// Ignore DeltaSeconds == 0.f, as that is used when we want to force an update followed by RefreshBoneTransforms.
 	// RefreshBoneTransforms will need an updated graph.
-	return (MeshComponentUpdateFlag == EMeshComponentUpdateFlag::OnlyTickMontagesWhenNotRendered)
+	return (VisibilityBasedAnimTickOption == EVisibilityBasedAnimTickOption::OnlyTickMontagesWhenNotRendered)
 		&& !bRecentlyRendered 
 		&& (DeltaTime > 0.f);
 }
@@ -1078,7 +1078,7 @@ bool USkeletalMeshComponent::ShouldTickPose() const
 	const bool bShouldTickBasedOnAutonomousCheck = bIsAutonomousTickPose || (!bOnlyAllowAutonomousTickPose && !bAlreadyTickedThisFrame);
 	// When playing networked Root Motion Montages, we want these to play on dedicated servers and remote clients for networking and position correction purposes.
 	// So we force pose updates in that case to keep root motion and position in sync.
-	const bool bShouldTickBasedOnVisibility = ((MeshComponentUpdateFlag < EMeshComponentUpdateFlag::OnlyTickPoseWhenRendered) || bRecentlyRendered || IsPlayingNetworkedRootMotionMontage());
+	const bool bShouldTickBasedOnVisibility = ((VisibilityBasedAnimTickOption < EVisibilityBasedAnimTickOption::OnlyTickPoseWhenRendered) || bRecentlyRendered || IsPlayingNetworkedRootMotionMontage());
 
 	return (bShouldTickBasedOnVisibility && bShouldTickBasedOnAutonomousCheck && IsRegistered() && (AnimScriptInstance || PostProcessAnimInstance) && !bPauseAnims && GetWorld()->AreActorsInitialized() && !bNoSkeletonUpdate);
 }
diff --git a/Engine/Source/Runtime/Engine/Private/Components/SkinnedMeshComponent.cpp b/Engine/Source/Runtime/Engine/Private/Components/SkinnedMeshComponent.cpp
--- a/Engine/Source/Runtime/Engine/Private/Components/SkinnedMeshComponent.cpp
+++ b/Engine/Source/Runtime/Engine/Private/Components/SkinnedMeshComponent.cpp
@@ -320,7 +320,7 @@ USkinnedMeshComponent::USkinnedMeshComponent(const FObjectInitializer& ObjectIni
 	PrimaryComponentTick.TickGroup = TG_PrePhysics;
 	WireframeColor = FColor(221, 221, 28, 255);
 
-	MeshComponentUpdateFlag = EMeshComponentUpdateFlag::AlwaysTickPoseAndRefreshBones;
+	VisibilityBasedAnimTickOption = EVisibilityBasedAnimTickOption::AlwaysTickPoseAndRefreshBones;
 
 	SetCollisionProfileName(UCollisionProfile::NoCollision_ProfileName);
 
@@ -591,7 +591,7 @@ void USkinnedMeshComponent::SendRenderDynamicData_Concurrent()
 
 	// if we have not updated the transforms then no need to send them to the rendering thread
 	// @todo GIsEditor used to be bUpdateSkelWhenNotRendered. Look into it further to find out why it doesn't update animations in the AnimSetViewer, when a level is loaded in UED (like POC_Cover.gear).
-	if( MeshObject && SkeletalMesh && (bForceMeshObjectUpdate || (bRecentlyRendered || MeshComponentUpdateFlag == EMeshComponentUpdateFlag::AlwaysTickPoseAndRefreshBones || GIsEditor || MeshObject->bHasBeenUpdatedAtLeastOnce == false)) )
+	if( MeshObject && SkeletalMesh && (bForceMeshObjectUpdate || (bRecentlyRendered || VisibilityBasedAnimTickOption == EVisibilityBasedAnimTickOption::AlwaysTickPoseAndRefreshBones || GIsEditor || MeshObject->bHasBeenUpdatedAtLeastOnce == false)) )
 	{
 		SCOPE_CYCLE_COUNTER(STAT_MeshObjectUpdate);
 
@@ -675,12 +675,12 @@ void USkinnedMeshComponent::InitLODInfos()
 
 bool USkinnedMeshComponent::ShouldTickPose() const
 {
-	return ((MeshComponentUpdateFlag < EMeshComponentUpdateFlag::OnlyTickPoseWhenRendered) || bRecentlyRendered);
+	return ((VisibilityBasedAnimTickOption < EVisibilityBasedAnimTickOption::OnlyTickPoseWhenRendered) || bRecentlyRendered);
 }
 
 bool USkinnedMeshComponent::ShouldUpdateTransform(bool bLODHasChanged) const
 {
-	return (bRecentlyRendered || (MeshComponentUpdateFlag == EMeshComponentUpdateFlag::AlwaysTickPoseAndRefreshBones));
+	return (bRecentlyRendered || (VisibilityBasedAnimTickOption == EVisibilityBasedAnimTickOption::AlwaysTickPoseAndRefreshBones));
 }
 
 bool USkinnedMeshComponent::ShouldUseUpdateRateOptimizations() const
@@ -743,7 +743,7 @@ void USkinnedMeshComponent::TickComponent(float DeltaTime, enum ELevelTick TickT
 			RefreshBoneTransforms(ThisTickFunction);
 		}
 	}
-	else if(MeshComponentUpdateFlag == EMeshComponentUpdateFlag::AlwaysTickPose)
+	else if(VisibilityBasedAnimTickOption == EVisibilityBasedAnimTickOption::AlwaysTickPose)
 	{
 		// We are not refreshing bone transforms, but we do want to tick pose. We may need to kick off a parallel task
 		DispatchParallelTickPose(ThisTickFunction);
diff --git a/Engine/Source/Runtime/Engine/Private/SkeletalMeshComponentPhysics.cpp b/Engine/Source/Runtime/Engine/Private/SkeletalMeshComponentPhysics.cpp
--- a/Engine/Source/Runtime/Engine/Private/SkeletalMeshComponentPhysics.cpp
+++ b/Engine/Source/Runtime/Engine/Private/SkeletalMeshComponentPhysics.cpp
@@ -2923,7 +2923,7 @@ void USkeletalMeshComponent::TickClothing(float DeltaTime, FTickFunction& ThisTi
 	}
 
 	// Use the component update flag to gate simulation to respect the always tick options
-	bool bShouldTick = ((MeshComponentUpdateFlag < EMeshComponentUpdateFlag::OnlyTickPoseWhenRendered) || bRecentlyRendered);
+	bool bShouldTick = ((VisibilityBasedAnimTickOption < EVisibilityBasedAnimTickOption::OnlyTickPoseWhenRendered) || bRecentlyRendered);
 
 	if (bShouldTick)
 	{
diff --git a/Engine/Source/Runtime/MovieSceneTracks/Private/Evaluation/MovieSceneSkeletalAnimationTemplate.cpp b/Engine/Source/Runtime/MovieSceneTracks/Private/Evaluation/MovieSceneSkeletalAnimationTemplate.cpp
--- a/Engine/Source/Runtime/MovieSceneTracks/Private/Evaluation/MovieSceneSkeletalAnimationTemplate.cpp
+++ b/Engine/Source/Runtime/MovieSceneTracks/Private/Evaluation/MovieSceneSkeletalAnimationTemplate.cpp
@@ -77,7 +77,7 @@ struct FPreAnimatedAnimationTokenProducer : IMovieScenePreAnimatedTokenProducer
 			FToken(USkeletalMeshComponent* InComponent)
 			{
 				// Cache this object's current update flag and animation mode
-				MeshComponentUpdateFlag = InComponent->MeshComponentUpdateFlag;
+				VisibilityBasedAnimTickOption = InComponent->VisibilityBasedAnimTickOption;
 				AnimationMode = InComponent->GetAnimationMode();
 			}
 
@@ -94,7 +94,7 @@ struct FPreAnimatedAnimationTokenProducer : IMovieScenePreAnimatedTokenProducer
 				UAnimSequencerInstance::UnbindFromSkeletalMeshComponent(Component);
 
 				// Reset the mesh component update flag and animation mode to what they were before we animated the object
-				Component->MeshComponentUpdateFlag = MeshComponentUpdateFlag;
+				Component->VisibilityBasedAnimTickOption = VisibilityBasedAnimTickOption;
 				if (Component->GetAnimationMode() != AnimationMode)
 				{
 					// this SetAnimationMode reinitializes even if the mode is same
@@ -103,7 +103,7 @@ struct FPreAnimatedAnimationTokenProducer : IMovieScenePreAnimatedTokenProducer
 				}
 			}
 
-			EMeshComponentUpdateFlag::Type MeshComponentUpdateFlag;
+			EVisibilityBasedAnimTickOption VisibilityBasedAnimTickOption;
 			EAnimationMode::Type AnimationMode;
 		};
 
