Merging using Dev-AnimPhys<->Dev-Anim

Fix issue with a huge scale of root causing retarget to create invalid transform

diff --git a/Engine/Source/Runtime/Engine/Private/Animation/AnimSequence.cpp b/Engine/Source/Runtime/Engine/Private/Animation/AnimSequence.cpp
--- a/Engine/Source/Runtime/Engine/Private/Animation/AnimSequence.cpp
+++ b/Engine/Source/Runtime/Engine/Private/Animation/AnimSequence.cpp
@@ -2962,7 +2962,7 @@ int32 FindMeshBoneIndexFromBoneName(USkeleton * Skeleton, const FName &BoneName)
 
 	return BoneIndex;
 }
-void FillUpTransformBasedOnRig(USkeleton* Skeleton, TArray<FTransform>& NodeSpaceBases, TArray<FTransform> &Rotations, TArray<FTransform>& Translations, TArray<bool>& TranslationParentFlags)
+void FillUpTransformBasedOnRig(USkeleton* Skeleton, TArray<FTransform>& NodeSpaceBases, TArray<FTransform> &Rotations, TArray<FVector>& Translations, TArray<bool>& TranslationParentFlags)
 {
 	TArray<FTransform> SpaceBases;
 	FAnimationRuntime::FillUpComponentSpaceTransformsRetargetBasePose(Skeleton, SpaceBases);
@@ -3002,14 +3002,14 @@ void FillUpTransformBasedOnRig(USkeleton* Skeleton, TArray<FTransform>& NodeSpac
 					// add identity
 					NodeSpaceBases[Index].SetIdentity();
 					Rotations[Index].SetIdentity();
-					Translations[Index].SetIdentity();
+					Translations[Index] = FVector::ZeroVector;
 				}
 				else
 				{
 					// initialize with SpaceBases - assuming World Based
 					NodeSpaceBases[Index] = SpaceBases[BoneIndex];
 					Rotations[Index] = SpaceBases[BoneIndex];
-					Translations[Index] = SpaceBases[BoneIndex];
+					Translations[Index] = SpaceBases[BoneIndex].GetLocation();
 
 					const FTransformBase* TransformBase = Rig->GetTransformBaseByNodeName(NodeName);
 
@@ -3040,7 +3040,7 @@ void FillUpTransformBasedOnRig(USkeleton* Skeleton, TArray<FTransform>& NodeSpac
 							if (ParentBoneIndex != INDEX_NONE)
 							{
 								// I think translation has to include rotation, otherwise it won't work
-								Translations[Index] = SpaceBases[BoneIndex].GetRelativeTransform(SpaceBases[ParentBoneIndex]);
+								Translations[Index] = SpaceBases[BoneIndex].GetLocation() - SpaceBases[ParentBoneIndex].GetLocation();
 								TranslationParentFlags[Index] = true;
 							}
 						}
@@ -3092,7 +3092,7 @@ void UAnimSequence::RemapTracksToNewSkeleton( USkeleton* NewSkeleton, bool bConv
 				// first calculate component space ref pose to get the relative transform between
 				// two ref poses. It is very important update ref pose before getting here. 
 				TArray<FTransform> NewRotations, OldRotations, NewSpaceBases, OldSpaceBases;
-				TArray<FTransform> NewTranslations, OldTranslations;
+				TArray<FVector> NewTranslations, OldTranslations;
 				TArray<bool> NewTranslationParentFlags, OldTranslationParentFlags;
 				// get the spacebases transform
 				FillUpTransformBasedOnRig(NewSkeleton, NewSpaceBases, NewRotations, NewTranslations, NewTranslationParentFlags);
@@ -3121,8 +3121,8 @@ void UAnimSequence::RemapTracksToNewSkeleton( USkeleton* NewSkeleton, bool bConv
 					RelativeToNewSpaceBases[NodeIndex] = NewSpaceBases[NodeIndex].GetRelativeTransform(OldSpaceBases[NodeIndex]); 
 
 					// also savees the translation difference between old to new
-					FVector OldTranslation = OldTranslations[NodeIndex].GetTranslation();
-					FVector NewTranslation = NewTranslations[NodeIndex].GetTranslation();
+					FVector OldTranslation = OldTranslations[NodeIndex];
+					FVector NewTranslation = NewTranslations[NodeIndex];
 
 					// skip root because we don't really have clear relative point to test with it
 					if (NodeIndex != 0 && NewTranslationParentFlags[NodeIndex] == OldTranslationParentFlags[NodeIndex])
@@ -3206,7 +3206,8 @@ void UAnimSequence::RemapTracksToNewSkeleton( USkeleton* NewSkeleton, bool bConv
 					{
 						FVector ComponentSpaceTranslation = ComponentSpaceAnimations[TransParentTrackIndex][Key].TransformPosition(AnimatedLocalKey.GetTranslation());
 						ComponentSpaceAnimations[SrcTrackIndex][Key].SetTranslation(ComponentSpaceTranslation);
-						ComponentSpaceAnimations[SrcTrackIndex][Key].SetScale3D(AnimatedLocalKey.GetScale3D());
+						FVector ParentComponentSpaceScale3D = ComponentSpaceAnimations[TransParentTrackIndex][Key].GetScale3D();
+						ComponentSpaceAnimations[SrcTrackIndex][Key].SetScale3D(ParentComponentSpaceScale3D * AnimatedLocalKey.GetScale3D());
 					}
 					else
 					{
@@ -3245,9 +3246,9 @@ void UAnimSequence::RemapTracksToNewSkeleton( USkeleton* NewSkeleton, bool bConv
 
 					if(TransParentTrackIndex != INDEX_NONE)
 					{
-						FVector LocalTranslation = ConvertedSpaceAnimations[SrcTrackIndex][Key].GetRelativeTransform(ConvertedSpaceAnimations[TransParentTrackIndex][Key]).GetTranslation();
-						ConvertedLocalSpaceAnimations[SrcTrackIndex][Key].SetTranslation(LocalTranslation);
-						ConvertedLocalSpaceAnimations[SrcTrackIndex][Key].SetScale3D(ConvertedSpaceAnimations[SrcTrackIndex][Key].GetScale3D());
+						FTransform LocalTransform = ConvertedSpaceAnimations[SrcTrackIndex][Key].GetRelativeTransform(ConvertedSpaceAnimations[TransParentTrackIndex][Key]);
+						ConvertedLocalSpaceAnimations[SrcTrackIndex][Key].SetTranslation(LocalTransform.GetLocation());
+						ConvertedLocalSpaceAnimations[SrcTrackIndex][Key].SetScale3D(LocalTransform.GetScale3D());
 					}
 					else
 					{
