Merging 4070600 from Dev-Networking

Fix crash in ShooterGame rep graph after seamless travel from a map with streaming levels

diff --git a/Engine/Plugins/Runtime/ReplicationGraph/Source/Public/ReplicationGraph.h b/Engine/Plugins/Runtime/ReplicationGraph/Source/Public/ReplicationGraph.h
--- a/Engine/Plugins/Runtime/ReplicationGraph/Source/Public/ReplicationGraph.h
+++ b/Engine/Plugins/Runtime/ReplicationGraph/Source/Public/ReplicationGraph.h
@@ -637,6 +637,8 @@ public:
 
 	virtual void SetWorld(UWorld* InWorld) override;
 
+	virtual void ResetGameWorldState() override { }
+
 	/** Called by the NetDriver when the client connection is ready/added to the NetDriver's client connection list */
 	virtual void AddClientConnection(UNetConnection* NetConnection) override;
 
diff --git a/Engine/Source/Runtime/Engine/Classes/Engine/ReplicationDriver.h b/Engine/Source/Runtime/Engine/Classes/Engine/ReplicationDriver.h
--- a/Engine/Source/Runtime/Engine/Classes/Engine/ReplicationDriver.h
+++ b/Engine/Source/Runtime/Engine/Classes/Engine/ReplicationDriver.h
@@ -65,6 +65,8 @@ public:
 
 	virtual void SetWorld(UWorld* InWorld) PURE_VIRTUAL(UReplicationDriver::SetWorld, );
 
+	virtual void ResetGameWorldState() PURE_VIRTUAL(UReplicationDriver::ResetGameWorldState, );
+
 	virtual void AddClientConnection(UNetConnection* NetConnection) PURE_VIRTUAL(UReplicationDriver::AddClientConnection, );
 
 	virtual void RemoveClientConnection(UNetConnection* NetConnection) PURE_VIRTUAL(UReplicationDriver::RemoveClientConnection, );
diff --git a/Engine/Source/Runtime/Engine/Private/NetworkDriver.cpp b/Engine/Source/Runtime/Engine/Private/NetworkDriver.cpp
--- a/Engine/Source/Runtime/Engine/Private/NetworkDriver.cpp
+++ b/Engine/Source/Runtime/Engine/Private/NetworkDriver.cpp
@@ -4269,6 +4269,11 @@ void UNetDriver::ResetGameWorldState()
 	{
 		(*It)->ResetGameWorldState();
 	}
+
+	if (ReplicationDriver)
+	{
+		ReplicationDriver->ResetGameWorldState();
+	}
 }
 
 void UNetDriver::CleanPackageMaps()
