Fix infiltrator demo crash

diff --git a/Engine/Source/Developer/MeshSimplifier/Private/QuadricMeshReduction.cpp b/Engine/Source/Developer/MeshSimplifier/Private/QuadricMeshReduction.cpp
--- a/Engine/Source/Developer/MeshSimplifier/Private/QuadricMeshReduction.cpp
+++ b/Engine/Source/Developer/MeshSimplifier/Private/QuadricMeshReduction.cpp
@@ -681,8 +681,8 @@ public:
 		MeshSimp->SetAttributeWeights(AttributeWeights);
 		//MeshSimp->SetBoundaryLocked();
 		MeshSimp->InitCosts();
-
-		float MaxErrorSqr = MeshSimp->SimplifyMesh(MAX_FLT, FMath::Max(1, int32(NumTris * ReductionSettings.PercentTriangles)));
+		//We need a minimum of 2 triangles, to see the object on both side. If we use one, we will end up with zero triangle when we will remove a shared edge
+		float MaxErrorSqr = MeshSimp->SimplifyMesh(MAX_FLT, FMath::Max(2, int32(NumTris * ReductionSettings.PercentTriangles)));
 
 		NumVerts = MeshSimp->GetNumVerts();
 		NumTris = MeshSimp->GetNumTris();
diff --git a/Engine/Source/Runtime/Engine/Private/Components/StaticMeshComponent.cpp b/Engine/Source/Runtime/Engine/Private/Components/StaticMeshComponent.cpp
--- a/Engine/Source/Runtime/Engine/Private/Components/StaticMeshComponent.cpp
+++ b/Engine/Source/Runtime/Engine/Private/Components/StaticMeshComponent.cpp
@@ -112,7 +112,8 @@ public:
 
 						LODInfo.OverrideVertexColors = new FColorVertexBuffer;
 						LODInfo.OverrideVertexColors->InitFromColorArray(VertexColorLODData.VertexBufferColors);
-
+						
+						check(LODInfo.OverrideVertexColors->GetStride() > 0);
 						BeginInitResource(LODInfo.OverrideVertexColors);
 						bAppliedAnyData = true;
 					}
@@ -1087,6 +1088,7 @@ void UStaticMeshComponent::CopyInstanceVertexColorsIfCompatible( UStaticMeshComp
 							TargetLODInfo.OverrideVertexColors = new FColorVertexBuffer;
 							TargetLODInfo.OverrideVertexColors->InitFromColorArray( CopiedColors );
 						}
+						check(TargetLODInfo.OverrideVertexColors->GetStride() > 0);
 						BeginInitResource( TargetLODInfo.OverrideVertexColors );
 					}
 				}
@@ -1282,10 +1284,14 @@ void UStaticMeshComponent::PrivateFixupOverrideColors()
 				Vertex->Normal = CurRenderData.VertexBuffers.StaticMeshVertexBuffer.VertexTangentZ(VertIndex);
 				Vertex->Color = LODInfo.OverrideVertexColors->VertexColor(VertIndex);
 			}
+			BeginInitResource(LODInfo.OverrideVertexColors);
+			UpdateStaticMeshDeriveDataKey = true;
+		}
+		else
+		{
+			delete LODInfo.OverrideVertexColors;
+			LODInfo.OverrideVertexColors = nullptr;
 		}
-
-		BeginInitResource(LODInfo.OverrideVertexColors);
-		UpdateStaticMeshDeriveDataKey = true;
 	}
 
 	if (UpdateStaticMeshDeriveDataKey)
diff --git a/Engine/Source/Runtime/Engine/Private/StaticMesh.cpp b/Engine/Source/Runtime/Engine/Private/StaticMesh.cpp
--- a/Engine/Source/Runtime/Engine/Private/StaticMesh.cpp
+++ b/Engine/Source/Runtime/Engine/Private/StaticMesh.cpp
@@ -1344,7 +1344,7 @@ FArchive& operator<<(FArchive& Ar, FMeshBuildSettings& BuildSettings)
 // differences, etc.) replace the version GUID below with a new one.
 // In case of merge conflicts with DDC versions, you MUST generate a new GUID
 // and set this new GUID as the version.                                       
-#define STATICMESH_DERIVEDDATA_VER TEXT("C6BBBFF508784DAF85A95B95DC33FF82")
+#define STATICMESH_DERIVEDDATA_VER TEXT("397F52ED7403467A825B98E90EDCBD69")
 
 static const FString& GetStaticMeshDerivedDataVersion()
 {
