- Fix issue where keep state not working properly when montage is used for sequencer because montage is advancing and blends out automaticaly
- added option to "auto blend out" so that we can disable it.

diff --git a/Engine/Source/Runtime/Engine/Classes/Animation/AnimMontage.h b/Engine/Source/Runtime/Engine/Classes/Animation/AnimMontage.h
--- a/Engine/Source/Runtime/Engine/Classes/Animation/AnimMontage.h
+++ b/Engine/Source/Runtime/Engine/Classes/Animation/AnimMontage.h
@@ -309,6 +309,9 @@ struct FAnimMontageInstance
 	// Whether this in this tick's call to Advance we used marker based sync
 	bool bDidUseMarkerSyncThisTick;
 
+	// enable auto blend out. This is instance set up. You can override
+	bool bEnableAutoBlendOut;
+
 private:
 	struct FMontageSubStepper MontageSubStepper;
 
@@ -591,6 +594,10 @@ class UAnimMontage : public UAnimCompositeBase
 	UPROPERTY()
 	bool bEnableRootMotionRotation;
 
+	/** When it hits end, it automatically blends out. If this is false, it won't blend out but keep the last pose until stopped explicitly */
+	UPROPERTY(EditAnywhere, Category = BlendOption)
+	bool bEnableAutoBlendOut;
+
 	/** Root Bone will be locked to that position when extracting root motion. DEPRECATED in 4.5 root motion is controlled by anim sequences **/
 	UPROPERTY()
 	TEnumAsByte<ERootMotionRootLock::Type> RootMotionRootLock;
diff --git a/Engine/Source/Runtime/Engine/Private/Animation/AnimMontage.cpp b/Engine/Source/Runtime/Engine/Private/Animation/AnimMontage.cpp
--- a/Engine/Source/Runtime/Engine/Private/Animation/AnimMontage.cpp
+++ b/Engine/Source/Runtime/Engine/Private/Animation/AnimMontage.cpp
@@ -39,6 +39,7 @@ UAnimMontage::UAnimMontage(const FObjectInitializer& ObjectInitializer)
 	BlendIn.SetBlendTime(0.25f);
 	BlendOut.SetBlendTime(0.25f);
 	BlendOutTriggerTime = -1.f;
+	bEnableAutoBlendOut = true;
 	SyncSlotIndex = 0;
 
 	BlendInTime_DEPRECATED = -1.f;
@@ -1288,6 +1289,7 @@ FAnimMontageInstance::FAnimMontageInstance(UAnimInstance * InAnimInstance)
 	, bPlaying(false)
 	, DefaultBlendTimeMultiplier(1.0f)
 	, bDidUseMarkerSyncThisTick(false)
+	, bEnableAutoBlendOut(true)
 	, AnimInstance(InAnimInstance)
 	, InstanceID(INDEX_NONE)
 	, Position(0.f)
@@ -1317,6 +1319,7 @@ void FAnimMontageInstance::Play(float InPlayRate)
 	InitializeBlend(Montage->BlendIn);	
 	Blend.SetBlendTime(Montage->BlendIn.GetBlendTime() * DefaultBlendTimeMultiplier);
 	Blend.SetValueRange(CurrentWeight, 1.f);
+	bEnableAutoBlendOut = Montage->bEnableAutoBlendOut;
 }
 
 void FAnimMontageInstance::InitializeBlend(const FAlphaBlend& InAlphaBlend)
@@ -2214,7 +2217,7 @@ void FAnimMontageInstance::Advance(float DeltaTime, struct FRootMotionMovementPa
 				// If current section is last one, check to trigger a blend out and if it hasn't stopped yet, see if we should stop
 				// We check this even if we haven't moved, in case our position was different from last frame.
 				// (Code triggered a position jump).
-				if (!IsStopped())
+				if (!IsStopped() && bEnableAutoBlendOut)
 				{
 					const int32 CurrentSectionIndex = MontageSubStepper.GetCurrentSectionIndex();
 					check(NextSections.IsValidIndex(CurrentSectionIndex));
@@ -2675,6 +2678,8 @@ UAnimMontage* FAnimMontageInstance::SetSequencerMontagePosition(FName SlotName,
 			{
 				AnimInst->Montage_Play(PlayingMontage, 1.f, EMontagePlayReturnType::MontageLength, 0.f, false);
 				MontageInstanceToUpdate = AnimInst->GetActiveInstanceForMontage(PlayingMontage);
+				// this is sequencer set up, we disable auto blend out
+				MontageInstanceToUpdate->bEnableAutoBlendOut = false;
 			}
 		}
 
