Fix orion cooking issue

diff --git a/Engine/Source/Runtime/Engine/Classes/Components/SkinnedMeshComponent.h b/Engine/Source/Runtime/Engine/Classes/Components/SkinnedMeshComponent.h
--- a/Engine/Source/Runtime/Engine/Classes/Components/SkinnedMeshComponent.h
+++ b/Engine/Source/Runtime/Engine/Classes/Components/SkinnedMeshComponent.h
@@ -51,8 +51,8 @@ enum EPhysBodyOp
 };
 
 /** Skinned Mesh Animation Tick option based on rendered or not. This dictates "TickPose and RefreshBoneTransforms" */
-UENUM()
-enum class EVisibilityBasedAnimTickOption
+UENUM(BlueprintType)
+enum class EVisibilityBasedAnimTickOption : uint8
 {
 	/** Always Tick and Refresh BoneTransforms whether rendered or not. */
 	AlwaysTickPoseAndRefreshBones,
