Remove End User Settings and disable FEngineAnalytics for non-editor usage.
* Also enabled removing of FPlatformMisc::AllowSendAnonymousGameUsageDataToEpic()
#tests create a new editor project and look for settings to be gone.

diff --git a/Engine/Source/Editor/DetailCustomizations/Private/DetailCustomizations.cpp b/Engine/Source/Editor/DetailCustomizations/Private/DetailCustomizations.cpp
--- a/Engine/Source/Editor/DetailCustomizations/Private/DetailCustomizations.cpp
+++ b/Engine/Source/Editor/DetailCustomizations/Private/DetailCustomizations.cpp
@@ -361,7 +361,6 @@ void FDetailCustomizationsModule::RegisterObjectCustomizations()
 	RegisterCustomClassLayout("SequencerGenericKeyStruct", FOnGetDetailCustomizationInstance::CreateStatic(&FSequencerGenericKeyStructCustomization::MakeInstance));
 
 	RegisterCustomClassLayout("AnalyticsPrivacySettings", FOnGetDetailCustomizationInstance::CreateStatic(&FImportantToggleSettingCustomization::MakeInstance));
-	RegisterCustomClassLayout("EndUserSettings", FOnGetDetailCustomizationInstance::CreateStatic(&FImportantToggleSettingCustomization::MakeInstance));
 
 	RegisterCustomClassLayout("AssetViewerSettings", FOnGetDetailCustomizationInstance::CreateStatic(&FAssetViewerSettingsCustomization::MakeInstance));
 
diff --git a/Engine/Source/Editor/ProjectSettingsViewer/Private/ProjectSettingsViewerModule.cpp b/Engine/Source/Editor/ProjectSettingsViewer/Private/ProjectSettingsViewerModule.cpp
--- a/Engine/Source/Editor/ProjectSettingsViewer/Private/ProjectSettingsViewerModule.cpp
+++ b/Engine/Source/Editor/ProjectSettingsViewer/Private/ProjectSettingsViewerModule.cpp
@@ -34,7 +34,6 @@
 
 
 #include "AISystem.h"
-#include "Engine/EndUserSettings.h"
 #include "Runtime/Slate/Public/SlateSettings.h"
 
 #define LOCTEXT_NAMESPACE "FProjectSettingsViewerModule"
@@ -171,13 +170,6 @@ protected:
 			GetMutableDefault<UCrowdManager>()
 			);
 
-		// End-user settings
-		SettingsModule.RegisterSettings("Project", "Engine", "EndUser",
-			LOCTEXT("EndUserSettingsName", "End-User Settings"),
-			LOCTEXT("EndUserSettingsDescription", "Settings you may wish to expose to end-users of your game."),
-			GetMutableDefault<UEndUserSettings>()
-			);
-
 		SettingsModule.RegisterSettings("Project", "Engine", "Slate",
 			LOCTEXT("SlateSettingsName", "Slate Settings"),
 			LOCTEXT("SlateSettingsDescription", "Settings for Slate."),
diff --git a/Engine/Source/Runtime/Core/Public/GenericPlatform/GenericPlatformMisc.h b/Engine/Source/Runtime/Core/Public/GenericPlatform/GenericPlatformMisc.h
--- a/Engine/Source/Runtime/Core/Public/GenericPlatform/GenericPlatformMisc.h
+++ b/Engine/Source/Runtime/Core/Public/GenericPlatform/GenericPlatformMisc.h
@@ -1126,14 +1126,6 @@ public:
 	 */
 	static const TArray<FString>& GetConfidentialPlatforms();
 	
-	/**
-	 * Returns true if the platform allows network traffic for anonymous end user usage data
-	 */
-	static bool AllowSendAnonymousGameUsageDataToEpic()
-	{
-		return true;
-	}
-
 	/**
 	 * For mobile devices, this will crank up a haptic engine for the specified type to be played later with TriggerMobileHaptics
 	 * If this is called again before Release, it will switch to this type
diff --git a/Engine/Source/Runtime/Core/Public/Lumin/LuminPlatformMisc.h b/Engine/Source/Runtime/Core/Public/Lumin/LuminPlatformMisc.h
--- a/Engine/Source/Runtime/Core/Public/Lumin/LuminPlatformMisc.h
+++ b/Engine/Source/Runtime/Core/Public/Lumin/LuminPlatformMisc.h
@@ -77,12 +77,6 @@ struct CORE_API FLuminPlatformMisc : public FAndroidMisc
 	static const FString& GetApplicationApplicationPackageName();
 	static const FString& GetApplicationComponentName();
 
-	// Disable anonymous usage data network traffic on Lumin
-	static bool AllowSendAnonymousGameUsageDataToEpic()
-	{
-		return false;
-	}
-
 private:
 	static void InitApplicationPaths();
 
diff --git a/Engine/Source/Runtime/Engine/Classes/Engine/EndUserSettings.h b/Engine/Source/Runtime/Engine/Classes/Engine/EndUserSettings.h
deleted file mode 100644
diff --git a/Engine/Source/Runtime/Engine/Classes/Engine/Engine.h b/Engine/Source/Runtime/Engine/Classes/Engine/Engine.h
--- a/Engine/Source/Runtime/Engine/Classes/Engine/Engine.h
+++ b/Engine/Source/Runtime/Engine/Classes/Engine/Engine.h
@@ -2929,12 +2929,6 @@ public:
 
 	/** @return true if editor analytics are enabled */
 	virtual bool AreEditorAnalyticsEnabled() const { return false; }
-	/** @return true if in-game analytics are enabled */
-	bool AreGameAnalyticsEnabled() const;
-	/** @return true if in-game analytics are sent with an anonymous GUID rather than real account and machine ids */
-	bool AreGameAnalyticsAnonymous() const;
-	/** @return true if in-game MTBF analytics are enabled */
-	bool AreGameMTBFEventsEnabled() const;
 	virtual void CreateStartupAnalyticsAttributes( TArray<struct FAnalyticsEventAttribute>& StartSessionAttributes ) const {}
 	
 	/** @return true if the engine is autosaving a package */
diff --git a/Engine/Source/Runtime/Engine/Private/EndUserSettings.cpp b/Engine/Source/Runtime/Engine/Private/EndUserSettings.cpp
deleted file mode 100644
diff --git a/Engine/Source/Runtime/Engine/Private/EngineAnalytics.cpp b/Engine/Source/Runtime/Engine/Private/EngineAnalytics.cpp
--- a/Engine/Source/Runtime/Engine/Private/EngineAnalytics.cpp
+++ b/Engine/Source/Runtime/Engine/Private/EngineAnalytics.cpp
@@ -17,8 +17,6 @@
 #include "RHI.h"
 
 bool FEngineAnalytics::bIsInitialized;
-bool FEngineAnalytics::bIsEditorRun;
-bool FEngineAnalytics::bIsGameRun;
 TSharedPtr<IAnalyticsProviderET> FEngineAnalytics::Analytics;
 TSharedPtr<FEngineSessionManager> FEngineAnalytics::SessionManager;
 
@@ -58,24 +56,18 @@ void FEngineAnalytics::Initialize()
 	check(GEngine);
 
 #if WITH_EDITOR
-	// this will only be true for builds that have editor support (currently PC, Mac, Linux)
+	// this will only be true for builds that have editor support (desktop platforms)
 	// The idea here is to only send editor events for actual editor runs, not for things like -game runs of the editor.
-	bIsEditorRun = GIsEditor && !IsRunningCommandlet();
-	bIsGameRun = false;
+	bool bIsEditorRun = GIsEditor && !IsRunningCommandlet();
 #else
-	// We also want to identify a real run of a game, which is NOT necessarily the opposite of an editor run.
-	// Ideally we'd be able to tell explicitly, but with content-only games, it becomes difficult.
-	// So we ensure we are not an editor run, we don't have EDITOR stuff compiled in, we are not running a commandlet,
-	// we are not a generic, utility program, and we require cooked data.
-	bIsEditorRun = false;
-	bIsGameRun = !IsRunningCommandlet() && !FPlatformProperties::IsProgram() && FPlatformProperties::RequiresCookedData();
+	bool bIsEditorRun = false;
 #endif
 
 #if UE_BUILD_DEBUG
 	const bool bShouldInitAnalytics = false;
 #else
 	// Outside of the editor, the only engine analytics usage is the hardware survey
-	const bool bShouldInitAnalytics = (bIsEditorRun && GEngine->AreEditorAnalyticsEnabled()) || (bIsGameRun && GEngine->AreGameAnalyticsEnabled());
+	const bool bShouldInitAnalytics = bIsEditorRun && GEngine->AreEditorAnalyticsEnabled();
 #endif
 
 	if (bShouldInitAnalytics)
@@ -96,15 +88,7 @@ void FEngineAnalytics::Initialize()
 			FString UE4TypeOverride;
 			bool bHasOverride = GConfig->GetString(TEXT("Analytics"), TEXT("UE4TypeOverride"), UE4TypeOverride, GEngineIni);
 			const TCHAR* UE4TypeStr = bHasOverride ? *UE4TypeOverride : FEngineBuildSettings::IsPerforceBuild() ? TEXT("Perforce") : TEXT("UnrealEngine");
-			if (bIsEditorRun)
-			{
-				Config.APIKeyET = FString::Printf(TEXT("UEEditor.%s.%s"), UE4TypeStr, BuildTypeStr);
-			}
-			else
-			{
-				const UGeneralProjectSettings& ProjectSettings = *GetDefault<UGeneralProjectSettings>();
-				Config.APIKeyET = FString::Printf(TEXT("UEGame.%s.%s|%s|%s"), UE4TypeStr, BuildTypeStr, *ProjectSettings.ProjectID.ToString(), *ProjectSettings.ProjectName);
-			}
+			Config.APIKeyET = FString::Printf(TEXT("UEEditor.%s.%s"), UE4TypeStr, BuildTypeStr);
 		}
 		if (Config.APIServerET.IsEmpty())
 		{
@@ -125,23 +109,7 @@ void FEngineAnalytics::Initialize()
 
 		if (Analytics.IsValid())
 		{
-			// Use an anonymous user id in-game
-			if (bIsGameRun && GEngine->AreGameAnalyticsAnonymous())
-			{
-				FString AnonymousId;
-				if (!FPlatformMisc::GetStoredValue(TEXT("Epic Games"), TEXT("Unreal Engine/Privacy"), TEXT("AnonymousID"), AnonymousId) || AnonymousId.IsEmpty())
-				{
-					AnonymousId = FGuid::NewGuid().ToString(EGuidFormats::DigitsWithHyphensInBraces);
-					FPlatformMisc::SetStoredValue(TEXT("Epic Games"), TEXT("Unreal Engine/Privacy"), TEXT("AnonymousID"), AnonymousId);
-				}
-
-				// Place the anonymous user id into the first field of the UserID set, as it most closely matches the LoginID semantics.
-				Analytics->SetUserID(FString::Printf(TEXT("ANON-%s||"), *AnonymousId));
-			}
-			else
-			{
-				Analytics->SetUserID(FString::Printf(TEXT("%s|%s|%s"), *FPlatformMisc::GetLoginId(), *FPlatformMisc::GetEpicAccountId(), *FPlatformMisc::GetOperatingSystemId()));
-			}
+			Analytics->SetUserID(FString::Printf(TEXT("%s|%s|%s"), *FPlatformMisc::GetLoginId(), *FPlatformMisc::GetEpicAccountId(), *FPlatformMisc::GetOperatingSystemId()));
 
 			TArray<FAnalyticsEventAttribute> StartSessionAttributes;
 			GEngine->CreateStartupAnalyticsAttributes( StartSessionAttributes );
@@ -179,11 +147,8 @@ void FEngineAnalytics::Initialize()
 		// Create the session manager singleton
 		if (!SessionManager.IsValid())
 		{
-			if (bIsEditorRun || (bIsGameRun && GEngine->AreGameMTBFEventsEnabled()))
-			{
-				SessionManager = MakeShareable(new FEngineSessionManager(bIsEditorRun ? EEngineSessionManagerMode::Editor : EEngineSessionManagerMode::Game));
-				SessionManager->Initialize();
-			}
+			SessionManager = MakeShareable(new FEngineSessionManager(EEngineSessionManagerMode::Editor));
+			SessionManager->Initialize();
 		}
 	}
 }
diff --git a/Engine/Source/Runtime/Engine/Private/UnrealEngine.cpp b/Engine/Source/Runtime/Engine/Private/UnrealEngine.cpp
--- a/Engine/Source/Runtime/Engine/Private/UnrealEngine.cpp
+++ b/Engine/Source/Runtime/Engine/Private/UnrealEngine.cpp
@@ -187,7 +187,6 @@ UnrealEngine.cpp: Implements the UEngine class and helpers.
 #include "Performance/EnginePerformanceTargets.h"
 
 #include "InstancedReferenceSubobjectHelper.h"
-#include "Engine/EndUserSettings.h"
 
 #include "Engine/LODActor.h"
 #include "Engine/AssetManager.h"
@@ -12548,21 +12547,6 @@ UGameViewportClient* UEngine::GameViewportForWorld(const UWorld *InWorld) const
 	return (Context ? Context->GameViewport : NULL);
 }
 
-bool UEngine::AreGameAnalyticsEnabled() const
-{ 
-	return FPlatformMisc::AllowSendAnonymousGameUsageDataToEpic() && GetDefault<UEndUserSettings>()->bSendAnonymousUsageDataToEpic;
-}
-
-bool UEngine::AreGameAnalyticsAnonymous() const
-{
-	return !GetDefault<UEndUserSettings>()->bAllowUserIdInUsageData;
-}
-
-bool UEngine::AreGameMTBFEventsEnabled() const
-{
-	return GetDefault<UEndUserSettings>()->bSendMeanTimeBetweenFailureDataToEpic;
-}
-
 void UEngine::SetIsVanillaProduct(bool bInIsVanillaProduct)
 {
 	// set bIsVanillaProduct and if it changes broadcast the core delegate
diff --git a/Engine/Source/Runtime/Engine/Public/EngineAnalytics.h b/Engine/Source/Runtime/Engine/Public/EngineAnalytics.h
--- a/Engine/Source/Runtime/Engine/Public/EngineAnalytics.h
+++ b/Engine/Source/Runtime/Engine/Public/EngineAnalytics.h
@@ -9,15 +9,14 @@ class IAnalyticsProvider;
 class IAnalyticsProviderET;
 
 /**
- * The public interface for the engine's analytics provider singleton.
- * For Epic builds, this will point to epic's internal analytics provider.
- * For licensee builds, it will be NULL by default unless they provide their own
- * configuration.
+ * The public interface for the editor's analytics provider singleton.
  * 
- * WARNING: This is an analytics provider instance that is created whenever UE4 is launched. 
+ * WARNING: This is an analytics provider instance that is created whenever UE4 editor is launched. 
  * It is intended ONLY for use by Epic Games. This is NOT intended for games to send 
  * game-specific telemetry. Create your own provider instance for your game and configure
  * it independently.
+ *
+ * It is called FEngineAnalytics for legacy reasons, and is only used for editor telemetry.
  */
 class FEngineAnalytics : FNoncopyable
 {
@@ -33,17 +32,11 @@ public:
 	static ENGINE_API void Initialize();
 	/** Called to shut down the singleton */
 	static ENGINE_API void Shutdown(bool bIsEngineShutdown = false);
-	/** Is this an editor run? (false for games and commandlets) */
-	static ENGINE_API bool IsEditorRun() { return bIsEditorRun; }
-	/** Is this a game run? (false for the editor, commandlets, programs and non-cooked games) */
-	static ENGINE_API bool IsGameRun() { return bIsGameRun; }
 
 	static ENGINE_API void Tick(float DeltaTime);
 
 private:
 	static bool bIsInitialized;
-	static bool bIsEditorRun;
-	static bool bIsGameRun;
 	static ENGINE_API TSharedPtr<IAnalyticsProviderET> Analytics;
 	static TSharedPtr<FEngineSessionManager> SessionManager;
 };
