ML - Old MorphTargets were deserialized incorrectly

diff --git a/Engine/Source/Runtime/Engine/Classes/Animation/MorphTarget.h b/Engine/Source/Runtime/Engine/Classes/Animation/MorphTarget.h
--- a/Engine/Source/Runtime/Engine/Classes/Animation/MorphTarget.h
+++ b/Engine/Source/Runtime/Engine/Classes/Animation/MorphTarget.h
@@ -29,22 +29,12 @@ struct FMorphTargetDelta
 	/** pipe operator */
 	friend FArchive& operator<<(FArchive& Ar, FMorphTargetDelta& V)
 	{
-		if (Ar.UE4Ver() < VER_UE4_MORPHTARGET_CPU_TANGENTZDELTA_FORMATCHANGE)
+		if ((Ar.UE4Ver() < VER_UE4_MORPHTARGET_CPU_TANGENTZDELTA_FORMATCHANGE) && Ar.IsLoading())
 		{
 			/** old format of change in tangent basis normal */
-			FPackedNormal	TangentZDelta_DEPRECATED;
-
-			if (Ar.IsSaving())
-			{
-				TangentZDelta_DEPRECATED = FPackedNormal(V.TangentZDelta);
-			}
-
+			FDeprecatedSerializedPackedNormal TangentZDelta_DEPRECATED;
 			Ar << V.PositionDelta << TangentZDelta_DEPRECATED << V.SourceIdx;
-
-			if (Ar.IsLoading())
-			{
-				V.TangentZDelta = TangentZDelta_DEPRECATED.ToFVector();
-			}
+			V.TangentZDelta = TangentZDelta_DEPRECATED;
 		}
 		else
 		{
