Fix issue where master bone map does not exists when master pose component is serialized
- allow force update master pose component

diff --git a/Engine/Source/Runtime/Engine/Classes/Components/SkinnedMeshComponent.h b/Engine/Source/Runtime/Engine/Classes/Components/SkinnedMeshComponent.h
--- a/Engine/Source/Runtime/Engine/Classes/Components/SkinnedMeshComponent.h
+++ b/Engine/Source/Runtime/Engine/Classes/Components/SkinnedMeshComponent.h
@@ -942,7 +942,7 @@ public:
 	 * @param NewMasterBoneComponent New MasterPoseComponent
 	 */
 	UFUNCTION(BlueprintCallable, Category="Components|SkinnedMesh")
-	void SetMasterPoseComponent(USkinnedMeshComponent* NewMasterBoneComponent);
+	void SetMasterPoseComponent(USkinnedMeshComponent* NewMasterBoneComponent, bool bForceUpdate = false);
 
 protected:
 	/** Add a slave component to the SlavePoseComponents array */
diff --git a/Engine/Source/Runtime/Engine/Private/Components/SkinnedMeshComponent.cpp b/Engine/Source/Runtime/Engine/Private/Components/SkinnedMeshComponent.cpp
--- a/Engine/Source/Runtime/Engine/Private/Components/SkinnedMeshComponent.cpp
+++ b/Engine/Source/Runtime/Engine/Private/Components/SkinnedMeshComponent.cpp
@@ -455,8 +455,8 @@ void USkinnedMeshComponent::OnRegister()
 
 	if (MasterPoseComponent.IsValid())
 	{
-		// this has to be called again during register so that it can do related initialization
-		SetMasterPoseComponent(MasterPoseComponent.Get());
+		// we have to make sure it updates the mastesr pose
+		SetMasterPoseComponent(MasterPoseComponent.Get(), true);
 	}
 	else
 	{
@@ -1329,10 +1329,10 @@ void USkinnedMeshComponent::SetPhysicsAsset(class UPhysicsAsset* InPhysicsAsset,
 	PhysicsAssetOverride = InPhysicsAsset;
 }
 
-void USkinnedMeshComponent::SetMasterPoseComponent(class USkinnedMeshComponent* NewMasterBoneComponent)
+void USkinnedMeshComponent::SetMasterPoseComponent(class USkinnedMeshComponent* NewMasterBoneComponent, bool bForceUpdate)
 {
 	// Early out if we're already setup.
-	if (NewMasterBoneComponent == MasterPoseComponent)
+	if (!bForceUpdate && NewMasterBoneComponent == MasterPoseComponent)
 	{
 		return;
 	}
