PR #4755: Fix minor typo in comment (Contributed by [REDACTED_PERSON])
#rnx

diff --git a/Engine/Source/Runtime/Engine/Classes/GameFramework/GameModeBase.h b/Engine/Source/Runtime/Engine/Classes/GameFramework/GameModeBase.h
--- a/Engine/Source/Runtime/Engine/Classes/GameFramework/GameModeBase.h
+++ b/Engine/Source/Runtime/Engine/Classes/GameFramework/GameModeBase.h
@@ -523,7 +523,7 @@ protected:
 	/** Replicates the current level streaming status to the given PlayerController */
 	virtual void ReplicateStreamingStatus(APlayerController* PC);
 
-	/** Return true of FindPlayerStart should use the StartSpot stored on Player instead of calling ChoosePlayerStart */
+	/** Return true if FindPlayerStart should use the StartSpot stored on Player instead of calling ChoosePlayerStart */
 	virtual bool ShouldSpawnAtStartSpot(AController* Player);
 
 	/** Handles second half of RestartPlayer */
