Text Render Component: Fix crash on shutdown.

diff --git a/Engine/Source/Runtime/Engine/Private/Components/TextRenderComponent.cpp b/Engine/Source/Runtime/Engine/Private/Components/TextRenderComponent.cpp
--- a/Engine/Source/Runtime/Engine/Private/Components/TextRenderComponent.cpp
+++ b/Engine/Source/Runtime/Engine/Private/Components/TextRenderComponent.cpp
@@ -412,6 +412,11 @@ public:
 		return *InstancePtr;
 	}
 
+	static bool IsValid()
+	{
+		return InstancePtr != nullptr;
+	}
+
 	FMIDDataRef GetMIDData(UMaterialInterface* InMaterial, UFont* InFont)
 	{
 		checkfSlow(IsInGameThread(), TEXT("FTextRenderComponentMIDCache::GetMIDData is only expected to be called from the game thread!"));
@@ -659,7 +664,10 @@ FTextRenderSceneProxy::FTextRenderSceneProxy( UTextRenderComponent* Component) :
 
 	if (Font && Font->FontCacheType == EFontCacheType::Offline)
 	{
-		FontMIDs = FTextRenderComponentMIDCache::Get().GetMIDData(TextMaterial, Font);
+		if (FTextRenderComponentMIDCache::IsValid())
+		{
+			FontMIDs = FTextRenderComponentMIDCache::Get().GetMIDData(TextMaterial, Font);
+		}
 	}
 
 	// The MID from the cache isn't known by the UTextRenderComponent
