Deleting keys, undoing the deletion, selecting them and then redoing the deletion no longer causes a crash in Sequencer/UMG's Curve Editors. KeyHandleLookupTables were not being serialized so when Undo/Redo transactions happened their map of the key handles would not be restored. This caused a mis-match in the data (the map contained more entries than actual data) so the map would falsely report a valid index when there was not one. Graph Editor selections are not sync'd with the Undo/Redo system currently so redoing a deletion after selecting keys would leave you in a state where you tried to request now out of bounds keys on the next Tick to update the UI.

diff --git a/Engine/Source/Runtime/Core/Public/Misc/Optional.h b/Engine/Source/Runtime/Core/Public/Misc/Optional.h
--- a/Engine/Source/Runtime/Core/Public/Misc/Optional.h
+++ b/Engine/Source/Runtime/Core/Public/Misc/Optional.h
@@ -7,6 +7,8 @@
 #include "Templates/TypeCompatibleBytes.h"
 #include "Templates/UnrealTemplate.h"
 
+class FArchive;
+
 /**
  * When we have an optional value IsSet() returns true, and GetValue() is meaningful.
  * Otherwise GetValue() is not meaningful.
@@ -138,11 +140,41 @@ public:
 		}
 		return (*(OptionalType*)&lhs.Value) == (*(OptionalType*)&rhs.Value);
 	}
+
 	friend bool operator!=(const TOptional& lhs, const TOptional& rhs)
 	{
 		return !(lhs == rhs);
 	}
 
+	friend FArchive& operator<<(FArchive& Ar, TOptional& Optional)
+	{
+		bool bIsSet = Optional.bIsSet;
+		Ar << bIsSet;
+		if (Ar.IsLoading())
+		{
+			if (bIsSet)
+			{
+				if (!Optional.bIsSet)
+				{
+					Optional.Emplace();
+				}
+				Ar << Optional.GetValue();
+			}
+			else
+			{
+				Optional.Reset();
+			}
+		}
+		else
+		{
+			if (bIsSet)
+			{
+				Ar << Optional.GetValue();
+			}
+		}
+		return Ar;
+	}
+
 	/** @return true when the value is meaningful; false if calling GetValue() is undefined. */
 	bool IsSet() const { return bIsSet; }
 	FORCEINLINE explicit operator bool() const { return bIsSet; }
diff --git a/Engine/Source/Runtime/Engine/Classes/Curves/KeyHandle.h b/Engine/Source/Runtime/Engine/Classes/Curves/KeyHandle.h
--- a/Engine/Source/Runtime/Engine/Classes/Curves/KeyHandle.h
+++ b/Engine/Source/Runtime/Engine/Classes/Curves/KeyHandle.h
@@ -116,8 +116,11 @@ struct TStructOpsTypeTraits<FKeyHandleMap>
  * Maintains a map of key handle to last known index,
  * and an array of optional key handles that's used to validate map entries.
  */
+USTRUCT()
 struct FKeyHandleLookupTable
 {
+	GENERATED_BODY()
+public:
 	/**
 	 * Get the index that corresponds to the specified key handle
 	 *
@@ -162,11 +165,29 @@ struct FKeyHandleLookupTable
 	 */
 	ENGINE_API void Reset();
 
-private:
+	/** ICPPStructOps implementation */
+	ENGINE_API bool Serialize(FArchive& Ar);
+	friend FArchive& operator<<(FArchive& Ar, FKeyHandleLookupTable& P)
+	{
+		P.Serialize(Ar);
+		return Ar;
+	}
 
+private:
 	/** Array of optional key handles that reside in corresponding indices for an externally owned serial data structure */
 	TArray<TOptional<FKeyHandle>> KeyHandles;
 
 	/** Map of which key handles go to which indices. Indices may point to incorrect indices. Check against entry in KeyHandles array to verify. */
 	TMap<FKeyHandle, int32> KeyHandlesToIndices;
+
 };
+
+template<>
+struct TStructOpsTypeTraits<FKeyHandleLookupTable>
+	: public TStructOpsTypeTraitsBase2<FKeyHandleLookupTable>
+{
+	enum
+	{
+		WithSerializer = true,
+	};
+};
\ No newline at end of file
diff --git a/Engine/Source/Runtime/Engine/Private/Curves/KeyHandle.cpp b/Engine/Source/Runtime/Engine/Private/Curves/KeyHandle.cpp
--- a/Engine/Source/Runtime/Engine/Private/Curves/KeyHandle.cpp
+++ b/Engine/Source/Runtime/Engine/Private/Curves/KeyHandle.cpp
@@ -221,3 +221,15 @@ void FKeyHandleLookupTable::Reset()
 	KeyHandles.Reset();
 	KeyHandlesToIndices.Reset();
 }
+
+bool FKeyHandleLookupTable::Serialize(FArchive& Ar)
+{
+	// We're only concerned with Undo/Redo transactions
+	if (Ar.IsTransacting())
+	{
+		Ar << KeyHandles;
+		Ar << KeyHandlesToIndices;
+	}
+
+	return true;
+}
\ No newline at end of file
diff --git a/Engine/Source/Runtime/MovieScene/Public/Channels/MovieSceneChannelData.h b/Engine/Source/Runtime/MovieScene/Public/Channels/MovieSceneChannelData.h
--- a/Engine/Source/Runtime/MovieScene/Public/Channels/MovieSceneChannelData.h
+++ b/Engine/Source/Runtime/MovieScene/Public/Channels/MovieSceneChannelData.h
@@ -7,12 +7,18 @@
 #include "Misc/FrameNumber.h"
 #include "Misc/FrameTime.h"
 
+#include "MovieSceneChannelData.generated.h"
+
 struct FFrameRate;
 struct FKeyDataOptimizationParams;
 
 /** A map of key handles that is copyable, but does not copy data on copy */
-struct FMovieSceneKeyHandleMap : FKeyHandleLookupTable
+USTRUCT()
+struct FMovieSceneKeyHandleMap : public FKeyHandleLookupTable
 {
+	GENERATED_BODY()
+
+public:
 	FMovieSceneKeyHandleMap() = default;
 	FMovieSceneKeyHandleMap(const FMovieSceneKeyHandleMap& RHS){}
 	FMovieSceneKeyHandleMap& operator=(const FMovieSceneKeyHandleMap& RHS)
@@ -22,6 +28,16 @@ struct FMovieSceneKeyHandleMap : FKeyHandleLookupTable
 	}
 };
 
+template<>
+struct TStructOpsTypeTraits<FMovieSceneKeyHandleMap>
+	: public TStructOpsTypeTraitsBase2<FMovieSceneKeyHandleMap>
+{
+	enum
+	{
+		WithSerializer = true,
+	};
+};
+
 namespace MovieScene
 {
 	/**
diff --git a/Engine/Source/Runtime/MovieScene/Public/Channels/MovieSceneFloatChannel.h b/Engine/Source/Runtime/MovieScene/Public/Channels/MovieSceneFloatChannel.h
--- a/Engine/Source/Runtime/MovieScene/Public/Channels/MovieSceneFloatChannel.h
+++ b/Engine/Source/Runtime/MovieScene/Public/Channels/MovieSceneFloatChannel.h
@@ -269,6 +269,8 @@ private:
 	UPROPERTY()
 	bool bHasDefaultValue;
 
+	/** This needs to be a UPROPERTY so it gets saved into editor transactions but transient so it doesn't get saved into assets. */
+	UPROPERTY(Transient)
 	FMovieSceneKeyHandleMap KeyHandles;
 
 	UPROPERTY()
