Fix for PIE HighDPI touch location.

diff --git a/Engine/Source/Runtime/Engine/Private/Slate/SceneViewport.cpp b/Engine/Source/Runtime/Engine/Private/Slate/SceneViewport.cpp
--- a/Engine/Source/Runtime/Engine/Private/Slate/SceneViewport.cpp
+++ b/Engine/Source/Runtime/Engine/Private/Slate/SceneViewport.cpp
@@ -754,7 +754,7 @@ FReply FSceneViewport::OnTouchStarted( const FGeometry& MyGeometry, const FPoint
 		// Switch to the viewport clients world before processing input
 		FScopedConditionalWorldSwitcher WorldSwitcher( ViewportClient );
 
-		const FVector2D TouchPosition = MyGeometry.AbsoluteToLocal(TouchEvent.GetScreenSpacePosition());
+		const FVector2D TouchPosition = MyGeometry.AbsoluteToLocal(TouchEvent.GetScreenSpacePosition()) * MyGeometry.Scale;
 
 		if(ViewportClient->InputTouch( this, TouchEvent.GetUserIndex(), TouchEvent.GetPointerIndex(), ETouchType::Began, TouchPosition, TouchEvent.GetTouchForce(), FDateTime::Now(), TouchEvent.GetTouchpadIndex()) )
 		{
@@ -786,7 +786,7 @@ FReply FSceneViewport::OnTouchMoved( const FGeometry& MyGeometry, const FPointer
 		// Switch to the viewport clients world before processing input
 		FScopedConditionalWorldSwitcher WorldSwitcher( ViewportClient );
 
-		const FVector2D TouchPosition = MyGeometry.AbsoluteToLocal(TouchEvent.GetScreenSpacePosition());
+		const FVector2D TouchPosition = MyGeometry.AbsoluteToLocal(TouchEvent.GetScreenSpacePosition()) * MyGeometry.Scale;
 
 		if( !ViewportClient->InputTouch( this, TouchEvent.GetUserIndex(), TouchEvent.GetPointerIndex(), ETouchType::Moved, TouchPosition, TouchEvent.GetTouchForce(), FDateTime::Now(), TouchEvent.GetTouchpadIndex()) )
 		{
@@ -817,7 +817,7 @@ FReply FSceneViewport::OnTouchEnded( const FGeometry& MyGeometry, const FPointer
 		// Switch to the viewport clients world before processing input
 		FScopedConditionalWorldSwitcher WorldSwitcher( ViewportClient );
 
-		const FVector2D TouchPosition = MyGeometry.AbsoluteToLocal(TouchEvent.GetScreenSpacePosition());
+		const FVector2D TouchPosition = MyGeometry.AbsoluteToLocal(TouchEvent.GetScreenSpacePosition()) * MyGeometry.Scale;
 
 		if( !ViewportClient->InputTouch( this, TouchEvent.GetUserIndex(), TouchEvent.GetPointerIndex(), ETouchType::Ended, TouchPosition, 0.0f, FDateTime::Now(), TouchEvent.GetTouchpadIndex()) )
 		{
