Duplicating in Release-4.20:
Change: 4087159
UE-58249 Late update has FP precision issues on PSVR
-In lateupdate the TranslatedViewMatrix and its inverse were being calculated from the ViewMatrix, by subtracting the PreViewTranslation.  This doesn't work because the translation numerical innacuracy is already in the ViewMatrix.
-I also rearanged the UpdateViewMatrix code some to make it look more like the FViewMatrix constructor code that sets this up, just to make it easier to spot the difference.

Files:
//UE4/Dev-VR/Engine/Source/Runtime/Engine/Private/SceneView.cpp#71

diff --git a/Engine/Source/Runtime/Engine/Private/SceneView.cpp b/Engine/Source/Runtime/Engine/Private/SceneView.cpp
--- a/Engine/Source/Runtime/Engine/Private/SceneView.cpp
+++ b/Engine/Source/Runtime/Engine/Private/SceneView.cpp
@@ -878,28 +878,30 @@ void FViewMatrices::UpdateViewMatrix(const FVector& ViewLocation, const FRotator
 		FPlane(0, 1, 0, 0),
 		FPlane(0, 0, 0, 1));
 
-	ViewMatrix = FTranslationMatrix(-ViewLocation);
-	ViewMatrix = ViewMatrix * FInverseRotationMatrix(ViewRotation);
-	ViewMatrix = ViewMatrix * ViewPlanesMatrix;
+	const FMatrix ViewRotationMatrix = FInverseRotationMatrix(ViewRotation) * ViewPlanesMatrix;
 
-	InvViewMatrix = FTranslationMatrix(-ViewMatrix.GetOrigin()) * ViewMatrix.RemoveTranslation().GetTransposed();
+	ViewMatrix = FTranslationMatrix(-ViewLocation) * ViewRotationMatrix;
 
 	// Duplicate HMD rotation matrix with roll removed
 	FRotator HMDViewRotation = ViewRotation;
 	HMDViewRotation.Roll = 0.f;
 	HMDViewMatrixNoRoll = FInverseRotationMatrix(HMDViewRotation) * ViewPlanesMatrix;
 
+	ViewProjectionMatrix = GetViewMatrix() * GetProjectionMatrix();
+
+	InvViewMatrix = ViewRotationMatrix.GetTransposed() * FTranslationMatrix(ViewLocation);
+	InvViewProjectionMatrix = GetInvProjectionMatrix() * GetInvViewMatrix();
+
 	PreViewTranslation = -ViewOrigin;
-	//using mathematical equality rule for matrix inverse: (A*B)^-1 == B^-1 * A^-1
-	OverriddenTranslatedViewMatrix = TranslatedViewMatrix = FTranslationMatrix(-PreViewTranslation) * ViewMatrix;
-	OverriddenInvTranslatedViewMatrix = InvTranslatedViewMatrix = InvViewMatrix * FTranslationMatrix(PreViewTranslation);
+
+	TranslatedViewMatrix = ViewRotationMatrix;
+	InvTranslatedViewMatrix = TranslatedViewMatrix.GetTransposed();
+	OverriddenTranslatedViewMatrix = FTranslationMatrix(-PreViewTranslation) * ViewMatrix;
+	OverriddenInvTranslatedViewMatrix = InvViewMatrix * FTranslationMatrix(PreViewTranslation);
 
 	// Compute a transform from view origin centered world-space to clip space.
 	TranslatedViewProjectionMatrix = GetTranslatedViewMatrix() * GetProjectionMatrix();
 	InvTranslatedViewProjectionMatrix = GetInvProjectionMatrix() * GetInvTranslatedViewMatrix();
-
-	ViewProjectionMatrix = GetViewMatrix() * GetProjectionMatrix();
-	InvViewProjectionMatrix = GetInvProjectionMatrix() * GetInvViewMatrix();
 }
 
 void FSceneView::UpdateViewMatrix()
