Allow the default physics volume to be spawned during construction script execution if it needs to be lazily created

diff --git a/Engine/Source/Runtime/Engine/Private/World.cpp b/Engine/Source/Runtime/Engine/Private/World.cpp
--- a/Engine/Source/Runtime/Engine/Private/World.cpp
+++ b/Engine/Source/Runtime/Engine/Private/World.cpp
@@ -4062,8 +4062,11 @@ APhysicsVolume* UWorld::GetDefaultPhysicsVolume() const
 		}
 
 		// Spawn volume
+		FActorSpawnParameters SpawnParams;
+		SpawnParams.bAllowDuringConstructionScript = true;
+
 		UWorld* MutableThis = const_cast<UWorld*>(this);
-		MutableThis->DefaultPhysicsVolume = MutableThis->SpawnActor<APhysicsVolume>(DefaultPhysicsVolumeClass);
+		MutableThis->DefaultPhysicsVolume = MutableThis->SpawnActor<APhysicsVolume>(DefaultPhysicsVolumeClass, SpawnParams);
 		MutableThis->DefaultPhysicsVolume->Priority = -1000000;
 	}
 	return DefaultPhysicsVolume;
