No need to load AnimationBlueprintEditor so early
#rnx

diff --git a/Engine/Source/Runtime/Engine/Classes/Animation/AnimBlueprint.h b/Engine/Source/Runtime/Engine/Classes/Animation/AnimBlueprint.h
--- a/Engine/Source/Runtime/Engine/Classes/Animation/AnimBlueprint.h
+++ b/Engine/Source/Runtime/Engine/Classes/Animation/AnimBlueprint.h
@@ -175,8 +175,4 @@ private:
 	UPROPERTY(duplicatetransient, AssetRegistrySearchable)
 	TSoftObjectPtr<class USkeletalMesh> PreviewSkeletalMesh;
 #endif
-
-#if WITH_EDITOR
-	virtual void LoadModulesRequiredForCompilation() override;
-#endif //WITH_EDITOR
 };
diff --git a/Engine/Source/Runtime/Engine/Private/AnimBlueprint.cpp b/Engine/Source/Runtime/Engine/Private/AnimBlueprint.cpp
--- a/Engine/Source/Runtime/Engine/Private/AnimBlueprint.cpp
+++ b/Engine/Source/Runtime/Engine/Private/AnimBlueprint.cpp
@@ -180,12 +180,6 @@ bool UAnimBlueprint::CanRecompileWhilePlayingInEditor() const
 {
 	return GetDefault<UEditorExperimentalSettings>()->bEnableLiveRecompilationOfAnimationBlueprints;
 }
-
-void UAnimBlueprint::LoadModulesRequiredForCompilation()
-{
-	static const FName ModuleName(TEXT("AnimationBlueprintEditor"));
-	FModuleManager::Get().LoadModule(ModuleName);
-}
 #endif
 
 USkeletalMesh* UAnimBlueprint::GetPreviewMesh(bool bFindIfNotSet/*=false*/)
