PR #4633: Fixed DPI awareness of two functions (Contributed by [REDACTED_PERSON])

diff --git a/Engine/Source/Runtime/Engine/Private/Slate/SceneViewport.cpp b/Engine/Source/Runtime/Engine/Private/Slate/SceneViewport.cpp
--- a/Engine/Source/Runtime/Engine/Private/Slate/SceneViewport.cpp
+++ b/Engine/Source/Runtime/Engine/Private/Slate/SceneViewport.cpp
@@ -324,8 +324,8 @@ FVector2D FSceneViewport::VirtualDesktopPixelToViewport(FIntPoint VirtualDesktop
 	// Virtual Desktop Pixel to local slate unit
 	const FVector2D TransformedPoint = CachedGeometry.AbsoluteToLocal(FVector2D(VirtualDesktopPointPx.X, VirtualDesktopPointPx.Y));
 
-	// Pixels to normalized coordinates
-	return FVector2D( TransformedPoint.X / SizeX, TransformedPoint.Y / SizeY );
+	// Pixels to normalized coordinates and correct for DPI scale
+	return FVector2D( TransformedPoint.X / SizeX * CachedGeometry.Scale, TransformedPoint.Y / SizeY * CachedGeometry.Scale );
 }
 
 FIntPoint FSceneViewport::ViewportToVirtualDesktopPixel(FVector2D ViewportCoordinate) const
@@ -335,7 +335,8 @@ FIntPoint FSceneViewport::ViewportToVirtualDesktopPixel(FVector2D ViewportCoordi
 	// Local slate unit to virtual desktop pixel.
 	const FVector2D TransformedPoint = FVector2D( CachedGeometry.LocalToAbsolute( LocalCoordinateInSu ) );
 
-	return FIntPoint( FMath::TruncToInt(TransformedPoint.X), FMath::TruncToInt(TransformedPoint.Y) );
+	//Correct for DPI
+	return FIntPoint( FMath::TruncToInt(TransformedPoint.X / CachedGeometry.Scale), FMath::TruncToInt(TransformedPoint.Y / CachedGeometry.Scale) );
 }
 
 void FSceneViewport::OnDrawViewport( const FGeometry& AllottedGeometry, const FSlateRect& MyCullingRect, FSlateWindowElementList& OutDrawElements, int32 LayerId, const FWidgetStyle& InWidgetStyle, bool bParentEnabled )
