Added the possibility to see the Cluster tree bounds in the Show->Advanced Menu

diff --git a/Engine/Source/Runtime/Core/Public/UObject/ReleaseObjectVersion.h b/Engine/Source/Runtime/Core/Public/UObject/ReleaseObjectVersion.h
--- a/Engine/Source/Runtime/Core/Public/UObject/ReleaseObjectVersion.h
+++ b/Engine/Source/Runtime/Core/Public/UObject/ReleaseObjectVersion.h
@@ -52,6 +52,9 @@ struct CORE_API FReleaseObjectVersion
 		// Removed objects that were serialized as part of this material feature
 		RemovedMaterialSharedInputCollection,
 
+		// HISMC Cluster Tree migration to add new data
+		HISMCClusterTreeMigration,
+
 		// -----<new versions can be added above this line>-------------------------------------------------
 		VersionPlusOne,
 		LatestVersion = VersionPlusOne - 1
diff --git a/Engine/Source/Runtime/Engine/Classes/Components/HierarchicalInstancedStaticMeshComponent.h b/Engine/Source/Runtime/Engine/Classes/Components/HierarchicalInstancedStaticMeshComponent.h
--- a/Engine/Source/Runtime/Engine/Classes/Components/HierarchicalInstancedStaticMeshComponent.h
+++ b/Engine/Source/Runtime/Engine/Classes/Components/HierarchicalInstancedStaticMeshComponent.h
@@ -14,6 +14,51 @@
 class FClusterBuilder;
 class FStaticLightingTextureMapping_InstancedStaticMesh;
 
+// Due to BulkSerialize we can't edit the struct, so we must deprecated this one and create a new one
+USTRUCT()
+struct FClusterNode_DEPRECATED
+{
+	GENERATED_USTRUCT_BODY()
+
+		UPROPERTY()
+		FVector BoundMin;
+	UPROPERTY()
+		int32 FirstChild;
+	UPROPERTY()
+		FVector BoundMax;
+	UPROPERTY()
+		int32 LastChild;
+	UPROPERTY()
+		int32 FirstInstance;
+	UPROPERTY()
+		int32 LastInstance;
+
+	FClusterNode_DEPRECATED()
+		: BoundMin(MAX_flt, MAX_flt, MAX_flt)
+		, FirstChild(-1)
+		, BoundMax(MIN_flt, MIN_flt, MIN_flt)
+		, LastChild(-1)
+		, FirstInstance(-1)
+		, LastInstance(-1)
+	{
+	}
+
+	friend FArchive& operator<<(FArchive& Ar, FClusterNode_DEPRECATED& NodeData)
+	{
+		// @warning BulkSerialize: FClusterNode is serialized as memory dump
+		// See TArray::BulkSerialize for detailed description of implied limitations.
+		Ar << NodeData.BoundMin;
+		Ar << NodeData.FirstChild;
+		Ar << NodeData.BoundMax;
+		Ar << NodeData.LastChild;
+		Ar << NodeData.FirstInstance;
+		Ar << NodeData.LastInstance;
+
+		return Ar;
+	}
+};
+
+
 USTRUCT()
 struct FClusterNode
 {
@@ -32,6 +77,11 @@ struct FClusterNode
 	UPROPERTY()
 	int32 LastInstance;
 
+	UPROPERTY()
+	FVector MinInstanceScale;
+	UPROPERTY()
+	FVector MaxInstanceScale;
+
 	FClusterNode()
 		: BoundMin(MAX_flt, MAX_flt, MAX_flt)
 		, FirstChild(-1)
@@ -39,6 +89,20 @@ struct FClusterNode
 		, LastChild(-1)
 		, FirstInstance(-1)
 		, LastInstance(-1)
+		, MinInstanceScale(MAX_flt)
+		, MaxInstanceScale(-MAX_flt)
+	{
+	}
+
+	FClusterNode(const FClusterNode_DEPRECATED& OldNode)
+		: BoundMin(OldNode.BoundMin)
+		, FirstChild(OldNode.FirstChild)
+		, BoundMax(OldNode.BoundMax)
+		, LastChild(OldNode.LastChild)
+		, FirstInstance(OldNode.FirstChild)
+		, LastInstance(OldNode.LastInstance)
+		, MinInstanceScale(MAX_flt)
+		, MaxInstanceScale(-MAX_flt)
 	{
 	}
 
@@ -52,6 +116,9 @@ struct FClusterNode
 		Ar << NodeData.LastChild;
 		Ar << NodeData.FirstInstance;
 		Ar << NodeData.LastInstance;
+		Ar << NodeData.MinInstanceScale;
+		Ar << NodeData.MaxInstanceScale;
+		
 		return Ar;
 	}
 };
diff --git a/Engine/Source/Runtime/Engine/Private/HierarchicalInstancedStaticMesh.cpp b/Engine/Source/Runtime/Engine/Private/HierarchicalInstancedStaticMesh.cpp
--- a/Engine/Source/Runtime/Engine/Private/HierarchicalInstancedStaticMesh.cpp
+++ b/Engine/Source/Runtime/Engine/Private/HierarchicalInstancedStaticMesh.cpp
@@ -29,6 +29,7 @@
 #include "InstancedStaticMesh.h"
 #include "SceneManagement.h"
 #include "HAL/LowLevelMemTracker.h"
+#include "UObject/ReleaseObjectVersion.h"
 
 static TAutoConsoleVariable<int32> CVarFoliageSplitFactor(
 	TEXT("foliage.SplitFactor"),
@@ -446,6 +447,11 @@ public:
 				const FMatrix& ThisInstTrans = Transforms[SortedInstances[InstanceIndex]];
 				FBox ThisInstBox = InstBox.TransformBy(ThisInstTrans);
 				NodeBox += ThisInstBox;
+
+				FVector CurrentScale = ThisInstTrans.GetScaleVector();
+
+				Node.MinInstanceScale = Node.MinInstanceScale.ComponentMin(CurrentScale);
+				Node.MaxInstanceScale = Node.MaxInstanceScale.ComponentMax(CurrentScale);
 			}
 			Node.BoundMin = NodeBox.Min;
 			Node.BoundMax = NodeBox.Max;
@@ -621,6 +627,9 @@ public:
 						FClusterNode& ChildNode = Result->Nodes[ChildIndex];
 						NodeBox += ChildNode.BoundMin;
 						NodeBox += ChildNode.BoundMax;
+
+						Node.MinInstanceScale = Node.MinInstanceScale.ComponentMin(ChildNode.MinInstanceScale);
+						Node.MaxInstanceScale = Node.MaxInstanceScale.ComponentMax(ChildNode.MaxInstanceScale);
 					}
 					Node.BoundMin = NodeBox.Min;
 					Node.BoundMax = NodeBox.Max;
@@ -1113,7 +1122,7 @@ static FORCEINLINE_DEBUGGABLE bool CullNode(const FFoliageCullInstanceParams& Pa
 	return false;
 }
 
-inline void CalcLOD(int32& InOutMinLOD, int32& InOutMaxLOD, const FVector& BoundMin, const FVector& BoundMax, const FVector& ViewOriginInLocalZero, const FVector& ViewOriginInLocalOne, const float (&LODPlanesMin)[MAX_STATIC_MESH_LODS], const float (&LODPlanesMax)[MAX_STATIC_MESH_LODS])
+inline void CalcLOD(int32& InOutMinLOD, int32& InOutMaxLOD, const FVector& BoundMin, const FVector& BoundMax, const FVector& ViewOriginInLocalZero, const FVector& ViewOriginInLocalOne, const float (&LODPlanesMin)[MAX_STATIC_MESH_LODS], const float (&LODPlanesMax)[MAX_STATIC_MESH_LODS], float LODDistanceScaleFactor)
 {
 	if (InOutMinLOD != InOutMaxLOD)
 	{
@@ -1124,11 +1133,11 @@ inline void CalcLOD(int32& InOutMinLOD, int32& InOutMaxLOD, const FVector& Bound
 		const float NearDot = FMath::Min(DistCenterZero, DistCenterOne) - HalfWidth;
 		const float FarDot = FMath::Max(DistCenterZero, DistCenterOne) + HalfWidth;
 
-		while (InOutMaxLOD > InOutMinLOD && NearDot > LODPlanesMax[InOutMinLOD])
+		while (InOutMaxLOD > InOutMinLOD && NearDot > LODPlanesMax[InOutMinLOD] * LODDistanceScaleFactor)
 		{
 			InOutMinLOD++;
 		}
-		while (InOutMaxLOD > InOutMinLOD && FarDot < LODPlanesMin[InOutMaxLOD - 1])
+		while (InOutMaxLOD > InOutMinLOD && FarDot < LODPlanesMin[InOutMaxLOD - 1] * LODDistanceScaleFactor)
 		{
 			InOutMaxLOD--;
 		}
@@ -1163,7 +1172,10 @@ void FHierarchicalStaticMeshSceneProxy::Traverse(const FFoliageCullInstanceParam
 
 	if (MinLOD != MaxLOD)
 	{
-		CalcLOD(MinLOD, MaxLOD, Node.BoundMin, Node.BoundMax, Params.ViewOriginInLocalZero, Params.ViewOriginInLocalOne, Params.LODPlanesMin, Params.LODPlanesMax);
+		FVector ScaleAverage = (Node.MaxInstanceScale - Node.MinInstanceScale) / 2.0f;
+		float DistanceScale = FMath::Max(FMath::Max3(ScaleAverage.X, ScaleAverage.Y, ScaleAverage.Z), 0.001f);
+
+		CalcLOD(MinLOD, MaxLOD, Node.BoundMin, Node.BoundMax, Params.ViewOriginInLocalZero, Params.ViewOriginInLocalOne, Params.LODPlanesMin, Params.LODPlanesMax, DistanceScale);
 
 		if (MinLOD >= Params.LODs)
 		{
@@ -1529,7 +1541,7 @@ void FHierarchicalStaticMeshSceneProxy::GetDynamicMeshElements(const TArray<cons
 				float FinalCull = MAX_flt;
 				if (MinSize > 0.0)
 				{
-					FinalCull = ComputeBoundsDrawDistance(MinSize, SphereRadius, View->ViewMatrices.GetProjectionMatrix()) * LODScale;
+					FinalCull = ComputeBoundsDrawDistance(MinSize, SphereRadius * FMath::Max(FMath::Min3(ClusterTree[0].MinInstanceScale.X, ClusterTree[0].MinInstanceScale.Y, ClusterTree[0].MinInstanceScale.Z), 0.001f), View->ViewMatrices.GetProjectionMatrix()) * LODScale;
 				}
 				if (UserData_AllInstances.EndCullDistance > 0.0f)
 				{
@@ -1714,16 +1726,18 @@ void FHierarchicalStaticMeshSceneProxy::GetDynamicMeshElements(const TArray<cons
 						LODPlanesMin[NumLODs - 1] = FinalCull - LODRandom;
 						LODPlanesMax[NumLODs - 1] = FinalCull;				
 
+						// NOTE: in case of unbuilt we can't really apply the instance scales so the LOD won't be optimal until the build is completed
+
 						// calculate runs
 						int32 MinLOD = 0;
 						int32 MaxLOD = NumLODs;
-						CalcLOD(MinLOD, MaxLOD, UnbuiltBounds[0].Min, UnbuiltBounds[0].Max, ViewOriginInLocalZero, ViewOriginInLocalOne, LODPlanesMin, LODPlanesMax);
+						CalcLOD(MinLOD, MaxLOD, UnbuiltBounds[0].Min, UnbuiltBounds[0].Max, ViewOriginInLocalZero, ViewOriginInLocalOne, LODPlanesMin, LODPlanesMax, 1.0f);
 						int32 FirstIndexInRun = 0;
 						for (int32 Index = 1; Index < UnbuiltInstanceCount; ++Index)
 						{
 							int32 TempMinLOD = 0;
 							int32 TempMaxLOD = NumLODs;
-							CalcLOD(TempMinLOD, TempMaxLOD, UnbuiltBounds[Index].Min, UnbuiltBounds[Index].Max, ViewOriginInLocalZero, ViewOriginInLocalOne, LODPlanesMin, LODPlanesMax);
+							CalcLOD(TempMinLOD, TempMaxLOD, UnbuiltBounds[Index].Min, UnbuiltBounds[Index].Max, ViewOriginInLocalZero, ViewOriginInLocalOne, LODPlanesMin, LODPlanesMax, 1.0f);
 							if (TempMinLOD != MinLOD)
 							{
 								if (MinLOD < NumLODs)
@@ -1746,13 +1760,26 @@ void FHierarchicalStaticMeshSceneProxy::GetDynamicMeshElements(const TArray<cons
 				FillDynamicMeshElements(Collector, ElementParams, InstanceParams);
 			}
 
-			if (View->Family->EngineShowFlags.FoliageOcclusionBounds)
+			if (View->Family->EngineShowFlags.HISMCOcclusionBounds)
 			{
 				for (auto& OcclusionBound : OcclusionBounds)
 				{
 					DrawWireBox(Collector.GetPDI(ViewIndex), OcclusionBound.GetBox(), FColor(255, 0, 0), View->Family->EngineShowFlags.Game ? SDPG_World : SDPG_Foreground);
 				}
 			}
+
+			if (View->Family->EngineShowFlags.HISMCClusterTree)
+			{
+				FColor StartingColor(100, 0, 0);
+
+				for (const FClusterNode& CulsterNode : ClusterTree)
+				{
+					DrawWireBox(Collector.GetPDI(ViewIndex), FBox(CulsterNode.BoundMin, CulsterNode.BoundMax), StartingColor, View->Family->EngineShowFlags.Game ? SDPG_World : SDPG_Foreground);
+					StartingColor.R += 5;
+					StartingColor.G += 5;
+					StartingColor.B += 5;
+				}
+			}
 		}
 	}
 }
@@ -1872,6 +1899,8 @@ void UHierarchicalInstancedStaticMeshComponent::Serialize(FArchive& Ar)
 {
 	LLM_SCOPE(ELLMTag::StaticMesh);	
 
+	Ar.UsingCustomVersion(FReleaseObjectVersion::GUID);
+
 	// Before serializing the content to cook, wait for the async task to be completed
 	if (bIsAsyncBuilding && Ar.IsSaving() && Ar.IsCooking())
 	{
@@ -1899,10 +1928,24 @@ void UHierarchicalInstancedStaticMeshComponent::Serialize(FArchive& Ar)
 	{
 		ClusterTreePtr = MakeShareable(new TArray<FClusterNode>);
 	}
-	TArray<FClusterNode>& ClusterTree = *ClusterTreePtr;
-	ClusterTree.BulkSerialize(Ar);
+
+	if (Ar.IsLoading() && Ar.CustomVer(FReleaseObjectVersion::GUID) < FReleaseObjectVersion::HISMCClusterTreeMigration)
+	{
+		// Skip the serialized tree, we will regenerate it correctly to contains the new data
+		TArray<FClusterNode_DEPRECATED> ClusterTree_DEPRECATED;
+		ClusterTree_DEPRECATED.BulkSerialize(Ar);
+
+		BuildTreeIfOutdated(false, true);
+	}
+	else
+	{
+		ClusterTreePtr->BulkSerialize(Ar);
+	}
+
 	if (Ar.IsLoading() && !BuiltInstanceBounds.IsValid)
 	{
+		TArray<FClusterNode>& ClusterTree = *ClusterTreePtr.Get();
+
 		BuiltInstanceBounds = (ClusterTree.Num() > 0 ? FBox(ClusterTree[0].BoundMin, ClusterTree[0].BoundMax) : FBox(ForceInit));
 	}
 }
@@ -3047,3 +3090,5 @@ static FAutoConsoleCommand RebuildFoliageTreesCmd(
 	FConsoleCommandWithArgsDelegate::CreateStatic(&RebuildFoliageTrees)
 	);
 
+
+PRAGMA_ENABLE_OPTIMIZATION
\ No newline at end of file
diff --git a/Engine/Source/Runtime/Engine/Public/ShowFlags.h b/Engine/Source/Runtime/Engine/Public/ShowFlags.h
--- a/Engine/Source/Runtime/Engine/Public/ShowFlags.h
+++ b/Engine/Source/Runtime/Engine/Public/ShowFlags.h
@@ -333,7 +333,8 @@ private:
 		SetHLODColoration(false);
 		SetVisualizeLPV(false);
 		SetStreamingBounds(false);
-		SetFoliageOcclusionBounds(false);
+		SetHISMCOcclusionBounds(false);
+		SetHISMCClusterTree(false);
 		SetConstraints(false);
 		SetMassProperties(false);
 		SetCameraFrustums(false);
diff --git a/Engine/Source/Runtime/Engine/Public/ShowFlagsValues.inl b/Engine/Source/Runtime/Engine/Public/ShowFlagsValues.inl
--- a/Engine/Source/Runtime/Engine/Public/ShowFlagsValues.inl
+++ b/Engine/Source/Runtime/Engine/Public/ShowFlagsValues.inl
@@ -141,7 +141,9 @@ SHOWFLAG_ALWAYS_ACCESSIBLE(InstancedStaticMeshes, SFG_Advanced, NSLOCTEXT("Unrea
 /** Draws instanced foliage, for now SHOWFLAG_ALWAYS_ACCESSIBLE because it's exposed in SceneCapture */
 SHOWFLAG_ALWAYS_ACCESSIBLE(InstancedFoliage, SFG_Advanced, NSLOCTEXT("UnrealEd", "InstancedFoliageSF", "Foliage"))
 /** Allow to see the foliage bounds used in the occlusion test */
-SHOWFLAG_FIXED_IN_SHIPPING(0, FoliageOcclusionBounds, SFG_Advanced, NSLOCTEXT("UnrealEd", "FoliageOcclusionBoundsSF", "Foliage Occlusion Bounds"))
+SHOWFLAG_FIXED_IN_SHIPPING(0, HISMCOcclusionBounds, SFG_Advanced, NSLOCTEXT("UnrealEd", "HISMOcclusionBoundsSF", "HISM/Foliage Occlusion Bounds"))
+/** Allow to see the cluster tree bounds used used to generate the occlusion bounds and in the culling */
+SHOWFLAG_FIXED_IN_SHIPPING(0, HISMCClusterTree, SFG_Advanced, NSLOCTEXT("UnrealEd", "HISMClusterTreeSF", "HISM/Foliage Cluster Tree"))
 /** Draws instanced grass, for now SHOWFLAG_ALWAYS_ACCESSIBLE because it's exposed in SceneCapture */
 SHOWFLAG_ALWAYS_ACCESSIBLE(InstancedGrass, SFG_Advanced, NSLOCTEXT("UnrealEd", "InstancedGrassSF", "Grass"))
 /** non baked shadows, for now SHOWFLAG_ALWAYS_ACCESSIBLE because it's exposed in SceneCapture */
