Add an "IsTerminating" event to the record of active sessions. Should allow detecting whether RequestExit() was called before an abnormal shutdown was triggered.

diff --git a/Engine/Source/Runtime/Engine/Private/EngineSessionManager.cpp b/Engine/Source/Runtime/Engine/Private/EngineSessionManager.cpp
--- a/Engine/Source/Runtime/Engine/Private/EngineSessionManager.cpp
+++ b/Engine/Source/Runtime/Engine/Private/EngineSessionManager.cpp
@@ -54,6 +54,7 @@ namespace SessionManagerDefs
 	static const FString GPUCrashStoreKey(TEXT("IsGPUCrash"));
 	static const FString DeactivatedStoreKey(TEXT("IsDeactivated"));
 	static const FString BackgroundStoreKey(TEXT("IsInBackground"));
+	static const FString TerminatingKey(TEXT("Terminating"));
 	static const FString EngineVersionStoreKey(TEXT("EngineVersion"));
 	static const FString TimestampStoreKey(TEXT("Timestamp"));
 	static const FString StartupTimeStoreKey(TEXT("StartupTimestamp"));
@@ -99,6 +100,7 @@ void FEngineSessionManager::Initialize()
 	FCoreDelegates::ApplicationWillDeactivateDelegate.AddRaw(this, &FEngineSessionManager::OnAppDeactivate);
 	FCoreDelegates::ApplicationWillEnterBackgroundDelegate.AddRaw(this, &FEngineSessionManager::OnAppBackground);
 	FCoreDelegates::ApplicationHasEnteredForegroundDelegate.AddRaw(this, &FEngineSessionManager::OnAppForeground);
+	FCoreDelegates::ApplicationWillTerminateDelegate.AddRaw(this, &FEngineSessionManager::OnTerminate);
 	FUserActivityTracking::OnActivityChanged.AddRaw(this, &FEngineSessionManager::OnUserActivity);
 	FCoreDelegates::IsVanillaProductChanged.AddRaw(this, &FEngineSessionManager::OnVanillaStateChanged);
 	FSlateApplication::Get().GetOnModalLoopTickEvent().AddRaw(this, &FEngineSessionManager::Tick);
@@ -236,6 +238,7 @@ void FEngineSessionManager::Shutdown()
 	FCoreDelegates::ApplicationWillDeactivateDelegate.RemoveAll(this);
 	FCoreDelegates::ApplicationWillEnterBackgroundDelegate.RemoveAll(this);
 	FCoreDelegates::ApplicationHasEnteredForegroundDelegate.RemoveAll(this);
+	FCoreDelegates::ApplicationWillTerminateDelegate.RemoveAll(this);
 	FCoreDelegates::IsVanillaProductChanged.RemoveAll(this);
 	FSlateApplication::Get().GetOnModalLoopTickEvent().RemoveAll(this);
 
@@ -256,6 +259,7 @@ void FEngineSessionManager::Shutdown()
 			FPlatformMisc::DeleteStoredValue(SessionManagerDefs::StoreId, CurrentSessionSectionName, SessionManagerDefs::BackgroundStoreKey);
 			FPlatformMisc::DeleteStoredValue(SessionManagerDefs::StoreId, CurrentSessionSectionName, SessionManagerDefs::UserActivityStoreKey);
 			FPlatformMisc::DeleteStoredValue(SessionManagerDefs::StoreId, CurrentSessionSectionName, SessionManagerDefs::VanillaStoreKey);
+			FPlatformMisc::DeleteStoredValue(SessionManagerDefs::StoreId, CurrentSessionSectionName, SessionManagerDefs::TerminatingKey);
 
 #if PLATFORM_SUPPORTS_WATCHDOG
 			if (!WatchdogSectionName.IsEmpty())
@@ -320,6 +324,8 @@ bool FEngineSessionManager::BeginReadWriteRecords()
 			FPlatformMisc::GetStoredValue(SessionManagerDefs::StoreId, SectionName, SessionManagerDefs::VanillaStoreKey, IsVanillaString);
 			FString IsGPUCrashString = SessionManagerDefs::FalseValueString;
 			FPlatformMisc::GetStoredValue(SessionManagerDefs::StoreId, SectionName, SessionManagerDefs::GPUCrashStoreKey, IsGPUCrashString);
+			FString IsTerminatingString = SessionManagerDefs::FalseValueString;
+			FPlatformMisc::GetStoredValue(SessionManagerDefs::StoreId, SectionName, SessionManagerDefs::TerminatingKey, IsTerminatingString);
 
 			// Create new record from the read values
 			FSessionRecord NewRecord;
@@ -336,6 +342,7 @@ bool FEngineSessionManager::BeginReadWriteRecords()
 			NewRecord.bIsInBackground = IsInBackgroundString == SessionManagerDefs::TrueValueString;
 			NewRecord.CurrentUserActivity = UserActivityString;
 			NewRecord.bIsVanilla = IsVanillaString == SessionManagerDefs::TrueValueString;
+			NewRecord.bIsTerminating = IsTerminatingString == SessionManagerDefs::TrueValueString;
 
 			SessionRecords.Add(NewRecord);
 		}
@@ -354,6 +361,7 @@ bool FEngineSessionManager::BeginReadWriteRecords()
 			FPlatformMisc::DeleteStoredValue(SessionManagerDefs::StoreId, SectionName, SessionManagerDefs::BackgroundStoreKey);
 			FPlatformMisc::DeleteStoredValue(SessionManagerDefs::StoreId, SectionName, SessionManagerDefs::UserActivityStoreKey);
 			FPlatformMisc::DeleteStoredValue(SessionManagerDefs::StoreId, SectionName, SessionManagerDefs::VanillaStoreKey);
+			FPlatformMisc::DeleteStoredValue(SessionManagerDefs::StoreId, SectionName, SessionManagerDefs::TerminatingKey);
 		}
 	}
 
@@ -399,6 +407,7 @@ void FEngineSessionManager::DeleteStoredRecord(const FSessionRecord& Record)
 	FPlatformMisc::DeleteStoredValue(SessionManagerDefs::StoreId, SectionName, SessionManagerDefs::BackgroundStoreKey);
 	FPlatformMisc::DeleteStoredValue(SessionManagerDefs::StoreId, SectionName, SessionManagerDefs::UserActivityStoreKey);
 	FPlatformMisc::DeleteStoredValue(SessionManagerDefs::StoreId, SectionName, SessionManagerDefs::VanillaStoreKey);
+	FPlatformMisc::DeleteStoredValue(SessionManagerDefs::StoreId, SectionName, SessionManagerDefs::TerminatingKey);
 
 	// Remove the session record from SessionRecords list
 	SessionRecords.RemoveAll([&SessionId](const FSessionRecord& X){ return X.SessionId == SessionId; });
@@ -425,6 +434,7 @@ void FEngineSessionManager::DeleteStoredRecord(const FSessionRecord& Record)
  * @EventParam CurrentUserActivity - If one was set when the session abnormally terminated, this is the activity taken from the FUserActivityTracking API.
  * @EventParam IsVanilla - Value from the engine's IsVanillaProduct() method. Basically if this is a Epic-distributed Editor with zero third party plugins or game code modules.
  * @EventParam GPUCrash - A GPU Hang or Crash was detected before the final assert, fatal log, or other exit.
+ * @EventParam Terminating - The RequestExit() function was called on the process; it may have been forcibly terminated or in the process of shutting down.
  
  *
  * @TODO: Debugger should be a completely separate flag, since it's orthogonal to whether we detect a crash or shutdown.
@@ -485,6 +495,7 @@ void FEngineSessionManager::SendAbnormalShutdownReport(const FSessionRecord& Rec
 	FString IsVanillaString = Record.bIsVanilla ? SessionManagerDefs::TrueValueString : SessionManagerDefs::FalseValueString;
 	FString WasDebuggedString = Record.bWasEverDebugger ? SessionManagerDefs::TrueValueString : SessionManagerDefs::FalseValueString;
 	FString GPUCrashedString = Record.bGPUCrashed ? SessionManagerDefs::TrueValueString : SessionManagerDefs::FalseValueString;
+	FString TerminatingString = Record.bIsTerminating ? SessionManagerDefs::TrueValueString : SessionManagerDefs::FalseValueString;
 
 	TArray< FAnalyticsEventAttribute > AbnormalShutdownAttributes;
 	AbnormalShutdownAttributes.Add(FAnalyticsEventAttribute(TEXT("RunType"), RunTypeString));
@@ -498,6 +509,7 @@ void FEngineSessionManager::SendAbnormalShutdownReport(const FSessionRecord& Rec
 	AbnormalShutdownAttributes.Add(FAnalyticsEventAttribute(TEXT("IsVanilla"), IsVanillaString));
 	AbnormalShutdownAttributes.Add(FAnalyticsEventAttribute(TEXT("WasDebugged"), WasDebuggedString));
 	AbnormalShutdownAttributes.Add(FAnalyticsEventAttribute(TEXT("GPUCrash"), GPUCrashedString));
+	AbnormalShutdownAttributes.Add(FAnalyticsEventAttribute(TEXT("Terminating"), TerminatingString));
 
 	FEngineAnalytics::GetProvider().RecordEvent(TEXT("Engine.AbnormalShutdown"), AbnormalShutdownAttributes);
 
@@ -534,6 +546,7 @@ void FEngineSessionManager::CreateAndWriteRecordForSession()
 	FString IsDeactivatedString = CurrentSession.bIsDeactivated ? SessionManagerDefs::TrueValueString : SessionManagerDefs::FalseValueString;
 	FString IsInBackgroundString = CurrentSession.bIsInBackground ? SessionManagerDefs::TrueValueString : SessionManagerDefs::FalseValueString;
 	FString IsVanillaString = CurrentSession.bIsVanilla ? SessionManagerDefs::TrueValueString : SessionManagerDefs::FalseValueString;
+	FString IsTerminatingString = CurrentSession.bIsTerminating ? SessionManagerDefs::TrueValueString : SessionManagerDefs::FalseValueString;
 
 	FPlatformMisc::SetStoredValue(SessionManagerDefs::StoreId, CurrentSessionSectionName, SessionManagerDefs::ModeStoreKey, ModeString);
 	FPlatformMisc::SetStoredValue(SessionManagerDefs::StoreId, CurrentSessionSectionName, SessionManagerDefs::ProjectNameStoreKey, CurrentSession.ProjectName);
@@ -546,6 +559,7 @@ void FEngineSessionManager::CreateAndWriteRecordForSession()
 	FPlatformMisc::SetStoredValue(SessionManagerDefs::StoreId, CurrentSessionSectionName, SessionManagerDefs::BackgroundStoreKey, IsInBackgroundString);
 	FPlatformMisc::SetStoredValue(SessionManagerDefs::StoreId, CurrentSessionSectionName, SessionManagerDefs::UserActivityStoreKey, CurrentSession.CurrentUserActivity);
 	FPlatformMisc::SetStoredValue(SessionManagerDefs::StoreId, CurrentSessionSectionName, SessionManagerDefs::VanillaStoreKey, IsVanillaString);
+	FPlatformMisc::SetStoredValue(SessionManagerDefs::StoreId, CurrentSessionSectionName, SessionManagerDefs::TerminatingKey, IsTerminatingString);
 
 	SessionRecords.Add(CurrentSession);
 
@@ -615,6 +629,15 @@ void FEngineSessionManager::OnAppForeground()
 	}
 }
 
+void FEngineSessionManager::OnTerminate()
+{
+	if (!CurrentSession.bIsTerminating)
+	{
+		CurrentSession.bIsTerminating = true;
+		FPlatformMisc::SetStoredValue(SessionManagerDefs::StoreId, CurrentSessionSectionName, SessionManagerDefs::TerminatingKey, SessionManagerDefs::TrueValueString);
+	}
+}
+
 FString FEngineSessionManager::GetStoreSectionString(FString InSuffix)
 {
 	check(Mode == EEngineSessionManagerMode::Editor || Mode == EEngineSessionManagerMode::Game)
diff --git a/Engine/Source/Runtime/Engine/Private/EngineSessionManager.h b/Engine/Source/Runtime/Engine/Private/EngineSessionManager.h
--- a/Engine/Source/Runtime/Engine/Private/EngineSessionManager.h
+++ b/Engine/Source/Runtime/Engine/Private/EngineSessionManager.h
@@ -59,6 +59,7 @@ private:
 		bool bIsInBackground;
 		FString CurrentUserActivity;
 		bool bIsVanilla;
+		bool bIsTerminating;
 	};
 
 private:
@@ -73,6 +74,7 @@ private:
 	void OnAppDeactivate();
 	void OnAppBackground();
 	void OnAppForeground();
+	void OnTerminate();
 	FString GetStoreSectionString(FString InSuffix);
 	void OnUserActivity(const FUserActivity& UserActivity);
 	void OnVanillaStateChanged(bool bIsVanilla);
