Fix crash when re registering a skeletal mesh component with a post process instance that references invalid virtual bones

diff --git a/Engine/Source/Runtime/Engine/Private/Components/SkeletalMeshComponent.cpp b/Engine/Source/Runtime/Engine/Private/Components/SkeletalMeshComponent.cpp
--- a/Engine/Source/Runtime/Engine/Private/Components/SkeletalMeshComponent.cpp
+++ b/Engine/Source/Runtime/Engine/Private/Components/SkeletalMeshComponent.cpp
@@ -753,7 +753,7 @@ bool USkeletalMeshComponent::InitializeAnimScriptInstance(bool bForceReinit)
 			PostProcessAnimInstance = nullptr;
 		}
 
-		if(NeedToSpawnPostPhysicsInstance())
+		if(bForceReinit || NeedToSpawnPostPhysicsInstance())
 		{
 			PostProcessAnimInstance = NewObject<UAnimInstance>(this, *SkeletalMesh->PostProcessAnimBlueprint);
 
