Make sure that animations that have transform curves are up to date on loading.

#4320

diff --git a/Engine/Source/Runtime/Engine/Private/Animation/AnimSequence.cpp b/Engine/Source/Runtime/Engine/Private/Animation/AnimSequence.cpp
--- a/Engine/Source/Runtime/Engine/Private/Animation/AnimSequence.cpp
+++ b/Engine/Source/Runtime/Engine/Private/Animation/AnimSequence.cpp
@@ -259,6 +259,12 @@ void UAnimSequence::Serialize(FArchive& Ar)
 				Ar << SourceRawAnimationData;
 			}
 		}
+
+		// If we have transform curves but no SourceRawAnimationData then we need to rebake
+		if (DoesContainTransformCurves() && RawAnimationData.Num() > 0 && SourceRawAnimationData.Num() == 0)
+		{
+			bNeedsRebake = true;
+		}
 #endif // WITH_EDITORONLY_DATA
 	}
 
