Fix for crash introduced in 4105998

diff --git a/Engine/Source/Runtime/Engine/Classes/Components/SkeletalMeshComponent.h b/Engine/Source/Runtime/Engine/Classes/Components/SkeletalMeshComponent.h
--- a/Engine/Source/Runtime/Engine/Classes/Components/SkeletalMeshComponent.h
+++ b/Engine/Source/Runtime/Engine/Classes/Components/SkeletalMeshComponent.h
@@ -1800,7 +1800,7 @@ public:
 
 protected:
 	bool NeedToSpawnAnimScriptInstance() const;
-	bool NeedToSpawnPostPhysicsInstance() const;
+	bool NeedToSpawnPostPhysicsInstance(bool bForceReinit) const;
 
 	bool ShouldBlendPhysicsBones() const;
 
diff --git a/Engine/Source/Runtime/Engine/Private/Components/SkeletalMeshComponent.cpp b/Engine/Source/Runtime/Engine/Private/Components/SkeletalMeshComponent.cpp
--- a/Engine/Source/Runtime/Engine/Private/Components/SkeletalMeshComponent.cpp
+++ b/Engine/Source/Runtime/Engine/Private/Components/SkeletalMeshComponent.cpp
@@ -476,7 +476,7 @@ bool USkeletalMeshComponent::NeedToSpawnAnimScriptInstance() const
 	return false;
 }
 
-bool USkeletalMeshComponent::NeedToSpawnPostPhysicsInstance() const
+bool USkeletalMeshComponent::NeedToSpawnPostPhysicsInstance(bool bForceReinit) const
 {
 	if(SkeletalMesh)
 	{
@@ -485,7 +485,7 @@ bool USkeletalMeshComponent::NeedToSpawnPostPhysicsInstance() const
 		const UClass* CurrentClass = PostProcessAnimInstance ? PostProcessAnimInstance->GetClass() : nullptr;
 
 		// We need to have an instance, and we have the wrong class (different or null)
-		if(ClassToUse && ClassToUse != CurrentClass && MainInstanceClass != ClassToUse)
+		if(ClassToUse && (ClassToUse != CurrentClass || bForceReinit ) && MainInstanceClass != ClassToUse)
 		{
 			return true;
 		}
@@ -753,7 +753,7 @@ bool USkeletalMeshComponent::InitializeAnimScriptInstance(bool bForceReinit)
 			PostProcessAnimInstance = nullptr;
 		}
 
-		if(bForceReinit || NeedToSpawnPostPhysicsInstance())
+		if(NeedToSpawnPostPhysicsInstance(bForceReinit))
 		{
 			PostProcessAnimInstance = NewObject<UAnimInstance>(this, *SkeletalMesh->PostProcessAnimBlueprint);
 
@@ -769,6 +769,10 @@ bool USkeletalMeshComponent::InitializeAnimScriptInstance(bool bForceReinit)
 				bInitializedPostInstance = true;
 			}
 		}
+		else if (!SkeletalMesh->PostProcessAnimBlueprint.Get())
+		{
+			PostProcessAnimInstance = nullptr;
+		}
 
 		if (AnimScriptInstance && !bInitializedMainInstance && bForceReinit)
 		{
