#4714

diff --git a/Engine/Source/Runtime/Engine/Private/HierarchicalInstancedStaticMesh.cpp b/Engine/Source/Runtime/Engine/Private/HierarchicalInstancedStaticMesh.cpp
--- a/Engine/Source/Runtime/Engine/Private/HierarchicalInstancedStaticMesh.cpp
+++ b/Engine/Source/Runtime/Engine/Private/HierarchicalInstancedStaticMesh.cpp
@@ -1474,6 +1474,12 @@ void FHierarchicalStaticMeshSceneProxy::GetDynamicMeshElements(const TArray<cons
 						FConvexVolume LeftEyeBounds, RightEyeBounds;
 						GetViewFrustumBounds(LeftEyeBounds, LeftEyeLocalViewProjForCulling, false);
 						GetViewFrustumBounds(RightEyeBounds, RightEyeLocalViewProjForCulling, false);
+
+						// Invalid bounds retrieved, so skip render of this frame
+						if (LeftEyeBounds.Planes.Num() < 5 || RightEyeBounds.Planes.Num() < 5)
+						{
+							continue;
+						}
 						
 						InstanceParams.ViewFrustumLocal.Planes.Empty(5);
 						InstanceParams.ViewFrustumLocal.Planes.Add(LeftEyeBounds.Planes[0]);
