Fix issue with double ticking a reused animation instance on set skeletal mesh because of parallel evaluation.

diff --git a/Engine/Source/Runtime/Engine/Private/Components/SkeletalMeshComponent.cpp b/Engine/Source/Runtime/Engine/Private/Components/SkeletalMeshComponent.cpp
--- a/Engine/Source/Runtime/Engine/Private/Components/SkeletalMeshComponent.cpp
+++ b/Engine/Source/Runtime/Engine/Private/Components/SkeletalMeshComponent.cpp
@@ -2446,6 +2446,17 @@ void USkeletalMeshComponent::SetSkeletalMesh(USkeletalMesh* InSkelMesh, bool bRe
 		return;
 	}
 
+	if(!bReinitPose)	// To stop double ticking when reusing the anim instance we need to make sure 
+						// we have completed animation parallel work before continuing with SetSkeletalMesh	
+	{
+		// We may be doing parallel evaluation on the current anim instance
+		// Calling this here with true will block this init till that thread completes
+		// and it is safe to continue
+		const bool bBlockOnTask = true; // wait on evaluation task so it is safe to continue with Init
+		const bool bPerformPostAnimEvaluation = true;
+		HandleExistingParallelEvaluationTask(bBlockOnTask, bPerformPostAnimEvaluation);
+	}
+
 	UPhysicsAsset* OldPhysAsset = GetPhysicsAsset();
 
 	{
