Fix for crash on connecting curve param node with no texture set

diff --git a/Engine/Source/Editor/MaterialEditor/Private/SMaterialLayersFunctionsTree.cpp b/Engine/Source/Editor/MaterialEditor/Private/SMaterialLayersFunctionsTree.cpp
--- a/Engine/Source/Editor/MaterialEditor/Private/SMaterialLayersFunctionsTree.cpp
+++ b/Engine/Source/Editor/MaterialEditor/Private/SMaterialLayersFunctionsTree.cpp
@@ -669,7 +669,7 @@ void SMaterialLayersFunctionsInstanceTreeItem::Construct(const FArguments& InArg
 				]
 			];
 		}
-		if (ScalarParam && ScalarParam->AtlasData.bIsUsedAsAtlasPosition)
+		else if (ScalarParam && ScalarParam->AtlasData.bIsUsedAsAtlasPosition)
 		{
 			const FText ParameterName = FText::FromName(StackParameterData->Parameter->ParameterInfo.Name);
 
diff --git a/Engine/Source/Runtime/Engine/Private/Materials/MaterialExpressions.cpp b/Engine/Source/Runtime/Engine/Private/Materials/MaterialExpressions.cpp
--- a/Engine/Source/Runtime/Engine/Private/Materials/MaterialExpressions.cpp
+++ b/Engine/Source/Runtime/Engine/Private/Materials/MaterialExpressions.cpp
@@ -14483,27 +14483,29 @@ int32 UMaterialExpressionCurveAtlasRowParameter::Compile(class FMaterialCompiler
 	if (Atlas && Curve)
 	{
 		SlotIndex = Atlas->GradientCurves.Find(Curve);
-	}
-	if (SlotIndex != INDEX_NONE)
-	{
-		float NewValue = ((float)SlotIndex * Atlas->GradientPixelSize) / Atlas->TextureSize + (0.5f * Atlas->GradientPixelSize) / Atlas->TextureSize;
-		SetParameterValue(ParameterName, NewValue);
-	}
-	// if the input is hooked up, use it, otherwise use the internal constant
-	int32 Arg1 = InputTime.GetTracedInput().Expression ? InputTime.Compile(Compiler) : Compiler->Constant(0);
-	int32 Arg2 = Compiler->ScalarParameter(ParameterName, DefaultValue);
 
-	int32 UV = Compiler->AppendVector(Arg1, Arg2);
-	return CompileTextureSample(
-		Compiler,
-		Atlas,
-		UV,
-		SAMPLERTYPE_LinearColor,
-		TOptional<FName>(),
-		INDEX_NONE,
-		INDEX_NONE,
-		TMVM_None,
-		SSM_Clamp_WorldGroupSettings);
+		if (SlotIndex != INDEX_NONE)
+		{
+			float NewValue = ((float)SlotIndex * Atlas->GradientPixelSize) / Atlas->TextureSize + (0.5f * Atlas->GradientPixelSize) / Atlas->TextureSize;
+			SetParameterValue(ParameterName, NewValue);
+		}
+		// if the input is hooked up, use it, otherwise use the internal constant
+		int32 Arg1 = InputTime.GetTracedInput().Expression ? InputTime.Compile(Compiler) : Compiler->Constant(0);
+		int32 Arg2 = Compiler->ScalarParameter(ParameterName, DefaultValue);
+
+		int32 UV = Compiler->AppendVector(Arg1, Arg2);
+		return CompileTextureSample(
+			Compiler,
+			Atlas,
+			UV,
+			SAMPLERTYPE_LinearColor,
+			TOptional<FName>(),
+			INDEX_NONE,
+			INDEX_NONE,
+			TMVM_None,
+			SSM_Clamp_WorldGroupSettings);
+	}
+	return INDEX_NONE;
 }
 
 void UMaterialExpressionCurveAtlasRowParameter::PostEditChangeProperty(FPropertyChangedEvent& PropertyChangedEvent)
