Fixed crash when running in standalone editor build using sub-instances that copy properties

DestProperty serialization was always a bad idea, this is another artifact of the problem where the property persists but the class it belongs to gets trashed. Re-mapping the property by name to the correct outer fixes this for a minor startup cost.

diff --git a/Engine/Source/Runtime/Engine/Private/Animation/AnimNodeBase.cpp b/Engine/Source/Runtime/Engine/Private/Animation/AnimNodeBase.cpp
--- a/Engine/Source/Runtime/Engine/Private/Animation/AnimNodeBase.cpp
+++ b/Engine/Source/Runtime/Engine/Private/Animation/AnimNodeBase.cpp
@@ -506,6 +506,14 @@ void FExposedValueHandler::Initialize(FAnimNode_Base* AnimNode, UObject* AnimIns
 			CopyRecord.CachedSourceProperty = AnimInstanceObject->GetClass()->FindPropertyByName(CopyRecord.SourcePropertyName);
 		}
 		check(CopyRecord.CachedSourceProperty);
+
+		// Re-find our dest property as it (or its class outer) may have changed
+		if (CopyRecord.DestProperty != nullptr && CopyRecord.DestProperty->GetOuter() != AnimInstanceObject->GetClass())
+		{
+			CopyRecord.DestProperty = AnimInstanceObject->GetClass()->FindPropertyByName(CopyRecord.DestProperty->GetFName());
+		}
+		check(CopyRecord.DestProperty);
+
 		if (UArrayProperty* SourceArrayProperty = Cast<UArrayProperty>(CopyRecord.CachedSourceProperty))
 		{
 			// the compiler should not be generating any code that calls down this path at the moment - it is untested
