Back out changelist 4108877

This was causing a 100% crash in an automated test

diff --git a/Engine/Source/Runtime/Engine/Private/Animation/AnimNodeBase.cpp b/Engine/Source/Runtime/Engine/Private/Animation/AnimNodeBase.cpp
--- a/Engine/Source/Runtime/Engine/Private/Animation/AnimNodeBase.cpp
+++ b/Engine/Source/Runtime/Engine/Private/Animation/AnimNodeBase.cpp
@@ -506,14 +506,6 @@ void FExposedValueHandler::Initialize(FAnimNode_Base* AnimNode, UObject* AnimIns
 			CopyRecord.CachedSourceProperty = AnimInstanceObject->GetClass()->FindPropertyByName(CopyRecord.SourcePropertyName);
 		}
 		check(CopyRecord.CachedSourceProperty);
-
-		// Re-find our dest property as it (or its class outer) may have changed
-		if (CopyRecord.DestProperty != nullptr && CopyRecord.DestProperty->GetOuter() != AnimInstanceObject->GetClass())
-		{
-			CopyRecord.DestProperty = AnimInstanceObject->GetClass()->FindPropertyByName(CopyRecord.DestProperty->GetFName());
-		}
-		check(CopyRecord.DestProperty);
-
 		if (UArrayProperty* SourceArrayProperty = Cast<UArrayProperty>(CopyRecord.CachedSourceProperty))
 		{
 			// the compiler should not be generating any code that calls down this path at the moment - it is untested
