Make sure rotations are normalized before creating blended transform

diff --git a/Engine/Source/Runtime/Engine/Private/Animation/AnimSequence.cpp b/Engine/Source/Runtime/Engine/Private/Animation/AnimSequence.cpp
--- a/Engine/Source/Runtime/Engine/Private/Animation/AnimSequence.cpp
+++ b/Engine/Source/Runtime/Engine/Private/Animation/AnimSequence.cpp
@@ -890,6 +890,17 @@ void UAnimSequence::ExtractBoneTransform(const struct FRawAnimSequenceTrack& Raw
 	// 	UE_LOG(LogAnimation, Log, TEXT(" *  *  *  Position. PosKeyIndex1: %3d, PosKeyIndex2: %3d, Alpha: %f"), PosKeyIndex1, PosKeyIndex2, Alpha);
 	// 	UE_LOG(LogAnimation, Log, TEXT(" *  *  *  Rotation. RotKeyIndex1: %3d, RotKeyIndex2: %3d, Alpha: %f"), RotKeyIndex1, RotKeyIndex2, Alpha);
 
+	// Ensure rotations are normalized (Added for Jira UE-53971)
+	if (!ensureMsgf(KeyAtom1.IsRotationNormalized(), TEXT("Rotation isn't normalized (Anim:%s)"), *GetPathName()))
+	{
+		KeyAtom1.NormalizeRotation();
+	}
+
+	if (!ensureMsgf(KeyAtom2.IsRotationNormalized(), TEXT("Rotation isn't normalized (Anim:%s)"), *GetPathName()))
+	{
+		KeyAtom2.NormalizeRotation();
+	}
+
 	OutAtom.Blend(KeyAtom1, KeyAtom2, Alpha);
 	OutAtom.NormalizeRotation();
 }
