Fix crash when trying to extract transform from raw track with no track data.

diff --git a/Engine/Source/Runtime/Engine/Private/Animation/AnimSequence.cpp b/Engine/Source/Runtime/Engine/Private/Animation/AnimSequence.cpp
--- a/Engine/Source/Runtime/Engine/Private/Animation/AnimSequence.cpp
+++ b/Engine/Source/Runtime/Engine/Private/Animation/AnimSequence.cpp
@@ -1195,11 +1195,26 @@ void UAnimSequence::BuildPoseFromRawDataInternal(const TArray<FRawAnimSequenceTr
 
 				const FRawAnimSequenceTrack& TrackToExtract = InAnimationData[TrackIndex];
 
-				InOutPose[PoseBoneIndex] = ExtractTransformForKey(KeyIndex1, TrackToExtract);
+				// Bail out (with rather wacky data) if data is empty for some reason.
+				if (TrackToExtract.PosKeys.Num() == 0 || TrackToExtract.RotKeys.Num() == 0)
+				{
+					UE_LOG(LogAnimation, Warning, TEXT("UAnimSequence::GetBoneTransform : No anim data in AnimSequence '%s' Track '%s'"), *GetPathName(), *AnimationTrackNames[TrackIndex].ToString() );
+
+					InOutPose[PoseBoneIndex].SetIdentity();
 
-				if (bInterpolateT)
+					if (bInterpolateT)
+					{
+						Key2Pose[PoseBoneIndex].SetIdentity();
+					}
+				}
+				else
 				{
-					Key2Pose[PoseBoneIndex] = ExtractTransformForKey(KeyIndex2, TrackToExtract);
+					InOutPose[PoseBoneIndex] = ExtractTransformForKey(KeyIndex1, TrackToExtract);
+
+					if (bInterpolateT)
+					{
+						Key2Pose[PoseBoneIndex] = ExtractTransformForKey(KeyIndex2, TrackToExtract);
+					}
 				}
 
 				RetargetTracking.Add(FRetargetTracking(PoseBoneIndex, SkeletonBoneIndex));
