[NO NON-PII COMMIT MESSAGE]

diff --git a/Engine/Source/Runtime/Engine/Private/HierarchicalInstancedStaticMesh.cpp b/Engine/Source/Runtime/Engine/Private/HierarchicalInstancedStaticMesh.cpp
--- a/Engine/Source/Runtime/Engine/Private/HierarchicalInstancedStaticMesh.cpp
+++ b/Engine/Source/Runtime/Engine/Private/HierarchicalInstancedStaticMesh.cpp
@@ -3094,7 +3094,4 @@ static FAutoConsoleCommand RebuildFoliageTreesCmd(
 	TEXT("foliage.RebuildFoliageTrees"),
 	TEXT("Rebuild the trees for non-grass foliage."),
 	FConsoleCommandWithArgsDelegate::CreateStatic(&RebuildFoliageTrees)
-	);
-
-
-PRAGMA_ENABLE_OPTIMIZATION
\ No newline at end of file
+	);
\ No newline at end of file
