[NO NON-PII COMMIT MESSAGE]

diff --git a/Engine/Source/Runtime/Engine/Private/HierarchicalInstancedStaticMesh.cpp b/Engine/Source/Runtime/Engine/Private/HierarchicalInstancedStaticMesh.cpp
--- a/Engine/Source/Runtime/Engine/Private/HierarchicalInstancedStaticMesh.cpp
+++ b/Engine/Source/Runtime/Engine/Private/HierarchicalInstancedStaticMesh.cpp
@@ -1940,8 +1940,6 @@ void UHierarchicalInstancedStaticMeshComponent::Serialize(FArchive& Ar)
 		// Skip the serialized tree, we will regenerate it correctly to contains the new data
 		TArray<FClusterNode_DEPRECATED> ClusterTree_DEPRECATED;
 		ClusterTree_DEPRECATED.BulkSerialize(Ar);
-
-		BuildTreeIfOutdated(false, true);
 	}
 	else
 	{
@@ -2545,7 +2543,8 @@ bool UHierarchicalInstancedStaticMeshComponent::BuildTreeIfOutdated(bool Async,
 		|| NumBuiltInstances != PerInstanceSMData.Num() 
 		|| (GetStaticMesh() != nullptr && CacheMeshExtendedBounds != GetStaticMesh()->GetBounds())
 		|| UnbuiltInstanceBoundsList.Num() > 0
-		|| GetLinkerUE4Version() < VER_UE4_REBUILD_HIERARCHICAL_INSTANCE_TREES)
+		|| GetLinkerUE4Version() < VER_UE4_REBUILD_HIERARCHICAL_INSTANCE_TREES
+		|| GetLinkerCustomVersion(FReleaseObjectVersion::GUID) < FReleaseObjectVersion::HISMCClusterTreeMigration)
 	{
 		if (GetStaticMesh() != nullptr && !GetStaticMesh()->HasAnyFlags(RF_NeedLoad)) // we can build the tree if the static mesh is not even loaded, and we can't call PostLoad as the load is not even done
 		{
