Fix non editor builds

diff --git a/Engine/Source/Runtime/Engine/Private/Animation/AnimSequence.cpp b/Engine/Source/Runtime/Engine/Private/Animation/AnimSequence.cpp
--- a/Engine/Source/Runtime/Engine/Private/Animation/AnimSequence.cpp
+++ b/Engine/Source/Runtime/Engine/Private/Animation/AnimSequence.cpp
@@ -1198,7 +1198,9 @@ void UAnimSequence::BuildPoseFromRawDataInternal(const TArray<FRawAnimSequenceTr
 				// Bail out (with rather wacky data) if data is empty for some reason.
 				if (TrackToExtract.PosKeys.Num() == 0 || TrackToExtract.RotKeys.Num() == 0)
 				{
+#if WITH_EDITORONLY_DATA
 					UE_LOG(LogAnimation, Warning, TEXT("UAnimSequence::GetBoneTransform : No anim data in AnimSequence '%s' Track '%s'"), *GetPathName(), *AnimationTrackNames[TrackIndex].ToString() );
+#endif
 
 					InOutPose[PoseBoneIndex].SetIdentity();
 
