UStaticMeshComponent::PostEditUndo now initializes its UStaticMesh's rendering resources.  This is necessary for cases where the UStaticMesh is in the same transaction (PostEditUndo will be called AFTER the component).  However, when calling UStaticMesh::InitResources we must use a FStaticMeshComponentRecreateRenderStateContext to safely handle any components using the UStaticMesh which are not in the transaction.

diff --git a/Engine/Source/Runtime/Engine/Classes/Engine/StaticMesh.h b/Engine/Source/Runtime/Engine/Classes/Engine/StaticMesh.h
--- a/Engine/Source/Runtime/Engine/Classes/Engine/StaticMesh.h
+++ b/Engine/Source/Runtime/Engine/Classes/Engine/StaticMesh.h
@@ -688,6 +688,9 @@ protected:
 	UPROPERTY(EditAnywhere, AdvancedDisplay, Instanced, Category = StaticMesh)
 	TArray<UAssetUserData*> AssetUserData;
 
+	/** Tracks whether InitResources has been called, and rendering resources are initialized. */
+	bool bRenderingResourcesInitialized;
+
 public:
 	/** The editable mesh representation of this static mesh */
 	// @todo: Maybe we don't want this visible in the details panel in the end; for now, this might aid debugging.
@@ -944,6 +947,8 @@ public:
 	/** Calculates the extended bounds */
 	ENGINE_API void CalculateExtendedBounds();
 
+	inline bool AreRenderingResourcesInitialized() const { return bRenderingResourcesInitialized; }
+
 #if WITH_EDITOR
 
 	/**
diff --git a/Engine/Source/Runtime/Engine/Private/Components/StaticMeshComponent.cpp b/Engine/Source/Runtime/Engine/Private/Components/StaticMeshComponent.cpp
--- a/Engine/Source/Runtime/Engine/Private/Components/StaticMeshComponent.cpp
+++ b/Engine/Source/Runtime/Engine/Private/Components/StaticMeshComponent.cpp
@@ -1544,6 +1544,17 @@ void UStaticMeshComponent::PreEditUndo()
 
 void UStaticMeshComponent::PostEditUndo()
 {
+	// If the StaticMesh was also involved in this transaction, it may need reinitialization first
+	// In this case, the StaticMesh will have PostEditUndo called later in this transaction, which is too late to register this component
+	if (GetStaticMesh() && !GetStaticMesh()->AreRenderingResourcesInitialized())
+	{
+		// We need to recreate the render state of any components using the static mesh before modifying its rendering resources
+		// However, we must not create the rendering state of any components in the transaction which have had PreEditUndo called and must not be referenced by the rendering thread until their PostEditUndo
+		// FStaticMeshComponentRecreateRenderStateContext handles this by only recreating rendering state if the component had rendering state created in the first place.
+		FStaticMeshComponentRecreateRenderStateContext RecreateContext(GetStaticMesh(), false, false);
+		GetStaticMesh()->InitResources();
+	}
+
 	// The component's light-maps are loaded from the transaction, so their resources need to be reinitialized.
 	InitResources();
 
diff --git a/Engine/Source/Runtime/Engine/Private/StaticMesh.cpp b/Engine/Source/Runtime/Engine/Private/StaticMesh.cpp
--- a/Engine/Source/Runtime/Engine/Private/StaticMesh.cpp
+++ b/Engine/Source/Runtime/Engine/Private/StaticMesh.cpp
@@ -1738,6 +1738,7 @@ UStaticMesh::UStaticMesh(const FObjectInitializer& ObjectInitializer)
 	MinLOD.Default = 0;
 
 	bSupportUniformlyDistributedSampling = false;
+	bRenderingResourcesInitialized = false;
 #if WITH_EDITOR
 	BuildCacheAutomationTestGuid.Invalidate();
 #endif
@@ -1759,6 +1760,8 @@ void UStaticMesh::PostInitProperties()
  */
 void UStaticMesh::InitResources()
 {
+	bRenderingResourcesInitialized = true;
+
 	UpdateUVChannelData(false);
 
 	if (RenderData)
@@ -2074,6 +2077,8 @@ void UStaticMesh::ReleaseResources()
 	
 	// insert a fence to signal when these commands completed
 	ReleaseResourcesFence.BeginFence();
+
+	bRenderingResourcesInitialized = false;
 }
 
 #if WITH_EDITOR
