Fixed static analysis warning caused by CL 4109096

diff --git a/Engine/Source/Runtime/Engine/Private/Animation/AnimNodeBase.cpp b/Engine/Source/Runtime/Engine/Private/Animation/AnimNodeBase.cpp
--- a/Engine/Source/Runtime/Engine/Private/Animation/AnimNodeBase.cpp
+++ b/Engine/Source/Runtime/Engine/Private/Animation/AnimNodeBase.cpp
@@ -556,14 +556,14 @@ void FExposedValueHandler::Initialize(FAnimNode_Base* AnimNode, UObject* AnimIns
 				CopyRecord.CachedDestContainer = AnimNode;
 			}
 		}
-		else
+		else if(CopyRecord.DestProperty != nullptr)
 		{
 			CopyRecord.Dest = CopyRecord.DestProperty->ContainerPtrToValuePtr<uint8>(AnimNode, CopyRecord.DestArrayIndex);
 
 			if (CopyRecord.bInstanceIsTarget)
 			{
 				// Re-find our dest property as it (or its class outer) may have changed
-				if (CopyRecord.DestProperty != nullptr && CopyRecord.DestProperty->GetOuter() != AnimInstanceObject->GetClass())
+				if (CopyRecord.DestProperty->GetOuter() != AnimInstanceObject->GetClass())
 				{
 					CopyRecord.DestProperty = AnimInstanceObject->GetClass()->FindPropertyByName(CopyRecord.DestProperty->GetFName());
 				}
