#rnx

diff --git a/Engine/Source/Runtime/Engine/Private/InstancedStaticMesh.cpp b/Engine/Source/Runtime/Engine/Private/InstancedStaticMesh.cpp
--- a/Engine/Source/Runtime/Engine/Private/InstancedStaticMesh.cpp
+++ b/Engine/Source/Runtime/Engine/Private/InstancedStaticMesh.cpp
@@ -236,7 +236,7 @@ void FStaticMeshInstanceBuffer::InitRHI()
 		QUICK_SCOPE_CYCLE_COUNTER(STAT_FStaticMeshInstanceBuffer_InitRHI);
 		auto AccessFlags = BUF_Static;
 		CreateVertexBuffer(InstanceData->GetOriginResourceArray(), AccessFlags | BUF_ShaderResource, 16, PF_A32B32G32R32F, InstanceOriginBuffer.VertexBufferRHI, InstanceOriginSRV);
-		CreateVertexBuffer(InstanceData->GetTransformResourceArray(), AccessFlags | BUF_ShaderResource, InstanceData->GetTranslationUsesHalfs() ? 8 : 16, InstanceData->GetTranslationUsesHalfs() ? PF_FloatRGBA : PF_A32B32G32R32F, InstanceTransformBuffer.VertexBufferRHI, InstanceTransformSRV);
+		CreateVertexBuffer(InstanceData->GetTransformResourceArray(), AccessFlags | BUF_ShaderResource, InstanceData->GetTranslationUsesHalfs() && GVertexElementTypeSupport.IsSupported(VET_Half2) ? 8 : 16, InstanceData->GetTranslationUsesHalfs() && GVertexElementTypeSupport.IsSupported(VET_Half2) ? PF_FloatRGBA : PF_A32B32G32R32F, InstanceTransformBuffer.VertexBufferRHI, InstanceTransformSRV);
 		CreateVertexBuffer(InstanceData->GetLightMapResourceArray(), AccessFlags | BUF_ShaderResource, 8, PF_R16G16B16A16_SNORM, InstanceLightmapBuffer.VertexBufferRHI, InstanceLightmapSRV);
 	}
 }
@@ -305,8 +305,8 @@ void FStaticMeshInstanceBuffer::BindInstanceVertexBuffer(const class FVertexFact
 			EVertexStreamUsage::ManualFetch | EVertexStreamUsage::Instancing
 		);
 
-		EVertexElementType TransformType = InstanceData->GetTranslationUsesHalfs() ? VET_Half4 : VET_Float4;
-		uint32 TransformStride = InstanceData->GetTranslationUsesHalfs() ? 8 : 16;
+		EVertexElementType TransformType = InstanceData->GetTranslationUsesHalfs() && GVertexElementTypeSupport.IsSupported(VET_Half2) ? VET_Half4 : VET_Float4;
+		uint32 TransformStride = InstanceData->GetTranslationUsesHalfs() && GVertexElementTypeSupport.IsSupported(VET_Half2) ? 8 : 16;
 
 		InstancedStaticMeshData.InstanceTransformComponent[0] = FVertexStreamComponent(
 			&InstanceTransformBuffer,
