Make sure bone references are up to date on update (skeletal control nodes can crash otherwise)

diff --git a/Engine/Source/Runtime/Engine/Private/Animation/AnimInstanceProxy.cpp b/Engine/Source/Runtime/Engine/Private/Animation/AnimInstanceProxy.cpp
--- a/Engine/Source/Runtime/Engine/Private/Animation/AnimInstanceProxy.cpp
+++ b/Engine/Source/Runtime/Engine/Private/Animation/AnimInstanceProxy.cpp
@@ -1083,6 +1083,8 @@ void FAnimInstanceProxy::UpdateAnimation()
 	ANIM_MT_SCOPE_CYCLE_COUNTER(ProxyUpdateAnimation, !IsInGameThread());
 	FScopeCycleCounterUObject AnimScope(GetAnimInstanceObject());
 
+	CacheBones();
+
 	// update native update
 	{
 		SCOPE_CYCLE_COUNTER(STAT_NativeUpdateAnimation);
