Extra logging for marker sync issue

diff --git a/Engine/Source/Runtime/Engine/Classes/Animation/AnimationAsset.h b/Engine/Source/Runtime/Engine/Classes/Animation/AnimationAsset.h
--- a/Engine/Source/Runtime/Engine/Classes/Animation/AnimationAsset.h
+++ b/Engine/Source/Runtime/Engine/Classes/Animation/AnimationAsset.h
@@ -51,6 +51,12 @@ struct FMarkerTickRecord
 	bool IsValid() const { return PreviousMarker.MarkerIndex != MarkerIndexSpecialValues::Unitialized && NextMarker.MarkerIndex != MarkerIndexSpecialValues::Unitialized; }
 
 	void Reset() { PreviousMarker.Reset(); NextMarker.Reset(); }
+
+	/** Debug output function*/
+	FString ToString() const
+	{
+		return FString::Printf(TEXT("[PreviousMarker Index/Time %i/%.2f, NextMarker Index/Time %i/%.2f]"), PreviousMarker.MarkerIndex, PreviousMarker.TimeToMarker, NextMarker.MarkerIndex, NextMarker.TimeToMarker);
+	}
 };
 
 /** Transform definition */
diff --git a/Engine/Source/Runtime/Engine/Private/Animation/AnimInstanceProxy.cpp b/Engine/Source/Runtime/Engine/Private/Animation/AnimInstanceProxy.cpp
--- a/Engine/Source/Runtime/Engine/Private/Animation/AnimInstanceProxy.cpp
+++ b/Engine/Source/Runtime/Engine/Private/Animation/AnimInstanceProxy.cpp
@@ -613,9 +613,13 @@ void FAnimInstanceProxy::TickAssetPlayerInstances(float DeltaSeconds)
 				}
 			}
 
+#if DO_CHECK
 			//For debugging UE-54705
 			FName InitialMarkerPrevious = TickContext.MarkerTickContext.GetMarkerSyncStartPosition().PreviousMarkerName;
 			FName InitialMarkerEnd = TickContext.MarkerTickContext.GetMarkerSyncStartPosition().NextMarkerName;
+			const bool bIsLeaderRecordValidPre = SyncGroup.ActivePlayers[0].MarkerTickRecord->IsValid();
+			FMarkerTickRecord LeaderPreMarkerTickRecord = *SyncGroup.ActivePlayers[0].MarkerTickRecord;
+#endif
 
 			// initialize to invalidate first
 			ensureMsgf(SyncGroup.GroupLeaderIndex == INDEX_NONE, TEXT("SyncGroup with GroupIndex=%d had a non -1 group leader index of %d in asset %s"), GroupIndex, SyncGroup.GroupLeaderIndex, *GetNameSafe(SkeletalMeshComponent));
@@ -673,28 +677,34 @@ void FAnimInstanceProxy::TickAssetPlayerInstances(float DeltaSeconds)
 
 				if (!bStartMarkerValid)
 				{
+#if DO_CHECK
 					FString ErrorMsg = FString(TEXT("Prev Marker name not valid for sync group.\n"));
 					ErrorMsg += FString::Format(TEXT("\tMarker {0} : SyncGroupName {1} : Leader {2}\n"), { MarkerStart.PreviousMarkerName.ToString(), SyncGroupName.ToString(), LeaderAnimName });
 					ErrorMsg += FString::Format(TEXT("\tInitalPrev {0} : InitialNext {1} : GroupLeaderIndex {2}\n"), { InitialMarkerPrevious.ToString(), InitialMarkerEnd.ToString(), GroupLeaderIndex });
+					ErrorMsg += FString::Format(TEXT("\tLeader (0 index) was originally valid: {0} | Record: {1}\n"), { bIsLeaderRecordValidPre, LeaderPreMarkerTickRecord.ToString() });
 					ErrorMsg += FString::Format(TEXT("\t Valid Markers : {0}\n"), { SyncGroup.ValidMarkers.Num() });
 					for (int32 MarkerIndex = 0; MarkerIndex < SyncGroup.ValidMarkers.Num(); ++MarkerIndex)
 					{
 						ErrorMsg += FString::Format(TEXT("\t\t{0}) '{1}'\n"), {MarkerIndex, SyncGroup.ValidMarkers[MarkerIndex].ToString()});
 					}
 					ensureMsgf(false, TEXT("%s"), *ErrorMsg);
+#endif
 					TickContext.InvalidateMarkerSync();
 				}
 				else if (!bEndMarkerValid)
 				{
+#if DO_CHECK
 					FString ErrorMsg = FString(TEXT("Next Marker name not valid for sync group.\n"));
 					ErrorMsg += FString::Format(TEXT("\tMarker {0} : SyncGroupName {1} : Leader {2}\n"), { MarkerStart.NextMarkerName.ToString(), SyncGroupName.ToString(), LeaderAnimName });
 					ErrorMsg += FString::Format(TEXT("\tInitalPrev {0} : InitialNext {1} : GroupLeaderIndex {2}\n"), { InitialMarkerPrevious.ToString(), InitialMarkerEnd.ToString(), GroupLeaderIndex });
+					ErrorMsg += FString::Format(TEXT("\tLeader (0 index) was originally valid: {0} | Record: {1}\n"), { bIsLeaderRecordValidPre, LeaderPreMarkerTickRecord.ToString() });
 					ErrorMsg += FString::Format(TEXT("\t Valid Markers : {0}\n"), { SyncGroup.ValidMarkers.Num() });
 					for (int32 MarkerIndex = 0; MarkerIndex < SyncGroup.ValidMarkers.Num(); ++MarkerIndex)
 					{
 						ErrorMsg += FString::Format(TEXT("\t\t{0}) '{1}'\n"), { MarkerIndex, SyncGroup.ValidMarkers[MarkerIndex].ToString() });
 					}
 					ensureMsgf(false, TEXT("%s"), *ErrorMsg);
+#endif
 					TickContext.InvalidateMarkerSync();
 				}
 			}
