Audio: immediately stop sounds when destroying it's data

diff --git a/Engine/Source/Runtime/Engine/Private/AudioDevice.cpp b/Engine/Source/Runtime/Engine/Private/AudioDevice.cpp
--- a/Engine/Source/Runtime/Engine/Private/AudioDevice.cpp
+++ b/Engine/Source/Runtime/Engine/Private/AudioDevice.cpp
@@ -4751,7 +4751,7 @@ void FAudioDevice::StopSourcesUsingBuffer(FSoundBuffer* SoundBuffer)
 			if (Src && Src->Buffer == SoundBuffer)
 			{
 				// Make sure the buffer is no longer referenced by anything
-				Src->Stop();
+				Src->StopNow();
 				break;
 			}
 		}
