RL - Fix CPU skin emulation, where the SSE code was unpacking the underlying tangent SNORM data as UNORM

diff --git a/Engine/Source/Runtime/Core/Public/Math/UnrealMathDirectX.h b/Engine/Source/Runtime/Core/Public/Math/UnrealMathDirectX.h
--- a/Engine/Source/Runtime/Core/Public/Math/UnrealMathDirectX.h
+++ b/Engine/Source/Runtime/Core/Public/Math/UnrealMathDirectX.h
@@ -677,6 +677,15 @@ FORCEINLINE VectorRegister VectorLoadByte4Reverse( const uint8* Ptr )
  */
 #define VectorStoreByte4( Vec, Ptr )		DirectX::PackedVector::XMStoreUByte4( (DirectX::PackedVector::XMUBYTE4*)(Ptr), Vec )
 
+ /**
+ * Converts the 4 FLOATs in the vector to 4 BYTEs, clamped to [-127,127], and stores to unaligned memory.
+ * IMPORTANT: You need to call VectorResetFloatRegisters() before using scalar FLOATs after you've used this intrinsic!
+ *
+ * @param Vec			Vector containing 4 FLOATs
+ * @param Ptr			Unaligned memory pointer to store the 4 BYTEs.
+ */
+#define VectorStoreSignedByte4( Vec, Ptr )		DirectX::PackedVector::XMStoreByte4( (DirectX::PackedVector::XMBYTE4*)(Ptr), Vec )
+
 /**
 * Loads packed RGB10A2(4 bytes) from unaligned memory and converts them into 4 FLOATs.
 * IMPORTANT: You need to call VectorResetFloatRegisters() before using scalar FLOATs after you've used this intrinsic!
diff --git a/Engine/Source/Runtime/Core/Public/Math/UnrealMathFPU.h b/Engine/Source/Runtime/Core/Public/Math/UnrealMathFPU.h
--- a/Engine/Source/Runtime/Core/Public/Math/UnrealMathFPU.h
+++ b/Engine/Source/Runtime/Core/Public/Math/UnrealMathFPU.h
@@ -1026,6 +1026,23 @@ FORCEINLINE void VectorStoreByte4( const VectorRegister& Vec, void* Ptr )
 	BytePtr[3] = uint8( Vec.V[3] );
 }
 
+/**
+* Converts the 4 FLOATs in the vector to 4 BYTEs, clamped to [-127,127], and stores to unaligned memory.
+* IMPORTANT: You need to call VectorResetFloatRegisters() before using scalar FLOATs after you've used this intrinsic!
+*
+* @param Vec			Vector containing 4 FLOATs
+* @param Ptr			Unaligned memory pointer to store the 4 BYTEs.
+*/
+FORCEINLINE void VectorStoreSignedByte4(const VectorRegister& Vec, void* Ptr)
+{
+	int8 *BytePtr = (int8*)Ptr;
+	BytePtr[0] = int8(Vec.V[0]);
+	BytePtr[1] = int8(Vec.V[1]);
+	BytePtr[2] = int8(Vec.V[2]);
+	BytePtr[3] = int8(Vec.V[3]);
+}
+
+
 /**
 * Loads packed RGB10A2(4 bytes) from unaligned memory and converts them into 4 FLOATs.
 * IMPORTANT: You need to call VectorResetFloatRegisters() before using scalar FLOATs after you've used this intrinsic!
diff --git a/Engine/Source/Runtime/Core/Public/Math/UnrealMathNeon.h b/Engine/Source/Runtime/Core/Public/Math/UnrealMathNeon.h
--- a/Engine/Source/Runtime/Core/Public/Math/UnrealMathNeon.h
+++ b/Engine/Source/Runtime/Core/Public/Math/UnrealMathNeon.h
@@ -991,6 +991,23 @@ FORCEINLINE void VectorStoreByte4( VectorRegister Vec, void* Ptr )
 	*(uint32_t *)Ptr = buf[0]; 
 }
 
+/**
+* Converts the 4 floats in the vector to 4 int8s, clamped to [-127, 127], and stores to unaligned memory.
+* IMPORTANT: You need to call VectorResetFloatRegisters() before using scalar floats after you've used this intrinsic!
+*
+* @param Vec			Vector containing 4 floats
+* @param Ptr			Unaligned memory pointer to store the 4 uint8s.
+*/
+FORCEINLINE void VectorStoreSignedByte4(VectorRegister Vec, void* Ptr)
+{
+	int16x8_t s16x8 = (int16x8_t)vcvtq_s32_f32(VectorMax(VectorMin(Vec, GlobalVectorConstants::Float127), GlobalVectorConstants::FloatNeg127));
+	int8x8_t s8x8 = (int8x8_t)vget_low_s16(vuzpq_s16(s16x8, s16x8).val[0]);
+	s8x8 = vuzp_s8(s8x8, s8x8).val[0];
+	int32_t buf[2];
+	vst1_s8((int8_t *)buf, s8x8);
+	*(int32_t *)Ptr = buf[0];
+}
+
 /**
  * Converts the 4 floats in the vector to 4 fp16 and stores based off bool to [un]aligned memory.
  *
diff --git a/Engine/Source/Runtime/Core/Public/Math/UnrealMathSSE.h b/Engine/Source/Runtime/Core/Public/Math/UnrealMathSSE.h
--- a/Engine/Source/Runtime/Core/Public/Math/UnrealMathSSE.h
+++ b/Engine/Source/Runtime/Core/Public/Math/UnrealMathSSE.h
@@ -933,6 +933,21 @@ FORCEINLINE void VectorStoreByte4( const VectorRegister& Vec, void* Ptr )
 	*(int32*)Ptr = _mm_cvtsi128_si32(_mm_packus_epi16(_mm_packs_epi32(_mm_cvttps_epi32(Vec), _mm_setzero_si128()), _mm_setzero_si128()));
 }
 
+/**
+* Converts the 4 FLOATs in the vector to 4 BYTEs, clamped to [-127,127], and stores to unaligned memory.
+* IMPORTANT: You need to call VectorResetFloatRegisters() before using scalar FLOATs after you've used this intrinsic!
+*
+* @param Vec			Vector containing 4 FLOATs
+* @param Ptr			Unaligned memory pointer to store the 4 BYTEs.
+*/
+FORCEINLINE void VectorStoreSignedByte4(const VectorRegister& Vec, void* Ptr)
+{
+	// Looks complex but is really quite straightforward:
+	// Convert 4x floats to 4x 32-bit ints, then pack into 4x 16-bit ints, then into 4x 8-bit unsigned ints, then store as a 32-bit value
+	*(int32*)Ptr = _mm_cvtsi128_si32(_mm_packs_epi16(_mm_packs_epi32(_mm_cvttps_epi32(Vec), _mm_setzero_si128()), _mm_setzero_si128()));
+}
+
+
 /**
 * Loads packed RGB10A2(4 bytes) from unaligned memory and converts them into 4 FLOATs.
 * IMPORTANT: You need to call VectorResetFloatRegisters() before using scalar FLOATs after you've used this intrinsic!
diff --git a/Engine/Source/Runtime/Core/Public/Math/UnrealMathVectorConstants.h b/Engine/Source/Runtime/Core/Public/Math/UnrealMathVectorConstants.h
--- a/Engine/Source/Runtime/Core/Public/Math/UnrealMathVectorConstants.h
+++ b/Engine/Source/Runtime/Core/Public/Math/UnrealMathVectorConstants.h
@@ -60,6 +60,8 @@ namespace GlobalVectorConstants
 	static const VectorRegister OneOverTwoPi = MakeVectorRegister(1.0f / (2.0f*PI), 1.0f / (2.0f*PI), 1.0f / (2.0f*PI), 1.0f / (2.0f*PI));
 
 	static const VectorRegister Float255 = MakeVectorRegister(255.0f, 255.0f, 255.0f, 255.0f);
+	static const VectorRegister Float127 = MakeVectorRegister(127.0f, 127.0f, 127.0f, 127.0f);
+	static const VectorRegister FloatNeg127 = MakeVectorRegister(-127.0f, -127.0f, -127.0f, -127.0f);
 	static const VectorRegister Float360 = MakeVectorRegister(360.f, 360.f, 360.f, 360.f);
 	static const VectorRegister Float180 = MakeVectorRegister(180.f, 180.f, 180.f, 180.f);
 
diff --git a/Engine/Source/Runtime/Engine/Private/SkeletalRender.h b/Engine/Source/Runtime/Engine/Private/SkeletalRender.h
--- a/Engine/Source/Runtime/Engine/Private/SkeletalRender.h
+++ b/Engine/Source/Runtime/Engine/Private/SkeletalRender.h
@@ -40,14 +40,6 @@ ENGINE_API void UpdateRefToLocalMatrices( TArray<FMatrix>& ReferenceToLocal, con
 */
 ENGINE_API void UpdatePreviousRefToLocalMatrices(TArray<FMatrix>& ReferenceToLocal, const USkinnedMeshComponent* InMeshComponent, const FSkeletalMeshRenderData* InSkeletalMeshRenderData, int32 LODIndex, const TArray<FBoneIndexType>* ExtraRequiredBoneIndices = NULL);
 
-extern ENGINE_API const VectorRegister		VECTOR_PACK_127_5;
-extern ENGINE_API const VectorRegister		VECTOR4_PACK_127_5;
-
-extern ENGINE_API const VectorRegister		VECTOR_INV_127_5;
-extern ENGINE_API const VectorRegister		VECTOR4_INV_127_5;
-
-extern ENGINE_API const VectorRegister		VECTOR_UNPACK_MINUS_1;
-extern ENGINE_API const VectorRegister		VECTOR4_UNPACK_MINUS_1;
 
 extern ENGINE_API const VectorRegister		VECTOR_0001;
 
@@ -60,7 +52,7 @@ extern ENGINE_API const VectorRegister		VECTOR_0001;
 */
 static FORCEINLINE VectorRegister Unpack3( const uint32 *PackedNormal )
 {
-	return VectorMultiplyAdd( VectorLoadByte4(PackedNormal), VECTOR_INV_127_5, VECTOR_UNPACK_MINUS_1 );
+	return VectorMultiply(VectorLoadSignedByte4(PackedNormal), VectorSetFloat3(1.0f / 127.0f, 1.0f / 127.0f, 1.0f / 127.0f));
 }
 
 /**
@@ -72,8 +64,8 @@ static FORCEINLINE VectorRegister Unpack3( const uint32 *PackedNormal )
 */
 static FORCEINLINE void Pack3( VectorRegister Normal, uint32 *PackedNormal )
 {
-	Normal = VectorMultiplyAdd(Normal, VECTOR_PACK_127_5, VECTOR_PACK_127_5);
-	VectorStoreByte4( Normal, PackedNormal );
+	Normal = VectorMultiply(Normal, VectorSetFloat3(127.0f, 127.0f, 127.0f));
+	VectorStoreSignedByte4( Normal, PackedNormal );
 }
 
 /**
@@ -85,7 +77,7 @@ static FORCEINLINE void Pack3( VectorRegister Normal, uint32 *PackedNormal )
 */
 static FORCEINLINE VectorRegister Unpack4( const uint32 *PackedNormal )
 {
-	return VectorMultiplyAdd( VectorLoadByte4(PackedNormal), VECTOR4_INV_127_5, VECTOR4_UNPACK_MINUS_1 );
+	return VectorMultiply(VectorLoadSignedByte4(PackedNormal), VectorSetFloat1(1.0f / 127.0f));
 }
 
 /**
@@ -97,7 +89,7 @@ static FORCEINLINE VectorRegister Unpack4( const uint32 *PackedNormal )
 */
 static FORCEINLINE void Pack4( VectorRegister Normal, uint32 *PackedNormal )
 {
-	Normal = VectorMultiplyAdd(Normal, VECTOR4_PACK_127_5, VECTOR4_PACK_127_5);
-	VectorStoreByte4( Normal, PackedNormal );
+	Normal = VectorMultiply(Normal, VectorSetFloat1(127.0f));
+	VectorStoreSignedByte4( Normal, PackedNormal );
 }
 
diff --git a/Engine/Source/Runtime/Engine/Private/SkeletalRenderCPUSkin.cpp b/Engine/Source/Runtime/Engine/Private/SkeletalRenderCPUSkin.cpp
--- a/Engine/Source/Runtime/Engine/Private/SkeletalRenderCPUSkin.cpp
+++ b/Engine/Source/Runtime/Engine/Private/SkeletalRenderCPUSkin.cpp
@@ -507,7 +507,7 @@ FORCEINLINE void ApplyMorphBlend( VertexType& DestVertex, const FMorphTargetDelt
 	DestVertex.Position += SrcMorph.PositionDelta * Weight;
 
 	// Save W before = operator. That overwrites W to be 127.
-	uint8 W = DestVertex.TangentZ.Vector.W;
+	int8 W = DestVertex.TangentZ.Vector.W;
 
 	FVector TanZ = DestVertex.TangentZ.ToFVector();
 
@@ -555,16 +555,6 @@ FORCEINLINE void UpdateMorphedVertex( VertexType& MorphedVertex, const VertexTyp
 MSVC_PRAGMA(warning(push))
 MSVC_PRAGMA(warning(disable : 4730)) //mixing _m64 and floating point expressions may result in incorrect code
 
-
-const VectorRegister		VECTOR_PACK_127_5		= DECLARE_VECTOR_REGISTER(127.5f, 127.5f, 127.5f, 0.f);
-const VectorRegister		VECTOR4_PACK_127_5		= DECLARE_VECTOR_REGISTER(127.5f, 127.5f, 127.5f, 127.5f);
-
-const VectorRegister		VECTOR_INV_127_5		= DECLARE_VECTOR_REGISTER(1.f / 127.5f, 1.f / 127.5f, 1.f / 127.5f, 0.f);
-const VectorRegister		VECTOR4_INV_127_5		= DECLARE_VECTOR_REGISTER(1.f / 127.5f, 1.f / 127.5f, 1.f / 127.5f, 1.f / 127.5f);
-
-const VectorRegister		VECTOR_UNPACK_MINUS_1	= DECLARE_VECTOR_REGISTER(-1.f, -1.f, -1.f, 0.f);
-const VectorRegister		VECTOR4_UNPACK_MINUS_1	= DECLARE_VECTOR_REGISTER(-1.f, -1.f, -1.f, -1.f);
-
 const VectorRegister		VECTOR_0001				= DECLARE_VECTOR_REGISTER(0.f, 0.f, 0.f, 1.f);
 
 #define FIXED_VERTEX_INDEX 0xFFFF
