Changed virtual path when streaming textures, so that it clears the IO requests on an async task instead of in the renderthread.

diff --git a/Engine/Source/Runtime/Engine/Private/Streaming/Texture2DStreamIn_IO_Virtual.cpp b/Engine/Source/Runtime/Engine/Private/Streaming/Texture2DStreamIn_IO_Virtual.cpp
--- a/Engine/Source/Runtime/Engine/Private/Streaming/Texture2DStreamIn_IO_Virtual.cpp
+++ b/Engine/Source/Runtime/Engine/Private/Streaming/Texture2DStreamIn_IO_Virtual.cpp
@@ -32,9 +32,20 @@ void FTexture2DStreamIn_IO_Virtual::LoadMips(const FContext& Context)
 
 	SetIORequests(Context);
 
-	PushTask(Context, TT_Render, TEXTURE2D_UPDATE_CALLBACK(Finalize), TT_Async, TEXTURE2D_UPDATE_CALLBACK(CancelIO));
+	PushTask(Context, TT_Async, TEXTURE2D_UPDATE_CALLBACK(PostLoadMips), TT_Async, TEXTURE2D_UPDATE_CALLBACK(CancelIO));
 }
 
+void FTexture2DStreamIn_IO_Virtual::PostLoadMips(const FContext& Context)
+{
+	DECLARE_SCOPE_CYCLE_COUNTER(TEXT("FTexture2DStreamIn_IO_Virtual::PostLoadMips"), STAT_Texture2DStreamInIOVirtual_LoadMips, STATGROUP_StreamingDetails);
+	check(Context.CurrentThread == TT_Async);
+
+	ClearIORequests(Context);
+
+	PushTask(Context, TT_Render, TEXTURE2D_UPDATE_CALLBACK(Finalize), TT_Render, TEXTURE2D_UPDATE_CALLBACK(Cancel));
+}
+
+
 void FTexture2DStreamIn_IO_Virtual::Finalize(const FContext& Context)
 {
 	DECLARE_SCOPE_CYCLE_COUNTER(TEXT("FTexture2DStreamIn_IO_Virtual::Finalize"), STAT_Texture2DStreamInIOVirtual_Finalize, STATGROUP_StreamingDetails);
diff --git a/Engine/Source/Runtime/Engine/Private/Streaming/Texture2DStreamIn_IO_Virtual.h b/Engine/Source/Runtime/Engine/Private/Streaming/Texture2DStreamIn_IO_Virtual.h
--- a/Engine/Source/Runtime/Engine/Private/Streaming/Texture2DStreamIn_IO_Virtual.h
+++ b/Engine/Source/Runtime/Engine/Private/Streaming/Texture2DStreamIn_IO_Virtual.h
@@ -25,6 +25,8 @@ protected:
 	void LockMips(const FContext& Context);
 	// Create load requests into each locked mips. (AsyncThread)
 	void LoadMips(const FContext& Context);
+	// Clear the IO requests and call Finalize. (AsyncThread)
+	void PostLoadMips(const FContext& Context);
 	// Apply the intermediate texture and cleanup. (RenderThread)
 	void Finalize(const FContext& Context);
 
