Fix crash when remapping pose asset to new skeleton

diff --git a/Engine/Source/Runtime/Engine/Private/Animation/PoseAsset.cpp b/Engine/Source/Runtime/Engine/Private/Animation/PoseAsset.cpp
--- a/Engine/Source/Runtime/Engine/Private/Animation/PoseAsset.cpp
+++ b/Engine/Source/Runtime/Engine/Private/Animation/PoseAsset.cpp
@@ -1340,6 +1340,18 @@ void FPoseDataContainer::DeleteTrack(int32 TrackIndex)
 void UPoseAsset::RemapTracksToNewSkeleton(USkeleton* NewSkeleton, bool bConvertSpaces)
 {
 	Super::RemapTracksToNewSkeleton(NewSkeleton, bConvertSpaces);
+
+	// after remap, should verify if the names are still valid in this skeleton
+	if (NewSkeleton)
+	{
+		NewSkeleton->VerifySmartNames(USkeleton::AnimCurveMappingName, PoseContainer.PoseNames);
+
+		for (auto& Curve : PoseContainer.Curves)
+		{
+			NewSkeleton->VerifySmartName(USkeleton::AnimCurveMappingName, Curve.Name);
+		}
+	}
+
 	// @todo: should allow deleted tracks and so on
 	RecacheTrackmap();
 }
