Fixed asset thumbnail crash when exiting at the same time the game engine is closing.

diff --git a/Engine/Source/Runtime/Engine/Private/Texture2D.cpp b/Engine/Source/Runtime/Engine/Private/Texture2D.cpp
--- a/Engine/Source/Runtime/Engine/Private/Texture2D.cpp
+++ b/Engine/Source/Runtime/Engine/Private/Texture2D.cpp
@@ -596,19 +596,22 @@ bool UTexture2D::UpdateStreamingStatus( bool bWaitForMipFading /*= false*/ )
 		PendingUpdate = nullptr;
 
 #if WITH_EDITOR
-		// When all the requested mips are streamed in, generate an empty property changed event, to force the
-		// ResourceSize asset registry tag to be recalculated.
-		FPropertyChangedEvent EmptyPropertyChangedEvent(nullptr);
-		FCoreUObjectDelegates::OnObjectPropertyChanged.Broadcast(this, EmptyPropertyChangedEvent);
-
-		// We can't load the source art from a bulk data object if the texture itself is pending kill because the linker will have been detached.
-		// In this case we don't rebuild the data and instead let the streaming request be cancelled. This will let the garbage collector finish
-		// destroying the object.
-		if (bRebuildPlatformData)
+		if (GIsEditor)
 		{
-			ForceRebuildPlatformData();
-			// @TODO this can not be called from this callstack since the entry needs to be removed completely from the streamer.
-			// UpdateResource();
+			// When all the requested mips are streamed in, generate an empty property changed event, to force the
+			// ResourceSize asset registry tag to be recalculated.
+			FPropertyChangedEvent EmptyPropertyChangedEvent(nullptr);
+			FCoreUObjectDelegates::OnObjectPropertyChanged.Broadcast(this, EmptyPropertyChangedEvent);
+
+			// We can't load the source art from a bulk data object if the texture itself is pending kill because the linker will have been detached.
+			// In this case we don't rebuild the data and instead let the streaming request be cancelled. This will let the garbage collector finish
+			// destroying the object.
+			if (bRebuildPlatformData)
+			{
+				ForceRebuildPlatformData();
+				// @TODO this can not be called from this callstack since the entry needs to be removed completely from the streamer.
+				// UpdateResource();
+			}
 		}
 #endif
 	}
