We set default setting to tick on init - as it's more user friendly option

diff --git a/Engine/Source/Runtime/Engine/Private/Animation/AnimationSettings.cpp b/Engine/Source/Runtime/Engine/Private/Animation/AnimationSettings.cpp
--- a/Engine/Source/Runtime/Engine/Private/Animation/AnimationSettings.cpp
+++ b/Engine/Source/Runtime/Engine/Private/Animation/AnimationSettings.cpp
@@ -21,6 +21,7 @@ UAnimationSettings::UAnimationSettings(const FObjectInitializer& ObjectInitializ
 	, bTryLinearKeyRemovalCompression(true)
 	, bTryIntervalKeyRemoval(true)
 	, bEnablePerformanceLog(false)
+	, bTickAnimationOnSkeletalMeshInit(true)
 {
 	SectionName = TEXT("Animation");
 
