Fixed texture ending up in the unkown ref heuristic when changing material quality levels.

diff --git a/Engine/Source/Runtime/Engine/Private/Streaming/DynamicTextureInstanceManager.h b/Engine/Source/Runtime/Engine/Private/Streaming/DynamicTextureInstanceManager.h
--- a/Engine/Source/Runtime/Engine/Private/Streaming/DynamicTextureInstanceManager.h
+++ b/Engine/Source/Runtime/Engine/Private/Streaming/DynamicTextureInstanceManager.h
@@ -27,10 +27,8 @@ public:
 
 	void IncrementalUpdate(FRemovedTextureArray& RemovedTextures, float Percentage);
 
-#if !UE_BUILD_SHIPPING
 	// Get all (non removed) components refered by the manager. Debug only.
 	void GetReferencedComponents(TArray<const UPrimitiveComponent*>& Components) { StateSync.SyncAndGetState()->GetReferencedComponents(Components); }
-#endif
 
 	/** Remove all pending components that are marked for delete. This prevents searching in the pending list for each entry. */
 	void OnPreGarbageCollect(FRemovedTextureArray& RemovedTextures);
diff --git a/Engine/Source/Runtime/Engine/Private/Streaming/StreamingManagerTexture.cpp b/Engine/Source/Runtime/Engine/Private/Streaming/StreamingManagerTexture.cpp
--- a/Engine/Source/Runtime/Engine/Private/Streaming/StreamingManagerTexture.cpp
+++ b/Engine/Source/Runtime/Engine/Private/Streaming/StreamingManagerTexture.cpp
@@ -352,16 +352,30 @@ void FStreamingManagerTexture::ConditionalUpdateStaticData()
 		for (FStreamingTexture& StreamingTexture : StreamingTextures)
 		{
 			StreamingTexture.UpdateStaticData(Settings);
+
+			// When the material quality changes, some textures could stop being used.
+			// Refreshing their removed timestamp ensures not texture ends up in the unkwown 
+			// ref heuristic (which would force load them).
+			if (PreviousSettings.MaterialQualityLevel != Settings.MaterialQualityLevel)
+			{
+				StreamingTexture.InstanceRemovedTimestamp = FApp::GetCurrentTime();
+			}
 		}
 		STAT(GatheredStats.SetupAsyncTaskCycles -= (int32)FPlatformTime::Cycles();)
 
 #if !UE_BUILD_SHIPPING
-		// Don't put anything here that could normally change in game.
-		// This is used to test regression / improvements, and insertion perfs.
-		if (PreviousSettings.bUseMaterialData != Settings.bUseMaterialData ||
+		// Those debug settings are config that are not expected to change in-game.
+		const bool bDebugSettingsChanged = 
+			PreviousSettings.bUseMaterialData != Settings.bUseMaterialData ||
 			PreviousSettings.bUseNewMetrics != Settings.bUseNewMetrics ||
 			PreviousSettings.bUsePerTextureBias != Settings.bUsePerTextureBias || 
-			PreviousSettings.MaxTextureUVDensity != Settings.MaxTextureUVDensity)
+			PreviousSettings.MaxTextureUVDensity != Settings.MaxTextureUVDensity;
+#else
+		const bool bDebugSettingsChanged = false;
+#endif
+
+		// If the material quality changes, everything needs to be updated.
+		if (bDebugSettingsChanged || PreviousSettings.MaterialQualityLevel != Settings.MaterialQualityLevel)
 		{
 			TArray<ULevel*, TInlineAllocator<32> > Levels;
 
@@ -386,7 +400,6 @@ void FStreamingManagerTexture::ConditionalUpdateStaticData()
 				NotifyPrimitiveUpdated_Concurrent(Primitive);
 			}
 		}
-#endif
 
 		// Update the cache variables.
 		PreviousLightmapStreamingFactor = GLightmapStreamingFactor;
diff --git a/Engine/Source/Runtime/Engine/Private/Streaming/StreamingTexture.cpp b/Engine/Source/Runtime/Engine/Private/Streaming/StreamingTexture.cpp
--- a/Engine/Source/Runtime/Engine/Private/Streaming/StreamingTexture.cpp
+++ b/Engine/Source/Runtime/Engine/Private/Streaming/StreamingTexture.cpp
@@ -20,7 +20,7 @@ FStreamingTexture::FStreamingTexture(UTexture2D* InTexture, const int32 NumStrea
 		UpdateOptionalMipsState_Async();
 	}
 
-	InstanceRemovedTimestamp = -FLT_MAX;
+	InstanceRemovedTimestamp = FApp::GetCurrentTime();
 	DynamicBoostFactor = 1.f;
 
 	bHasUpdatePending = InTexture && InTexture->bHasStreamingUpdatePending;
diff --git a/Engine/Source/Runtime/Engine/Private/Streaming/TextureStreamingHelpers.cpp b/Engine/Source/Runtime/Engine/Private/Streaming/TextureStreamingHelpers.cpp
--- a/Engine/Source/Runtime/Engine/Private/Streaming/TextureStreamingHelpers.cpp
+++ b/Engine/Source/Runtime/Engine/Private/Streaming/TextureStreamingHelpers.cpp
@@ -5,6 +5,7 @@ TextureStreamingHelpers.cpp: Definitions of classes used for texture streaming.
 =============================================================================*/
 
 #include "Streaming/TextureStreamingHelpers.h"
+#include "UnrealEngine.h"
 #include "Engine/Texture2D.h"
 #include "GenericPlatform/GenericPlatformMemoryPoolStats.h"
 
@@ -252,6 +253,8 @@ void FTextureStreamingSettings::Update()
 	bUseMaterialData = bUseNewMetrics && CVarStreamingUseMaterialData.GetValueOnAnyThread() != 0;
 	HiddenPrimitiveScale = bUseNewMetrics ? CVarStreamingHiddenPrimitiveScale.GetValueOnAnyThread() : 1.f;
 
+	MaterialQualityLevel = (int32)GetCachedScalabilityCVars().MaterialQualityLevel;
+
 	if (MinMipForSplitRequest <= 0)
 	{
 		MinMipForSplitRequest = MAX_TEXTURE_MIP_COUNT + 1;
diff --git a/Engine/Source/Runtime/Engine/Private/Streaming/TextureStreamingHelpers.h b/Engine/Source/Runtime/Engine/Private/Streaming/TextureStreamingHelpers.h
--- a/Engine/Source/Runtime/Engine/Private/Streaming/TextureStreamingHelpers.h
+++ b/Engine/Source/Runtime/Engine/Private/Streaming/TextureStreamingHelpers.h
@@ -85,6 +85,7 @@ struct FTextureStreamingSettings
 	int32 MinMipForSplitRequest;
 	float MinLevelTextureScreenSize;
 	float MaxTextureUVDensity;
+	int32 MaterialQualityLevel;
 
 protected:
 
