Fix bitpacking regression in FBodyInstance
#rnx

diff --git a/Engine/Source/Runtime/Engine/Classes/PhysicsEngine/BodyInstance.h b/Engine/Source/Runtime/Engine/Classes/PhysicsEngine/BodyInstance.h
--- a/Engine/Source/Runtime/Engine/Classes/PhysicsEngine/BodyInstance.h
+++ b/Engine/Source/Runtime/Engine/Classes/PhysicsEngine/BodyInstance.h
@@ -327,9 +327,9 @@ public:
 	 * Some games need to dynamically modify collision per skeletal body. These provide game code a way to 
 	 * do that, until we're able to refactor how skeletal bodies work.
 	 */
-	uint32 bHACK_DisableCollisionResponse : 1;
+	uint8 bHACK_DisableCollisionResponse : 1;
 	/* By default, an owning skel mesh component will override the body's collision filter. This will disable that behavior. */
-	uint32 bHACK_DisableSkelComponentFilterOverriding : 1;
+	uint8 bHACK_DisableSkelComponentFilterOverriding : 1;
 
 protected:
 
