Keep source raw animation data in sync with raw when cropping.

diff --git a/Engine/Source/Runtime/Engine/Private/Animation/AnimSequence.cpp b/Engine/Source/Runtime/Engine/Private/Animation/AnimSequence.cpp
--- a/Engine/Source/Runtime/Engine/Private/Animation/AnimSequence.cpp
+++ b/Engine/Source/Runtime/Engine/Private/Animation/AnimSequence.cpp
@@ -1995,6 +1995,11 @@ bool UAnimSequence::CropRawAnimData( float CurrentTime, bool bFromStart )
 		CropRawTrack(RawAnimationData[i], StartKey, NumKeys, TotalNumOfFrames);
 	}
 
+	for (int32 i = 0; i < SourceRawAnimationData.Num(); ++i)
+	{
+		CropRawTrack(SourceRawAnimationData[i], StartKey, NumKeys, TotalNumOfFrames);
+	}
+
 	// Double check that everything is fine
 	for(int32 i=0; i<RawAnimationData.Num(); i++)
 	{
