Prevent Master pose component chain

diff --git a/Engine/Source/Runtime/Engine/Private/Components/SkinnedMeshComponent.cpp b/Engine/Source/Runtime/Engine/Private/Components/SkinnedMeshComponent.cpp
--- a/Engine/Source/Runtime/Engine/Private/Components/SkinnedMeshComponent.cpp
+++ b/Engine/Source/Runtime/Engine/Private/Components/SkinnedMeshComponent.cpp
@@ -1344,16 +1344,32 @@ void USkinnedMeshComponent::SetMasterPoseComponent(class USkinnedMeshComponent*
 	// now add to slave components list, 
 	if (MasterPoseComponent.IsValid())
 	{
+		// verify if my current master pose is valid
+		// we can't have chain of master poses, so 
+		// we'll find the root master pose component
+		USkinnedMeshComponent* CurrentMasterPose = MasterPoseComponent.Get();
+		while (CurrentMasterPose->MasterPoseComponent.IsValid())
+		{
+			MasterPoseComponent = CurrentMasterPose->MasterPoseComponent;
+			CurrentMasterPose = MasterPoseComponent.Get();
+		}
+
+		// ensure if master is not same as input, which means it has changed. 
+		USkinnedMeshComponent* RawMasterPoseComponent = MasterPoseComponent.Get();
+		ensureAlwaysMsgf(RawMasterPoseComponent == NewMasterBoneComponent,
+			TEXT("MasterPoseComponent chain is detected (%s). We re-route to top-most MasterPoseComponent (%s)"),
+				*GetNameSafe(RawMasterPoseComponent), *GetNameSafe(NewMasterBoneComponent));
+
 		bool bAddNew = true;
 		// make sure no empty element is there, this is weak obj ptr, so it will go away unless there is 
 		// other reference, this is intentional as master to slave reference is weak
-		for (auto Iter = MasterPoseComponent->SlavePoseComponents.CreateIterator(); Iter; ++Iter)
+		for (auto Iter = RawMasterPoseComponent->SlavePoseComponents.CreateIterator(); Iter; ++Iter)
 		{
 			TWeakObjectPtr<USkinnedMeshComponent> Comp = (*Iter);
 			if (Comp.IsValid() == false)
 			{
 				// remove
-				MasterPoseComponent->SlavePoseComponents.RemoveAt(Iter.GetIndex());
+				RawMasterPoseComponent->SlavePoseComponents.RemoveAt(Iter.GetIndex());
 				--Iter;
 			}
 			// if it has same as me, ignore to add
@@ -1365,11 +1381,11 @@ void USkinnedMeshComponent::SetMasterPoseComponent(class USkinnedMeshComponent*
 
 		if (bAddNew)
 		{
-			MasterPoseComponent->AddSlavePoseComponent(this);
+			RawMasterPoseComponent->AddSlavePoseComponent(this);
 		}
 
 		// set up tick dependency between master & slave components
-		PrimaryComponentTick.AddPrerequisite(MasterPoseComponent.Get(), MasterPoseComponent->PrimaryComponentTick);
+		PrimaryComponentTick.AddPrerequisite(RawMasterPoseComponent, RawMasterPoseComponent->PrimaryComponentTick);
 	}
 
 	if ((OldMasterPoseComponent != nullptr) && (OldMasterPoseComponent != NewMasterBoneComponent))
