Merging CL# 4121591 from Release-4.20 to Release-Staging-4.20
#rnx

diff --git a/Engine/Source/Runtime/Engine/Private/HierarchicalInstancedStaticMesh.cpp b/Engine/Source/Runtime/Engine/Private/HierarchicalInstancedStaticMesh.cpp
--- a/Engine/Source/Runtime/Engine/Private/HierarchicalInstancedStaticMesh.cpp
+++ b/Engine/Source/Runtime/Engine/Private/HierarchicalInstancedStaticMesh.cpp
@@ -2307,6 +2307,13 @@ void UHierarchicalInstancedStaticMeshComponent::BuildTree()
 
 	QUICK_SCOPE_CYCLE_COUNTER(STAT_UHierarchicalInstancedStaticMeshComponent_BuildTree);
 
+	// upload instance edits to GPU, before validing if the mesh is valid, as it's possible that PerInstanceSMData.Num() == 0, so we have to hide everything before doing the build
+	if (GIsEditor && InstanceUpdateCmdBuffer.NumInlineCommands() > 0)
+	{
+		// this is allowed only in editor, at runtime upload will happen when buffer is built from component data
+		PerInstanceRenderData->UpdateFromCommandBuffer(InstanceUpdateCmdBuffer);
+	}
+
 	// all pending edits will be updated
 	InstanceUpdateCmdBuffer.Reset();
 
