ULevelStreaming::bShouldBlockOnLoad now properly blocks for the entire loading process not just the file load
Added bShouldBlockOnUnload  to ULevelStreaming a UGameplayStatics::UnloadStreamLevel

diff --git a/Engine/Source/Runtime/Engine/Classes/Engine/LevelStreaming.h b/Engine/Source/Runtime/Engine/Classes/Engine/LevelStreaming.h
--- a/Engine/Source/Runtime/Engine/Classes/Engine/LevelStreaming.h
+++ b/Engine/Source/Runtime/Engine/Classes/Engine/LevelStreaming.h
@@ -20,13 +20,13 @@ class FStreamLevelAction : public FPendingLatentAction
 public:
 	bool			bLoading;
 	bool			bMakeVisibleAfterLoad;
-	bool			bShouldBlockOnLoad;
+	bool			bShouldBlock;
 	ULevelStreaming* Level;
 	FName			LevelName;
 
 	FLatentActionInfo LatentInfo;
 
-	FStreamLevelAction(bool bIsLoading, const FName& InLevelName, bool bIsMakeVisibleAfterLoad, bool bIsShouldBlockOnLoad, const FLatentActionInfo& InLatentInfo, UWorld* World);
+	FStreamLevelAction(bool bIsLoading, const FName& InLevelName, bool bIsMakeVisibleAfterLoad, bool bShouldBlock, const FLatentActionInfo& InLatentInfo, UWorld* World);
 
 	/**
 	 * Given a level name, returns level name that will work with Play on Editor or Play on Console
@@ -178,6 +178,10 @@ public:
 	UPROPERTY(Category=LevelStreaming, BlueprintReadWrite)
 	uint8 bShouldBlockOnLoad:1;
 
+	/** Whether we want to force a blocking unload																				*/
+	UPROPERTY(Category=LevelStreaming, BlueprintReadWrite)
+	uint8 bShouldBlockOnUnload:1;
+
 	/** 
 	 *  Whether this level streaming object should be ignored by world composition distance streaming, 
 	 *  so streaming state can be controlled by other systems (ex: in blueprints)
diff --git a/Engine/Source/Runtime/Engine/Classes/Engine/World.h b/Engine/Source/Runtime/Engine/Classes/Engine/World.h
--- a/Engine/Source/Runtime/Engine/Classes/Engine/World.h
+++ b/Engine/Source/Runtime/Engine/Classes/Engine/World.h
@@ -2292,7 +2292,7 @@ public:
 	 * @param Level				Level object we should add
 	 * @param LevelTransform	Transformation to apply to each actor in the level
 	 */
-	void AddToWorld( ULevel* Level, const FTransform& LevelTransform = FTransform::Identity );
+	void AddToWorld( ULevel* Level, const FTransform& LevelTransform = FTransform::Identity, bool bConsiderTimeLimit = true );
 
 	/** 
 	 * Dissociates the passed in level from the world. The removal is blocking.
diff --git a/Engine/Source/Runtime/Engine/Classes/Kismet/GameplayStatics.h b/Engine/Source/Runtime/Engine/Classes/Kismet/GameplayStatics.h
--- a/Engine/Source/Runtime/Engine/Classes/Kismet/GameplayStatics.h
+++ b/Engine/Source/Runtime/Engine/Classes/Kismet/GameplayStatics.h
@@ -160,7 +160,7 @@ class ENGINE_API UGameplayStatics : public UBlueprintFunctionLibrary
 
 	/** Unload a streamed in level */
 	UFUNCTION(BlueprintCallable, meta=(WorldContext="WorldContextObject", Latent = "", LatentInfo = "LatentInfo"), Category="Game")
-	static void UnloadStreamLevel(const UObject* WorldContextObject, FName LevelName, FLatentActionInfo LatentInfo);
+	static void UnloadStreamLevel(const UObject* WorldContextObject, FName LevelName, FLatentActionInfo LatentInfo, bool bShouldBlockOnUnload);
 	
 	/** Returns level streaming object with specified level package name */
 	UFUNCTION(BlueprintPure, meta=(WorldContext="WorldContextObject"), Category="Game")
diff --git a/Engine/Source/Runtime/Engine/Private/GameplayStatics.cpp b/Engine/Source/Runtime/Engine/Private/GameplayStatics.cpp
--- a/Engine/Source/Runtime/Engine/Private/GameplayStatics.cpp
+++ b/Engine/Source/Runtime/Engine/Private/GameplayStatics.cpp
@@ -591,14 +591,14 @@ void UGameplayStatics::LoadStreamLevel(const UObject* WorldContextObject, FName
 	}
 }
 
-void UGameplayStatics::UnloadStreamLevel(const UObject* WorldContextObject, FName LevelName,FLatentActionInfo LatentInfo)
+void UGameplayStatics::UnloadStreamLevel(const UObject* WorldContextObject, FName LevelName,FLatentActionInfo LatentInfo,bool bShouldBlockOnUnload)
 {
 	if (UWorld* World = GEngine->GetWorldFromContextObject(WorldContextObject, EGetWorldErrorMode::LogAndReturnNull))
 	{
 		FLatentActionManager& LatentManager = World->GetLatentActionManager();
 		if (LatentManager.FindExistingAction<FStreamLevelAction>(LatentInfo.CallbackTarget, LatentInfo.UUID) == nullptr)
 		{
-			FStreamLevelAction* NewAction = new FStreamLevelAction(false, LevelName, false, false, LatentInfo, World );
+			FStreamLevelAction* NewAction = new FStreamLevelAction(false, LevelName, false, bShouldBlockOnUnload, LatentInfo, World );
 			LatentManager.AddNewAction(LatentInfo.CallbackTarget, LatentInfo.UUID, NewAction );
 		}
 	}
diff --git a/Engine/Source/Runtime/Engine/Private/LevelStreaming.cpp b/Engine/Source/Runtime/Engine/Private/LevelStreaming.cpp
--- a/Engine/Source/Runtime/Engine/Private/LevelStreaming.cpp
+++ b/Engine/Source/Runtime/Engine/Private/LevelStreaming.cpp
@@ -76,10 +76,10 @@ static void NetDriverRenameStreamingLevelPackageForPIE(const UWorld* World, FNam
 	}
 }
 
-FStreamLevelAction::FStreamLevelAction(bool bIsLoading, const FName& InLevelName, bool bIsMakeVisibleAfterLoad, bool bIsShouldBlockOnLoad, const FLatentActionInfo& InLatentInfo, UWorld* World)
+FStreamLevelAction::FStreamLevelAction(bool bIsLoading, const FName& InLevelName, bool bIsMakeVisibleAfterLoad, bool bInShouldBlock, const FLatentActionInfo& InLatentInfo, UWorld* World)
 	: bLoading(bIsLoading)
 	, bMakeVisibleAfterLoad(bIsMakeVisibleAfterLoad)
-	, bShouldBlockOnLoad(bIsShouldBlockOnLoad)
+	, bShouldBlock(bInShouldBlock)
 	, LevelName(InLevelName)
 	, LatentInfo(InLatentInfo)
 {
@@ -171,12 +171,12 @@ void FStreamLevelAction::ActivateLevel( ULevelStreaming* LevelStreamingObject )
 	if (LevelStreamingObject)
 	{
 		// Loading.
-		if( bLoading )
+		if (bLoading)
 		{
 			UE_LOG(LogStreaming, Log, TEXT("Streaming in level %s (%s)..."),*LevelStreamingObject->GetName(),*LevelStreamingObject->GetWorldAssetPackageName());
 			LevelStreamingObject->SetShouldBeLoaded(true);
 			LevelStreamingObject->SetShouldBeVisible(LevelStreamingObject->GetShouldBeVisibleFlag()	|| bMakeVisibleAfterLoad);
-			LevelStreamingObject->bShouldBlockOnLoad	= bShouldBlockOnLoad;
+			LevelStreamingObject->bShouldBlockOnLoad = bShouldBlock;
 		}
 		// Unloading.
 		else 
@@ -184,13 +184,14 @@ void FStreamLevelAction::ActivateLevel( ULevelStreaming* LevelStreamingObject )
 			UE_LOG(LogStreaming, Log, TEXT("Streaming out level %s (%s)..."),*LevelStreamingObject->GetName(),*LevelStreamingObject->GetWorldAssetPackageName());
 			LevelStreamingObject->SetShouldBeLoaded(false);
 			LevelStreamingObject->SetShouldBeVisible(false);
+			LevelStreamingObject->bShouldBlockOnUnload = bShouldBlock;
 		}
 
 		// If we have a valid world
 		if (UWorld* LevelWorld = LevelStreamingObject->GetWorld())
 		{
 			// Notify players of the change
-			for( FConstPlayerControllerIterator Iterator = LevelWorld->GetPlayerControllerIterator(); Iterator; ++Iterator )
+			for (FConstPlayerControllerIterator Iterator = LevelWorld->GetPlayerControllerIterator(); Iterator; ++Iterator)
 			{
 				APlayerController* PlayerController = Iterator->Get();
 
@@ -201,17 +202,14 @@ void FStreamLevelAction::ActivateLevel( ULevelStreaming* LevelStreamingObject )
 					*LevelStreamingObject->GetWorldAssetPackageName(), 
 					bShouldBeLoaded, 
 					bShouldBeVisible, 
-					LevelStreamingObject->bShouldBlockOnLoad );
-
-
+					bShouldBlock);
 
 				PlayerController->LevelStreamingStatusChanged( 
 					LevelStreamingObject, 
 					bShouldBeLoaded, 
 					bShouldBeVisible,
-					LevelStreamingObject->bShouldBlockOnLoad, 
+					bShouldBlock, 
 					INDEX_NONE);
-
 			}
 		}
 	}
@@ -563,7 +561,7 @@ void ULevelStreaming::UpdateStreamingState(bool& bOutUpdateAgain, bool& bOutRede
 	case ECurrentState::MakingVisible:
 		if (ensure(LoadedLevel))
 		{
-			World->AddToWorld(LoadedLevel, LevelTransform);
+			World->AddToWorld(LoadedLevel, LevelTransform, !bShouldBlockOnLoad);
 
 			if (LoadedLevel->bIsVisible)
 			{
@@ -588,7 +586,7 @@ void ULevelStreaming::UpdateStreamingState(bool& bOutUpdateAgain, bool& bOutRede
 		if (ensure(LoadedLevel))
 		{
 			// Hide loaded level, incrementally if necessary
-			World->RemoveFromWorld(LoadedLevel, World->IsGameWorld());
+			World->RemoveFromWorld(LoadedLevel, !bShouldBlockOnUnload && World->IsGameWorld());
 
 			// Inform the scene once we have finished making the level invisible
 			if (!LoadedLevel->bIsVisible)
diff --git a/Engine/Source/Runtime/Engine/Private/World.cpp b/Engine/Source/Runtime/Engine/Private/World.cpp
--- a/Engine/Source/Runtime/Engine/Private/World.cpp
+++ b/Engine/Source/Runtime/Engine/Private/World.cpp
@@ -2045,7 +2045,7 @@ private:
 
 #endif // PERF_TRACK_DETAILED_ASYNC_STATS
 
-void UWorld::AddToWorld( ULevel* Level, const FTransform& LevelTransform )
+void UWorld::AddToWorld( ULevel* Level, const FTransform& LevelTransform, bool bConsiderTimeLimit )
 {
 	SCOPE_CYCLE_COUNTER(STAT_AddToWorldTime);
 
@@ -2090,7 +2090,7 @@ void UWorld::AddToWorld( ULevel* Level, const FTransform& LevelTransform )
 	}
 
 	// Don't consider the time limit if the match hasn't started as we need to ensure that the levels are fully loaded
-	const bool bConsiderTimeLimit = bMatchStarted && IsGameWorld();
+	bConsiderTimeLimit &= bMatchStarted && IsGameWorld();
 	double TimeLimit = GLevelStreamingActorsUpdateTimeLimit;
 
 	if (bConsiderTimeLimit)
