Duplicate [CL 4007400, 4008252] - Implementing an alternative to backing out 3790473. This will re-add a bug where renamed material parameters that happen to reside within a material function in a material's expression tree will may have their values lost of replaced incorrectly. The reasoning here is to address an issue where duplicated material functions containing parameters will have duplicated/matching GUIDs causing them to pick up values from the wrong parameter thinking it's a rename. A proper fix for these issues goes beyond the scope of a short term fix so for now we're reverting behavior to match previous releases.

diff --git a/Engine/Source/Runtime/Engine/Classes/Materials/MaterialInstance.h b/Engine/Source/Runtime/Engine/Classes/Materials/MaterialInstance.h
--- a/Engine/Source/Runtime/Engine/Classes/Materials/MaterialInstance.h
+++ b/Engine/Source/Runtime/Engine/Classes/Materials/MaterialInstance.h
@@ -628,6 +628,73 @@ protected:
 	friend class FMaterialUpdateContext;
 };
 
+/** Workaround - Similar to base call but evaluates all expressions found, not just the first */
+template<typename ExpressionType>
+void FindClosestExpressionByGUIDRecursive(const FName& InName, const FGuid& InGUID, const TArray<UMaterialExpression*>& InMaterialExpression, ExpressionType*& OutExpression)
+{
+	for (int32 ExpressionIndex = 0; ExpressionIndex < InMaterialExpression.Num(); ExpressionIndex++)
+	{
+		UMaterialExpression* ExpressionPtr = InMaterialExpression[ExpressionIndex];
+		UMaterialExpressionMaterialFunctionCall* MaterialFunctionCall = Cast<UMaterialExpressionMaterialFunctionCall>(ExpressionPtr);
+		UMaterialExpressionMaterialAttributeLayers* MaterialLayers = Cast<UMaterialExpressionMaterialAttributeLayers>(ExpressionPtr);
+
+		if (ExpressionPtr && ExpressionPtr->GetParameterExpressionId() == InGUID)
+		{
+			check(ExpressionPtr->bIsParameterExpression);
+			if (ExpressionType* ParamExpression = Cast<ExpressionType>(ExpressionPtr))
+			{
+				// UE-57086, workaround - To deal with duplicated parameters with matching GUIDs we walk
+				// through every parameter rather than taking the first. Either we return the first matching GUID
+				// we encounter (as before), or if we find another with the same name that can take precedence.
+				// Only taking the first parameter means we can incorrectly treat the parameter as a rename and
+				// lose/move data when we encounter an illegal GUID duplicate.
+				// Properly fixing duplicate GUIDs is beyond the scope of a hotfix, see UE-47863 for more info.
+				// NOTE: The case where a parameter in a function is renamed but another function in the material
+				// contains a duplicate GUID is still broken and may lose the data. This still leaves us in a
+				// more consistent state than 4.18 and should minimize the impact to a rarer occurrence.
+				if (!OutExpression || InName == ParamExpression->ParameterName)
+				{
+					OutExpression = ParamExpression;
+				}
+			}
+		}
+		else if (MaterialFunctionCall && MaterialFunctionCall->MaterialFunction)
+		{
+			if (const TArray<UMaterialExpression*>* FunctionExpressions = MaterialFunctionCall->MaterialFunction->GetFunctionExpressions())
+			{
+				FindClosestExpressionByGUIDRecursive<ExpressionType>(InName, InGUID, *FunctionExpressions, OutExpression);
+			}
+		}
+		else if (MaterialLayers)
+		{
+			const TArray<UMaterialFunctionInterface*>& Layers = MaterialLayers->GetLayers();
+			const TArray<UMaterialFunctionInterface*>& Blends = MaterialLayers->GetBlends();
+
+			for (const auto* Layer : Layers)
+			{
+				if (Layer)
+				{
+					if (const TArray<UMaterialExpression*>* FunctionExpressions = Layer->GetFunctionExpressions())
+					{
+						FindClosestExpressionByGUIDRecursive<ExpressionType>(InName, InGUID, *FunctionExpressions, OutExpression);
+					}
+				}
+			}
+
+			for (const auto* Blend : Blends)
+			{
+				if (Blend)
+				{
+					if (const TArray<UMaterialExpression*>* FunctionExpressions = Blend->GetFunctionExpressions())
+					{
+						FindClosestExpressionByGUIDRecursive<ExpressionType>(InName, InGUID, *FunctionExpressions, OutExpression);
+					}
+				}
+			}
+		}
+	}
+}
+
 /**
  * This function takes a array of parameter structs and attempts to establish a reference to the expression object each parameter represents.
  * If a reference exists, the function checks to see if the parameter has been renamed.
@@ -652,7 +719,8 @@ bool UpdateParameterSet(TArray<ParameterType>& Parameters, UMaterial* ParentMate
 
 		if (Parameter.ExpressionGUID.IsValid())
 		{
-			ExpressionType* Expression = ParentMaterial->FindExpressionByGUID<ExpressionType>(Parameter.ExpressionGUID);
+			ExpressionType* Expression = nullptr;
+			FindClosestExpressionByGUIDRecursive<ExpressionType>(Parameter.ParameterInfo.Name, Parameter.ExpressionGUID, ParentMaterial->Expressions, Expression);
 
 			// Check to see if the parameter name was changed.
 			if (Expression)
