[NO NON-PII COMMIT MESSAGE]

diff --git a/Engine/Source/Runtime/Engine/Private/SoundWave.cpp b/Engine/Source/Runtime/Engine/Private/SoundWave.cpp
--- a/Engine/Source/Runtime/Engine/Private/SoundWave.cpp
+++ b/Engine/Source/Runtime/Engine/Private/SoundWave.cpp
@@ -972,15 +972,12 @@ void USoundWave::Parse( FAudioDevice* AudioDevice, const UPTRINT NodeWaveInstanc
 			WaveInstance->SetVolume(WaveInstanceVolume * 0.5f);
 		}
 
-		if (!bUseSoundClassDefaultReverb)
-		{
-			// Copy reverb send settings
-			WaveInstance->ReverbSendMethod = ParseParams.ReverbSendMethod;
-			WaveInstance->ManualReverbSendLevel = ParseParams.ManualReverbSendLevel;
-			WaveInstance->CustomRevebSendCurve = ParseParams.CustomReverbSendCurve;
-			WaveInstance->ReverbSendLevelRange = ParseParams.ReverbSendLevelRange;
-			WaveInstance->ReverbSendLevelDistanceRange = ParseParams.ReverbSendLevelDistanceRange;
-		}
+		// Copy reverb send settings
+		WaveInstance->ReverbSendMethod = ParseParams.ReverbSendMethod;
+		WaveInstance->ManualReverbSendLevel = ParseParams.ManualReverbSendLevel;
+		WaveInstance->CustomRevebSendCurve = ParseParams.CustomReverbSendCurve;
+		WaveInstance->ReverbSendLevelRange = ParseParams.ReverbSendLevelRange;
+		WaveInstance->ReverbSendLevelDistanceRange = ParseParams.ReverbSendLevelDistanceRange;
 
 		// Get the envelope follower settings
 		WaveInstance->EnvelopeFollowerAttackTime = ParseParams.EnvelopeFollowerAttackTime;
