UE-57094 - Error: Loading non-streamed mips from an external bulk file spam running app on Android device
r.TextureStreaming is set to 0 whenever there's no hardware support for streaming.

diff --git a/Engine/Source/Runtime/Engine/Private/ContentStreaming.cpp b/Engine/Source/Runtime/Engine/Private/ContentStreaming.cpp
--- a/Engine/Source/Runtime/Engine/Private/ContentStreaming.cpp
+++ b/Engine/Source/Runtime/Engine/Private/ContentStreaming.cpp
@@ -1050,6 +1050,9 @@ void FStreamingManagerCollection::AddOrRemoveTextureStreamingManagerIfNeeded(boo
 	if( !GRHISupportsTextureStreaming || IsRunningDedicatedServer() )
 	{
 		bUseTextureStreaming = false;
+
+		// some code relies on r.TextureStreaming so we're going to disable it here to reflect the hardware capabilities and system needs
+		CVarSetTextureStreaming.AsVariable()->Set(0, ECVF_SetByCode);
 	}
 #endif
 
