Remove code taken out in CL# 4109849 and CL# 4109923 which was reintroduced in CL# 4112852 as a result of a bad merge caused by CL# 4109110 not being a merge, but instead being an edit.
#rnx

diff --git a/Engine/Source/Runtime/Engine/Private/HierarchicalInstancedStaticMesh.cpp b/Engine/Source/Runtime/Engine/Private/HierarchicalInstancedStaticMesh.cpp
--- a/Engine/Source/Runtime/Engine/Private/HierarchicalInstancedStaticMesh.cpp
+++ b/Engine/Source/Runtime/Engine/Private/HierarchicalInstancedStaticMesh.cpp
@@ -1940,8 +1940,6 @@ void UHierarchicalInstancedStaticMeshComponent::Serialize(FArchive& Ar)
 		// Skip the serialized tree, we will regenerate it correctly to contains the new data
 		TArray<FClusterNode_DEPRECATED> ClusterTree_DEPRECATED;
 		ClusterTree_DEPRECATED.BulkSerialize(Ar);
-
-		BuildTreeIfOutdated(false, true);
 	}
 	else
 	{
@@ -3102,7 +3100,4 @@ static FAutoConsoleCommand RebuildFoliageTreesCmd(
 	TEXT("foliage.RebuildFoliageTrees"),
 	TEXT("Rebuild the trees for non-grass foliage."),
 	FConsoleCommandWithArgsDelegate::CreateStatic(&RebuildFoliageTrees)
-	);
-
-
-PRAGMA_ENABLE_OPTIMIZATION
\ No newline at end of file
+	);
\ No newline at end of file
