Possible fix for marker sync assert

diff --git a/Engine/Source/Runtime/Engine/Private/Animation/AnimSequence.cpp b/Engine/Source/Runtime/Engine/Private/Animation/AnimSequence.cpp
--- a/Engine/Source/Runtime/Engine/Private/Animation/AnimSequence.cpp
+++ b/Engine/Source/Runtime/Engine/Private/Animation/AnimSequence.cpp
@@ -5317,18 +5317,20 @@ void UAnimSequence::GetMarkerIndicesForPosition(const FMarkerSyncAnimPosition& S
 	// If we're not looping, assume we're playing a transition and we need to stay where we are.
 	if (!bLooping)
 	{
-		OutPrevMarker.MarkerIndex = -1;
-		OutNextMarker.MarkerIndex = -1;
+		OutPrevMarker.MarkerIndex = INDEX_NONE;
+		OutNextMarker.MarkerIndex = INDEX_NONE;
 
 		for (int32 Idx = 0; Idx<AuthoredSyncMarkers.Num(); Idx++)
 		{
-			const float MarkerTime = AuthoredSyncMarkers[Idx].Time;
-			if (OutCurrentTime > MarkerTime)
+			const FAnimSyncMarker& SyncMarker = AuthoredSyncMarkers[Idx];
+			const float MarkerTime = SyncMarker.Time;
+
+			if (OutCurrentTime > MarkerTime && SyncMarker.MarkerName == SyncPosition.PreviousMarkerName)
 			{
 				OutPrevMarker.MarkerIndex = Idx;
 				OutPrevMarker.TimeToMarker = MarkerTime - OutCurrentTime;
 			}
-			else if (OutCurrentTime < MarkerTime)
+			else if (OutCurrentTime < MarkerTime && SyncMarker.MarkerName == SyncPosition.NextMarkerName)
 			{
 				OutNextMarker.MarkerIndex = Idx;
 				OutNextMarker.TimeToMarker = MarkerTime - OutCurrentTime;
