Test for incorrect time in possible source of marker sync issue

diff --git a/Engine/Source/Runtime/Engine/Private/Animation/AnimSequence.cpp b/Engine/Source/Runtime/Engine/Private/Animation/AnimSequence.cpp
--- a/Engine/Source/Runtime/Engine/Private/Animation/AnimSequence.cpp
+++ b/Engine/Source/Runtime/Engine/Private/Animation/AnimSequence.cpp
@@ -5338,6 +5338,7 @@ void UAnimSequence::GetMarkerIndicesForPosition(const FMarkerSyncAnimPosition& S
 			}
 		}
 
+		ensureMsgf(OutCurrentTime >= 0.f && OutCurrentTime <= SequenceLength, TEXT("Current time inside of GetMarkerIndicesForPosition is out of range %.3f of 0.0 to %.3f\n    Sequence: %s"), OutCurrentTime, SequenceLength, *GetFullName());
 		return;
 	}
 
