[NO NON-PII COMMIT MESSAGE]

diff --git a/Engine/Source/Runtime/Engine/Private/InstancedStaticMesh.cpp b/Engine/Source/Runtime/Engine/Private/InstancedStaticMesh.cpp
--- a/Engine/Source/Runtime/Engine/Private/InstancedStaticMesh.cpp
+++ b/Engine/Source/Runtime/Engine/Private/InstancedStaticMesh.cpp
@@ -1537,11 +1537,9 @@ void UInstancedStaticMeshComponent::SerializeRenderData(FArchive& Ar)
 		
 			if (PerInstanceRenderData.IsValid())
 			{
-				FStaticMeshInstanceData* InstanceData = PerInstanceRenderData->InstanceBuffer.GetInstanceData();
-
-				if (InstanceData && InstanceData->GetNumInstances() > 0)
+				if (PerInstanceRenderData->InstanceBuffer.GetInstanceData()->GetNumInstances() > 0)
 				{
-					int32 NumInstances = InstanceData->GetNumInstances();
+					int32 NumInstances = PerInstanceRenderData->InstanceBuffer.GetInstanceData()->GetNumInstances();
 
 					// Clear editor data for the cooked data
 					for (int32 Index = 0; Index < NumInstances; ++Index)
@@ -1553,10 +1551,10 @@ void UInstancedStaticMeshComponent::SerializeRenderData(FArchive& Ar)
 							continue;
 						}
 
-						InstanceData->ClearInstanceEditorData(RenderIndex);
+						PerInstanceRenderData->InstanceBuffer.GetInstanceData()->ClearInstanceEditorData(RenderIndex);
 					}
 
-					InstanceData->Serialize(Ar);
+					PerInstanceRenderData->InstanceBuffer.GetInstanceData()->Serialize(Ar);
 
 #if WITH_EDITOR
 					// Restore back the state we were in
@@ -1581,7 +1579,7 @@ void UInstancedStaticMeshComponent::SerializeRenderData(FArchive& Ar)
 							HitProxyColor = HitProxies[Index]->Id.GetColor();
 						}
 
-						InstanceData->SetInstanceEditorData(RenderIndex, HitProxyColor, bSelected);
+						PerInstanceRenderData->InstanceBuffer.GetInstanceData()->SetInstanceEditorData(RenderIndex, HitProxyColor, bSelected);
 					}
 #endif					
 				}
