Fix crash when pending net game is unable to create the net driver

diff --git a/Engine/Source/Runtime/Engine/Private/PendingNetGame.cpp b/Engine/Source/Runtime/Engine/Private/PendingNetGame.cpp
--- a/Engine/Source/Runtime/Engine/Private/PendingNetGame.cpp
+++ b/Engine/Source/Runtime/Engine/Private/PendingNetGame.cpp
@@ -40,7 +40,13 @@ void UPendingNetGame::InitNetDriver()
 		{
 			NetDriver = GEngine->FindNamedNetDriver(this, NAME_PendingNetDriver);
 		}
-		check(NetDriver);
+
+		if (NetDriver == nullptr)
+		{
+			UE_LOG(LogNet, Warning, TEXT("Error initializing the pending net driver.  Check the configuration of NetDriverDefinitions and make sure module/plugin dependencies are correct."));
+			ConnectionError = NSLOCTEXT("Engine", "NetworkDriverInit", "Error creating network driver.").ToString();
+			return;
+		}
 
 		if( NetDriver->InitConnect( this, URL, ConnectionError ) )
 		{
