RL - removed previous hack by disabling MVF on OpenGL and only doing Unorm to Snorm conversion hack when running skincache in the skincache shaders

diff --git a/Engine/Shaders/Private/Common.ush b/Engine/Shaders/Private/Common.ush
--- a/Engine/Shaders/Private/Common.ush
+++ b/Engine/Shaders/Private/Common.ush
@@ -722,21 +722,10 @@ float4 ConvertTangentSnormToUnorm16(float4 Input)
 	return clamp(Res / 65535, 0.0f, 0.99f);
 }
 
-#ifndef FORCE_BIAS_CONVERSION
-#define FORCE_BIAS_CONVERSION	0
-#endif
-
 // Tangent space bias/unbias
 // We don't use a function so we can avoid type promotion/ coercion.
-#if OPENGL_PROFILE && (MANUAL_VERTEX_FETCH || FORCE_BIAS_CONVERSION)
-#define TangentBias(X)		ConvertTangentUnormToSnorm8(X)
-#define TangentUnbias(X)	ConvertTangentSnormToUnorm8(X)
-#else
 #define TangentBias(X)		(X)
 #define TangentUnbias(X)	(X)
-#endif // OPENGL_PROFILE && (MANUAL_VERTEX_FETCH || FORCE_BIAS_CONVERSION)
-
-
 
 float Square( float x )
 {
diff --git a/Engine/Shaders/Private/GpuSkinCacheComputeShader.usf b/Engine/Shaders/Private/GpuSkinCacheComputeShader.usf
--- a/Engine/Shaders/Private/GpuSkinCacheComputeShader.usf
+++ b/Engine/Shaders/Private/GpuSkinCacheComputeShader.usf
@@ -4,8 +4,6 @@
 	GpuSkinCacheComputeShader.usf: Perform vertex skinning into a buffer to avoid skinning in the vertex shader.
 =============================================================================*/
 
-#define FORCE_BIAS_CONVERSION	1
-
 #include "Common.ush"
 #include "Definitions.usf"
 #include "VertexFactoryCommon.ush"
@@ -153,8 +151,8 @@ FVertexUnpacked UnpackedVertex(uint MorphIndex, uint ClothIndex, uint VertexInde
 {
 	FVertexUnpacked Unpacked;
 
-	Unpacked.TangentX = TangentBias(TangentInputBuffer[2 * VertexIndex + GPUSKIN_VB_OFFSET_TANGENT_X]);
-	Unpacked.TangentZ = TangentBias(TangentInputBuffer[2 * VertexIndex + GPUSKIN_VB_OFFSET_TANGENT_Z]);
+	Unpacked.TangentX = TangentBias_SkinCache(TangentInputBuffer[2 * VertexIndex + GPUSKIN_VB_OFFSET_TANGENT_X]);
+	Unpacked.TangentZ = TangentBias_SkinCache(TangentInputBuffer[2 * VertexIndex + GPUSKIN_VB_OFFSET_TANGENT_Z]);
 
 	uint BlendIndices = InputWeightStream[WeightOffset + GPUSKIN_VB_OFFSET_INFLUENCEBONES];
 	Unpacked.BlendIndices.x = BlendIndices & 0xff;
@@ -311,8 +309,8 @@ void SkinCacheUpdateBatchCS(uint3 GroupID : SV_GroupID,
 	PositionBufferUAV[OutputOffset * 3 + 1] = SPos.y;
 	PositionBufferUAV[OutputOffset * 3 + 2] = SPos.z;
 
-	TangentBufferUAV[2 * OutputOffset + GPUSKIN_RWBUFFER_OFFSET_TANGENT_X] = TangentUnbias(float4(TangentX, Unpacked.TangentX.w));
-	TangentBufferUAV[2 * OutputOffset + GPUSKIN_RWBUFFER_OFFSET_TANGENT_Z] = TangentUnbias(float4(TangentZ, Unpacked.TangentZ.w));
+	TangentBufferUAV[2 * OutputOffset + GPUSKIN_RWBUFFER_OFFSET_TANGENT_X] = TangentUnbias_SkinCache(float4(TangentX, Unpacked.TangentX.w));
+	TangentBufferUAV[2 * OutputOffset + GPUSKIN_RWBUFFER_OFFSET_TANGENT_Z] = TangentUnbias_SkinCache(float4(TangentZ, Unpacked.TangentZ.w));
 
 #if GPUSKIN_DEBUG_PASSTHROUGH
 	// Passthrough debug code
diff --git a/Engine/Shaders/Private/GpuSkinCommon.ush b/Engine/Shaders/Private/GpuSkinCommon.ush
--- a/Engine/Shaders/Private/GpuSkinCommon.ush
+++ b/Engine/Shaders/Private/GpuSkinCommon.ush
@@ -24,6 +24,8 @@
 //	InfluenceBones		0,1							0,2	
 //	InfluenceWeights	1,1							2,2	
 
+#include "Common.ush"
+
 #define	GPUSKIN_VB_OFFSET_TANGENT_X			0
 #define	GPUSKIN_VB_OFFSET_TANGENT_Z			1
 
@@ -83,4 +85,12 @@ uint PackS8x3N(float3 Unpacked)
 		(uint((Unpacked.y + 1.0f) * 127.499f) << 8) |
 		(uint((Unpacked.z + 1.0f) * 127.499f) << 16));
 	return Packed;
-}
\ No newline at end of file
+}
+
+#if OPENGL_PROFILE
+#define TangentBias_SkinCache(X)	ConvertTangentUnormToSnorm8(X)
+#define TangentUnbias_SkinCache(X)	ConvertTangentSnormToUnorm16(X)
+#else
+#define TangentBias_SkinCache(X)	(X)
+#define TangentUnbias_SkinCache(X)	(X)
+#endif // OPENGL_PROFILE
\ No newline at end of file
diff --git a/Engine/Shaders/Private/RecomputeTangentsCommon.ush b/Engine/Shaders/Private/RecomputeTangentsCommon.ush
--- a/Engine/Shaders/Private/RecomputeTangentsCommon.ush
+++ b/Engine/Shaders/Private/RecomputeTangentsCommon.ush
@@ -4,7 +4,6 @@
 	RecomputeTangentsCommon.usf: Recompute Skin Tangents
 =============================================================================*/
 
-#define FORCE_BIAS_CONVERSION	1
 #include "Common.ush"
 #include "GpuSkinCommon.ush"
 
diff --git a/Engine/Shaders/Private/RecomputeTangentsPerVertexPass.usf b/Engine/Shaders/Private/RecomputeTangentsPerVertexPass.usf
--- a/Engine/Shaders/Private/RecomputeTangentsPerVertexPass.usf
+++ b/Engine/Shaders/Private/RecomputeTangentsPerVertexPass.usf
@@ -46,7 +46,7 @@ void MainCS(
 		uint VertexIndex = DispatchThreadId;
 		uint Offset = VertexIndex + SkinCacheStart;
 
-		TangentBufferUAV[2 * Offset + GPUSKIN_RWBUFFER_OFFSET_TANGENT_X] = TangentUnbias(float4(TangentX, 0.0f));
-		TangentBufferUAV[2 * Offset + GPUSKIN_RWBUFFER_OFFSET_TANGENT_Z] = TangentUnbias(float4(TangentZ, Orientation));
+		TangentBufferUAV[2 * Offset + GPUSKIN_RWBUFFER_OFFSET_TANGENT_X] = TangentUnbias_SkinCache(float4(TangentX, 0.0f));
+		TangentBufferUAV[2 * Offset + GPUSKIN_RWBUFFER_OFFSET_TANGENT_Z] = TangentUnbias_SkinCache(float4(TangentZ, Orientation));
 	}
 }
diff --git a/Engine/Source/Runtime/Engine/Private/GPUSkinCache.cpp b/Engine/Source/Runtime/Engine/Private/GPUSkinCache.cpp
--- a/Engine/Source/Runtime/Engine/Private/GPUSkinCache.cpp
+++ b/Engine/Source/Runtime/Engine/Private/GPUSkinCache.cpp
@@ -107,14 +107,14 @@ FAutoConsoleVariableRef CVarGPUSkinCacheAllowDupedVertesForRecomputeTangents(
 
 static int32 GGPUSkinCacheFlushCounter = 0;
 
-bool IsGPUSkinCacheAvailable()
+static inline bool DoesPlatformSupportGPUSkinCache(EShaderPlatform Platform)
 {
-	return GEnableGPUSkinCacheShaders != 0 || GForceRecomputeTangents != 0;
+	return Platform == SP_PCD3D_SM5 || Platform == SP_METAL_SM5 || Platform == SP_METAL_SM5_NOTESS || Platform == SP_METAL_MRT_MAC || Platform == SP_METAL_MRT || IsVulkanSM5Platform(Platform) || Platform == SP_OPENGL_SM5;
 }
 
-static inline bool DoesPlatformSupportGPUSkinCache(EShaderPlatform Platform)
+ENGINE_API bool IsGPUSkinCacheAvailable()
 {
-	return Platform == SP_PCD3D_SM5 || Platform == SP_METAL_SM5 || Platform == SP_METAL_SM5_NOTESS || Platform == SP_METAL_MRT_MAC || Platform == SP_METAL_MRT || IsVulkanSM5Platform(Platform) || Platform == SP_OPENGL_SM5;
+	return (GEnableGPUSkinCacheShaders != 0 || GForceRecomputeTangents != 0) && DoesPlatformSupportGPUSkinCache(GMaxRHIShaderPlatform);
 }
 
 // We don't have it always enabled as it's not clear if this has a performance cost
@@ -124,7 +124,7 @@ ENGINE_API bool DoSkeletalMeshIndexBuffersNeedSRV()
 {
 	// currently only implemented and tested on Window SM5 (needs Compute, Atomics, SRV for index buffers, UAV for VertexBuffers)
 	//#todo-gpuskin: Enable on PS4 when SRVs for IB exist
-	return DoesPlatformSupportGPUSkinCache(GMaxRHIShaderPlatform) && IsGPUSkinCacheAvailable();
+	return IsGPUSkinCacheAvailable();
 }
 
 ENGINE_API bool DoRecomputeSkinTangentsOnGPU_RT()
@@ -334,6 +334,8 @@ public:
 		Data.SourceVertexFactory = InSourceVertexFactory;
 		Data.TargetVertexFactory = InTargetVertexFactory;
 
+		InTargetVertexFactory->InvalidateStreams();
+
 		int32 RecomputeTangentsMode = GForceRecomputeTangents > 0 ? 1 : GSkinCacheRecomputeTangents;
 		if (RecomputeTangentsMode > 0)
 		{
diff --git a/Engine/Source/Runtime/Engine/Private/GPUSkinVertexFactory.cpp b/Engine/Source/Runtime/Engine/Private/GPUSkinVertexFactory.cpp
--- a/Engine/Source/Runtime/Engine/Private/GPUSkinVertexFactory.cpp
+++ b/Engine/Source/Runtime/Engine/Private/GPUSkinVertexFactory.cpp
@@ -641,6 +641,9 @@ void FGPUSkinPassthroughVertexFactory::InternalUpdateVertexDeclaration(FGPUBaseS
 		Data.PositionComponentSRV = PositionRWBuffer->SRV;
 	}
 
+	Data.TangentBasisComponents[0] = SourceVertexFactory->GetTangentStreamComponent(0);
+	Data.TangentBasisComponents[1] = SourceVertexFactory->GetTangentStreamComponent(1);
+
 	if (TangentRWBuffer)
 	{
 		Data.TangentBasisComponents[0].VertexBuffer = &TangentVBAlias;
diff --git a/Engine/Source/Runtime/Engine/Private/PositionVertexBuffer.cpp b/Engine/Source/Runtime/Engine/Private/PositionVertexBuffer.cpp
--- a/Engine/Source/Runtime/Engine/Private/PositionVertexBuffer.cpp
+++ b/Engine/Source/Runtime/Engine/Private/PositionVertexBuffer.cpp
@@ -7,6 +7,7 @@
 #include "Components.h"
 
 #include "StaticMeshVertexData.h"
+#include "GPUSkinCache.h"
 
 /*-----------------------------------------------------------------------------
 FPositionVertexBuffer
@@ -183,7 +184,7 @@ void FPositionVertexBuffer::InitRHI()
 		VertexBufferRHI = RHICreateVertexBuffer(ResourceArray->GetResourceDataSize(), BUF_Static | BUF_ShaderResource, CreateInfo);
 
 		// we have decide to create the SRV based on GMaxRHIShaderPlatform because this is created once and shared between feature levels for editor preview.
-		if (RHISupportsManualVertexFetch(GMaxRHIShaderPlatform))
+		if (RHISupportsManualVertexFetch(GMaxRHIShaderPlatform) || IsGPUSkinCacheAvailable())
 		{
 			PositionComponentSRV = RHICreateShaderResourceView(VertexBufferRHI, 4, PF_R32_FLOAT);
 		}
diff --git a/Engine/Source/Runtime/Engine/Private/Rendering/StaticMeshVertexBuffer.cpp b/Engine/Source/Runtime/Engine/Private/Rendering/StaticMeshVertexBuffer.cpp
--- a/Engine/Source/Runtime/Engine/Private/Rendering/StaticMeshVertexBuffer.cpp
+++ b/Engine/Source/Runtime/Engine/Private/Rendering/StaticMeshVertexBuffer.cpp
@@ -3,6 +3,7 @@
 #include "Rendering/StaticMeshVertexBuffer.h"
 #include "EngineUtils.h"
 #include "Components.h"
+#include "GPUSkinCache.h"
 
 FStaticMeshVertexBuffer::FStaticMeshVertexBuffer() :
 	TangentsData(nullptr),
@@ -255,7 +256,7 @@ void FStaticMeshVertexBuffer::InitRHI()
 			// Create the vertex buffer.
 			FRHIResourceCreateInfo CreateInfo(ResourceArray);
 			TangentsVertexBuffer.VertexBufferRHI = RHICreateVertexBuffer(ResourceArray->GetResourceDataSize(), BUF_Static | BUF_ShaderResource, CreateInfo);
-			if (RHISupportsManualVertexFetch(GMaxRHIShaderPlatform))
+			if (RHISupportsManualVertexFetch(GMaxRHIShaderPlatform) || IsGPUSkinCacheAvailable())
 			{
 				TangentsSRV = RHICreateShaderResourceView(TangentsVertexBuffer.VertexBufferRHI, GetUseHighPrecisionTangentBasis() ? 8 : 4, GetUseHighPrecisionTangentBasis() ? PF_R16G16B16A16_SNORM : PF_R8G8B8A8_SNORM);
 			}
diff --git a/Engine/Source/Runtime/Engine/Public/GPUSkinVertexFactory.h b/Engine/Source/Runtime/Engine/Public/GPUSkinVertexFactory.h
--- a/Engine/Source/Runtime/Engine/Public/GPUSkinVertexFactory.h
+++ b/Engine/Source/Runtime/Engine/Public/GPUSkinVertexFactory.h
@@ -371,12 +371,20 @@ public:
 	virtual const FShaderResourceViewRHIRef GetColorComponentsSRV() const = 0;
 	virtual const uint32 GetColorIndexMask() const = 0;
 
+	inline const FVertexStreamComponent& GetTangentStreamComponent(int Index)
+	{
+		check(TangentStreamComponents[Index].VertexBuffer != nullptr);
+		return TangentStreamComponents[Index];
+	}
+
 protected:
 	/** dynamic data need for setting the shader */ 
 	FShaderDataType ShaderData;
 	/** Pool of buffers for bone matrices. */
 	static TGlobalResource<FBoneBufferPool> BoneBufferPool;
 
+	FVertexStreamComponent TangentStreamComponents[2];
+
 private:
 	uint32 NumVertices;
 };
@@ -434,6 +442,8 @@ public:
 	void SetData(const FDataType& InData)
 	{
 		Data = InData;
+		this->TangentStreamComponents[0] = InData.TangentBasisComponents[0];
+		this->TangentStreamComponents[1] = InData.TangentBasisComponents[1];
 		FGPUBaseSkinVertexFactory::UpdateRHI();
 	}
 
@@ -530,6 +540,12 @@ public:
 	//TODO should be supported
 	bool SupportsPositionOnlyStream() const override { return false; }
 
+	inline void InvalidateStreams()
+	{
+		PositionStreamIndex = -1;
+		TangentStreamIndex = -1;
+	}
+
 protected:
 	// Vertex buffer required for creating the Vertex Declaration
 	FVertexBuffer PositionVBAlias;
@@ -577,6 +593,8 @@ public:
 	void SetData(const FDataType& InData)
 	{
 		MorphData = InData;
+		this->TangentStreamComponents[0] = InData.TangentBasisComponents[0];
+		this->TangentStreamComponents[1] = InData.TangentBasisComponents[1];
 		FGPUBaseSkinVertexFactory::UpdateRHI();
 	}
 
@@ -897,6 +915,8 @@ public:
 	{
         Super::SetData(InData);
 		MeshMappingData = InData;
+		this->TangentStreamComponents[0] = InData.TangentBasisComponents[0];
+		this->TangentStreamComponents[1] = InData.TangentBasisComponents[1];
 		FGPUBaseSkinVertexFactory::UpdateRHI();
 	}
 
diff --git a/Engine/Source/Runtime/RHI/Public/RHI.h b/Engine/Source/Runtime/RHI/Public/RHI.h
--- a/Engine/Source/Runtime/RHI/Public/RHI.h
+++ b/Engine/Source/Runtime/RHI/Public/RHI.h
@@ -162,7 +162,7 @@ inline bool RHISupports4ComponentUAVReadWrite(EShaderPlatform Platform)
 	Shader Platform must not use the mobile renderer, and for Metal, the shader lanugage must be at least 2. */
 inline bool RHISupportsManualVertexFetch(EShaderPlatform InShaderPlatform)
 {
-	return !IsMobilePlatform(InShaderPlatform) && (!IsMetalPlatform(InShaderPlatform) || RHIGetShaderLanguageVersion(InShaderPlatform) >= 2);
+	return !IsOpenGLPlatform(InShaderPlatform) && !IsMobilePlatform(InShaderPlatform) && (!IsMetalPlatform(InShaderPlatform) || RHIGetShaderLanguageVersion(InShaderPlatform) >= 2);
 }
 
 // Wrapper for GRHI## global variables, allows values to be overridden for mobile preview modes.
diff --git a/Engine/Source/Runtime/ShaderCore/Public/VertexFactory.h b/Engine/Source/Runtime/ShaderCore/Public/VertexFactory.h
--- a/Engine/Source/Runtime/ShaderCore/Public/VertexFactory.h
+++ b/Engine/Source/Runtime/ShaderCore/Public/VertexFactory.h
@@ -430,7 +430,7 @@ public:
 	bool SupportsManualVertexFetch(ERHIFeatureLevel::Type InFeatureLevel) const 
 	{ 
 		check(InFeatureLevel != ERHIFeatureLevel::Num);
-		return bSupportsManualVertexFetch && (InFeatureLevel > ERHIFeatureLevel::ES3_1) && !IsMobileOpenGlPlatform(GMaxRHIShaderPlatform) && (!IsMetalPlatform(GMaxRHIShaderPlatform) || (RHIGetShaderLanguageVersion(GMaxRHIShaderPlatform) >= 2));
+		return bSupportsManualVertexFetch && (InFeatureLevel > ERHIFeatureLevel::ES3_1) && RHISupportsManualVertexFetch(GMaxRHIShaderPlatform);
 	}
 
 protected:
