Removing unneeded log spam for every loaded skeletal mesh

#tests n/a

diff --git a/Engine/Source/Runtime/Engine/Private/Rendering/MultiSizeIndexContainer.cpp b/Engine/Source/Runtime/Engine/Private/Rendering/MultiSizeIndexContainer.cpp
--- a/Engine/Source/Runtime/Engine/Private/Rendering/MultiSizeIndexContainer.cpp
+++ b/Engine/Source/Runtime/Engine/Private/Rendering/MultiSizeIndexContainer.cpp
@@ -161,15 +161,6 @@ void FMultiSizeIndexContainer::Serialize(FArchive& Ar, bool bNeedsCPUAccess)
 	}
 
 	IndexBuffer->Serialize(Ar);
-
-	if (bNeedsCPUAccess && IndexBuffer)
-	{
-		UE_LOG(LogSkeletalMesh, Log, TEXT("Index buffer Needs CPU: %p Num: %d"), IndexBuffer, IndexBuffer->Num());
-	}
-	else
-	{
-		UE_LOG(LogSkeletalMesh, Log, TEXT("Index buffer No CPU: %p Num: %d"), IndexBuffer, IndexBuffer->Num());
-	}
 }
 
 #if WITH_EDITOR
