#codereivew jack.porter

diff --git a/Engine/Source/Runtime/Engine/Private/HierarchicalInstancedStaticMesh.cpp b/Engine/Source/Runtime/Engine/Private/HierarchicalInstancedStaticMesh.cpp
--- a/Engine/Source/Runtime/Engine/Private/HierarchicalInstancedStaticMesh.cpp
+++ b/Engine/Source/Runtime/Engine/Private/HierarchicalInstancedStaticMesh.cpp
@@ -2306,7 +2306,7 @@ void UHierarchicalInstancedStaticMeshComponent::BuildTree()
 	QUICK_SCOPE_CYCLE_COUNTER(STAT_UHierarchicalInstancedStaticMeshComponent_BuildTree);
 
 	// upload instance edits to GPU, before validing if the mesh is valid, as it's possible that PerInstanceSMData.Num() == 0, so we have to hide everything before doing the build
-	if (GIsEditor && InstanceUpdateCmdBuffer.NumInlineCommands() > 0)
+	if (GIsEditor && InstanceUpdateCmdBuffer.NumInlineCommands() > 0 && PerInstanceRenderData.IsValid())
 	{
 		// this is allowed only in editor, at runtime upload will happen when buffer is built from component data
 		PerInstanceRenderData->UpdateFromCommandBuffer(InstanceUpdateCmdBuffer);
@@ -2593,7 +2593,7 @@ void UHierarchicalInstancedStaticMeshComponent::BuildTreeAsync()
 	check(BuildTreeAsyncTasks.Num() == 0);
 
 	// upload instance edits to GPU, before validing if the mesh is valid, as it's possible that PerInstanceSMData.Num() == 0, so we have to hide everything before doing the build
-	if (GIsEditor && InstanceUpdateCmdBuffer.NumInlineCommands() > 0)
+	if (GIsEditor && InstanceUpdateCmdBuffer.NumInlineCommands() > 0 && PerInstanceRenderData.IsValid())
 	{
 		// this is allowed only in editor, at runtime upload will happen when buffer is built from component data
 		PerInstanceRenderData->UpdateFromCommandBuffer(InstanceUpdateCmdBuffer);
