Fix seamless travel in PIE not marking the new world or its streaming sublevels as PIE worlds/packages
Fix package name sent via ServerNotifyLoadedWorld not having PIE prefix removed

diff --git a/Engine/Source/Runtime/Engine/Private/World.cpp b/Engine/Source/Runtime/Engine/Private/World.cpp
--- a/Engine/Source/Runtime/Engine/Private/World.cpp
+++ b/Engine/Source/Runtime/Engine/Private/World.cpp
@@ -5127,10 +5127,22 @@ void FSeamlessTravelHandler::SetHandlerLoadedData(UObject* InLevelPackage, UWorl
 /** callback sent to async loading code to inform us when the level package is complete */
 void FSeamlessTravelHandler::SeamlessTravelLoadCallback(const FName& PackageName, UPackage* LevelPackage, EAsyncLoadingResult::Type Result)
 {
-	// make sure we remove the name, even if travel was cancelled.
+	// make sure we remove the name, even if travel was canceled.
 	const FName URLMapFName = FName(*PendingTravelURL.Map);
 	UWorld::WorldTypePreLoadMap.Remove(URLMapFName);
 
+#if WITH_EDITOR
+	if (GIsEditor)
+	{
+		FWorldContext &WorldContext = GEngine->GetWorldContextFromHandleChecked(WorldContextHandle);
+		if (WorldContext.WorldType == EWorldType::PIE)
+		{
+			FString URLMapPackageName = UWorld::ConvertToPIEPackageName(PendingTravelURL.Map, WorldContext.PIEInstance);
+			UWorld::WorldTypePreLoadMap.Remove(FName(*URLMapPackageName));
+		}
+	}
+#endif
+
 	// defer until tick when it's safe to perform the transition
 	if (IsInTransition())
 	{
@@ -5260,7 +5272,14 @@ bool FSeamlessTravelHandler::StartTravel(UWorld* InCurrentWorld, const FURL& InU
 			else if (TransitionMap.IsEmpty())
 			{
 				// If a default transition map doesn't exist, create a dummy World to use as the transition
-				SetHandlerLoadedData(NULL, UWorld::CreateWorld(EWorldType::None, false));
+				if (CurrentWorld->WorldType == EWorldType::PIE)
+				{
+					SetHandlerLoadedData(NULL, UWorld::CreateWorld(EWorldType::PIE, false));
+				}
+				else
+				{
+					SetHandlerLoadedData(NULL, UWorld::CreateWorld(EWorldType::None, false));
+				}
 			}
 			else
 			{
@@ -5373,9 +5392,8 @@ void FSeamlessTravelHandler::StartLoadingDestination()
 
 		CurrentWorld->GetGameInstance()->PreloadContentForURL(PendingTravelURL);
 
-		// Set the world type in the static map, so that UWorld::PostLoad can set the world type
 		const FName URLMapFName = FName(*PendingTravelURL.Map);
-		UWorld::WorldTypePreLoadMap.FindOrAdd(URLMapFName) = CurrentWorld->WorldType;
+
 		// In PIE we might want to mangle MapPackageName when traveling to a map loaded in the editor
 		FString URLMapPackageName = PendingTravelURL.Map;
 		FString URLMapPackageToLoadFrom = PendingTravelURL.Map;
@@ -5398,6 +5416,10 @@ void FSeamlessTravelHandler::StartLoadingDestination()
 			}
 		}
 #endif
+
+		// Set the world type in the static map, so that UWorld::PostLoad can set the world type
+		UWorld::WorldTypePreLoadMap.FindOrAdd(FName(*URLMapPackageName)) = CurrentWorld->WorldType;
+
 		LoadPackageAsync(
 			URLMapPackageName, 
 			PendingTravelGuid.IsValid() ? &PendingTravelGuid : NULL,
@@ -5828,6 +5850,8 @@ UWorld* FSeamlessTravelHandler::Tick()
 			// calling it after InitializeActorsForPlay has been called to have all potential bounding boxed initialized
 			FNavigationSystem::AddNavigationSystemToWorld(*LoadedWorld, FNavigationSystemRunMode::GameMode);
 
+			FName LoadedWorldName = FName(*UWorld::RemovePIEPrefix(LoadedWorld->GetOutermost()->GetName()));
+
 			// send loading complete notifications for all local players
 			for (FLocalPlayerIterator It(GEngine, LoadedWorld); It; ++It)
 			{
@@ -5839,8 +5863,8 @@ UWorld* FSeamlessTravelHandler::Tick()
 					LOG_SCOPE_VERBOSITY_OVERRIDE(LogNetTraffic, ELogVerbosity::VeryVerbose);
 					UE_LOG(LogNet, Verbose, TEXT("NotifyLoadedWorld Begin"));
 #endif
-					It->PlayerController->NotifyLoadedWorld(LoadedWorld->GetOutermost()->GetFName(), bSwitchedToDefaultMap);
-					It->PlayerController->ServerNotifyLoadedWorld(LoadedWorld->GetOutermost()->GetFName());
+					It->PlayerController->NotifyLoadedWorld(LoadedWorldName, bSwitchedToDefaultMap);
+					It->PlayerController->ServerNotifyLoadedWorld(LoadedWorldName);
 #if !UE_BUILD_SHIPPING
 					UE_LOG(LogNet, Verbose, TEXT("NotifyLoadedWorld End"));
 #endif
