Fixed some unlocalizable messages

diff --git a/Engine/Source/Runtime/Engine/Private/UnrealEngine.cpp b/Engine/Source/Runtime/Engine/Private/UnrealEngine.cpp
--- a/Engine/Source/Runtime/Engine/Private/UnrealEngine.cpp
+++ b/Engine/Source/Runtime/Engine/Private/UnrealEngine.cpp
@@ -9315,12 +9315,12 @@ float DrawMapWarnings(UWorld* World, FViewport* Viewport, FCanvas* Canvas, UCanv
 
 		if (NumLightingScenariosEnabled > 1)
 		{
-			SmallTextItem.Text = FText::FromString( FString(TEXT("MULTIPLE LIGHTING SCENARIO LEVELS ENABLED")) );		
+			SmallTextItem.Text = LOCTEXT("MULTIPLE_LIGHTING_SCENARIO_LEVELS_ENABLED", "MULTIPLE LIGHTING SCENARIO LEVELS ENABLED");		
 		}
 		else
 		{
 			// Use 'DumpUnbuiltLightInteractions' to investigate, if lighting is still unbuilt after a lighting build
-			SmallTextItem.Text = FText::FromString( FString::Printf(TEXT("LIGHTING NEEDS TO BE REBUILT (%u unbuilt object(s))"), World->NumLightingUnbuiltObjects) );		
+			SmallTextItem.Text = FText::Format(LOCTEXT("LIGHTING_NEEDS_TO_BE_REBUILT_FMT", "LIGHTING NEEDS TO BE REBUILT ({0} unbuilt {0}|plural(one=object,other=objects))"), World->NumLightingUnbuiltObjects);		
 		}
 
 		Canvas->DrawItem(SmallTextItem, FVector2D(MessageX, MessageY));
@@ -9344,7 +9344,7 @@ float DrawMapWarnings(UWorld* World, FViewport* Viewport, FCanvas* Canvas, UCanv
 		if (NumLightingScenariosEnabled <= 1)
 		{
 			SmallTextItem.SetColor(FLinearColor::White);
-			SmallTextItem.Text = FText::FromString( FString::Printf(TEXT("REFLECTION CAPTURES NEED TO BE REBUILT (%u unbuilt)"), World->NumUnbuiltReflectionCaptures) );		
+			SmallTextItem.Text = FText::Format(LOCTEXT("REFLECTION_CAPTURES_NEED_TO_BE_REBUILT_FMT", "REFLECTION CAPTURES NEED TO BE REBUILT ({0} unbuilt)"), World->NumUnbuiltReflectionCaptures);		
 
 			Canvas->DrawItem(SmallTextItem, FVector2D(MessageX, MessageY));
 			MessageY += FontSizeY;
@@ -9379,7 +9379,7 @@ float DrawMapWarnings(UWorld* World, FViewport* Viewport, FCanvas* Canvas, UCanv
 		if (GUnbuiltHLODCount)
 		{
 			SmallTextItem.SetColor(FLinearColor::Red);
-			SmallTextItem.Text = FText::FromString(FString::Printf(TEXT("HLOD CLUSTER(S) NEED TO BE REBUILT (%u unbuilt object(s))"), GUnbuiltHLODCount));
+			SmallTextItem.Text = FText::Format(LOCTEXT("HLOD_CLUSTERS_NEED_TO_BE_REBUILT_FMT", "HLOD CLUSTERS NEED TO BE REBUILT ({0} unbuilt {0}|plural(one=object,other=objects))"), GUnbuiltHLODCount);
 			Canvas->DrawItem(SmallTextItem, FVector2D(MessageX, MessageY));
 			MessageY += FontSizeY;
 		}
@@ -9389,7 +9389,7 @@ float DrawMapWarnings(UWorld* World, FViewport* Viewport, FCanvas* Canvas, UCanv
 	if (World->NumTextureStreamingUnbuiltComponents > 0 || World->NumTextureStreamingDirtyResources > 0)
 	{
 		SmallTextItem.SetColor(FLinearColor::Red);
-		SmallTextItem.Text = FText::FromString(FString::Printf(TEXT("TEXTURE STREAMING NEEDS TO BE REBUILT (%u Components, %u Resource Refs)"), World->NumTextureStreamingUnbuiltComponents, World->NumTextureStreamingDirtyResources));
+		SmallTextItem.Text = FText::Format(LOCTEXT("TEXTURE_STREAMING_NEEDS_TO_BE_REBUILT_FMT", "TEXTURE STREAMING NEEDS TO BE REBUILT ({0} {0}|plural(one=component,other=components), {1} {1}|plural(one=\"resource ref\",other=\"resource refs\"))"), World->NumTextureStreamingUnbuiltComponents, World->NumTextureStreamingDirtyResources);
 		Canvas->DrawItem(SmallTextItem, FVector2D(MessageX, MessageY));
 		MessageY += FontSizeY;
 	}
@@ -9400,7 +9400,7 @@ float DrawMapWarnings(UWorld* World, FViewport* Viewport, FCanvas* Canvas, UCanv
 		if (MemOver > 0)
 		{
 			SmallTextItem.SetColor(FLinearColor::Red);
-			SmallTextItem.Text = FText::FromString(FString::Printf(TEXT("TEXTURE STREAMING POOL OVER %0.2f MB"), (float)MemOver / 1024.0f / 1024.0f));
+			SmallTextItem.Text = FText::Format(LOCTEXT("TEXTURE_STREAMING_POOL_OVER_BUDGET_FMT", "TEXTURE STREAMING POOL OVER {0} BUDGET"), FText::AsMemory(MemOver));
 			Canvas->DrawItem(SmallTextItem, FVector2D(MessageX, MessageY));
 			MessageY += FontSizeY;
 		}
@@ -9419,7 +9419,7 @@ float DrawMapWarnings(UWorld* World, FViewport* Viewport, FCanvas* Canvas, UCanv
 
 	if (GShaderCompilingManager && GShaderCompilingManager->IsCompiling())
 	{
-		SmallTextItem.Text = FText::FromString(FString::Printf(TEXT("Shaders Compiling (%u)"), GShaderCompilingManager->GetNumRemainingJobs()));
+		SmallTextItem.Text = FText::Format(LOCTEXT("ShadersCompilingFmt", "Shaders Compiling ({0})"), GShaderCompilingManager->GetNumRemainingJobs());
 		Canvas->DrawItem(SmallTextItem, FVector2D(MessageX, MessageY));
 		MessageY += FontSizeY;
 	}
@@ -9611,11 +9611,11 @@ void DrawStatsHUD( UWorld* World, FViewport* Viewport, FCanvas* Canvas, UCanvas*
 		{
 			int32 XSize;
 			int32 YSize;
-			FString String = FString::Printf(TEXT("VisLog recording active"));
-			StringSize(GEngine->GetSmallFont(), XSize, YSize, *String);
+			FText Text = LOCTEXT("VisLogRecordingActive", "VisLog recording active");
+			StringSize(GEngine->GetSmallFont(), XSize, YSize, *Text.ToString());
 
 			SmallTextItem.Position = FVector2D((int32)Viewport->GetSizeXY().X - XSize - 16, 36);
-			SmallTextItem.Text = FText::FromString(String);
+			SmallTextItem.Text = Text;
 			SmallTextItem.SetColor(FLinearColor::Red);
 			SmallTextItem.EnableShadow(FLinearColor::Black);
 			Canvas->DrawItem(SmallTextItem);
@@ -9664,8 +9664,7 @@ void DrawStatsHUD( UWorld* World, FViewport* Viewport, FCanvas* Canvas, UCanvas*
 		if ( FCsvProfiler::Get()->IsCapturing() )
 		{
 			SmallTextItem.SetColor(FLinearColor(0.0f, 1.0f, 0.0f, 1.0f));
-			FString ProfilerScreenText = FString::Printf(TEXT("CsvProfiler frame: %d"), FCsvProfiler::Get()->GetCaptureFrameNumber() );
-			SmallTextItem.Text = FText::FromString(ProfilerScreenText);
+			SmallTextItem.Text = FText::Format(LOCTEXT("CsvProfilerFrameFmt", "CsvProfiler frame: {0}"), FCsvProfiler::Get()->GetCaptureFrameNumber());
 
 			MessageY += 250.0f;
 			Canvas->DrawItem(SmallTextItem, FVector2D(MessageX, MessageY));
@@ -9677,8 +9676,7 @@ void DrawStatsHUD( UWorld* World, FViewport* Viewport, FCanvas* Canvas, UCanvas*
 		if (FTracingProfiler::Get()->IsCapturing())
 		{
 			SmallTextItem.SetColor(FLinearColor(0.0f, 1.0f, 0.0f, 1.0f));
-			FString ProfilerScreenText = FString::Printf(TEXT("TracingProfiler frame: %d"), FTracingProfiler::Get()->GetCaptureFrameNumber());
-			SmallTextItem.Text = FText::FromString(ProfilerScreenText);
+			SmallTextItem.Text = FText::Format(LOCTEXT("TracingProfilerFrameFmt", "TracingProfiler frame: {0}"), FTracingProfiler::Get()->GetCaptureFrameNumber());
 
 			MessageY += 250.0f;
 			Canvas->DrawItem(SmallTextItem, FVector2D(MessageX, MessageY));
@@ -9719,7 +9717,7 @@ void DrawStatsHUD( UWorld* World, FViewport* Viewport, FCanvas* Canvas, UCanvas*
 		if (MessageY != MessageStartY)
 		{
 			SmallTextItem.SetColor(FLinearColor(.05f, .05f, .05f, .2f));
-			SmallTextItem.Text = FText::FromString(FString(TEXT("'DisableAllScreenMessages' to suppress")));
+			SmallTextItem.Text = LOCTEXT("DisableAllScreenMessagesToSuppress", "'DisableAllScreenMessages' to suppress");
 			Canvas->DrawItem(SmallTextItem, FVector2D(MessageX + 50, MessageY));
 			MessageY += 16;
 		}
