DR - add missing guard for MFV

diff --git a/Engine/Source/Runtime/Engine/Private/StaticMeshRender.cpp b/Engine/Source/Runtime/Engine/Private/StaticMeshRender.cpp
--- a/Engine/Source/Runtime/Engine/Private/StaticMeshRender.cpp
+++ b/Engine/Source/Runtime/Engine/Private/StaticMeshRender.cpp
@@ -1510,14 +1510,17 @@ FStaticMeshSceneProxy::FLODInfo::FLODInfo(const UStaticMeshComponent* InComponen
 				OverrideColorVertexBuffer = ComponentLODInfo.OverrideVertexColors;
 				check(OverrideColorVertexBuffer->GetStride() == sizeof(FColor)); //assumed when we set up the stream
 
-				TUniformBufferRef<FLocalVertexFactoryUniformShaderParameters>* UniformBufferPtr = &OverrideColorVFUniformBuffer;
-				const FLocalVertexFactory* LocalVF = &VFs.VertexFactoryOverrideColorVertexBuffer;
-				FColorVertexBuffer* VertexBuffer = OverrideColorVertexBuffer;
-				ENQUEUE_RENDER_COMMAND(FLocalVertexFactoryCopyData)(
-					[UniformBufferPtr, LocalVF, VertexBuffer](FRHICommandListImmediate& RHICmdList)
+				if (RHISupportsManualVertexFetch(GMaxRHIShaderPlatform))
+				{
+					TUniformBufferRef<FLocalVertexFactoryUniformShaderParameters>* UniformBufferPtr = &OverrideColorVFUniformBuffer;
+					const FLocalVertexFactory* LocalVF = &VFs.VertexFactoryOverrideColorVertexBuffer;
+					FColorVertexBuffer* VertexBuffer = OverrideColorVertexBuffer;
+					ENQUEUE_RENDER_COMMAND(FLocalVertexFactoryCopyData)(
+						[UniformBufferPtr, LocalVF, VertexBuffer](FRHICommandListImmediate& RHICmdList)
 					{
 						*UniformBufferPtr = CreateLocalVFUniformBuffer(LocalVF, VertexBuffer);
 					});
+				}
 			}
 		}
 	}
