Fixed a crash with mobile Feature Level Preview when selecting an actor with a Vector Field component
Fixed an ensure when selecting an actor with a Vector Field component

diff --git a/Engine/Source/Runtime/Engine/Classes/Components/VectorFieldComponent.h b/Engine/Source/Runtime/Engine/Classes/Components/VectorFieldComponent.h
--- a/Engine/Source/Runtime/Engine/Classes/Components/VectorFieldComponent.h
+++ b/Engine/Source/Runtime/Engine/Classes/Components/VectorFieldComponent.h
@@ -54,6 +54,7 @@ public:
 	//~ Begin UPrimitiveComponent Interface.
 	virtual FPrimitiveSceneProxy* CreateSceneProxy() override;
 	virtual FBoxSphereBounds CalcBounds(const FTransform& LocalToWorld) const override;
+	virtual void GetUsedMaterials(TArray<UMaterialInterface*>& OutMaterials, bool bGetDebugMaterials) const override;
 	//~ End UPrimitiveComponent Interface.
 
 	//~ Begin UActorComponent Interface.
diff --git a/Engine/Source/Runtime/Engine/Private/VectorField.cpp b/Engine/Source/Runtime/Engine/Private/VectorField.cpp
--- a/Engine/Source/Runtime/Engine/Private/VectorField.cpp
+++ b/Engine/Source/Runtime/Engine/Private/VectorField.cpp
@@ -26,6 +26,7 @@
 #include "VectorField/VectorFieldStatic.h"
 #include "Components/VectorFieldComponent.h"
 #include "PrimitiveSceneProxy.h"
+#include "Materials/Material.h"
 
 #if WITH_EDITORONLY_DATA
 	#include "EditorFramework/AssetImportData.h"
@@ -518,6 +519,13 @@ FBoxSphereBounds UVectorFieldComponent::CalcBounds(const FTransform& LocalToWorl
 	return NewBounds.TransformBy( LocalToWorld );
 }
 
+void UVectorFieldComponent::GetUsedMaterials(TArray<UMaterialInterface*>& OutMaterials, bool bGetDebugMaterials) const
+{
+	if (GEngine->LevelColorationUnlitMaterial != nullptr)
+	{
+		OutMaterials.Add(GEngine->LevelColorationUnlitMaterial);
+	}
+}
 
 void UVectorFieldComponent::OnRegister()
 {
diff --git a/Engine/Source/Runtime/Engine/Private/VectorFieldVisualization.cpp b/Engine/Source/Runtime/Engine/Private/VectorFieldVisualization.cpp
--- a/Engine/Source/Runtime/Engine/Private/VectorFieldVisualization.cpp
+++ b/Engine/Source/Runtime/Engine/Private/VectorFieldVisualization.cpp
@@ -130,7 +130,8 @@ void FVectorFieldVisualizationVertexFactory::ReleaseRHI()
  */
 bool FVectorFieldVisualizationVertexFactory::ShouldCompilePermutation(EShaderPlatform Platform, const class FMaterial* Material, const class FShaderType* ShaderType)
 {
-	return Material->IsSpecialEngineMaterial() && SupportsGPUParticles(Platform) && IsFeatureLevelSupported(Platform, ERHIFeatureLevel::SM4);
+	// We support vector fields on mobile, but not the visualization, so we limit this VF to Feature Level Preview shader platforms
+	return Material->IsSpecialEngineMaterial() && SupportsGPUParticles(Platform) && (IsFeatureLevelSupported(Platform, ERHIFeatureLevel::SM4) || IsPCPlatform(Platform));
 }
 
 /**
