Shift an assert to an early out as it can be handled.

diff --git a/Engine/Source/Runtime/Engine/Private/GameUserSettings.cpp b/Engine/Source/Runtime/Engine/Private/GameUserSettings.cpp
--- a/Engine/Source/Runtime/Engine/Private/GameUserSettings.cpp
+++ b/Engine/Source/Runtime/Engine/Private/GameUserSettings.cpp
@@ -879,10 +879,9 @@ void UGameUserSettings::EnableHDRDisplayOutput(bool bEnable, int32 DisplayNits /
 
 	if (ensure(CVarHDROutputDevice && CVarHDRColorGamut && CVarHDROutputEnabled))
 	{
-		check( !bEnable || IsHDRAllowed() );
-		if (bEnable && !GRHISupportsHDROutput)
+		if (bEnable && !(GRHISupportsHDROutput && IsHDRAllowed()))
 		{
-			UE_LOG(LogConsoleResponse, Display, TEXT("Tried to enable HDR display output but unsupported, forcing off."));
+			UE_LOG(LogConsoleResponse, Display, TEXT("Tried to enable HDR display output but unsupported or disallowed, forcing off."));
 			bEnable = false;
 		}
 
