[Release-4.20] #jira UE-60546 fix for precache state not being updated by FPCMSampleLoader. #rb aaron.mcleran

diff --git a/Engine/Source/Runtime/Engine/Private/SampleBuffer.cpp b/Engine/Source/Runtime/Engine/Private/SampleBuffer.cpp
--- a/Engine/Source/Runtime/Engine/Private/SampleBuffer.cpp
+++ b/Engine/Source/Runtime/Engine/Private/SampleBuffer.cpp
@@ -71,7 +71,9 @@ namespace Audio
 					check(!bIsLoaded);
 					delete SoundWave->AudioDecompressor;
 					SoundWave->AudioDecompressor = nullptr;
+					SoundWave->SetPrecacheState(ESoundWavePrecacheState::Done);
 				}
+
 				bIsLoading = false;
 				bIsLoaded = true;
 
