Merge fix for using mbp.

diff --git a/Engine/Source/Runtime/Engine/Classes/Engine/World.h b/Engine/Source/Runtime/Engine/Classes/Engine/World.h
--- a/Engine/Source/Runtime/Engine/Classes/Engine/World.h
+++ b/Engine/Source/Runtime/Engine/Classes/Engine/World.h
@@ -2065,7 +2065,7 @@ public:
 	float GetMonoFarFieldCullingDistance() const;
 
 	/** Creates a new physics scene for this world. */
-	void CreatePhysicsScene();
+	void CreatePhysicsScene(const AWorldSettings* Settings = nullptr);
 
 	/** Returns a pointer to the physics scene for this world. */
 	FPhysScene* GetPhysicsScene() const { return PhysicsScene; }
diff --git a/Engine/Source/Runtime/Engine/Classes/GameFramework/WorldSettings.h b/Engine/Source/Runtime/Engine/Classes/GameFramework/WorldSettings.h
--- a/Engine/Source/Runtime/Engine/Classes/GameFramework/WorldSettings.h
+++ b/Engine/Source/Runtime/Engine/Classes/GameFramework/WorldSettings.h
@@ -377,6 +377,37 @@ public:
 #endif // WITH_EDITOR
 };
 
+/** Settings pertaining to which PhysX broadphase to use, and settings for MBP if that is the chosen broadphase type */
+USTRUCT()
+struct FBroadphaseSettings
+{
+	GENERATED_BODY();
+
+	FBroadphaseSettings()
+		: bUseMBPOnClient(false)
+		, bUseMBPOnServer(false)
+		, MBPBounds(EForceInit::ForceInitToZero)
+		, MBPNumSubdivs(2)
+	{
+
+	}
+
+	/** Whether to use MBP (Multi Broadphase Pruning */
+	UPROPERTY(EditAnywhere, Category = Broadphase)
+	bool bUseMBPOnClient;
+
+	UPROPERTY(EditAnywhere, Category = Broadphase)
+	bool bUseMBPOnServer;
+
+	/** Total bounds for MBP, must cover the game world or collisions are disabled for out of bounds actors */
+	UPROPERTY(EditAnywhere, Category = Broadphase, meta = (EditCondition = bUseMBP))
+	FBox MBPBounds;
+
+	/** Number of times to subdivide the MBP bounds, final number of regions is MBPNumSubdivs^2 */
+	UPROPERTY(EditAnywhere, Category = Broadphase, meta = (EditCondition = bUseMBP, ClampMin=1, ClampMax=16))
+	uint32 MBPNumSubdivs;
+};
+
 /**
  * Actor containing all script accessible world properties.
  */
@@ -665,6 +696,12 @@ public:
 	UPROPERTY()
 	TArray<struct FNetViewer> ReplicationViewers;
 
+	UPROPERTY(config, EditAnywhere, Category = Broadphase)
+	bool bOverrideDefaultBroadphaseSettings;
+
+	UPROPERTY(config, EditAnywhere, Category = Broadphase)
+	FBroadphaseSettings BroadphaseSettings;
+
 	// ************************************
 
 	/** Maximum number of bookmarks	*/
diff --git a/Engine/Source/Runtime/Engine/Classes/PhysicsEngine/PhysicsSettings.h b/Engine/Source/Runtime/Engine/Classes/PhysicsEngine/PhysicsSettings.h
--- a/Engine/Source/Runtime/Engine/Classes/PhysicsEngine/PhysicsSettings.h
+++ b/Engine/Source/Runtime/Engine/Classes/PhysicsEngine/PhysicsSettings.h
@@ -13,6 +13,7 @@
 #include "Engine/DeveloperSettings.h"
 #include "PhysicsEngine/PhysicsSettingsEnums.h"
 #include "PhysicsEngine/BodySetupEnums.h"
+#include "GameFramework/WorldSettings.h"
 #include "PhysicsSettings.generated.h"
 
 /**
@@ -72,33 +73,6 @@ namespace ESettingsLockedAxis
 	};
 }
 
-/** Settings pertaining to which PhysX broadphase to use, and settings for MBP if that is the chosen broadphase type */
-USTRUCT()
-struct FBroadphaseSettings
-{
-	GENERATED_BODY();
-
-	FBroadphaseSettings()
-		: bUseMBP(false)
-		, MBPBounds(EForceInit::ForceInitToZero)
-		, MBPNumSubdivs(2)
-	{
-
-	}
-
-	/** Whether to use MBP (Multi Broadphase Pruning */
-	UPROPERTY(EditAnywhere, Category = Broadphase)
-	bool bUseMBP;
-
-	/** Total bounds for MBP, must cover the game world or collisions are disabled for out of bounds actors */
-	UPROPERTY(EditAnywhere, Category = Broadphase, meta = (EditCondition = bUseMBP))
-	FBox MBPBounds;
-
-	/** Number of times to subdivide the MBP bounds, final number of regions is MBPNumSubdivs^2 */
-	UPROPERTY(EditAnywhere, Category = Broadphase, meta = (EditCondition = bUseMBP, ClampMin=1, ClampMax=16))
-	uint32 MBPNumSubdivs;
-};
-
 /**
  * Default physics settings.
  */
@@ -286,11 +260,9 @@ class ENGINE_API UPhysicsSettings : public UDeveloperSettings
 	UPROPERTY(config, EditAnywhere, Category=PhysicalSurfaces)
 	TArray<FPhysicalSurfaceName> PhysicalSurfaces;
 
+	/** If we want to Enable MPB or not globally. This is then overridden by project settings if not enabled. **/
 	UPROPERTY(config, EditAnywhere, Category = Broadphase)
-	FBroadphaseSettings ClientBroadphaseSettings;
-
-	UPROPERTY(config, EditAnywhere, Category = Broadphase)
-	FBroadphaseSettings ServerBroadphaseSettings;
+	FBroadphaseSettings DefaultBroadphaseSettings;
 
 public:
 
diff --git a/Engine/Source/Runtime/Engine/Private/PhysicsEngine/PhysScene.cpp b/Engine/Source/Runtime/Engine/Private/PhysicsEngine/PhysScene.cpp
--- a/Engine/Source/Runtime/Engine/Private/PhysicsEngine/PhysScene.cpp
+++ b/Engine/Source/Runtime/Engine/Private/PhysicsEngine/PhysScene.cpp
@@ -495,7 +495,7 @@ static FAutoConsoleCommandWithWorldAndArgs GSetPhysXTreeRebuildRate(TEXT("p.Phys
 
 
 /** Exposes creation of physics-engine scene outside Engine (for use with Physics Asset Editor for example). */
-FPhysScene::FPhysScene()
+FPhysScene::FPhysScene(const AWorldSettings* Settings)
 {
 	LineBatcher = NULL;
 	OwningWorld = NULL;
@@ -519,7 +519,7 @@ FPhysScene::FPhysScene()
 	for (uint32 SceneType = 0; SceneType < NumPhysScenes; ++SceneType)
 	{
 		// Create the physics scene
-		InitPhysScene(SceneType);
+		InitPhysScene(SceneType, Settings);
 
 		// Also initialize scene data
 		bPhysXSceneExecuting[SceneType] = false;
@@ -1990,7 +1990,7 @@ void FPhysScene::ApplyWorldOffset(FVector InOffset)
 #endif
 }
 
-void FPhysScene::InitPhysScene(uint32 SceneType)
+void FPhysScene::InitPhysScene(uint32 SceneType, const AWorldSettings* Settings)
 {
 	check(SceneType < NumPhysScenes);
 
@@ -2128,8 +2128,8 @@ void FPhysScene::InitPhysScene(uint32 SceneType)
 	}
 
 	// Setup MBP desc settings if required
-	FBroadphaseSettings& BroadphaseSettings = IsRunningDedicatedServer() ? UPhysicsSettings::Get()->ServerBroadphaseSettings : UPhysicsSettings::Get()->ClientBroadphaseSettings;
-	bool bUseMBP = (GPhysXForceMbp_Server != 0 && IsRunningDedicatedServer()) || (GPhysXForceMbp_Client != 0 && !IsRunningDedicatedServer()) || BroadphaseSettings.bUseMBP;
+	const FBroadphaseSettings& BroadphaseSettings = (Settings && Settings->bOverrideDefaultBroadphaseSettings) ? Settings->BroadphaseSettings : UPhysicsSettings::Get()->DefaultBroadphaseSettings;
+	bool bUseMBP = IsRunningDedicatedServer() ? BroadphaseSettings.bUseMBPOnServer : BroadphaseSettings.bUseMBPOnClient;
 
 	if(bUseMBP)
 	{
diff --git a/Engine/Source/Runtime/Engine/Private/World.cpp b/Engine/Source/Runtime/Engine/Private/World.cpp
--- a/Engine/Source/Runtime/Engine/Private/World.cpp
+++ b/Engine/Source/Runtime/Engine/Private/World.cpp
@@ -1206,12 +1206,13 @@ void UWorld::InitWorld(const InitializationValues IVS)
 
 	FWorldDelegates::OnPreWorldInitialization.Broadcast(this, IVS);
 
+	AWorldSettings* WorldSettings = GetWorldSettings();
 	if (IVS.bInitializeScenes)
 	{
 		if (IVS.bCreatePhysicsScene)
 		{
 			// Create the physics scene
-			CreatePhysicsScene();
+			CreatePhysicsScene(WorldSettings);
 		}
 
 		bShouldSimulatePhysics = IVS.bShouldSimulatePhysics;
@@ -1222,7 +1223,6 @@ void UWorld::InitWorld(const InitializationValues IVS)
 	}
 
 	// Prepare AI systems
-	AWorldSettings* WorldSettings = GetWorldSettings();
 	if (WorldSettings)
 	{
 		if (IVS.bCreateNavigation || IVS.bCreateAISystem)
@@ -4022,9 +4022,9 @@ float UWorld::GetMonoFarFieldCullingDistance() const
 	return Result;
 }
 
-void UWorld::CreatePhysicsScene()
+void UWorld::CreatePhysicsScene(const AWorldSettings* Settings)
 {
-	SetPhysicsScene(new FPhysScene());
+	SetPhysicsScene(new FPhysScene(Settings));
 }
 
 void UWorld::SetPhysicsScene(FPhysScene* InScene)
diff --git a/Engine/Source/Runtime/Engine/Public/PhysicsPublic.h b/Engine/Source/Runtime/Engine/Public/PhysicsPublic.h
--- a/Engine/Source/Runtime/Engine/Public/PhysicsPublic.h
+++ b/Engine/Source/Runtime/Engine/Public/PhysicsPublic.h
@@ -522,7 +522,7 @@ public:
 	/** Utility for looking up the ApexScene of the given EPhysicsSceneType associated with this FPhysScene.  SceneType must be in the range [0,PST_MAX). */
 	ENGINE_API nvidia::apex::Scene*				GetApexScene(uint32 SceneType) const;
 #endif
-	ENGINE_API FPhysScene();
+	ENGINE_API FPhysScene(const AWorldSettings* Settings = nullptr);
 	ENGINE_API ~FPhysScene();
 
 	/** Start simulation on the physics scene of the given type */
@@ -627,7 +627,7 @@ public:
 	
 private:
 	/** Initialize a scene of the given type.  Must only be called once for each scene type. */
-	void InitPhysScene(uint32 SceneType);
+	void InitPhysScene(uint32 SceneType, const AWorldSettings* Settings = nullptr);
 
 	/** Terminate a scene of the given type.  Must only be called once for each scene type. */
 	void TermPhysScene(uint32 SceneType);
