[NO NON-PII COMMIT MESSAGE]

diff --git a/Engine/Source/Runtime/Engine/Private/HierarchicalInstancedStaticMesh.cpp b/Engine/Source/Runtime/Engine/Private/HierarchicalInstancedStaticMesh.cpp
--- a/Engine/Source/Runtime/Engine/Private/HierarchicalInstancedStaticMesh.cpp
+++ b/Engine/Source/Runtime/Engine/Private/HierarchicalInstancedStaticMesh.cpp
@@ -1170,10 +1170,12 @@ void FHierarchicalStaticMeshSceneProxy::Traverse(const FFoliageCullInstanceParam
 		}
 	}
 
+	float DistanceScale = 1.0f;
+
 	if (MinLOD != MaxLOD)
 	{
 		FVector ScaleAverage = Node.MinInstanceScale + ((Node.MaxInstanceScale - Node.MinInstanceScale) / 2.0f);
-		float DistanceScale = FMath::Max(FMath::Max3(ScaleAverage.X, ScaleAverage.Y, ScaleAverage.Z), 0.001f);
+		//DistanceScale = FMath::Max(FMath::Max3(ScaleAverage.X, ScaleAverage.Y, ScaleAverage.Z), 0.001f);
 
 		CalcLOD(MinLOD, MaxLOD, Node.BoundMin, Node.BoundMax, Params.ViewOriginInLocalZero, Params.ViewOriginInLocalOne, Params.LODPlanesMin, Params.LODPlanesMax, DistanceScale);
 
@@ -1196,7 +1198,7 @@ void FHierarchicalStaticMeshSceneProxy::Traverse(const FFoliageCullInstanceParam
 
 	bool bShouldGroup = Node.FirstChild < 0
 		|| ((Node.LastInstance - Node.FirstInstance + 1) < Params.MinInstancesToSplit[MinLOD]
-			&& CanGroup(Node.BoundMin, Node.BoundMax, Params.ViewOriginInLocalZero, Params.ViewOriginInLocalOne, Params.LODPlanesMax[Params.LODs - 1]));
+			&& CanGroup(Node.BoundMin, Node.BoundMax, Params.ViewOriginInLocalZero, Params.ViewOriginInLocalOne, Params.LODPlanesMax[Params.LODs - 1] * DistanceScale));
 	bool bSplit = (!bFullyContained || MinLOD < MaxLOD || Index < Params.FirstOcclusionNode)
 		&& !bShouldGroup;
 
