Fix world settings being hidden in editor causing components which are attached to the settings actor to not be visible.  This is apparent used by particle systems to spawn temporary psystem components

diff --git a/Engine/Source/Editor/UnrealEd/Private/EditorActor.cpp b/Engine/Source/Editor/UnrealEd/Private/EditorActor.cpp
--- a/Engine/Source/Editor/UnrealEd/Private/EditorActor.cpp
+++ b/Engine/Source/Editor/UnrealEd/Private/EditorActor.cpp
@@ -1753,7 +1753,7 @@ void UUnrealEdEngine::edactSelectAll( UWorld* InWorld )
 	for( FActorIterator It(InWorld); It; ++It )
 	{
 		AActor* Actor = *It;
-		if( !Actor->IsSelected() && !Actor->IsHiddenEd() )
+		if( !Actor->IsSelected() && !Actor->IsHiddenEd() && Actor->IsSelectable())
 		{
 			SelectActor( Actor, 1, 0 );
 		}
diff --git a/Engine/Source/Runtime/Engine/Classes/GameFramework/WorldSettings.h b/Engine/Source/Runtime/Engine/Classes/GameFramework/WorldSettings.h
--- a/Engine/Source/Runtime/Engine/Classes/GameFramework/WorldSettings.h
+++ b/Engine/Source/Runtime/Engine/Classes/GameFramework/WorldSettings.h
@@ -727,6 +727,7 @@ public:
 	//~ Begin AActor Interface.
 #if WITH_EDITOR
 	virtual void CheckForErrors() override;
+	virtual bool IsSelectable() const override { return false; }
 #endif // WITH_EDITOR
 	virtual void PostInitProperties() override;
 	virtual void PreInitializeComponents() override;
diff --git a/Engine/Source/Runtime/Engine/Private/WorldSettings.cpp b/Engine/Source/Runtime/Engine/Private/WorldSettings.cpp
--- a/Engine/Source/Runtime/Engine/Private/WorldSettings.cpp
+++ b/Engine/Source/Runtime/Engine/Private/WorldSettings.cpp
@@ -406,7 +406,6 @@ void AWorldSettings::PostLoad()
 		Entry.MergeSetting.LODSelectionType = EMeshLODSelectionType::CalculateLOD;
 	}
 
-	SetIsTemporarilyHiddenInEditor(true);
 #endif// WITH_EDITOR
 
 	if (bEnableNavigationSystem && NavigationSystemConfig == nullptr)
