[NO NON-PII COMMIT MESSAGE]

diff --git a/Engine/Source/Runtime/Engine/Classes/Components/InstancedStaticMeshComponent.h b/Engine/Source/Runtime/Engine/Classes/Components/InstancedStaticMeshComponent.h
--- a/Engine/Source/Runtime/Engine/Classes/Components/InstancedStaticMeshComponent.h
+++ b/Engine/Source/Runtime/Engine/Classes/Components/InstancedStaticMeshComponent.h
@@ -283,7 +283,7 @@ public:
 	void ClearInstanceSelection();
 
 	/** Initialize the Per Instance Render Data */
-	void InitPerInstanceRenderData(bool InitializeFromCurrentData, FStaticMeshInstanceData* InSharedInstanceBufferData = nullptr);
+	void InitPerInstanceRenderData(bool InitializeFromCurrentData, FStaticMeshInstanceData* InSharedInstanceBufferData = nullptr, bool InRequireCPUAccess = false);
 
 	/** Transfers ownership of instance render data to a render thread. Instance render data will be released in scene proxy destructor or on render thread task. */
 	void ReleasePerInstanceRenderData();
diff --git a/Engine/Source/Runtime/Engine/Private/HierarchicalInstancedStaticMesh.cpp b/Engine/Source/Runtime/Engine/Private/HierarchicalInstancedStaticMesh.cpp
--- a/Engine/Source/Runtime/Engine/Private/HierarchicalInstancedStaticMesh.cpp
+++ b/Engine/Source/Runtime/Engine/Private/HierarchicalInstancedStaticMesh.cpp
@@ -2836,10 +2836,17 @@ static void GatherInstanceTransformsInArea(const UHierarchicalInstancedStaticMes
 					}
 					else if (Component.PerInstanceRenderData.IsValid())
 					{
-						// if there's no PerInstanceSMData (e.g. for grass), we'll go ge the transform from the render buffer
-						FMatrix XformMat;
-						Component.PerInstanceRenderData->InstanceBuffer.GetInstanceTransform(i, XformMat);
-						InstanceToComponent = FTransform(XformMat);
+						if (Component.PerInstanceRenderData->InstanceBuffer.RequireCPUAccess)
+						{
+							// if there's no PerInstanceSMData (e.g. for grass), we'll go get the transform from the render buffer
+							FMatrix XformMat;
+							Component.PerInstanceRenderData->InstanceBuffer.GetInstanceTransform(i, XformMat);
+							InstanceToComponent = FTransform(XformMat);
+						}
+						else
+						{
+							UE_LOG(LogStaticMesh, Warning, TEXT("Trying to query the Instance buffer for information but we don't have a CPU copy to provide the data. Please set KeepInstanceBufferCPUCopy from the Grass variety to true."));
+						}
 					}
 					
 					if (!InstanceToComponent.GetScale3D().IsZero())
diff --git a/Engine/Source/Runtime/Engine/Private/InstancedStaticMesh.cpp b/Engine/Source/Runtime/Engine/Private/InstancedStaticMesh.cpp
--- a/Engine/Source/Runtime/Engine/Private/InstancedStaticMesh.cpp
+++ b/Engine/Source/Runtime/Engine/Private/InstancedStaticMesh.cpp
@@ -547,6 +547,7 @@ FPerInstanceRenderData::~FPerInstanceRenderData()
 
 void FPerInstanceRenderData::UpdateFromPreallocatedData(FStaticMeshInstanceData& InOther)
 {
+	InstanceBuffer.RequireCPUAccess = (InOther.GetOriginResourceArray()->GetAllowCPUAccess() || InOther.GetTransformResourceArray()->GetAllowCPUAccess() || InOther.GetLightMapResourceArray()->GetAllowCPUAccess()) ? true : InstanceBuffer.RequireCPUAccess;
 	ResourceSize = InOther.GetResourceSize(); // 2x when with CPU access?
 	InstanceBuffer.UpdateFromPreallocatedData_Concurrent(InOther);
 }
@@ -2020,7 +2021,7 @@ static bool ComponentRequestsCPUAccess(UInstancedStaticMeshComponent* InComponen
 	return false;
 }
 
-void UInstancedStaticMeshComponent::InitPerInstanceRenderData(bool InitializeFromCurrentData, FStaticMeshInstanceData* InSharedInstanceBufferData)
+void UInstancedStaticMeshComponent::InitPerInstanceRenderData(bool InitializeFromCurrentData, FStaticMeshInstanceData* InSharedInstanceBufferData, bool InRequireCPUAccess)
 {
 	if (PerInstanceRenderData.IsValid())
 	{
@@ -2038,7 +2039,7 @@ void UInstancedStaticMeshComponent::InitPerInstanceRenderData(bool InitializeFro
 	UWorld* World = GetWorld();
 	ERHIFeatureLevel::Type FeatureLevel = World != nullptr ? World->FeatureLevel : GMaxRHIFeatureLevel;
 
-	bool KeepInstanceBufferCPUAccess = GIsEditor || ComponentRequestsCPUAccess(this, FeatureLevel);
+	bool KeepInstanceBufferCPUAccess = GIsEditor || InRequireCPUAccess || ComponentRequestsCPUAccess(this, FeatureLevel);
 
 	if (InSharedInstanceBufferData != nullptr)
 	{
diff --git a/Engine/Source/Runtime/Landscape/Classes/LandscapeGrassType.h b/Engine/Source/Runtime/Landscape/Classes/LandscapeGrassType.h
--- a/Engine/Source/Runtime/Landscape/Classes/LandscapeGrassType.h
+++ b/Engine/Source/Runtime/Landscape/Classes/LandscapeGrassType.h
@@ -103,6 +103,10 @@ struct FGrassVariety
 	UPROPERTY(EditAnywhere, Category = Grass)
 	bool bCastDynamicShadow;
 
+	/** Whether we should keep a cpu copy of the instance buffer. This should be set to true if you plan on using GetOverlappingXXXXCount functions of the component otherwise it won't return any data.**/
+	UPROPERTY(EditAnywhere, Category = Grass)
+	bool bKeepInstanceBufferCPUCopy;
+
 	FGrassVariety()
 		: GrassMesh(nullptr)
 		, GrassDensity(400)
@@ -120,6 +124,7 @@ struct FGrassVariety
 		, bUseLandscapeLightmap(false)
 		, bReceivesDecals(true)
 		, bCastDynamicShadow(true)
+		, bKeepInstanceBufferCPUCopy(false)
 	{
 	}
 };
diff --git a/Engine/Source/Runtime/Landscape/Private/LandscapeGrass.cpp b/Engine/Source/Runtime/Landscape/Private/LandscapeGrass.cpp
--- a/Engine/Source/Runtime/Landscape/Private/LandscapeGrass.cpp
+++ b/Engine/Source/Runtime/Landscape/Private/LandscapeGrass.cpp
@@ -1464,6 +1464,7 @@ struct FAsyncGrassBuilder : public FGrassBuilderBase
 	FVector2D ShadowmapBaseScale;
 	FVector2D LightMapComponentBias;
 	FVector2D LightMapComponentScale;
+	bool RequireCPUAccess;
 
 	// output
 	FStaticMeshInstanceData InstanceBuffer;
@@ -1499,6 +1500,7 @@ struct FAsyncGrassBuilder : public FGrassBuilderBase
 		, ShadowmapBaseScale(FVector2D::UnitVector)
 		, LightMapComponentBias(FVector2D::ZeroVector)
 		, LightMapComponentScale(FVector2D::UnitVector)
+		, RequireCPUAccess(GrassVariety.bKeepInstanceBufferCPUCopy)
 
 		// output
 		, InstanceBuffer(/*bSupportsVertexHalfFloat*/ GVertexElementTypeSupport.IsSupported(VET_Half2))
@@ -1507,6 +1509,8 @@ struct FAsyncGrassBuilder : public FGrassBuilderBase
 	{
 		bHaveValidData = bHaveValidData && GrassData.IsValid();
 
+		InstanceBuffer.SetAllowCPUAccess(RequireCPUAccess);
+
 		check(DesiredInstancesPerLeaf > 0);
 
 		if (UseLandscapeLightmap)
@@ -2485,7 +2489,7 @@ void ALandscapeProxy::UpdateGrass(const TArray<FVector>& Cameras, bool bForceSyn
 
 						if (!HierarchicalInstancedStaticMeshComponent->PerInstanceRenderData.IsValid())
 						{
-							HierarchicalInstancedStaticMeshComponent->InitPerInstanceRenderData(true, &Inner.Builder->InstanceBuffer);
+							HierarchicalInstancedStaticMeshComponent->InitPerInstanceRenderData(true, &Inner.Builder->InstanceBuffer, Inner.Builder->RequireCPUAccess);
 						}
 						else
 						{
