Fix an error in package map object serialization when the client rpcs a struct with a non-null default object reference

diff --git a/Engine/Source/Runtime/Engine/Private/PackageMapClient.cpp b/Engine/Source/Runtime/Engine/Private/PackageMapClient.cpp
--- a/Engine/Source/Runtime/Engine/Private/PackageMapClient.cpp
+++ b/Engine/Source/Runtime/Engine/Private/PackageMapClient.cpp
@@ -869,13 +869,6 @@ FNetworkGUID UPackageMapClient::InternalLoadObject( FArchive & Ar, UObject *& Ob
 		// Remap name for PIE
 		GEngine->NetworkRemapPath( Connection->Driver, PathName, true );
 
-		if ( Object != NULL )
-		{
-			// If we already have the object, just do some sanity checking and return
-			SanityCheckExport( GuidCache.Get(), Object, NetGUID, PathName, ObjOuter, OuterGUID, ExportFlags );
-			return NetGUID;
-		}
-
 		if (NetGUID.IsDefault())
 		{
 			// This should be from the client
@@ -948,6 +941,12 @@ FNetworkGUID UPackageMapClient::InternalLoadObject( FArchive & Ar, UObject *& Ob
 
 			return NetGUID;
 		}
+		else if ( Object != nullptr )
+		{
+			// If we already have the object, just do some sanity checking and return
+			SanityCheckExport( GuidCache.Get(), Object, NetGUID, PathName, ObjOuter, OuterGUID, ExportFlags );
+			return NetGUID;
+		}
 
 		// If we are the server, we should have found the object by now
 		if ( IsNetGUIDAuthority() )
