Fix audio crash when unloading the world

diff --git a/Engine/Source/Runtime/Engine/Private/UnrealEngine.cpp b/Engine/Source/Runtime/Engine/Private/UnrealEngine.cpp
--- a/Engine/Source/Runtime/Engine/Private/UnrealEngine.cpp
+++ b/Engine/Source/Runtime/Engine/Private/UnrealEngine.cpp
@@ -11767,6 +11767,8 @@ bool UEngine::LoadMap( FWorldContext& WorldContext, FURL URL, class UPendingNetG
 	// Unload the current world
 	if( WorldContext.World() )
 	{
+		WorldContext.World()->bIsTearingDown = true;
+
 		if(!URL.HasOption(TEXT("quiet")) )
 		{
 			TransitionType = TT_Loading;
diff --git a/Engine/Source/Runtime/Engine/Private/World.cpp b/Engine/Source/Runtime/Engine/Private/World.cpp
--- a/Engine/Source/Runtime/Engine/Private/World.cpp
+++ b/Engine/Source/Runtime/Engine/Private/World.cpp
@@ -5565,6 +5565,9 @@ UWorld* FSeamlessTravelHandler::Tick()
 				CurrentWorld->GetGameState()->SeamlessTravelTransitionCheckpoint(!bSwitchedToDefaultMap);
 			}
 			
+
+			CurrentWorld->bIsTearingDown = true;
+
 			// If it's not still playing, destroy the demo net driver before we start renaming actors.
 			if ( CurrentWorld->DemoNetDriver && !CurrentWorld->DemoNetDriver->IsPlaying() && !CurrentWorld->DemoNetDriver->bRecordMapChanges)
 			{
