Don't assign a component to a property of the same name but different type

diff --git a/Engine/Source/Runtime/Engine/Private/SCS_Node.cpp b/Engine/Source/Runtime/Engine/Private/SCS_Node.cpp
--- a/Engine/Source/Runtime/Engine/Private/SCS_Node.cpp
+++ b/Engine/Source/Runtime/Engine/Private/SCS_Node.cpp
@@ -164,7 +164,14 @@ UActorComponent* USCS_Node::ExecuteNodeOnActor(AActor* Actor, USceneComponent* P
 			UClass* ActorClass = Actor->GetClass();
 			if (UObjectPropertyBase* Prop = FindField<UObjectPropertyBase>(ActorClass, VarName))
 			{
-				Prop->SetObjectPropertyValue_InContainer(Actor, NewActorComp);
+				if (NewActorComp->IsA(Prop->PropertyClass))
+				{
+					Prop->SetObjectPropertyValue_InContainer(Actor, NewActorComp);
+				}
+				else
+				{
+					UE_LOG(LogBlueprint, Log, TEXT("ExecuteNodeOnActor: Property '%s' on '%s' is of type '%s'. Could not assign '%s' to it."), *VarName.ToString(), *Actor->GetName(), *Prop->PropertyClass->GetName(), *NewActorComp->GetName());
+				}
 			}
 			else
 			{
