Fix now invalid check on flush. #jira none #rb Aaron.McLeran

diff --git a/Engine/Source/Runtime/Engine/Public/Audio.h b/Engine/Source/Runtime/Engine/Public/Audio.h
--- a/Engine/Source/Runtime/Engine/Public/Audio.h
+++ b/Engine/Source/Runtime/Engine/Public/Audio.h
@@ -648,7 +648,7 @@ public:
 	/** Stops the sound source. */
 	ENGINE_API virtual void Stop();
 
-	ENGINE_API virtual void StopNow() { checkf(false, TEXT("This should not get called in non-audio mixer.")); };
+	ENGINE_API virtual void StopNow() { Stop(); };
 
 	/** Whether or not the source is stopping. Only implemented in audio mixer. */
 	virtual bool IsStopping() { return false; }
