more fixes on master pose component issues

diff --git a/Engine/Source/Runtime/Engine/Private/Components/SkinnedMeshComponent.cpp b/Engine/Source/Runtime/Engine/Private/Components/SkinnedMeshComponent.cpp
--- a/Engine/Source/Runtime/Engine/Private/Components/SkinnedMeshComponent.cpp
+++ b/Engine/Source/Runtime/Engine/Private/Components/SkinnedMeshComponent.cpp
@@ -1338,38 +1338,55 @@ void USkinnedMeshComponent::SetMasterPoseComponent(class USkinnedMeshComponent*
 	}
 
 	USkinnedMeshComponent* OldMasterPoseComponent = MasterPoseComponent.Get();
-
-	MasterPoseComponent = NewMasterBoneComponent;
+	USkinnedMeshComponent* ValidNewMasterPose = NewMasterBoneComponent;
 
 	// now add to slave components list, 
-	if (MasterPoseComponent.IsValid())
+	if (ValidNewMasterPose)
 	{
 		// verify if my current master pose is valid
 		// we can't have chain of master poses, so 
 		// we'll find the root master pose component
-		USkinnedMeshComponent* CurrentMasterPose = MasterPoseComponent.Get();
-		while (CurrentMasterPose->MasterPoseComponent.IsValid())
+		USkinnedMeshComponent* Iterator = ValidNewMasterPose;
+		while (Iterator->MasterPoseComponent.IsValid())
 		{
-			MasterPoseComponent = CurrentMasterPose->MasterPoseComponent;
-			CurrentMasterPose = MasterPoseComponent.Get();
+			ValidNewMasterPose = Iterator->MasterPoseComponent.Get();
+			Iterator = ValidNewMasterPose;
+
+			// we have cycling, where in this chain, if it comes back to me, then reject it
+			if (Iterator == this)
+			{
+				ensureAlwaysMsgf(false,
+					TEXT("SetMasterPoseComponent detected loop (the input master pose chain point to itself. (%s <- %s)). Aborting... "),
+					*GetNameSafe(NewMasterBoneComponent), *GetNameSafe(this));
+				ValidNewMasterPose = nullptr;
+				break;
+			}
 		}
 
-		// ensure if master is not same as input, which means it has changed. 
-		USkinnedMeshComponent* RawMasterPoseComponent = MasterPoseComponent.Get();
-		ensureAlwaysMsgf(RawMasterPoseComponent == NewMasterBoneComponent,
-			TEXT("MasterPoseComponent chain is detected (%s). We re-route to top-most MasterPoseComponent (%s)"),
-				*GetNameSafe(RawMasterPoseComponent), *GetNameSafe(NewMasterBoneComponent));
+		// if we have valid master pose, compare with input data and we warn users
+		if (ValidNewMasterPose)
+		{
+			// ensure if master is not same as input, which means it has changed. 
+			ensureMsgf(ValidNewMasterPose == NewMasterBoneComponent,
+				TEXT("MasterPoseComponent chain is detected (%s). We re-route to top-most MasterPoseComponent (%s)"),
+				*GetNameSafe(ValidNewMasterPose), *GetNameSafe(NewMasterBoneComponent));
+		}
+	}
 
+	// now we have valid master pose, set it
+	MasterPoseComponent = ValidNewMasterPose;
+	if (MasterPoseComponent.IsValid())
+	{
 		bool bAddNew = true;
 		// make sure no empty element is there, this is weak obj ptr, so it will go away unless there is 
 		// other reference, this is intentional as master to slave reference is weak
-		for (auto Iter = RawMasterPoseComponent->SlavePoseComponents.CreateIterator(); Iter; ++Iter)
+		for (auto Iter = ValidNewMasterPose->SlavePoseComponents.CreateIterator(); Iter; ++Iter)
 		{
 			TWeakObjectPtr<USkinnedMeshComponent> Comp = (*Iter);
 			if (Comp.IsValid() == false)
 			{
 				// remove
-				RawMasterPoseComponent->SlavePoseComponents.RemoveAt(Iter.GetIndex());
+				ValidNewMasterPose->SlavePoseComponents.RemoveAt(Iter.GetIndex());
 				--Iter;
 			}
 			// if it has same as me, ignore to add
@@ -1381,14 +1398,14 @@ void USkinnedMeshComponent::SetMasterPoseComponent(class USkinnedMeshComponent*
 
 		if (bAddNew)
 		{
-			RawMasterPoseComponent->AddSlavePoseComponent(this);
+			ValidNewMasterPose->AddSlavePoseComponent(this);
 		}
 
 		// set up tick dependency between master & slave components
-		PrimaryComponentTick.AddPrerequisite(RawMasterPoseComponent, RawMasterPoseComponent->PrimaryComponentTick);
+		PrimaryComponentTick.AddPrerequisite(ValidNewMasterPose, ValidNewMasterPose->PrimaryComponentTick);
 	}
 
-	if ((OldMasterPoseComponent != nullptr) && (OldMasterPoseComponent != NewMasterBoneComponent))
+	if ((OldMasterPoseComponent != nullptr) && (OldMasterPoseComponent != ValidNewMasterPose))
 	{
 		OldMasterPoseComponent->RemoveSlavePoseComponent(this);
 
@@ -1401,15 +1418,30 @@ void USkinnedMeshComponent::SetMasterPoseComponent(class USkinnedMeshComponent*
 	UpdateMasterBoneMap();
 
 	// Update Slave in case Master has already been ticked, and we won't get an update for another frame.
-	if (MasterPoseComponent != nullptr)
+	if (ValidNewMasterPose)
 	{
+		// if I have master, but I also have slaves, they won't work anymore
+		// we have to reroute the slaves to new mastser
+		if (SlavePoseComponents.Num() > 0)
+		{
+			ensureMsgf(false,
+				TEXT("MasterPoseComponent chain is detected (%s). We re-route all children to mew MasterPoseComponent (%s)"),
+				*GetNameSafe(this), *GetNameSafe(ValidNewMasterPose));
+
+			for (auto Iter = SlavePoseComponents.CreateIterator(); Iter; ++Iter)
+			{
+				USkinnedMeshComponent* SlaveComp = Iter->Get();
+				SlaveComp->SetMasterPoseComponent(ValidNewMasterPose);
+			}
+		}
+
 		UpdateSlaveComponent();
 	}
 }
 
 void USkinnedMeshComponent::AddSlavePoseComponent(USkinnedMeshComponent* SkinnedMeshComponent)
 {
-	SlavePoseComponents.Add(SkinnedMeshComponent);
+	SlavePoseComponents.AddUnique(SkinnedMeshComponent);
 }
 
 void USkinnedMeshComponent::RemoveSlavePoseComponent(USkinnedMeshComponent* SkinnedMeshComponent)
