Log a warning and early out when trying to remove a local player as a client.  This is not supported and leaves the client/server in a bad state.

diff --git a/Engine/Source/Runtime/Engine/Private/GameInstance.cpp b/Engine/Source/Runtime/Engine/Private/GameInstance.cpp
--- a/Engine/Source/Runtime/Engine/Private/GameInstance.cpp
+++ b/Engine/Source/Runtime/Engine/Private/GameInstance.cpp
@@ -644,6 +644,16 @@ int32 UGameInstance::AddLocalPlayer(ULocalPlayer* NewLocalPlayer, int32 Controll
 
 bool UGameInstance::RemoveLocalPlayer(ULocalPlayer* ExistingPlayer)
 {
+	check(ExistingPlayer);
+
+	UWorld* World = GetWorld();
+
+	if (World && World->IsNetMode(NM_Client))
+	{
+		UE_LOG(LogPlayerManagement, Warning, TEXT("UGameInstance::RemoveLocalPlayer: Unable to remove player %s with ControllerId %i (not supported as a client)"), *ExistingPlayer->GetName(), ExistingPlayer->GetControllerId());
+		return false;
+	}
+
 	// FIXME: Notify server we want to leave the game if this is an online game
 	if (ExistingPlayer->PlayerController != NULL)
 	{
@@ -676,7 +686,7 @@ bool UGameInstance::RemoveLocalPlayer(ULocalPlayer* ExistingPlayer)
 	// Do this after notifications, as some of them require the ViewportClient.
 	ExistingPlayer->ViewportClient = NULL;
 
-	UE_LOG(LogPlayerManagement, Log, TEXT("UGameInstance::RemovePlayer: Removed player %s with ControllerId %i at index %i (%i remaining players)"), *ExistingPlayer->GetName(), ExistingPlayer->GetControllerId(), OldIndex, LocalPlayers.Num());
+	UE_LOG(LogPlayerManagement, Log, TEXT("UGameInstance::RemoveLocalPlayer: Removed player %s with ControllerId %i at index %i (%i remaining players)"), *ExistingPlayer->GetName(), ExistingPlayer->GetControllerId(), OldIndex, LocalPlayers.Num());
 
 	return true;
 }
