Make UObject*'s stored in ComponentInstanceDataCache's Bytes array visible to UObject system. This prevents them from going stale when blueprints are in an error state

diff --git a/Engine/Source/Runtime/Engine/Private/ComponentInstanceDataCache.cpp b/Engine/Source/Runtime/Engine/Private/ComponentInstanceDataCache.cpp
--- a/Engine/Source/Runtime/Engine/Private/ComponentInstanceDataCache.cpp
+++ b/Engine/Source/Runtime/Engine/Private/ComponentInstanceDataCache.cpp
@@ -16,10 +16,11 @@
 class FComponentPropertyWriter : public FObjectWriter
 {
 public:
-	FComponentPropertyWriter(const UActorComponent* InComponent, TArray<uint8>& InBytes, TArray<UObject*>& InInstancedObjects)
+	FComponentPropertyWriter(const UActorComponent* InComponent, TArray<uint8>& InBytes, TArray<UObject*>& InInstancedObjects, TArray<uint32_t>& InObjectReferenceIndicesInByteArray)
 		: FObjectWriter(InBytes)
 		, Component(InComponent)
 		, InstancedObjects(InInstancedObjects)
+		, ObjectReferenceIndicesInByteArray(InObjectReferenceIndicesInByteArray)
 	{
 		// Include properties that would normally skip tagged serialization (e.g. bulk serialization of array properties).
 		ArPortFlags |= PPF_ForceTaggedSerialization;
@@ -108,7 +109,18 @@ public:
 		}
 
 		// store the pointer to this object
+		if (SerializedObject)
+		{
+			ObjectReferenceIndicesInByteArray.Add(Bytes.Num());
+		}
 		Serialize(&SerializedObject, sizeof(UObject*));
+#if DO_CHECK
+		if (SerializedObject)
+		{
+			uint32_t PtrIdx = Bytes.Num() - sizeof(UObject*);
+			ensure(*(UObject**)(&Bytes[PtrIdx]) == SerializedObject);
+		}
+#endif
 
 		return *this;
 	}
@@ -120,6 +132,7 @@ private:
 	TSet<const UProperty*> PropertiesToSkip;
 	
 	TArray<UObject*>& InstancedObjects;
+	TArray<uint32_t>& ObjectReferenceIndicesInByteArray;
 
 	FUObjectAnnotationSparse<FDuplicatedObject,false> DuplicatedObjectAnnotation;
 };
@@ -194,7 +207,7 @@ FActorComponentInstanceData::FActorComponentInstanceData(const UActorComponent*
 
 	if (SourceComponent->IsEditableWhenInherited())
 	{
-		FComponentPropertyWriter ComponentPropertyWriter(SourceComponent, SavedProperties, InstancedObjects);
+		FComponentPropertyWriter ComponentPropertyWriter(SourceComponent, SavedProperties, InstancedObjects, ObjectReferenceIndicesInByteArray);
 
 		// Cache off the length of an array that will come from SerializeTaggedProperties that had no properties saved in to it.
 		auto GetSizeOfEmptyArchive = []() -> int32
@@ -202,13 +215,14 @@ FActorComponentInstanceData::FActorComponentInstanceData(const UActorComponent*
 			const UActorComponent* DummyComponent = GetDefault<UActorComponent>();
 			TArray<uint8> NoWrittenPropertyReference;
 			TArray<UObject*> NoInstances;
-			FComponentPropertyWriter NullWriter(nullptr, NoWrittenPropertyReference, NoInstances);
+			TArray<uint32_t> NoObjectReferences;
+			FComponentPropertyWriter NullWriter(nullptr, NoWrittenPropertyReference, NoInstances, NoObjectReferences);
 			UClass* ComponentClass = DummyComponent->GetClass();
 			
 			// By serializing the component with itself as its defaults we guarantee that no properties will be written out
 			ComponentClass->SerializeTaggedProperties(NullWriter, (uint8*)DummyComponent, ComponentClass, (uint8*)DummyComponent);
 
-			check(NoInstances.Num() == 0);
+			check(NoInstances.Num() == 0 && NoObjectReferences.Num() == 0);
 			return NoWrittenPropertyReference.Num();
 		};
 
@@ -288,6 +302,15 @@ void FActorComponentInstanceData::AddReferencedObjects(FReferenceCollector& Coll
 {
 	Collector.AddReferencedObject(SourceComponentTemplate);
 	Collector.AddReferencedObjects(InstancedObjects);
+
+	for (uint32_t Idx : ObjectReferenceIndicesInByteArray)
+	{
+		UObject** Obj = (UObject**)&SavedProperties[Idx];
+		if (*Obj)
+		{
+			Collector.AddReferencedObject(*Obj);
+		}
+	}
 }
 
 FComponentInstanceDataCache::FComponentInstanceDataCache(const AActor* Actor)
diff --git a/Engine/Source/Runtime/Engine/Public/ComponentInstanceDataCache.h b/Engine/Source/Runtime/Engine/Public/ComponentInstanceDataCache.h
--- a/Engine/Source/Runtime/Engine/Public/ComponentInstanceDataCache.h
+++ b/Engine/Source/Runtime/Engine/Public/ComponentInstanceDataCache.h
@@ -55,6 +55,7 @@ protected:
 
 	TArray<uint8> SavedProperties;
 	TArray<UObject*> InstancedObjects;
+	TArray<uint32_t> ObjectReferenceIndicesInByteArray;
 };
 
 /** 
