Updating an atlas texture now correctly updates materials and material instances

diff --git a/Engine/Source/Developer/AssetTools/Private/AssetTypeActions/AssetTypeActions_CurveLinearColorAtlas.h b/Engine/Source/Developer/AssetTools/Private/AssetTypeActions/AssetTypeActions_CurveLinearColorAtlas.h
--- a/Engine/Source/Developer/AssetTools/Private/AssetTypeActions/AssetTypeActions_CurveLinearColorAtlas.h
+++ b/Engine/Source/Developer/AssetTools/Private/AssetTypeActions/AssetTypeActions_CurveLinearColorAtlas.h
@@ -18,5 +18,5 @@ public:
 	virtual bool HasActions ( const TArray<UObject*>& InObjects ) const override { return true; }
 	virtual void GetActions( const TArray<UObject*>& InObjects, FMenuBuilder& MenuBuilder ) override;
 	virtual bool IsImportedAsset() const override { return false; }
-
+	virtual bool CanFilter() override { return true; }
 };
diff --git a/Engine/Source/Editor/UnrealEd/Private/PreviewMaterial.cpp b/Engine/Source/Editor/UnrealEd/Private/PreviewMaterial.cpp
--- a/Engine/Source/Editor/UnrealEd/Private/PreviewMaterial.cpp
+++ b/Engine/Source/Editor/UnrealEd/Private/PreviewMaterial.cpp
@@ -1342,6 +1342,15 @@ void UMaterialEditorInstanceConstant::CopyToSourceInstance(const bool bForceStat
 				if (ScalarParameterValue && ScalarParameterValue->bOverride)
 				{
 					SourceInstance->SetScalarParameterValueEditorOnly(ScalarParameterValue->ParameterInfo, ScalarParameterValue->ParameterValue);
+					// Copy from editor parameter to saved FParameter
+					if (ScalarParameterValue->AtlasData.bIsUsedAsAtlasPosition)
+					{
+						FScalarParameterAtlasInstanceData InAtlasData = FScalarParameterAtlasInstanceData();
+						InAtlasData.bIsUsedAsAtlasPosition = ScalarParameterValue->AtlasData.bIsUsedAsAtlasPosition;
+						InAtlasData.Curve = ScalarParameterValue->AtlasData.Curve;
+						InAtlasData.Atlas = ScalarParameterValue->AtlasData.Atlas;
+						SourceInstance->SetScalarParameterAtlasEditorOnly(ScalarParameterValue->ParameterInfo, InAtlasData);
+					}
 				}
 				else if (VectorParameterValue && VectorParameterValue->bOverride)
 				{
diff --git a/Engine/Source/Runtime/Engine/Classes/Materials/MaterialInstance.h b/Engine/Source/Runtime/Engine/Classes/Materials/MaterialInstance.h
--- a/Engine/Source/Runtime/Engine/Classes/Materials/MaterialInstance.h
+++ b/Engine/Source/Runtime/Engine/Classes/Materials/MaterialInstance.h
@@ -28,6 +28,34 @@ class FMaterialShaderMapId;
 class FSHAHash;
 
 /** Editable scalar parameter. */
+
+USTRUCT()
+struct FScalarParameterAtlasInstanceData
+{
+	GENERATED_BODY()
+public:
+	UPROPERTY()
+	bool bIsUsedAsAtlasPosition;
+
+	UPROPERTY()
+	TSoftObjectPtr<class UCurveLinearColor> Curve;
+
+	UPROPERTY()
+	TSoftObjectPtr<class UCurveLinearColorAtlas> Atlas;
+
+	bool operator==(const FScalarParameterAtlasInstanceData& Other) const
+	{
+		return
+			bIsUsedAsAtlasPosition == Other.bIsUsedAsAtlasPosition &&
+			Curve == Other.Curve &&
+			Atlas == Other.Atlas;
+	}
+	bool operator!=(const FScalarParameterAtlasInstanceData& Other) const
+	{
+		return !((*this) == Other);
+	}
+};
+
 USTRUCT(BlueprintType)
 struct FScalarParameterValue
 {
@@ -36,6 +64,9 @@ struct FScalarParameterValue
 #if WITH_EDITORONLY_DATA
 	UPROPERTY()
 	FName ParameterName_DEPRECATED;
+
+	UPROPERTY()
+	FScalarParameterAtlasInstanceData AtlasData;
 #endif
 
 	UPROPERTY(EditAnywhere, BlueprintReadWrite, Category=ScalarParameterValue)
@@ -570,7 +601,8 @@ protected:
 
 	void GetTextureExpressionValues(const FMaterialResource* MaterialResource, TArray<UTexture*>& OutTextures, TArray< TArray<int32> >* OutIndices = nullptr) const;
 
-	/** 
+	void GetAtlasTextureValues(const FMaterialResource* MaterialResource, TArray<UTexture*>& OutTextures) const;
+	/**
 	 * Updates StaticPermutationMaterialResources based on the value of bHasStaticPermutationResource
 	 * This is a helper used when recompiling MI's with static parameters.  
 	 * Assumes that the rendering thread command queue has been flushed by the caller.
@@ -592,6 +624,9 @@ protected:
 	bool SetVectorParameterByIndexInternal(int32 ParameterIndex, FLinearColor Value);
 	bool SetScalarParameterByIndexInternal(int32 ParameterIndex, float Value);
 	void SetScalarParameterValueInternal(const FMaterialParameterInfo& ParameterInfo, float Value);
+#if WITH_EDITOR
+	void SetScalarParameterAtlasInternal(const FMaterialParameterInfo& ParameterInfo, FScalarParameterAtlasInstanceData AtlasData);
+#endif
 	void SetTextureParameterValueInternal(const FMaterialParameterInfo& ParameterInfo, class UTexture* Value);
 	void SetFontParameterValueInternal(const FMaterialParameterInfo& ParameterInfo, class UFont* FontValue, int32 FontPage);
 	void ClearParameterValuesInternal(const bool bAllParameters = true);
diff --git a/Engine/Source/Runtime/Engine/Classes/Materials/MaterialInstanceConstant.h b/Engine/Source/Runtime/Engine/Classes/Materials/MaterialInstanceConstant.h
--- a/Engine/Source/Runtime/Engine/Classes/Materials/MaterialInstanceConstant.h
+++ b/Engine/Source/Runtime/Engine/Classes/Materials/MaterialInstanceConstant.h
@@ -63,6 +63,7 @@ class UMaterialInstanceConstant : public UMaterialInstance
 	 */
 	ENGINE_API void SetVectorParameterValueEditorOnly(const FMaterialParameterInfo& ParameterInfo, FLinearColor Value);
 	ENGINE_API void SetScalarParameterValueEditorOnly(const FMaterialParameterInfo& ParameterInfo, float Value);
+	ENGINE_API void SetScalarParameterAtlasEditorOnly(const FMaterialParameterInfo& ParameterInfo, FScalarParameterAtlasInstanceData AtlasData);
 	ENGINE_API void SetTextureParameterValueEditorOnly(const FMaterialParameterInfo& ParameterInfo, class UTexture* Value);
 	ENGINE_API void SetFontParameterValueEditorOnly(const FMaterialParameterInfo& ParameterInfo, class UFont* FontValue, int32 FontPage);
 
diff --git a/Engine/Source/Runtime/Engine/Private/Curves/CurveLinearColorAtlas.cpp b/Engine/Source/Runtime/Engine/Private/Curves/CurveLinearColorAtlas.cpp
--- a/Engine/Source/Runtime/Engine/Private/Curves/CurveLinearColorAtlas.cpp
+++ b/Engine/Source/Runtime/Engine/Private/Curves/CurveLinearColorAtlas.cpp
@@ -89,6 +89,7 @@ void UCurveLinearColorAtlas::PostLoad()
 			GradientCurves[i]->OnUpdateGradient.AddUObject(this, &UCurveLinearColorAtlas::UpdateGradientSlot);
 		}
 	}
+	SizeXY = { (float)TextureSize, (float)GradientPixelSize };
 	UpdateTextures();
 	Super::PostLoad();
 }
diff --git a/Engine/Source/Runtime/Engine/Private/EngineGlobals.cpp b/Engine/Source/Runtime/Engine/Private/EngineGlobals.cpp
--- a/Engine/Source/Runtime/Engine/Private/EngineGlobals.cpp
+++ b/Engine/Source/Runtime/Engine/Private/EngineGlobals.cpp
@@ -51,9 +51,6 @@ FSimpleMulticastDelegate FEditorSupportDelegates::WorldChange;
 FEditorSupportDelegates::FOnForcePropertyWindowRebuild FEditorSupportDelegates::ForcePropertyWindowRebuild;
 /** Sent when events happen that affect how the editors UI looks (mode changes, grid size changes, etc) */
 FSimpleMulticastDelegate FEditorSupportDelegates::UpdateUI;
-/** Called for a material after the user has change a texture's compression settings.
-	Needed to notify the material editors that the need to reattach their preview objects */
-FEditorSupportDelegates::FOnMaterialTextureSettingsChanged FEditorSupportDelegates::MaterialTextureSettingsChanged;
 /** Refresh property windows w/o creating/destroying controls */
 FSimpleMulticastDelegate FEditorSupportDelegates::RefreshPropertyWindows;
 /** Sent before the given windows message is handled in the given viewport */
diff --git a/Engine/Source/Runtime/Engine/Private/Materials/Material.cpp b/Engine/Source/Runtime/Engine/Private/Materials/Material.cpp
--- a/Engine/Source/Runtime/Engine/Private/Materials/Material.cpp
+++ b/Engine/Source/Runtime/Engine/Private/Materials/Material.cpp
@@ -939,6 +939,28 @@ void UMaterial::GetUsedTextures(TArray<UTexture*>& OutTextures, EMaterialQuality
 						}
 					}
 				}
+				// Also look for any scalar parameters that are acting as lookups for an atlas texture, and store the atlas texture
+				const TArray<TRefCountPtr<FMaterialUniformExpression> >* AtlasExpressions[1] =
+				{
+					&CurrentResource->GetUniformScalarParameterExpressions()
+				};
+				for (int32 TypeIndex = 0; TypeIndex < ARRAY_COUNT(AtlasExpressions); TypeIndex++)
+				{
+					// Iterate over each of the material's texture expressions.
+					for (FMaterialUniformExpression* Expression : *AtlasExpressions[TypeIndex])
+					{
+						const FMaterialUniformExpressionScalarParameter* ScalarExpression = static_cast<const FMaterialUniformExpressionScalarParameter*>(Expression);
+						bool bIsUsedAsAtlasPosition;
+						TSoftObjectPtr<class UCurveLinearColor> Curve;
+						TSoftObjectPtr<class UCurveLinearColorAtlas> Atlas;
+						ScalarExpression->GetGameThreadUsedAsAtlas(this, bIsUsedAsAtlasPosition, Curve, Atlas);
+
+						if (Atlas)
+						{
+							OutTextures.AddUnique(Atlas.Get());
+						}
+					}
+				}
 			}
 		}
 	}
diff --git a/Engine/Source/Runtime/Engine/Private/Materials/MaterialInstance.cpp b/Engine/Source/Runtime/Engine/Private/Materials/MaterialInstance.cpp
--- a/Engine/Source/Runtime/Engine/Private/Materials/MaterialInstance.cpp
+++ b/Engine/Source/Runtime/Engine/Private/Materials/MaterialInstance.cpp
@@ -482,6 +482,19 @@ void UMaterialInstance::InitResources()
 	GameThread_InitMIParameters(this, VectorParameterValues);
 	GameThread_InitMIParameters(this, TextureParameterValues);
 	GameThread_InitMIParameters(this, FontParameterValues);
+
+#if WITH_EDITOR
+	//recalculate any scalar params based on a curve position in an atlas in case the atlas changed
+	for (FScalarParameterValue ScalarParam : ScalarParameterValues)
+	{
+		IsScalarParameterUsedAsAtlasPosition(ScalarParam.ParameterInfo, ScalarParam.AtlasData.bIsUsedAsAtlasPosition, ScalarParam.AtlasData.Curve, ScalarParam.AtlasData.Atlas);
+		if (ScalarParam.AtlasData.bIsUsedAsAtlasPosition)
+		{
+			SetScalarParameterAtlasInternal(ScalarParam.ParameterInfo, ScalarParam.AtlasData);
+		}
+	}
+#endif
+
 	PropagateDataToMaterialProxy();
 
 	CacheMaterialInstanceUniformExpressions(this);
@@ -648,6 +661,18 @@ bool UMaterialInstance::IsScalarParameterUsedAsAtlasPosition(const FMaterialPara
 		return false;
 	}
 
+	// Instance override
+	const FScalarParameterValue* ParameterValue = GameThread_FindParameterByName(ScalarParameterValues, ParameterInfo);
+#if WITH_EDITOR
+	if (ParameterValue && ParameterValue->AtlasData.Curve.Get() != nullptr && ParameterValue->AtlasData.Atlas.Get() != nullptr)
+	{
+		OutValue = ParameterValue->AtlasData.bIsUsedAsAtlasPosition;
+		Curve = ParameterValue->AtlasData.Curve;
+		Atlas = ParameterValue->AtlasData.Atlas;
+		return true;
+	}
+#endif
+
 	// Instance-included default
 	if (ParameterInfo.Association != EMaterialParameterAssociation::GlobalParameter)
 	{
@@ -969,6 +994,33 @@ void UMaterialInstance::GetTextureExpressionValues(const FMaterialResource* Mate
 	}
 }
 
+void UMaterialInstance::GetAtlasTextureValues(const FMaterialResource* MaterialResource, TArray<UTexture*>& OutTextures) const
+{
+	check(MaterialResource);
+
+	const TArray<TRefCountPtr<FMaterialUniformExpression> >* AtlasExpressions[1] =
+	{
+		&MaterialResource->GetUniformScalarParameterExpressions()
+	};
+	for (int32 TypeIndex = 0; TypeIndex < ARRAY_COUNT(AtlasExpressions); TypeIndex++)
+	{
+		// Iterate over each of the material's scalar expressions.
+		for (FMaterialUniformExpression* Expression : *AtlasExpressions[TypeIndex])
+		{
+			const FMaterialUniformExpressionScalarParameter* ScalarExpression = static_cast<const FMaterialUniformExpressionScalarParameter*>(Expression);
+			bool bIsUsedAsAtlasPosition = false;
+			TSoftObjectPtr<class UCurveLinearColor> Curve;
+			TSoftObjectPtr<class UCurveLinearColorAtlas> Atlas;
+			ScalarExpression->GetGameThreadUsedAsAtlas(this, bIsUsedAsAtlasPosition, Curve, Atlas);
+
+			if (Atlas)
+			{
+				OutTextures.AddUnique(Atlas.Get());
+			}
+		}
+	}
+}
+
 void UMaterialInstance::GetUsedTextures(TArray<UTexture*>& OutTextures, EMaterialQualityLevel::Type QualityLevel, bool bAllQualityLevels, ERHIFeatureLevel::Type FeatureLevel, bool bAllFeatureLevels) const
 {
 	OutTextures.Empty();
@@ -1019,6 +1071,7 @@ void UMaterialInstance::GetUsedTextures(TArray<UTexture*>& OutTextures, EMateria
 					if (MaterialResource)
 					{
 						GetTextureExpressionValues(MaterialResource, OutTextures);
+						GetAtlasTextureValues(MaterialResource, OutTextures);
 					}
 				}
 			}
@@ -1530,7 +1583,6 @@ void UMaterialInstance::CopyMaterialInstanceParameters(UMaterialInterface* Sourc
 			}
 		}
 
-
 		// Now do the scalar params
 		OutParameterInfo.Reset();
 		Guids.Reset();
@@ -1544,6 +1596,9 @@ void UMaterialInstance::CopyMaterialInstanceParameters(UMaterialInterface* Sourc
 				ParameterValue->ParameterInfo = ParameterInfo;
 				ParameterValue->ExpressionGUID.Invalidate();
 				ParameterValue->ParameterValue = ScalarValue;
+#if WITH_EDITOR
+				IsScalarParameterUsedAsAtlasPosition(ParameterValue->ParameterInfo, ParameterValue->AtlasData.bIsUsedAsAtlasPosition, ParameterValue->AtlasData.Curve, ParameterValue->AtlasData.Atlas);
+#endif
 			}
 		}
 
@@ -3321,6 +3376,40 @@ void UMaterialInstance::SetScalarParameterValueInternal(const FMaterialParameter
 	}
 }
 
+#if WITH_EDITOR
+void UMaterialInstance::SetScalarParameterAtlasInternal(const FMaterialParameterInfo& ParameterInfo, FScalarParameterAtlasInstanceData AtlasData)
+{
+	FScalarParameterValue* ParameterValue = GameThread_FindParameterByName(ScalarParameterValues, ParameterInfo);
+
+	if (ParameterValue)
+	{
+		ParameterValue->AtlasData = AtlasData;
+		UCurveLinearColorAtlas* Atlas = Cast<UCurveLinearColorAtlas>(AtlasData.Atlas.Get());
+		UCurveLinearColor* Curve = Cast<UCurveLinearColor>(AtlasData.Curve.Get());
+		if (!Atlas || !Curve)
+		{
+			return;
+		}
+		int32 Index = Atlas->GradientCurves.Find(Curve);
+		if (Index == INDEX_NONE)
+		{
+			return;
+		}
+
+		float NewValue = ((float)Index * Atlas->GradientPixelSize) / Atlas->TextureSize + (0.5f * Atlas->GradientPixelSize) / Atlas->TextureSize;
+		
+		// Don't enqueue an update if it isn't needed
+		if (ParameterValue->ParameterValue != NewValue)
+		{
+			ParameterValue->ParameterValue = NewValue;
+			// Update the material instance data in the rendering thread.
+			GameThread_UpdateMIParameter(this, *ParameterValue);
+			CacheMaterialInstanceUniformExpressions(this);
+		}
+	}
+}
+#endif
+
 void UMaterialInstance::SetTextureParameterValueInternal(const FMaterialParameterInfo& ParameterInfo, UTexture* Value)
 {
 	FTextureParameterValue* ParameterValue = GameThread_FindParameterByName(TextureParameterValues, ParameterInfo);
diff --git a/Engine/Source/Runtime/Engine/Private/Materials/MaterialInstanceConstant.cpp b/Engine/Source/Runtime/Engine/Private/Materials/MaterialInstanceConstant.cpp
--- a/Engine/Source/Runtime/Engine/Private/Materials/MaterialInstanceConstant.cpp
+++ b/Engine/Source/Runtime/Engine/Private/Materials/MaterialInstanceConstant.cpp
@@ -5,7 +5,9 @@
 =============================================================================*/
 
 #include "Materials/MaterialInstanceConstant.h"
-
+#if WITH_EDITOR
+#include "MaterialEditor/DEditorScalarParameterValue.h"
+#endif
 
 UMaterialInstanceConstant::UMaterialInstanceConstant(const FObjectInitializer& ObjectInitializer)
 	: Super(ObjectInitializer)
@@ -65,6 +67,12 @@ void UMaterialInstanceConstant::SetScalarParameterValueEditorOnly(const FMateria
 	SetScalarParameterValueInternal(ParameterInfo,Value);
 }
 
+void UMaterialInstanceConstant::SetScalarParameterAtlasEditorOnly(const FMaterialParameterInfo& ParameterInfo, FScalarParameterAtlasInstanceData AtlasData)
+{
+	check(GIsEditor);
+	SetScalarParameterAtlasInternal(ParameterInfo, AtlasData);
+}
+
 void UMaterialInstanceConstant::SetTextureParameterValueEditorOnly(const FMaterialParameterInfo& ParameterInfo, UTexture* Value)
 {
 	check(GIsEditor);
diff --git a/Engine/Source/Runtime/Engine/Private/Materials/MaterialUniformExpressions.cpp b/Engine/Source/Runtime/Engine/Private/Materials/MaterialUniformExpressions.cpp
--- a/Engine/Source/Runtime/Engine/Private/Materials/MaterialUniformExpressions.cpp
+++ b/Engine/Source/Runtime/Engine/Private/Materials/MaterialUniformExpressions.cpp
@@ -812,6 +812,21 @@ void FMaterialUniformExpressionScalarParameter::GetGameThreadNumberValue(const U
 	}
 }
 
+void FMaterialUniformExpressionScalarParameter::GetGameThreadUsedAsAtlas(const UMaterialInterface* SourceMaterialToCopyFrom, bool& OutValue, TSoftObjectPtr<class UCurveLinearColor>& Curve, TSoftObjectPtr<class UCurveLinearColorAtlas>& Atlas) const
+{
+	check(IsInGameThread());
+	checkSlow(SourceMaterialToCopyFrom);
+
+	const UMaterialInterface* It = SourceMaterialToCopyFrom;
+
+	const UMaterialInstance* MatInst = Cast<UMaterialInstance>(It);
+
+	if (MatInst)
+	{
+		MatInst->IsScalarParameterUsedAsAtlasPosition(ParameterInfo, OutValue, Curve, Atlas);		
+	}
+}
+
 namespace
 {
 	void SerializeOptional(FArchive& Ar, TOptional<FName>& OptionalName)
diff --git a/Engine/Source/Runtime/Engine/Private/Materials/MaterialUniformExpressions.h b/Engine/Source/Runtime/Engine/Private/Materials/MaterialUniformExpressions.h
--- a/Engine/Source/Runtime/Engine/Private/Materials/MaterialUniformExpressions.h
+++ b/Engine/Source/Runtime/Engine/Private/Materials/MaterialUniformExpressions.h
@@ -202,6 +202,8 @@ public:
 	// faster than GetNumberValue(), good for run-time use
 	void GetGameThreadNumberValue(const UMaterialInterface* SourceMaterialToCopyFrom, float& OutValue) const;
 
+	void GetGameThreadUsedAsAtlas(const UMaterialInterface* SourceMaterialToCopyFrom, bool& OutValue, TSoftObjectPtr<class UCurveLinearColor>& Curve, TSoftObjectPtr<class UCurveLinearColorAtlas>& Atlas) const;
+	
 	virtual bool IsConstant() const
 	{
 		return false;
diff --git a/Engine/Source/Runtime/Engine/Private/Texture.cpp b/Engine/Source/Runtime/Engine/Private/Texture.cpp
--- a/Engine/Source/Runtime/Engine/Private/Texture.cpp
+++ b/Engine/Source/Runtime/Engine/Private/Texture.cpp
@@ -1149,8 +1149,6 @@ void GetAllDefaultTextureFormats(const class ITargetPlatform* TargetPlatform, TA
 
 void UTexture::NotifyMaterials()
 {
-	TArray<UMaterialInterface*> MaterialsThatUseThisTexture;
-
 	// Create a material update context to safely update materials.
 	{
 		FMaterialUpdateContext UpdateContext;
@@ -1162,8 +1160,7 @@ void UTexture::NotifyMaterials()
 			UMaterialInterface* MaterialInterface = *It;
 			if (DoesMaterialUseTexture(MaterialInterface, this))
 			{
-				MaterialsThatUseThisTexture.Add(MaterialInterface);
-
+				UpdateContext.AddMaterialInterface(MaterialInterface);
 				// This is a bit tricky. We want to make sure all materials using this texture are
 				// updated. Materials are always updated. Material instances may also have to be
 				// updated and if they have static permutations their children must be updated
@@ -1176,16 +1173,9 @@ void UTexture::NotifyMaterials()
 		// Go ahead and update any base materials that need to be.
 		for (TSet<UMaterial*>::TConstIterator It(BaseMaterialsThatUseThisTexture); It; ++It)
 		{
-			UpdateContext.AddMaterial(*It);
 			(*It)->PostEditChange();
 		}
 	}
-
-	// Now that all materials and instances have updated send necessary callbacks.
-	for (int32 i = 0; i < MaterialsThatUseThisTexture.Num(); ++i)
-	{
-		FEditorSupportDelegates::MaterialTextureSettingsChanged.Broadcast(MaterialsThatUseThisTexture[i]);
-	}
 }
 
 #endif //WITH_EDITOR
\ No newline at end of file
diff --git a/Engine/Source/Runtime/Engine/Public/EditorSupportDelegates.h b/Engine/Source/Runtime/Engine/Public/EditorSupportDelegates.h
--- a/Engine/Source/Runtime/Engine/Public/EditorSupportDelegates.h
+++ b/Engine/Source/Runtime/Engine/Public/EditorSupportDelegates.h
@@ -39,9 +39,6 @@ struct ENGINE_API FEditorSupportDelegates
 	static FOnForcePropertyWindowRebuild ForcePropertyWindowRebuild;
 	/** Sent when events happen that affect how the editors UI looks (mode changes, grid size changes, etc) */
 	static FSimpleMulticastDelegate UpdateUI;
-	/** Called for a material after the user has change a texture's compression settings.
-		Needed to notify the material editors that the need to reattach their preview objects */
-	static FOnMaterialTextureSettingsChanged MaterialTextureSettingsChanged;
 	/** Refresh property windows w/o creating/destroying controls */
 	static FSimpleMulticastDelegate RefreshPropertyWindows;
 	/** Sent before the given windows message is handled in the given viewport */
