Fix static mesh editor UV scale in high dpi

diff --git a/Engine/Source/Runtime/Engine/Private/PrimitiveDrawingUtils.cpp b/Engine/Source/Runtime/Engine/Private/PrimitiveDrawingUtils.cpp
--- a/Engine/Source/Runtime/Engine/Private/PrimitiveDrawingUtils.cpp
+++ b/Engine/Source/Runtime/Engine/Private/PrimitiveDrawingUtils.cpp
@@ -1669,7 +1669,7 @@ void DrawUVsInternal(FViewport* InViewport, FCanvas* InCanvas, int32 InTextYPos,
 		const uint32 MinX = BorderWidth;
 		const FVector2D UVBoxOrigin(MinX, MinY);
 		const FVector2D BoxOrigin( MinX - 1, MinY - 1 );
-		const uint32 UVBoxScale = FMath::Min(InViewport->GetSizeXY().X - MinX, InViewport->GetSizeXY().Y - MinY) - BorderWidth;
+		const uint32 UVBoxScale = FMath::Min(InViewport->GetSizeXY().X / InCanvas->GetDPIScale() - MinX, InViewport->GetSizeXY().Y / InCanvas->GetDPIScale() - MinY) - BorderWidth;
 		const uint32 BoxSize = UVBoxScale + 2;
 		FCanvasTileItem BoxBackgroundTileItem(BoxOrigin, GWhiteTexture, FVector2D(BoxSize, BoxSize), FLinearColor(0, 0, 0, 0.4f));
 		BoxBackgroundTileItem.BlendMode = SE_BLEND_AlphaComposite;
