SA Fix
#rnx

diff --git a/Engine/Source/Runtime/Engine/Private/ComponentInstanceDataCache.cpp b/Engine/Source/Runtime/Engine/Private/ComponentInstanceDataCache.cpp
--- a/Engine/Source/Runtime/Engine/Private/ComponentInstanceDataCache.cpp
+++ b/Engine/Source/Runtime/Engine/Private/ComponentInstanceDataCache.cpp
@@ -118,7 +118,7 @@ public:
 		if (SerializedObject)
 		{
 			uint32_t PtrIdx = Bytes.Num() - sizeof(UObject*);
-			ensure(*(UObject**)(&Bytes[PtrIdx]) == SerializedObject);
+			ensure(*((UObject**)(&Bytes[PtrIdx])) == SerializedObject);
 		}
 #endif
 
