possible fix for static anlalysis warning

diff --git a/Engine/Source/Runtime/Engine/Private/Components/SkinnedMeshComponent.cpp b/Engine/Source/Runtime/Engine/Private/Components/SkinnedMeshComponent.cpp
--- a/Engine/Source/Runtime/Engine/Private/Components/SkinnedMeshComponent.cpp
+++ b/Engine/Source/Runtime/Engine/Private/Components/SkinnedMeshComponent.cpp
@@ -1375,7 +1375,7 @@ void USkinnedMeshComponent::SetMasterPoseComponent(class USkinnedMeshComponent*
 
 	// now we have valid master pose, set it
 	MasterPoseComponent = ValidNewMasterPose;
-	if (MasterPoseComponent.IsValid())
+	if (ValidNewMasterPose)
 	{
 		bool bAddNew = true;
 		// make sure no empty element is there, this is weak obj ptr, so it will go away unless there is 
