SA Fix... again
#rnx

diff --git a/Engine/Source/Runtime/Engine/Private/ComponentInstanceDataCache.cpp b/Engine/Source/Runtime/Engine/Private/ComponentInstanceDataCache.cpp
--- a/Engine/Source/Runtime/Engine/Private/ComponentInstanceDataCache.cpp
+++ b/Engine/Source/Runtime/Engine/Private/ComponentInstanceDataCache.cpp
@@ -118,7 +118,9 @@ public:
 		if (SerializedObject)
 		{
 			uint32_t PtrIdx = Bytes.Num() - sizeof(UObject*);
-			ensure(*((UObject**)(&Bytes[PtrIdx])) == SerializedObject);
+			void* NewObjAddress = static_cast<void*>(&Bytes[PtrIdx]);
+			UObject* NewObjValue = *(reinterpret_cast<UObject**>(NewObjAddress));
+			ensure(NewObjValue == SerializedObject);
 		}
 #endif
 
