Allow UActorComponent::OnRep_IsActive to be overriden

diff --git a/Engine/Source/Runtime/Engine/Classes/Components/ActorComponent.h b/Engine/Source/Runtime/Engine/Classes/Components/ActorComponent.h
--- a/Engine/Source/Runtime/Engine/Classes/Components/ActorComponent.h
+++ b/Engine/Source/Runtime/Engine/Classes/Components/ActorComponent.h
@@ -279,7 +279,7 @@ public:
 	bool IsCreatedByConstructionScript() const;
 
 	UFUNCTION()
-	void OnRep_IsActive();
+	virtual void OnRep_IsActive();
 
 private:
 	AActor* GetActorOwnerNoninline() const;
