Fix for missing audio on Unix platform

diff --git a/Engine/Source/Runtime/Engine/Private/AudioDevice.cpp b/Engine/Source/Runtime/Engine/Private/AudioDevice.cpp
--- a/Engine/Source/Runtime/Engine/Private/AudioDevice.cpp
+++ b/Engine/Source/Runtime/Engine/Private/AudioDevice.cpp
@@ -4487,7 +4487,7 @@ void FAudioDevice::Flush(UWorld* WorldToFlush, bool bClearActivatedReverb)
 	}
 
 	// Do fadeout when flushing the audio device.
-	if (WorldToFlush == nullptr || WorldToFlush->bIsTearingDown)
+	if (WorldToFlush == nullptr)
 	{
 		FadeOut();
 	}
