Fix deterministic cooking issue with skeleton

diff --git a/Engine/Source/Runtime/Engine/Private/Animation/Skeleton.cpp b/Engine/Source/Runtime/Engine/Private/Animation/Skeleton.cpp
--- a/Engine/Source/Runtime/Engine/Private/Animation/Skeleton.cpp
+++ b/Engine/Source/Runtime/Engine/Private/Animation/Skeleton.cpp
@@ -205,6 +205,7 @@ void USkeleton::Serialize( FArchive& Ar )
 
 	if (Ar.UE4Ver() < VER_UE4_SKELETON_GUID_SERIALIZATION)
 	{
+		UE_LOG(LogAnimation, Warning, TEXT("Skeleton '%s' has not been saved since version 'VER_UE4_SKELETON_GUID_SERIALIZATION' This asset will not cook deterministically until it is resaved."), *GetPathName());
 		RegenerateGuid();
 	}
 	else
