Fix check when disabling a cooling down tick function during a pause frame

diff --git a/Engine/Source/Runtime/Engine/Private/TickTaskManager.cpp b/Engine/Source/Runtime/Engine/Private/TickTaskManager.cpp
--- a/Engine/Source/Runtime/Engine/Private/TickTaskManager.cpp
+++ b/Engine/Source/Runtime/Engine/Private/TickTaskManager.cpp
@@ -1208,17 +1208,24 @@ public:
 			break;
 
 		case FTickFunction::ETickState::CoolingDown:
-			FTickFunction* PrevComparisionFunction = nullptr;
+			// If a cooling function is in the reschedule list then we must be in a pause frame and it has already been set for
+			// reschedule and removed from the cooldown list so we won't find it there. This is fine as the reschedule will see
+			// the tick function is disabled and not reschedule it.
+			auto FindTickFunctionInRescheduleList = [TickFunction](const FTickScheduleDetails& TSD)
+			{
+				return (TSD.TickFunction == TickFunction);
+			};
+			bool bFound = TickFunctionsToReschedule.ContainsByPredicate(FindTickFunctionInRescheduleList);
+			FTickFunction* PrevComparisonFunction = nullptr;
 			FTickFunction* ComparisonFunction = AllCoolingDownTickFunctions.Head;
-			bool bFound = false;
 			while (ComparisonFunction && !bFound)
 			{
 				if (ComparisonFunction == TickFunction)
 				{
 					bFound = true;
-					if (PrevComparisionFunction)
+					if (PrevComparisonFunction)
 					{
-						PrevComparisionFunction->Next = TickFunction->Next;
+						PrevComparisonFunction->Next = TickFunction->Next;
 					}
 					else
 					{
@@ -1233,7 +1240,7 @@ public:
 				}
 				else
 				{
-					PrevComparisionFunction = ComparisonFunction;
+					PrevComparisonFunction = ComparisonFunction;
 					ComparisonFunction = ComparisonFunction->Next;
 				}
 			}
