Don't search all scalar parameters for atlases at runtime, avoids hitching during texture streaming
#rnx

diff --git a/Engine/Source/Runtime/Engine/Private/Materials/Material.cpp b/Engine/Source/Runtime/Engine/Private/Materials/Material.cpp
--- a/Engine/Source/Runtime/Engine/Private/Materials/Material.cpp
+++ b/Engine/Source/Runtime/Engine/Private/Materials/Material.cpp
@@ -939,6 +939,7 @@ void UMaterial::GetUsedTextures(TArray<UTexture*>& OutTextures, EMaterialQuality
 						}
 					}
 				}
+#if WITH_EDITOR
 				// Also look for any scalar parameters that are acting as lookups for an atlas texture, and store the atlas texture
 				const TArray<TRefCountPtr<FMaterialUniformExpression> >* AtlasExpressions[1] =
 				{
@@ -961,6 +962,7 @@ void UMaterial::GetUsedTextures(TArray<UTexture*>& OutTextures, EMaterialQuality
 						}
 					}
 				}
+#endif 
 			}
 		}
 	}
diff --git a/Engine/Source/Runtime/Engine/Private/Materials/MaterialInstance.cpp b/Engine/Source/Runtime/Engine/Private/Materials/MaterialInstance.cpp
--- a/Engine/Source/Runtime/Engine/Private/Materials/MaterialInstance.cpp
+++ b/Engine/Source/Runtime/Engine/Private/Materials/MaterialInstance.cpp
@@ -1071,7 +1071,6 @@ void UMaterialInstance::GetUsedTextures(TArray<UTexture*>& OutTextures, EMateria
 					if (MaterialResource)
 					{
 						GetTextureExpressionValues(MaterialResource, OutTextures);
-						GetAtlasTextureValues(MaterialResource, OutTextures);
 					}
 				}
 			}
