Make all engine and smaple uses of player controller iterator consistent at validating that the weak object pointer is valid before dereferencing.
#rnx

diff --git a/Engine/Plugins/FX/Niagara/Source/Niagara/Private/NiagaraSystemInstance.cpp b/Engine/Plugins/FX/Niagara/Source/Niagara/Private/NiagaraSystemInstance.cpp
--- a/Engine/Plugins/FX/Niagara/Source/Niagara/Private/NiagaraSystemInstance.cpp
+++ b/Engine/Plugins/FX/Niagara/Source/Niagara/Private/NiagaraSystemInstance.cpp
@@ -914,7 +914,7 @@ void FNiagaraSystemInstance::TickInstanceParameters(float DeltaSeconds)
 			for (FConstPlayerControllerIterator Iterator = World->GetPlayerControllerIterator(); Iterator; ++Iterator)
 			{
 				APlayerController* PlayerController = Iterator->Get();
-				if (PlayerController->IsLocalPlayerController())
+				if (PlayerController && PlayerController->IsLocalPlayerController())
 				{
 					FVector* POVLoc = new(PlayerViewLocations) FVector;
 					FRotator POVRotation;
diff --git a/Engine/Plugins/Online/OnlineSubsystemSteam/Source/Private/OnlineAuthInterfaceSteam.cpp b/Engine/Plugins/Online/OnlineSubsystemSteam/Source/Private/OnlineAuthInterfaceSteam.cpp
--- a/Engine/Plugins/Online/OnlineSubsystemSteam/Source/Private/OnlineAuthInterfaceSteam.cpp
+++ b/Engine/Plugins/Online/OnlineSubsystemSteam/Source/Private/OnlineAuthInterfaceSteam.cpp
@@ -317,15 +317,14 @@ bool FOnlineAuthSteam::KickPlayer(const FUniqueNetId& InUserId)
 			return false;
 		}
 
-		FConstPlayerControllerIterator Itr = World->GetPlayerControllerIterator();
-		for (; Itr; ++Itr)
+		for (FConstPlayerControllerIterator Itr = World->GetPlayerControllerIterator(); Itr; ++Itr)
 		{
-			if (Itr->IsValid() && Itr->Get()->PlayerState != nullptr &&
-				Itr->Get()->PlayerState->UniqueId.IsValid() && 
-				*(Itr->Get()->PlayerState->UniqueId.GetUniqueNetId()) == InUserId)
+			APlayerController* PC = It->Get();
+			if (PC && PC->PlayerState != nullptr && PC->PlayerState->UniqueId.IsValid() &&
+				*(PC->PlayerState->UniqueId.GetUniqueNetId()) == InUserId)
 			{
 				const FText AuthKickReason = NSLOCTEXT("NetworkErrors", "HostClosedConnection", "Host closed the connection.");
-				bKickSuccess = GameMode->GameSession->KickPlayer(Itr->Get(), AuthKickReason);
+				bKickSuccess = GameMode->GameSession->KickPlayer(PC, AuthKickReason);
 				break;
 			}
 		}
diff --git a/Engine/Source/Editor/WorldBrowser/Private/Tiles/WorldTileCollectionModel.cpp b/Engine/Source/Editor/WorldBrowser/Private/Tiles/WorldTileCollectionModel.cpp
--- a/Engine/Source/Editor/WorldBrowser/Private/Tiles/WorldTileCollectionModel.cpp
+++ b/Engine/Source/Editor/WorldBrowser/Private/Tiles/WorldTileCollectionModel.cpp
@@ -620,9 +620,11 @@ bool FWorldTileCollectionModel::GetPlayerView(FVector& Location, FRotator& Rotat
 		UWorld* SimulationWorld = GetSimulationWorld();
 		for (FConstPlayerControllerIterator Iterator = SimulationWorld->GetPlayerControllerIterator(); Iterator; ++Iterator)
 		{
-			APlayerController* PlayerActor = Iterator->Get();
-			PlayerActor->GetPlayerViewPoint(Location, Rotation);
-			return true;
+			if (APlayerController* PlayerActor = Iterator->Get())
+			{
+				PlayerActor->GetPlayerViewPoint(Location, Rotation);
+				return true;
+			}
 		}
 	}
 	
diff --git a/Engine/Source/Runtime/AugmentedReality/Private/ARDebugDrawHelpers.cpp b/Engine/Source/Runtime/AugmentedReality/Private/ARDebugDrawHelpers.cpp
--- a/Engine/Source/Runtime/AugmentedReality/Private/ARDebugDrawHelpers.cpp
+++ b/Engine/Source/Runtime/AugmentedReality/Private/ARDebugDrawHelpers.cpp
@@ -25,7 +25,7 @@ namespace ARDebugHelpers
 			for( FConstPlayerControllerIterator Iterator = InWorld->GetPlayerControllerIterator(); Iterator; ++Iterator )
 			{
 				APlayerController* PlayerController = Iterator->Get();
-				if (PlayerController->MyHUD && PlayerController->Player)
+				if (PlayerController && PlayerController->MyHUD && PlayerController->Player)
 				{
 					PlayerController->MyHUD->AddDebugText(Text, BaseAct, Duration, TextLocation, TextLocation, TextColor, true, true, false, nullptr, Scale, bDrawShadow);
 					bDrawnSuccessfully = true;
diff --git a/Engine/Source/Runtime/Engine/Private/Actor.cpp b/Engine/Source/Runtime/Engine/Private/Actor.cpp
--- a/Engine/Source/Runtime/Engine/Private/Actor.cpp
+++ b/Engine/Source/Runtime/Engine/Private/Actor.cpp
@@ -3289,7 +3289,10 @@ void AActor::DisableInput(APlayerController* PlayerController)
 		{
 			for (FConstPlayerControllerIterator PCIt = GetWorld()->GetPlayerControllerIterator(); PCIt; ++PCIt)
 			{
-				(*PCIt)->PopInputComponent(InputComponent);
+				if (APlayerController* PC = PCIt->Get())
+				{
+					PC->PopInputComponent(InputComponent);
+				}
 			}
 		}
 	}
diff --git a/Engine/Source/Runtime/Engine/Private/CheatManager.cpp b/Engine/Source/Runtime/Engine/Private/CheatManager.cpp
--- a/Engine/Source/Runtime/Engine/Private/CheatManager.cpp
+++ b/Engine/Source/Runtime/Engine/Private/CheatManager.cpp
@@ -511,24 +511,30 @@ void UCheatManager::SetLevelStreamingStatus(FName PackageName, bool bShouldBeLoa
 	{
 		for( FConstPlayerControllerIterator Iterator = World->GetPlayerControllerIterator(); Iterator; ++Iterator )
 		{
-			(*Iterator)->ClientUpdateLevelStreamingStatus((*Iterator)->NetworkRemapPath(PackageName, false), bShouldBeLoaded, bShouldBeVisible, false, INDEX_NONE );
+			if (APlayerController* PC = Iterator->Get())
+			{
+				PC->ClientUpdateLevelStreamingStatus(PC->NetworkRemapPath(PackageName, false), bShouldBeLoaded, bShouldBeVisible, false, INDEX_NONE);
+			}
 		}
 	}
 	else
 	{
 		for( FConstPlayerControllerIterator Iterator = World->GetPlayerControllerIterator(); Iterator; ++Iterator )
 		{
-			TArray<FUpdateLevelStreamingLevelStatus> LevelStatuses;
-			for (ULevelStreaming* StreamingLevel : World->GetStreamingLevels())
+			if (APlayerController* PC = Iterator->Get())
 			{
-				FUpdateLevelStreamingLevelStatus& LevelStatus = *new( LevelStatuses ) FUpdateLevelStreamingLevelStatus();
-				LevelStatus.PackageName = (*Iterator)->NetworkRemapPath(StreamingLevel->GetWorldAssetPackageFName(), false);
-				LevelStatus.bNewShouldBeLoaded = bShouldBeLoaded;
-				LevelStatus.bNewShouldBeVisible = bShouldBeVisible;
-				LevelStatus.bNewShouldBlockOnLoad = false;
-				LevelStatus.LODIndex = INDEX_NONE;
+				TArray<FUpdateLevelStreamingLevelStatus> LevelStatuses;
+				for (ULevelStreaming* StreamingLevel : World->GetStreamingLevels())
+				{
+					FUpdateLevelStreamingLevelStatus& LevelStatus = *new(LevelStatuses) FUpdateLevelStreamingLevelStatus();
+					LevelStatus.PackageName = PC->NetworkRemapPath(StreamingLevel->GetWorldAssetPackageFName(), false);
+					LevelStatus.bNewShouldBeLoaded = bShouldBeLoaded;
+					LevelStatus.bNewShouldBeVisible = bShouldBeVisible;
+					LevelStatus.bNewShouldBlockOnLoad = false;
+					LevelStatus.LODIndex = INDEX_NONE;
+				}
+				PC->ClientUpdateMultipleLevelsStreamingStatus(LevelStatuses);
 			}
-			( *Iterator )->ClientUpdateMultipleLevelsStreamingStatus( LevelStatuses );
 		}
 	}
 }
diff --git a/Engine/Source/Runtime/Engine/Private/DrawDebugHelpers.cpp b/Engine/Source/Runtime/Engine/Private/DrawDebugHelpers.cpp
--- a/Engine/Source/Runtime/Engine/Private/DrawDebugHelpers.cpp
+++ b/Engine/Source/Runtime/Engine/Private/DrawDebugHelpers.cpp
@@ -730,7 +730,7 @@ void DrawDebugString(const UWorld* InWorld, FVector const& TextLocation, const F
 		for( FConstPlayerControllerIterator Iterator = InWorld->GetPlayerControllerIterator(); Iterator; ++Iterator )
 		{
 			APlayerController* PlayerController = Iterator->Get();
-			if (PlayerController->MyHUD && PlayerController->Player)
+			if (PlayerController && PlayerController->MyHUD && PlayerController->Player)
 			{
 				PlayerController->MyHUD->AddDebugText(Text, BaseAct, Duration, TextLocation, TextLocation, TextColor, true, (TestBaseActor==NULL), false, nullptr, FontScale, bDrawShadow);
 			}
@@ -745,7 +745,7 @@ void FlushDebugStrings( const UWorld* InWorld )
 	{
 		// if it's a player
 		APlayerController* PlayerController = Iterator->Get();
-		if (PlayerController->MyHUD)
+		if (PlayerController && PlayerController->MyHUD)
 		{
 			PlayerController->MyHUD->RemoveAllDebugStrings();
 		}
diff --git a/Engine/Source/Runtime/Engine/Private/EngineUtils.cpp b/Engine/Source/Runtime/Engine/Private/EngineUtils.cpp
--- a/Engine/Source/Runtime/Engine/Private/EngineUtils.cpp
+++ b/Engine/Source/Runtime/Engine/Private/EngineUtils.cpp
@@ -403,39 +403,40 @@ TArray<FSubLevelStatus> GetSubLevelsStatus( UWorld* World )
 
 	for( FConstPlayerControllerIterator Iterator = World->GetPlayerControllerIterator(); Iterator; ++Iterator )
 	{
-		APlayerController* PlayerController = Iterator->Get();
-
-		if (APawn* PCPawn = PlayerController->GetPawn())
+		if (APlayerController* PlayerController = Iterator->Get())
 		{
-			// need to do a trace down here
-			//TraceActor = Trace( out_HitLocation, out_HitNormal, TraceDest, TraceStart, false, TraceExtent, HitInfo, true );
-			FHitResult Hit(1.f);
-
-			// this will not work for flying around :-(
-			PlayerController->GetWorld()->LineTraceSingleByObjectType(Hit, PCPawn->GetActorLocation(), (PCPawn->GetActorLocation()-FVector(0.f, 0.f, 256.f)), FCollisionObjectQueryParams(ECC_WorldStatic), FCollisionQueryParams(SCENE_QUERY_STAT(FindLevel), true, PCPawn));
+			if (APawn* PCPawn = PlayerController->GetPawn())
+			{
+				// need to do a trace down here
+				//TraceActor = Trace( out_HitLocation, out_HitNormal, TraceDest, TraceStart, false, TraceExtent, HitInfo, true );
+				FHitResult Hit(1.f);
 
-			ULevel* LevelPlayerIsIn = nullptr;
+				// this will not work for flying around :-(
+				PlayerController->GetWorld()->LineTraceSingleByObjectType(Hit, PCPawn->GetActorLocation(), (PCPawn->GetActorLocation() - FVector(0.f, 0.f, 256.f)), FCollisionObjectQueryParams(ECC_WorldStatic), FCollisionQueryParams(SCENE_QUERY_STAT(FindLevel), true, PCPawn));
 
-			if (AActor* HitActor = Hit.GetActor())
-			{
-				LevelPlayerIsIn = HitActor->GetLevel();
-			}
-			else if (UPrimitiveComponent* HitComponent = Hit.Component.Get())
-			{
-				LevelPlayerIsIn = HitComponent->GetComponentLevel();
-			}
+				ULevel* LevelPlayerIsIn = nullptr;
 
-			if (LevelPlayerIsIn)
-			{
-				FName LevelName = LevelPlayerIsIn->GetOutermost()->GetFName();
-				FSubLevelStatus* LevelStatusPlayerIn = Result.FindByPredicate([LevelName](const FSubLevelStatus& InLevelStatus)
+				if (AActor* HitActor = Hit.GetActor())
 				{
-					return InLevelStatus.PackageName == LevelName;
-				});
+					LevelPlayerIsIn = HitActor->GetLevel();
+				}
+				else if (UPrimitiveComponent* HitComponent = Hit.Component.Get())
+				{
+					LevelPlayerIsIn = HitComponent->GetComponentLevel();
+				}
 
-				if (LevelStatusPlayerIn)
+				if (LevelPlayerIsIn)
 				{
-					LevelStatusPlayerIn->bPlayerInside = true;
+					FName LevelName = LevelPlayerIsIn->GetOutermost()->GetFName();
+					FSubLevelStatus* LevelStatusPlayerIn = Result.FindByPredicate([LevelName](const FSubLevelStatus& InLevelStatus)
+					{
+						return InLevelStatus.PackageName == LevelName;
+					});
+
+					if (LevelStatusPlayerIn)
+					{
+						LevelStatusPlayerIn->bPlayerInside = true;
+					}
 				}
 			}
 		}
diff --git a/Engine/Source/Runtime/Engine/Private/GameInstance.cpp b/Engine/Source/Runtime/Engine/Private/GameInstance.cpp
--- a/Engine/Source/Runtime/Engine/Private/GameInstance.cpp
+++ b/Engine/Source/Runtime/Engine/Private/GameInstance.cpp
@@ -744,9 +744,10 @@ APlayerController* UGameInstance::GetFirstLocalPlayerController(UWorld* World) c
 		// Only return a local PlayerController from the given World.
 		for (FConstPlayerControllerIterator Iterator = World->GetPlayerControllerIterator(); Iterator; ++Iterator)
 		{
-			if (*Iterator != nullptr && (*Iterator)->IsLocalController())
+			APlayerController* PC = Iterator->Get();
+			if (PC && PC->IsLocalController())
 			{
-				return Iterator->Get();
+				return PC;
 			}
 		}
 	}
diff --git a/Engine/Source/Runtime/Engine/Private/GameMode.cpp b/Engine/Source/Runtime/Engine/Private/GameMode.cpp
--- a/Engine/Source/Runtime/Engine/Private/GameMode.cpp
+++ b/Engine/Source/Runtime/Engine/Private/GameMode.cpp
@@ -218,7 +218,7 @@ void AGameMode::HandleMatchHasStarted()
 	for( FConstPlayerControllerIterator Iterator = GetWorld()->GetPlayerControllerIterator(); Iterator; ++Iterator )
 	{
 		APlayerController* PlayerController = Iterator->Get();
-		if ((PlayerController->GetPawn() == nullptr) && PlayerCanRestart(PlayerController))
+		if (PlayerController && (PlayerController->GetPawn() == nullptr) && PlayerCanRestart(PlayerController))
 		{
 			RestartPlayer(PlayerController);
 		}
@@ -241,7 +241,7 @@ void AGameMode::HandleMatchHasStarted()
 		for( FConstPlayerControllerIterator Iterator = GetWorld()->GetPlayerControllerIterator(); Iterator; ++Iterator )
 		{
 			APlayerController* PlayerController = Iterator->Get();
-			if( PlayerController->CheatManager != nullptr )
+			if (PlayerController &&  PlayerController->CheatManager != nullptr)
 			{
 				PlayerController->CheatManager->BugItGoString( BugLocString, BugRotString );
 			}
@@ -597,7 +597,10 @@ void AGameMode::Broadcast( AActor* Sender, const FString& Msg, FName Type )
 
 	for( FConstPlayerControllerIterator Iterator = GetWorld()->GetPlayerControllerIterator(); Iterator; ++Iterator )
 	{
-		(*Iterator)->ClientTeamMessage( SenderPlayerState, Msg, Type );
+		if (APlayerController* PC = Iterator->Get())
+		{
+			PC->ClientTeamMessage(SenderPlayerState, Msg, Type);
+		}
 	}
 }
 
@@ -606,7 +609,10 @@ void AGameMode::BroadcastLocalized( AActor* Sender, TSubclassOf<ULocalMessage> M
 {
 	for( FConstPlayerControllerIterator Iterator = GetWorld()->GetPlayerControllerIterator(); Iterator; ++Iterator )
 	{
-		(*Iterator)->ClientReceiveLocalizedMessage( Message, Switch, RelatedPlayerState_1, RelatedPlayerState_2, OptionalObject );
+		if (APlayerController* PC = Iterator->Get())
+		{
+			PC->ClientReceiveLocalizedMessage(Message, Switch, RelatedPlayerState_1, RelatedPlayerState_2, OptionalObject);
+		}
 	}
 }
 
diff --git a/Engine/Source/Runtime/Engine/Private/GameModeBase.cpp b/Engine/Source/Runtime/Engine/Private/GameModeBase.cpp
--- a/Engine/Source/Runtime/Engine/Private/GameModeBase.cpp
+++ b/Engine/Source/Runtime/Engine/Private/GameModeBase.cpp
@@ -143,7 +143,7 @@ int32 AGameModeBase::GetNumPlayers()
 	for (FConstPlayerControllerIterator Iterator = GetWorld()->GetPlayerControllerIterator(); Iterator; ++Iterator)
 	{
 		APlayerController* PlayerActor = Iterator->Get();
-		if (PlayerActor->PlayerState && !MustSpectate(PlayerActor))
+		if (PlayerActor && PlayerActor->PlayerState && !MustSpectate(PlayerActor))
 		{
 			PlayerCount++;
 		}
@@ -157,7 +157,7 @@ int32 AGameModeBase::GetNumSpectators()
 	for (FConstPlayerControllerIterator Iterator = GetWorld()->GetPlayerControllerIterator(); Iterator; ++Iterator)
 	{
 		APlayerController* PlayerActor = Iterator->Get();
-		if (PlayerActor->PlayerState && MustSpectate(PlayerActor))
+		if (PlayerActor && PlayerActor->PlayerState && MustSpectate(PlayerActor))
 		{
 			PlayerCount++;
 		}
@@ -261,7 +261,7 @@ void AGameModeBase::ForceClearUnpauseDelegates(AActor* PauseActor)
 			for (FConstPlayerControllerIterator Iterator = GetWorld()->GetPlayerControllerIterator(); Iterator; ++Iterator)
 			{
 				APlayerController* Player = Iterator->Get();
-				if (Player->PlayerState != nullptr
+				if (Player && Player->PlayerState != nullptr
 					&&	Player->PlayerState != PC->PlayerState
 					&& !Player->IsPendingKillPending() && !Player->PlayerState->IsPendingKillPending())
 				{
@@ -355,17 +355,19 @@ APlayerController* AGameModeBase::ProcessClientTravel(FString& FURL, FGuid NextM
 	APlayerController* LocalPlayerController = nullptr;
 	for (FConstPlayerControllerIterator Iterator = GetWorld()->GetPlayerControllerIterator(); Iterator; ++Iterator)
 	{
-		APlayerController* PlayerController = Iterator->Get();
-		if (Cast<UNetConnection>(PlayerController->Player) != nullptr)
+		if (APlayerController* PlayerController = Iterator->Get())
 		{
-			// Remote player
-			PlayerController->ClientTravel(FURL, TRAVEL_Relative, bSeamless, NextMapGuid);
-		}
-		else
-		{
-			// Local player
-			LocalPlayerController = PlayerController;
-			PlayerController->PreClientTravel(FURL, bAbsolute ? TRAVEL_Absolute : TRAVEL_Relative, bSeamless);
+			if (Cast<UNetConnection>(PlayerController->Player) != nullptr)
+			{
+				// Remote player
+				PlayerController->ClientTravel(FURL, TRAVEL_Relative, bSeamless, NextMapGuid);
+			}
+			else
+			{
+				// Local player
+				LocalPlayerController = PlayerController;
+				PlayerController->PreClientTravel(FURL, bAbsolute ? TRAVEL_Absolute : TRAVEL_Relative, bSeamless);
+			}
 		}
 	}
 
diff --git a/Engine/Source/Runtime/Engine/Private/GameNetworkManager.cpp b/Engine/Source/Runtime/Engine/Private/GameNetworkManager.cpp
--- a/Engine/Source/Runtime/Engine/Private/GameNetworkManager.cpp
+++ b/Engine/Source/Runtime/Engine/Private/GameNetworkManager.cpp
@@ -109,7 +109,10 @@ void AGameNetworkManager::UpdateNetSpeeds(bool bIsLanMatch)
 		AdjustedNetSpeed = NewNetSpeed;
 		for( FConstPlayerControllerIterator Iterator = GetWorld()->GetPlayerControllerIterator(); Iterator; ++Iterator )
 		{
-			(*Iterator)->SetNetSpeed(AdjustedNetSpeed);
+			if (APlayerController* PC = Iterator->Get())
+			{
+				PC->SetNetSpeed(AdjustedNetSpeed);
+			}
 		}
 	}
 }
diff --git a/Engine/Source/Runtime/Engine/Private/GameSession.cpp b/Engine/Source/Runtime/Engine/Private/GameSession.cpp
--- a/Engine/Source/Runtime/Engine/Private/GameSession.cpp
+++ b/Engine/Source/Runtime/Engine/Private/GameSession.cpp
@@ -32,7 +32,7 @@ APlayerController* GetPlayerControllerFromNetId(UWorld* World, const FUniqueNetI
 		{
 			APlayerController* PlayerController = Iterator->Get();
 			// Determine if this is a player with replication
-			if (PlayerController->PlayerState != NULL && PlayerController->PlayerState->UniqueId.IsValid())
+			if (PlayerController && PlayerController->PlayerState != NULL && PlayerController->PlayerState->UniqueId.IsValid())
 			{
 				// If the ids match, then this is the right player.
 				if (*PlayerController->PlayerState->UniqueId == PlayerNetId)
@@ -64,7 +64,7 @@ void AGameSession::HandleMatchHasStarted()
 		for (FConstPlayerControllerIterator Iterator = World->GetPlayerControllerIterator(); Iterator; ++Iterator)
 		{
 			APlayerController* PlayerController = Iterator->Get();
-			if (!PlayerController->IsLocalController())
+			if (PlayerController && !PlayerController->IsLocalController())
 			{
 				PlayerController->ClientStartOnlineSession();
 			}
@@ -106,7 +106,7 @@ void AGameSession::HandleMatchHasEnded()
 		for (FConstPlayerControllerIterator Iterator = World->GetPlayerControllerIterator(); Iterator; ++Iterator)
 		{
 			APlayerController* PlayerController = Iterator->Get();
-			if (!PlayerController->IsLocalController())
+			if (PlayerController && !PlayerController->IsLocalController())
 			{
 				PlayerController->ClientEndOnlineSession();
 			}
diff --git a/Engine/Source/Runtime/Engine/Private/GameViewportClient.cpp b/Engine/Source/Runtime/Engine/Private/GameViewportClient.cpp
--- a/Engine/Source/Runtime/Engine/Private/GameViewportClient.cpp
+++ b/Engine/Source/Runtime/Engine/Private/GameViewportClient.cpp
@@ -1539,7 +1539,10 @@ void UGameViewportClient::Draw(FViewport* InViewport, FCanvas* SceneCanvas)
 	{
 		for( FConstPlayerControllerIterator Iterator = MyWorld->GetPlayerControllerIterator(); Iterator; ++Iterator )
 		{
-			(*Iterator)->GetPlayerViewPoint( PlayerCameraLocation, PlayerCameraRotation );
+			if (APlayerController* PC = Iterator->Get())
+			{
+				PC->GetPlayerViewPoint(PlayerCameraLocation, PlayerCameraRotation);
+			}
 		}
 	}
 
@@ -3431,7 +3434,7 @@ bool UGameViewportClient::RequestBugScreenShot(const TCHAR* Cmd, bool bDisplayHU
 		{
 			for (FConstPlayerControllerIterator Iterator = ViewportWorld->GetPlayerControllerIterator(); Iterator; ++Iterator)
 			{
-						APlayerController* PlayerController = Iterator->Get();
+				APlayerController* PlayerController = Iterator->Get();
 				if (PlayerController && PlayerController->GetHUD())
 				{
 					PlayerController->GetHUD()->HandleBugScreenShot();
diff --git a/Engine/Source/Runtime/Engine/Private/Interpolation.cpp b/Engine/Source/Runtime/Engine/Private/Interpolation.cpp
--- a/Engine/Source/Runtime/Engine/Private/Interpolation.cpp
+++ b/Engine/Source/Runtime/Engine/Private/Interpolation.cpp
@@ -924,7 +924,7 @@ PRAGMA_ENABLE_DEPRECATION_WARNINGS
 						APlayerController *PC = Iterator->Get();
 
 						// If it's a player and this sequence is compatible with the player...
-						if (IsMatineeCompatibleWithPlayer( PC ) )
+						if (PC && IsMatineeCompatibleWithPlayer( PC ) )
 						{
 							// create a new instance with this player
 							UInterpGroupInstDirector* NewGroupInstDir = NewObject<UInterpGroupInstDirector>(this, NAME_None, RF_Transactional);
@@ -1292,7 +1292,7 @@ void AMatineeActor::EnableCinematicMode(bool bEnable)
 		for (FConstPlayerControllerIterator Iterator = GetWorld()->GetPlayerControllerIterator(); Iterator; ++Iterator)
 		{
 			APlayerController *PC = Iterator->Get();
-			if (bReplicates || PC->IsLocalController())
+			if (PC && (bReplicates || PC->IsLocalController()))
 			{
 				PC->SetCinematicMode(bEnable, bHidePlayer, bHideHud, bDisableMovementInput, bDisableLookAtInput);
 			}
diff --git a/Engine/Source/Runtime/Engine/Private/LevelScriptActor.cpp b/Engine/Source/Runtime/Engine/Private/LevelScriptActor.cpp
--- a/Engine/Source/Runtime/Engine/Private/LevelScriptActor.cpp
+++ b/Engine/Source/Runtime/Engine/Private/LevelScriptActor.cpp
@@ -97,7 +97,10 @@ void ALevelScriptActor::SetCinematicMode(bool bCinematicMode, bool bHidePlayer,
 	// Loop through all player controllers and call their version of SetCinematicMode which is where all the magic happens and replication occurs
 	for( FConstPlayerControllerIterator Iterator = GetWorld()->GetPlayerControllerIterator(); Iterator; ++Iterator )
 	{
-		(*Iterator)->SetCinematicMode(bCinematicMode, bHidePlayer, bAffectsHUD, bAffectsMovement, bAffectsTurning);
+		if (APlayerController* PC = Iterator->Get())
+		{
+			PC->SetCinematicMode(bCinematicMode, bHidePlayer, bAffectsHUD, bAffectsMovement, bAffectsTurning);
+		}
 	}
 }
 
diff --git a/Engine/Source/Runtime/Engine/Private/LevelStreaming.cpp b/Engine/Source/Runtime/Engine/Private/LevelStreaming.cpp
--- a/Engine/Source/Runtime/Engine/Private/LevelStreaming.cpp
+++ b/Engine/Source/Runtime/Engine/Private/LevelStreaming.cpp
@@ -190,26 +190,27 @@ void FStreamLevelAction::ActivateLevel( ULevelStreaming* LevelStreamingObject )
 		// If we have a valid world
 		if (UWorld* LevelWorld = LevelStreamingObject->GetWorld())
 		{
+			const bool bShouldBeLoaded = LevelStreamingObject->ShouldBeLoaded();
+			const bool bShouldBeVisible = LevelStreamingObject->ShouldBeVisible();
+
+			UE_LOG(LogLevel, Log, TEXT("ActivateLevel %s %i %i %i"),
+				*LevelStreamingObject->GetWorldAssetPackageName(),
+				bShouldBeLoaded,
+				bShouldBeVisible,
+				bShouldBlock);
+
 			// Notify players of the change
 			for (FConstPlayerControllerIterator Iterator = LevelWorld->GetPlayerControllerIterator(); Iterator; ++Iterator)
 			{
-				APlayerController* PlayerController = Iterator->Get();
-
-				const bool bShouldBeLoaded = LevelStreamingObject->ShouldBeLoaded();
-				const bool bShouldBeVisible = LevelStreamingObject->ShouldBeVisible();
-
-				UE_LOG(LogLevel, Log, TEXT("ActivateLevel %s %i %i %i"), 
-					*LevelStreamingObject->GetWorldAssetPackageName(), 
-					bShouldBeLoaded, 
-					bShouldBeVisible, 
-					bShouldBlock);
-
-				PlayerController->LevelStreamingStatusChanged( 
-					LevelStreamingObject, 
-					bShouldBeLoaded, 
-					bShouldBeVisible,
-					bShouldBlock, 
-					INDEX_NONE);
+				if (APlayerController* PlayerController = Iterator->Get())
+				{
+					PlayerController->LevelStreamingStatusChanged(
+						LevelStreamingObject,
+						bShouldBeLoaded,
+						bShouldBeVisible,
+						bShouldBlock,
+						INDEX_NONE);
+				}
 			}
 		}
 	}
diff --git a/Engine/Source/Runtime/Engine/Private/LevelTick.cpp b/Engine/Source/Runtime/Engine/Private/LevelTick.cpp
--- a/Engine/Source/Runtime/Engine/Private/LevelTick.cpp
+++ b/Engine/Source/Runtime/Engine/Private/LevelTick.cpp
@@ -1518,14 +1518,16 @@ void UWorld::Tick( ELevelTick TickType, float DeltaSeconds )
 				// Update cameras last. This needs to be done before NetUpdates, and after all actors have been ticked.
 				for( FConstPlayerControllerIterator Iterator = GetPlayerControllerIterator(); Iterator; ++Iterator )
 				{
-					APlayerController* PlayerController = Iterator->Get();
-					if (!bIsPaused || PlayerController->ShouldPerformFullTickWhenPaused())
+					if (APlayerController* PlayerController = Iterator->Get())
 					{
-						PlayerController->UpdateCameraManager(DeltaSeconds);
-					}
-					else if (PlayerController->PlayerCameraManager && FCameraPhotographyManager::IsSupported(this))
-					{
-						PlayerController->PlayerCameraManager->UpdateCameraPhotographyOnly();
+						if (!bIsPaused || PlayerController->ShouldPerformFullTickWhenPaused())
+						{
+							PlayerController->UpdateCameraManager(DeltaSeconds);
+						}
+						else if (PlayerController->PlayerCameraManager && FCameraPhotographyManager::IsSupported(this))
+						{
+							PlayerController->PlayerCameraManager->UpdateCameraPhotographyOnly();
+						}
 					}
 				}
 
diff --git a/Engine/Source/Runtime/Engine/Private/Particles/ParticleComponents.cpp b/Engine/Source/Runtime/Engine/Private/Particles/ParticleComponents.cpp
--- a/Engine/Source/Runtime/Engine/Private/Particles/ParticleComponents.cpp
+++ b/Engine/Source/Runtime/Engine/Private/Particles/ParticleComponents.cpp
@@ -4906,7 +4906,7 @@ void UParticleSystemComponent::TickComponent(float DeltaTime, enum ELevelTick Ti
 		for (FConstPlayerControllerIterator Iterator = World->GetPlayerControllerIterator(); Iterator; ++Iterator)
 		{
 			APlayerController* PlayerController = Iterator->Get();
-			if (PlayerController->IsLocalPlayerController())
+			if (PlayerController && PlayerController->IsLocalPlayerController())
 			{
 				FVector POVLoc;
 				FRotator POVRotation;
@@ -6640,7 +6640,7 @@ int32 UParticleSystemComponent::DetermineLODLevelForLocation(const FVector& Effe
 			for( FConstPlayerControllerIterator Iterator = World->GetPlayerControllerIterator(); Iterator; ++Iterator )
 			{
 				APlayerController* PlayerController = Iterator->Get();
-				if(PlayerController->IsLocalPlayerController())
+				if(PlayerController && PlayerController->IsLocalPlayerController())
 				{
 					FVector* POVLoc = new(PlayerViewLocations) FVector;
 					FRotator POVRotation;
@@ -7456,7 +7456,7 @@ bool UParticleSystemComponent::ShouldComputeLODFromGameThread()
 		for (FConstPlayerControllerIterator Iterator = World->GetPlayerControllerIterator(); Iterator; ++Iterator)
 		{
 			APlayerController* PlayerController = Iterator->Get();
-			if (PlayerController->IsLocalPlayerController())
+			if (PlayerController && PlayerController->IsLocalPlayerController())
 			{
 				bUseGameThread = true;
 				break;
diff --git a/Engine/Source/Runtime/Engine/Private/Player.cpp b/Engine/Source/Runtime/Engine/Private/Player.cpp
--- a/Engine/Source/Runtime/Engine/Private/Player.cpp
+++ b/Engine/Source/Runtime/Engine/Private/Player.cpp
@@ -77,11 +77,12 @@ APlayerController* UPlayer::GetPlayerController(UWorld* InWorld) const
 		return PlayerController;
 	}
 
-	for ( FConstPlayerControllerIterator Iterator = InWorld->GetPlayerControllerIterator(); Iterator; ++Iterator)
+	for (FConstPlayerControllerIterator Iterator = InWorld->GetPlayerControllerIterator(); Iterator; ++Iterator)
 	{
-		if ( (*Iterator)->GetLocalPlayer() == this )
+		APlayerController* PC = Iterator->Get();
+		if (PC && PC->GetLocalPlayer() == this)
 		{
-			return Iterator->Get();
+			return PC;
 		}
 	}
 
diff --git a/Engine/Source/Runtime/Engine/Private/PlayerMuteList.cpp b/Engine/Source/Runtime/Engine/Private/PlayerMuteList.cpp
--- a/Engine/Source/Runtime/Engine/Private/PlayerMuteList.cpp
+++ b/Engine/Source/Runtime/Engine/Private/PlayerMuteList.cpp
@@ -185,30 +185,32 @@ FString DumpMutelistState(UWorld* World)
 	{
 		for(FConstPlayerControllerIterator Iterator = World->GetPlayerControllerIterator(); Iterator; ++Iterator)
 		{
-			APlayerController* PlayerController = Iterator->Get();
-			Output += FString::Printf(TEXT("Player: %s\n"), PlayerController->PlayerState ? *PlayerController->PlayerState->GetPlayerName() : TEXT("NONAME"));
-			Output += FString::Printf(TEXT("VoiceChannel: %d\n"), PlayerController->MuteList.VoiceChannelIdx);
-			Output += FString::Printf(TEXT("Handshake: %s\n"), PlayerController->MuteList.bHasVoiceHandshakeCompleted ? TEXT("true") : TEXT("false"));
-			
-			Output += FString(TEXT("System mutes:\n"));
-			for (int32 idx=0; idx < PlayerController->MuteList.VoiceMuteList.Num(); idx++)
+			if (APlayerController* PlayerController = Iterator->Get())
 			{
-				Output += FString::Printf(TEXT("%s\n"), *PlayerController->MuteList.VoiceMuteList[idx]->ToString());
+				Output += FString::Printf(TEXT("Player: %s\n"), PlayerController->PlayerState ? *PlayerController->PlayerState->GetPlayerName() : TEXT("NONAME"));
+				Output += FString::Printf(TEXT("VoiceChannel: %d\n"), PlayerController->MuteList.VoiceChannelIdx);
+				Output += FString::Printf(TEXT("Handshake: %s\n"), PlayerController->MuteList.bHasVoiceHandshakeCompleted ? TEXT("true") : TEXT("false"));
+
+				Output += FString(TEXT("System mutes:\n"));
+				for (int32 idx = 0; idx < PlayerController->MuteList.VoiceMuteList.Num(); idx++)
+				{
+					Output += FString::Printf(TEXT("%s\n"), *PlayerController->MuteList.VoiceMuteList[idx]->ToString());
+				}
+
+				Output += FString(TEXT("Gameplay mutes:\n"));
+				for (int32 idx = 0; idx < PlayerController->MuteList.GameplayVoiceMuteList.Num(); idx++)
+				{
+					Output += FString::Printf(TEXT("%s\n"), *PlayerController->MuteList.GameplayVoiceMuteList[idx]->ToString());
+				}
+
+				Output += FString(TEXT("Filter:\n"));
+				for (int32 idx = 0; idx < PlayerController->MuteList.VoicePacketFilter.Num(); idx++)
+				{
+					Output += FString::Printf(TEXT("%s\n"), *PlayerController->MuteList.VoicePacketFilter[idx]->ToString());
+				}
+
+				Output += TEXT("\n");
 			}
-
-			Output += FString(TEXT("Gameplay mutes:\n"));
-			for (int32 idx=0; idx < PlayerController->MuteList.GameplayVoiceMuteList.Num(); idx++)
-			{
-				Output += FString::Printf(TEXT("%s\n"), *PlayerController->MuteList.GameplayVoiceMuteList[idx]->ToString());
-			}
-
-			Output += FString(TEXT("Filter:\n"));
-			for (int32 idx=0; idx < PlayerController->MuteList.VoicePacketFilter.Num(); idx++)
-			{
-				Output += FString::Printf(TEXT("%s\n"), *PlayerController->MuteList.VoicePacketFilter[idx]->ToString());
-			}
-
-			Output += TEXT("\n");
 		}
 	}
 
diff --git a/Engine/Source/Runtime/Engine/Private/UnrealEngine.cpp b/Engine/Source/Runtime/Engine/Private/UnrealEngine.cpp
--- a/Engine/Source/Runtime/Engine/Private/UnrealEngine.cpp
+++ b/Engine/Source/Runtime/Engine/Private/UnrealEngine.cpp
@@ -6117,7 +6117,7 @@ bool UEngine::HandleMergeMeshCommand( const TCHAR* Cmd, FOutputDevice& Ar, UWorl
 	for (FConstPlayerControllerIterator Iterator = InWorld->GetPlayerControllerIterator(); Iterator; ++Iterator)
 	{
 		APlayerController* PlayerController = Iterator->Get();
-		if (PlayerController->GetCharacter() != NULL && PlayerController->GetCharacter()->GetMesh())
+		if (PlayerController && PlayerController->GetCharacter() != NULL && PlayerController->GetCharacter()->GetMesh())
 		{
 			PlayerPawn = PlayerController->GetCharacter();
 			PlayerMesh = PlayerController->GetCharacter()->GetMesh()->SkeletalMesh;
@@ -6139,7 +6139,7 @@ bool UEngine::HandleMergeMeshCommand( const TCHAR* Cmd, FOutputDevice& Ar, UWorl
 		for (FConstPlayerControllerIterator Iterator = InWorld->GetPlayerControllerIterator(); Iterator; ++Iterator)
 		{
 			APlayerController* PlayerController = Iterator->Get();
-			if (PlayerController->GetPawn() != NULL)
+			if (PlayerController && PlayerController->GetPawn() != NULL)
 			{
 				PlayerPawn = PlayerController->GetPawn();
 				break;
@@ -13192,12 +13192,15 @@ bool UEngine::CommitMapChange( FWorldContext &Context )
 						// notify players of the change
 						for( FConstPlayerControllerIterator Iterator = Context.World()->GetPlayerControllerIterator(); Iterator; ++Iterator )
 						{
-							(*Iterator)->LevelStreamingStatusChanged( 
-								StreamingLevel, 
-								bShouldBeLoaded, 
-								bShouldBeVisible,
-								StreamingLevel->bShouldBlockOnLoad,
-								StreamingLevel->GetLevelLODIndex());							
+							if (APlayerController* PlayerController = Iterator->Get())
+							{
+								PlayerController->LevelStreamingStatusChanged(
+									StreamingLevel,
+									bShouldBeLoaded,
+									bShouldBeVisible,
+									StreamingLevel->bShouldBlockOnLoad,
+									StreamingLevel->GetLevelLODIndex());
+							}
 						}
 					}
 #endif // WITH_SERVER_CODE
diff --git a/Engine/Source/Runtime/Engine/Private/World.cpp b/Engine/Source/Runtime/Engine/Private/World.cpp
--- a/Engine/Source/Runtime/Engine/Private/World.cpp
+++ b/Engine/Source/Runtime/Engine/Private/World.cpp
@@ -4339,7 +4339,7 @@ bool UWorld::DestroySwappedPC(UNetConnection* Connection)
 	for( FConstPlayerControllerIterator Iterator = GetPlayerControllerIterator(); Iterator; ++Iterator )
 	{
 		APlayerController* PlayerController = Iterator->Get();
-		if (PlayerController->Player == NULL && PlayerController->PendingSwapConnection == Connection)
+		if (PlayerController && PlayerController->Player == NULL && PlayerController->PendingSwapConnection == Connection)
 		{
 			DestroyActor(PlayerController);
 			return true;
@@ -5677,8 +5677,10 @@ UWorld* FSeamlessTravelHandler::Tick()
 
 			for (FConstPlayerControllerIterator Iterator = CurrentWorld->GetPlayerControllerIterator(); Iterator; ++Iterator)
 			{
-				APlayerController* Player = Iterator->Get();
-				ProcessActor(Player);
+				if (APlayerController* Player = Iterator->Get())
+				{
+					ProcessActor(Player);
+				}
 			}
 
 			bool bCreateNewGameMode = !bIsClient;
diff --git a/Engine/Source/Runtime/LevelSequence/Private/LevelSequencePlayer.cpp b/Engine/Source/Runtime/LevelSequence/Private/LevelSequencePlayer.cpp
--- a/Engine/Source/Runtime/LevelSequence/Private/LevelSequencePlayer.cpp
+++ b/Engine/Source/Runtime/LevelSequence/Private/LevelSequencePlayer.cpp
@@ -363,8 +363,8 @@ void ULevelSequencePlayer::EnableCinematicMode(bool bEnable)
 		{
 			for (FConstPlayerControllerIterator Iterator = World->GetPlayerControllerIterator(); Iterator; ++Iterator)
 			{
-				APlayerController *PC = Iterator->Get();
-				if (PC->IsLocalController())
+				APlayerController* PC = Iterator->Get();
+				if (PC && PC->IsLocalController())
 				{
 					PC->SetCinematicMode(bEnable, PlaybackSettings.bHidePlayer, PlaybackSettings.bHideHud, PlaybackSettings.bDisableMovementInput, PlaybackSettings.bDisableLookAtInput);
 				}
diff --git a/Engine/Source/Runtime/NavigationSystem/Private/NavMesh/RecastNavMeshGenerator.cpp b/Engine/Source/Runtime/NavigationSystem/Private/NavMesh/RecastNavMeshGenerator.cpp
--- a/Engine/Source/Runtime/NavigationSystem/Private/NavMesh/RecastNavMeshGenerator.cpp
+++ b/Engine/Source/Runtime/NavigationSystem/Private/NavMesh/RecastNavMeshGenerator.cpp
@@ -4090,9 +4090,10 @@ void FRecastNavMeshGenerator::SortPendingBuildTiles()
 	// Collect players positions
 	for (FConstPlayerControllerIterator PlayerIt = CurWorld->GetPlayerControllerIterator(); PlayerIt; ++PlayerIt)
 	{
-		if (PlayerIt->IsValid() && PlayerIt->Get()->GetPawn() != NULL)
+		APlayerController* PC = PlayerIt->Get();
+		if (PC && PC->GetPawn() != NULL)
 		{
-			const FVector2D SeedLoc((*PlayerIt)->GetPawn()->GetActorLocation());
+			const FVector2D SeedLoc(PC->GetPawn()->GetActorLocation());
 			SeedLocations.Add(SeedLoc);
 		}
 	}
