Add logging to two reported crash sites related to loaded UClasses, harden ChildActorComponent serialization against placeholder interactions

diff --git a/Engine/Source/Runtime/CoreUObject/Private/Blueprint/BlueprintSupport.cpp b/Engine/Source/Runtime/CoreUObject/Private/Blueprint/BlueprintSupport.cpp
--- a/Engine/Source/Runtime/CoreUObject/Private/Blueprint/BlueprintSupport.cpp
+++ b/Engine/Source/Runtime/CoreUObject/Private/Blueprint/BlueprintSupport.cpp
@@ -278,6 +278,21 @@ bool FBlueprintSupport::ShouldSuppressWarning(FName WarningIdentifier)
 	return BlueprintWarningsToSuppress.Find(WarningIdentifier) != nullptr;
 }
 
+bool FBlueprintSupport::IsClassPlaceholder(UClass* Class)
+{
+	while (Class)
+	{
+		if (Cast<ULinkerPlaceholderClass>(Class))
+		{
+			return true;
+		}
+
+		Class = Class->GetSuperClass();
+	}
+
+	return false;
+}
+
 #if WITH_EDITOR
 void FBlueprintSupport::ValidateNoRefsToOutOfDateClasses()
 {
diff --git a/Engine/Source/Runtime/CoreUObject/Public/Blueprint/BlueprintSupport.h b/Engine/Source/Runtime/CoreUObject/Public/Blueprint/BlueprintSupport.h
--- a/Engine/Source/Runtime/CoreUObject/Public/Blueprint/BlueprintSupport.h
+++ b/Engine/Source/Runtime/CoreUObject/Public/Blueprint/BlueprintSupport.h
@@ -102,6 +102,8 @@ struct FBlueprintSupport
 	COREUOBJECT_API static bool ShouldTreatWarningAsError(FName WarningIdentifier);
 	COREUOBJECT_API static bool ShouldSuppressWarning(FName WarningIdentifier);
 
+	COREUOBJECT_API static bool IsClassPlaceholder(UClass* Class);
+
 #if WITH_EDITOR
 	/** Function that walks the object graph, ensuring that there are no references to TRASH or REINST classes: */
 	COREUOBJECT_API static void ValidateNoRefsToOutOfDateClasses();
diff --git a/Engine/Source/Runtime/CoreUObject/Public/Templates/SubclassOf.h b/Engine/Source/Runtime/CoreUObject/Public/Templates/SubclassOf.h
--- a/Engine/Source/Runtime/CoreUObject/Public/Templates/SubclassOf.h
+++ b/Engine/Source/Runtime/CoreUObject/Public/Templates/SubclassOf.h
@@ -98,6 +98,15 @@ public:
 		return GetTypeHash(SubclassOf.Class);
 	}
 
+#if DO_CHECK
+	// This is a DEVELOPMENT ONLY debugging function and should not be relied upon. Client
+	// systems should never require unsafe access to the referenced UClass
+	UClass* DebugAccessRawClassPtr() const
+	{
+		return Class;
+	}
+#endif
+
 private:
 	UClass* Class;
 };
diff --git a/Engine/Source/Runtime/Engine/Private/Components/ChildActorComponent.cpp b/Engine/Source/Runtime/Engine/Private/Components/ChildActorComponent.cpp
--- a/Engine/Source/Runtime/Engine/Private/Components/ChildActorComponent.cpp
+++ b/Engine/Source/Runtime/Engine/Private/Components/ChildActorComponent.cpp
@@ -90,11 +90,23 @@ void UChildActorComponent::Serialize(FArchive& Ar)
 #if WITH_EDITOR
 	if (ChildActorClass == nullptr)
 	{
-		// It is unknown how this state can come to be, so for now we'll simply correct the issue and record that it occurs and 
-		// and if it is occurring frequently, then investigate how the state comes to pass
-		if (!ensureAlwaysMsgf(ChildActorTemplate == nullptr, TEXT("Found unexpected ChildActorTemplate %s when ChildActorClass is null"), *ChildActorTemplate->GetFullName()))
+#if DO_CHECK
+		if (FBlueprintSupport::IsClassPlaceholder(ChildActorClass.DebugAccessRawClassPtr()))
 		{
-			ChildActorTemplate = nullptr;
+			ensureMsgf(false, TEXT("Unexpectedly encountered placeholder class while serializing a component"));
+		}
+		else
+#endif
+		{
+			if (!FBlueprintSupport::IsDeferredDependencyPlaceholder(ChildActorTemplate))
+			{
+				// It is unknown how this state can come to be, so for now we'll simply correct the issue and record that it occurs and 
+				// and if it is occurring frequently, then investigate how the state comes to pass
+				if (!ensureAlwaysMsgf(ChildActorTemplate == nullptr, TEXT("Found unexpected ChildActorTemplate %s when ChildActorClass is null"), *ChildActorTemplate->GetFullName()))
+				{
+					ChildActorTemplate = nullptr;
+				}
+			}
 		}
 	}
 	// Since we sometimes serialize properties in instead of using duplication and we can end up pointing at the wrong template
diff --git a/Engine/Source/Runtime/Engine/Private/GameModeBase.cpp b/Engine/Source/Runtime/Engine/Private/GameModeBase.cpp
--- a/Engine/Source/Runtime/Engine/Private/GameModeBase.cpp
+++ b/Engine/Source/Runtime/Engine/Private/GameModeBase.cpp
@@ -132,6 +132,17 @@ TSubclassOf<AGameSession> AGameModeBase::GetGameSessionClass() const
 
 UClass* AGameModeBase::GetDefaultPawnClassForController_Implementation(AController* InController)
 {
+#if WITH_EDITOR && DO_CHECK
+	UClass* DefaultClass = DefaultPawnClass.DebugAccessRawClassPtr();
+	if (DefaultClass)
+	{
+		if (FBlueprintSupport::IsClassPlaceholder(DefaultClass))
+		{
+			ensureMsgf(false, TEXT("Trying to spawn class that is, directly or indirectly, a placeholder"));
+			return ADefaultPawn::StaticClass();
+		}
+	}
+#endif
 	return DefaultPawnClass;
 }
 
