Changed the UniquePtr by a SharedPtr and create a gamethread version that is updated immediatly. The renderthread version will get updated on the render thread by a render command.
Update the code to use the Gamethread version in most cases.

diff --git a/Engine/Source/Runtime/Engine/Private/HierarchicalInstancedStaticMesh.cpp b/Engine/Source/Runtime/Engine/Private/HierarchicalInstancedStaticMesh.cpp
--- a/Engine/Source/Runtime/Engine/Private/HierarchicalInstancedStaticMesh.cpp
+++ b/Engine/Source/Runtime/Engine/Private/HierarchicalInstancedStaticMesh.cpp
@@ -2797,7 +2797,7 @@ void UHierarchicalInstancedStaticMeshComponent::OnPostLoadPerInstanceData()
 			{
 				// create PerInstanceRenderData either from current data or pre-built instance buffer
 				InitPerInstanceRenderData(true, InstanceDataBuffers.Release());
-				NumBuiltRenderInstances = PerInstanceRenderData->InstanceBuffer.GetNumInstances();
+				NumBuiltRenderInstances = PerInstanceRenderData->InstanceBuffer_GameThread->GetNumInstances();
 			}
 
 			// If any of the data is out of sync, build the tree now!
@@ -2840,7 +2840,7 @@ static void GatherInstanceTransformsInArea(const UHierarchicalInstancedStaticMes
 						{
 							// if there's no PerInstanceSMData (e.g. for grass), we'll go get the transform from the render buffer
 							FMatrix XformMat;
-							Component.PerInstanceRenderData->InstanceBuffer.GetInstanceTransform(i, XformMat);
+							Component.PerInstanceRenderData->InstanceBuffer_GameThread->GetInstanceTransform(i, XformMat);
 							InstanceToComponent = FTransform(XformMat);
 						}
 						else
diff --git a/Engine/Source/Runtime/Engine/Private/InstancedStaticMesh.cpp b/Engine/Source/Runtime/Engine/Private/InstancedStaticMesh.cpp
--- a/Engine/Source/Runtime/Engine/Private/InstancedStaticMesh.cpp
+++ b/Engine/Source/Runtime/Engine/Private/InstancedStaticMesh.cpp
@@ -137,29 +137,11 @@ void FStaticMeshInstanceBuffer::InitFromPreallocatedData(FStaticMeshInstanceData
 {
 	QUICK_SCOPE_CYCLE_COUNTER(STAT_FStaticMeshInstanceBuffer_InitFromPreallocatedData);
 
-	InstanceData = MakeUnique<FStaticMeshInstanceData>();
+	InstanceData = MakeShared<FStaticMeshInstanceData, ESPMode::ThreadSafe>();
 	FMemory::Memswap(&Other, InstanceData.Get(), sizeof(FStaticMeshInstanceData));
 	InstanceData->SetAllowCPUAccess(RequireCPUAccess);
 }
 
-void FStaticMeshInstanceBuffer::UpdateFromPreallocatedData_Concurrent(FStaticMeshInstanceData& Other)
-{
-	QUICK_SCOPE_CYCLE_COUNTER(STAT_FStaticMeshInstanceBuffer_UpdateFromPreallocatedData_Concurrent);
-	
-	Other.SetAllowCPUAccess(RequireCPUAccess);
-
-	FStaticMeshInstanceBuffer* InstanceBuffer = this; 
-	FStaticMeshInstanceData* NewInstanceData = new FStaticMeshInstanceData();
-	FMemory::Memswap(&Other, NewInstanceData, sizeof(FStaticMeshInstanceData));
-		
-	ENQUEUE_RENDER_COMMAND(InstanceBuffer_UpdateFromPreallocatedData)(
-		[InstanceBuffer, NewInstanceData](FRHICommandListImmediate& RHICmdList)
-		{
-			InstanceBuffer->InstanceData.Reset(NewInstanceData);
-			InstanceBuffer->UpdateRHI();
-		});
-}
-
 void FStaticMeshInstanceBuffer::UpdateFromCommandBuffer_Concurrent(FInstanceUpdateCmdBuffer& CmdBuffer)
 {
 	QUICK_SCOPE_CYCLE_COUNTER(STAT_FStaticMeshInstanceBuffer_UpdateFromCommandBuffer_Concurrent);
@@ -536,11 +518,14 @@ FPerInstanceRenderData::FPerInstanceRenderData(FStaticMeshInstanceData& Other, E
 	, ResourceSize(Other.GetResourceSize()) // 2x when with CPU access?
 {
 	InstanceBuffer.InitFromPreallocatedData(Other);
+	InstanceBuffer_GameThread = InstanceBuffer.InstanceData;
+
 	BeginInitResource(&InstanceBuffer);
 }
 		
 FPerInstanceRenderData::~FPerInstanceRenderData()
 {
+	InstanceBuffer_GameThread.Reset();
 	// Should be always destructed on rendering thread
 	InstanceBuffer.ReleaseResource();
 }
@@ -549,7 +534,25 @@ void FPerInstanceRenderData::UpdateFromPreallocatedData(FStaticMeshInstanceData&
 {
 	InstanceBuffer.RequireCPUAccess = (InOther.GetOriginResourceArray()->GetAllowCPUAccess() || InOther.GetTransformResourceArray()->GetAllowCPUAccess() || InOther.GetLightMapResourceArray()->GetAllowCPUAccess()) ? true : InstanceBuffer.RequireCPUAccess;
 	ResourceSize = InOther.GetResourceSize(); // 2x when with CPU access?
-	InstanceBuffer.UpdateFromPreallocatedData_Concurrent(InOther);
+
+	InOther.SetAllowCPUAccess(InstanceBuffer.RequireCPUAccess);
+
+	FStaticMeshInstanceData* NewInstanceData = new FStaticMeshInstanceData();
+	FMemory::Memswap(&InOther, NewInstanceData, sizeof(FStaticMeshInstanceData));
+
+	InstanceBuffer_GameThread = MakeShared<FStaticMeshInstanceData, ESPMode::ThreadSafe>(*NewInstanceData);
+
+	typedef TSharedPtr<FStaticMeshInstanceData, ESPMode::ThreadSafe> FStaticMeshInstanceDataPtr;
+
+	ENQUEUE_UNIQUE_RENDER_COMMAND_TWOPARAMETER(
+		FInstanceBuffer_UpdateFromPreallocatedData,
+		FStaticMeshInstanceDataPtr, InInstanceBufferDataPtr, InstanceBuffer_GameThread,
+		FStaticMeshInstanceBuffer*, InInstanceBuffer, &InstanceBuffer,
+		{
+			InInstanceBuffer->InstanceData = InInstanceBufferDataPtr;
+			InInstanceBuffer->UpdateRHI();
+		}
+	);
 }
 
 void FPerInstanceRenderData::UpdateFromCommandBuffer(FInstanceUpdateCmdBuffer& CmdBuffer)
@@ -1538,9 +1541,9 @@ void UInstancedStaticMeshComponent::SerializeRenderData(FArchive& Ar)
 		
 			if (PerInstanceRenderData.IsValid())
 			{
-				if (PerInstanceRenderData->InstanceBuffer.GetInstanceData()->GetNumInstances() > 0)
+				if (PerInstanceRenderData->InstanceBuffer_GameThread && PerInstanceRenderData->InstanceBuffer_GameThread->GetNumInstances() > 0)
 				{
-					int32 NumInstances = PerInstanceRenderData->InstanceBuffer.GetInstanceData()->GetNumInstances();
+					int32 NumInstances = PerInstanceRenderData->InstanceBuffer_GameThread->GetNumInstances();
 
 					// Clear editor data for the cooked data
 					for (int32 Index = 0; Index < NumInstances; ++Index)
@@ -1552,10 +1555,10 @@ void UInstancedStaticMeshComponent::SerializeRenderData(FArchive& Ar)
 							continue;
 						}
 
-						PerInstanceRenderData->InstanceBuffer.GetInstanceData()->ClearInstanceEditorData(RenderIndex);
+						PerInstanceRenderData->InstanceBuffer_GameThread->ClearInstanceEditorData(RenderIndex);
 					}
 
-					PerInstanceRenderData->InstanceBuffer.GetInstanceData()->Serialize(Ar);
+					PerInstanceRenderData->InstanceBuffer_GameThread->Serialize(Ar);
 
 #if WITH_EDITOR
 					// Restore back the state we were in
@@ -1580,7 +1583,7 @@ void UInstancedStaticMeshComponent::SerializeRenderData(FArchive& Ar)
 							HitProxyColor = HitProxies[Index]->Id.GetColor();
 						}
 
-						PerInstanceRenderData->InstanceBuffer.GetInstanceData()->SetInstanceEditorData(RenderIndex, HitProxyColor, bSelected);
+						PerInstanceRenderData->InstanceBuffer_GameThread->SetInstanceEditorData(RenderIndex, HitProxyColor, bSelected);
 					}
 #endif					
 				}
diff --git a/Engine/Source/Runtime/Engine/Private/InstancedStaticMesh.h b/Engine/Source/Runtime/Engine/Private/InstancedStaticMesh.h
--- a/Engine/Source/Runtime/Engine/Private/InstancedStaticMesh.h
+++ b/Engine/Source/Runtime/Engine/Private/InstancedStaticMesh.h
@@ -63,7 +63,6 @@ public:
 	 * @param Other - instance data, this call assumes the memory, so this will be empty after the call
 	 */
 	ENGINE_API void InitFromPreallocatedData(FStaticMeshInstanceData& Other);
-	ENGINE_API void UpdateFromPreallocatedData_Concurrent(FStaticMeshInstanceData& Other);
 	ENGINE_API void UpdateFromCommandBuffer_Concurrent(FInstanceUpdateCmdBuffer& CmdBuffer);
 
 	/** Serializer. */
@@ -107,7 +106,7 @@ public:
 
 public:
 	/** The vertex data storage type */
-	TUniquePtr<FStaticMeshInstanceData> InstanceData;
+	TSharedPtr<FStaticMeshInstanceData, ESPMode::ThreadSafe> InstanceData;
 
 	/** Keep CPU copy of instance data*/
 	bool RequireCPUAccess;
@@ -382,13 +381,15 @@ struct FPerInstanceRenderData
 	*/
 	ENGINE_API void UpdateFromCommandBuffer(FInstanceUpdateCmdBuffer& CmdBuffer);
 
-	/** Instance buffer */
-	FStaticMeshInstanceBuffer			InstanceBuffer;
 	/** Hit proxies for the instances */
 	TArray<TRefCountPtr<HHitProxy>>		HitProxies;
 
 	/** cached per-instance resource size*/
 	SIZE_T								ResourceSize;
+
+	/** Instance buffer */
+	FStaticMeshInstanceBuffer			InstanceBuffer;
+	TSharedPtr<FStaticMeshInstanceData, ESPMode::ThreadSafe> InstanceBuffer_GameThread;
 };
 
 
