MIrrors CL 4152380.

Make material parameter collection instances request an end of frame update when they need a render state update.

diff --git a/Engine/Source/Runtime/Engine/Classes/Engine/World.h b/Engine/Source/Runtime/Engine/Classes/Engine/World.h
--- a/Engine/Source/Runtime/Engine/Classes/Engine/World.h
+++ b/Engine/Source/Runtime/Engine/Classes/Engine/World.h
@@ -1200,6 +1200,9 @@ private:
 	/** Is forcibly making streaming levels visible?																			*/
 	bool										bShouldForceVisibleStreamingLevels;
 
+	/** Is there at least one material parameter collection instance waiting for a deferred update?								*/
+	uint8										bMaterialParameterCollectionInstanceNeedsDeferredUpdate : 1;
+
 public:
 
 	bool GetShouldForceUnloadStreamingLevels() const { return bShouldForceUnloadStreamingLevels; }
@@ -2522,7 +2525,15 @@ public:
 	 */
 	void UpdateActorComponentEndOfFrameUpdateState(UActorComponent* Component) const;
 
-	bool HasEndOfFrameUpdates();
+	/** 
+	 * Used to indicate a UMaterialParameterCollectionInstance needs a deferred update 
+	 */
+	void SetMaterialParameterCollectionInstanceNeedsUpdate();
+
+	/** 
+	 * Returns true if we have any updates that have been deferred to the end of the current frame.
+	 */
+	bool HasEndOfFrameUpdates() const;
 
 	/**
 	 * Send all render updates to the rendering thread.
diff --git a/Engine/Source/Runtime/Engine/Private/LevelTick.cpp b/Engine/Source/Runtime/Engine/Private/LevelTick.cpp
--- a/Engine/Source/Runtime/Engine/Private/LevelTick.cpp
+++ b/Engine/Source/Runtime/Engine/Private/LevelTick.cpp
@@ -429,7 +429,7 @@ void UWorld::TickNetClient( float DeltaSeconds )
 	// and there isn't a PendingNetGame, throw a network failure error.
 	if( NetDriver->ServerConnection->State == USOCK_Closed )
 	{
-		if (GEngine->PendingNetGameFromWorld(this) == NULL)
+		if (GEngine->PendingNetGameFromWorld(this) == nullptr)
 		{
 			const FString Error = NSLOCTEXT("Engine", "ConnectionFailed", "Your connection to the host has been lost.").ToString();
 			GEngine->BroadcastNetworkFailure(this, NetDriver, ENetworkFailure::ConnectionLost, Error);
@@ -446,7 +446,7 @@ bool UWorld::IsPaused() const
 {
 	// pause if specifically set or if we're waiting for the end of the tick to perform streaming level loads (so actors don't fall through the world in the meantime, etc)
 	const AWorldSettings* Info = GetWorldSettings();
-	return ( (Info && Info->Pauser != NULL && TimeSeconds >= PauseDelay) ||
+	return ( (Info && Info->Pauser != nullptr && TimeSeconds >= PauseDelay) ||
 				(bRequestedBlockOnAsyncLoading && GetNetMode() == NM_Client) ||
 				(GEngine->ShouldCommitPendingMapChange(this)) ||
 				(IsPlayInEditor() && bDebugPauseExecution) );
@@ -711,8 +711,8 @@ void UWorld::ProcessLevelStreamingVolumes(FVector* OverrideViewLocation)
 
 			// Figure out whether level should be loaded and keep track of original state for notifications on change.
 			FVisibleLevelStreamingSettings* NewStreamingSettings= VisibleLevelStreamingObjects.Find( LevelStreamingObject );
-			bool bShouldAffectLoading							= LevelStreamingObjectsWithVolumesOtherThanBlockingLoad.Find( LevelStreamingObject ) != NULL;
-			bool bShouldBeLoaded								= (NewStreamingSettings != NULL);
+			bool bShouldAffectLoading							= LevelStreamingObjectsWithVolumesOtherThanBlockingLoad.Find( LevelStreamingObject ) != nullptr;
+			bool bShouldBeLoaded								= (NewStreamingSettings != nullptr);
 			bool bOriginalShouldBeLoaded						= LevelStreamingObject->ShouldBeLoaded();
 			bool bOriginalShouldBeVisible						= LevelStreamingObject->ShouldBeVisible();
 			bool bOriginalShouldBlockOnLoad						= LevelStreamingObject->bShouldBlockOnLoad;
@@ -891,14 +891,19 @@ void UWorld::MarkActorComponentForNeededEndOfFrameUpdate(UActorComponent* Compon
 	}
 }
 
-bool UWorld::HasEndOfFrameUpdates()
+void UWorld::SetMaterialParameterCollectionInstanceNeedsUpdate()
 {
-	return ComponentsThatNeedEndOfFrameUpdate_OnGameThread.Num() > 0 || ComponentsThatNeedEndOfFrameUpdate.Num() > 0;
+	bMaterialParameterCollectionInstanceNeedsDeferredUpdate = true;
+}
+
+bool UWorld::HasEndOfFrameUpdates() const
+{
+	return ComponentsThatNeedEndOfFrameUpdate_OnGameThread.Num() > 0 || ComponentsThatNeedEndOfFrameUpdate.Num() > 0 || bMaterialParameterCollectionInstanceNeedsDeferredUpdate;
 }
 
 TDrawEvent<FRHICommandList>* BeginSendEndOfFrameUpdatesDrawEvent()
 {
-	TDrawEvent<FRHICommandList>* DrawEvent = NULL;
+	TDrawEvent<FRHICommandList>* DrawEvent = nullptr;
 
 #if WANTS_DRAW_MESH_EVENTS
 	DrawEvent = new TDrawEvent<FRHICommandList>();
@@ -969,7 +974,7 @@ void UWorld::SendAllEndOfFrameUpdates()
 				if (NextComponent->IsRegistered() && !NextComponent->IsTemplate() && !NextComponent->IsPendingKill())
 				{
 					FScopeCycleCounterUObject ComponentScope(NextComponent);
-					FScopeCycleCounterUObject AdditionalScope(STATS ? NextComponent->AdditionalStatObject() : NULL);
+					FScopeCycleCounterUObject AdditionalScope(STATS ? NextComponent->AdditionalStatObject() : nullptr);
 					NextComponent->DoDeferredRenderUpdates_Concurrent();
 				}
 				check(NextComponent->GetMarkedForEndOfFrameUpdateState() == EComponentMarkedForEndOfFrameUpdateState::Marked);
@@ -988,7 +993,7 @@ void UWorld::SendAllEndOfFrameUpdates()
 					if (Component->IsRegistered() && !Component->IsTemplate() && !Component->IsPendingKill())
 					{
 						FScopeCycleCounterUObject ComponentScope(Component);
-						FScopeCycleCounterUObject AdditionalScope(STATS ? Component->AdditionalStatObject() : NULL);
+						FScopeCycleCounterUObject AdditionalScope(STATS ? Component->AdditionalStatObject() : nullptr);
 						Component->DoDeferredRenderUpdates_Concurrent();
 					}
 					check(Component->GetMarkedForEndOfFrameUpdateState() == EComponentMarkedForEndOfFrameUpdateState::MarkedForGameThread);
@@ -1016,6 +1021,7 @@ void UWorld::SendAllEndOfFrameUpdates()
 			ParameterCollectionInstance->DeferredUpdateRenderState();
 		}
 	}
+	bMaterialParameterCollectionInstanceNeedsDeferredUpdate = false;
 			
 	LocalComponentsThatNeedEndOfFrameUpdate.Reset();
 
@@ -1034,13 +1040,13 @@ static class FFileProfileWrapperExec: private FSelfRegisteringExec
 			if( FParse::Command( &Cmd, TEXT("File") ) ) // if they didn't use the list command, we will show usage
 			{
 				FProfiledPlatformFile* ProfilePlatformFile = (FProfiledPlatformFile*)FPlatformFileManager::Get().FindPlatformFile(TProfiledPlatformFile<FProfiledFileStatsFileDetailed>::GetTypeName());
-				if (ProfilePlatformFile == NULL)
+				if (ProfilePlatformFile == nullptr)
 				{
 					// Try 'simple' profiler file.
 					ProfilePlatformFile = (FProfiledPlatformFile*)FPlatformFileManager::Get().FindPlatformFile(TProfiledPlatformFile<FProfiledFileStatsFileSimple>::GetTypeName());
 				}				
 
-				if( ProfilePlatformFile != NULL)
+				if( ProfilePlatformFile != nullptr)
 				{
 					DisplayProfileData( ProfilePlatformFile->GetStats() );
 				}
@@ -1394,7 +1400,7 @@ void UWorld::Tick( ELevelTick TickType, float DeltaSeconds )
 	}
 	
 	// update world's subsystems (NavigationSystem for now)
-	if (NavigationSystem != NULL)
+	if (NavigationSystem != nullptr)
 	{
 		SCOPE_CYCLE_COUNTER(STAT_NavWorldTickTime);
 		NavigationSystem->Tick(DeltaSeconds);
@@ -1484,7 +1490,7 @@ void UWorld::Tick( ELevelTick TickType, float DeltaSeconds )
 			if( !bIsPaused )
 			{
 				// This will process any latent actions that have not been processed already
-				CurrentLatentActionManager.ProcessLatentActions(NULL, DeltaSeconds);
+				CurrentLatentActionManager.ProcessLatentActions(nullptr, DeltaSeconds);
 			}
 #if 0 // if you need to debug physics drawing in editor, use this. If you type pxvis collision, it will work. 
 			else
@@ -1572,7 +1578,7 @@ void UWorld::Tick( ELevelTick TickType, float DeltaSeconds )
 		FWorldDelegates::OnWorldPostActorTick.Broadcast(this, TickType, DeltaSeconds);
 
 #if WITH_PHYSX
-		if ( PhysicsScene != NULL )
+		if ( PhysicsScene != nullptr )
 		{
 			GPhysCommandHandler->Flush();
 		}
@@ -1611,7 +1617,7 @@ void UWorld::Tick( ELevelTick TickType, float DeltaSeconds )
 	}
 
 	// Tick the FX system.
-	if (!bIsPaused && FXSystem != NULL)
+	if (!bIsPaused && FXSystem != nullptr)
 	{
 		SCOPE_TIME_GUARD_MS(TEXT("UWorld::Tick - FX"), 5);
 		FXSystem->Tick(DeltaSeconds);
diff --git a/Engine/Source/Runtime/Engine/Private/Materials/ParameterCollection.cpp b/Engine/Source/Runtime/Engine/Private/Materials/ParameterCollection.cpp
--- a/Engine/Source/Runtime/Engine/Private/Materials/ParameterCollection.cpp
+++ b/Engine/Source/Runtime/Engine/Private/Materials/ParameterCollection.cpp
@@ -25,7 +25,7 @@ TMap<FGuid, FMaterialParameterCollectionInstanceResource*> GDefaultMaterialParam
 UMaterialParameterCollection::UMaterialParameterCollection(const FObjectInitializer& ObjectInitializer)
 	: Super(ObjectInitializer)
 {
-	DefaultResource = NULL;
+	DefaultResource = nullptr;
 }
 
 void UMaterialParameterCollection::PostInitProperties()
@@ -72,7 +72,7 @@ void UMaterialParameterCollection::BeginDestroy()
 		);
 
 		DefaultResource->GameThread_Destroy();
-		DefaultResource = NULL;
+		DefaultResource = nullptr;
 	}
 
 	Super::BeginDestroy();
@@ -254,7 +254,7 @@ void UMaterialParameterCollection::PostEditChangeProperty(FPropertyChangedEvent&
 					UpdateContext.AddMaterial(CurrentMaterial);
 
 					// Propagate the change to this material
-					CurrentMaterial->PreEditChange(NULL);
+					CurrentMaterial->PreEditChange(nullptr);
 					CurrentMaterial->PostEditChange();
 					CurrentMaterial->MarkPackageDirty();
 				}
@@ -395,7 +395,7 @@ const FCollectionScalarParameter* UMaterialParameterCollection::GetScalarParamet
 		}
 	}
 
-	return NULL;
+	return nullptr;
 }
 
 const FCollectionVectorParameter* UMaterialParameterCollection::GetVectorParameterByName(FName ParameterName) const
@@ -410,10 +410,10 @@ const FCollectionVectorParameter* UMaterialParameterCollection::GetVectorParamet
 		}
 	}
 
-	return NULL;
+	return nullptr;
 }
 
-FShaderUniformBufferParameter* ConstructCollectionUniformBufferParameter() { return NULL; }
+FShaderUniformBufferParameter* ConstructCollectionUniformBufferParameter() { return nullptr; }
 
 void UMaterialParameterCollection::CreateBufferStruct()
 {	
@@ -421,7 +421,7 @@ void UMaterialParameterCollection::CreateBufferStruct()
 	uint32 NextMemberOffset = 0;
 
 	const uint32 NumVectors = FMath::DivideAndRoundUp(ScalarParameters.Num(), 4) + VectorParameters.Num();
-	new(Members) FUniformBufferStruct::FMember(TEXT("Vectors"),TEXT(""),NextMemberOffset,UBMT_FLOAT32,EShaderPrecisionModifier::Half,1,4,NumVectors,NULL);
+	new(Members) FUniformBufferStruct::FMember(TEXT("Vectors"),TEXT(""),NextMemberOffset,UBMT_FLOAT32,EShaderPrecisionModifier::Half,1,4,NumVectors, nullptr);
 	const uint32 VectorArraySize = NumVectors * sizeof(FVector4);
 	NextMemberOffset += VectorArraySize;
 	static FName LayoutName(TEXT("MaterialCollection"));
@@ -582,7 +582,7 @@ bool UMaterialParameterCollectionInstance::GetScalarParameterValue(FName Paramet
 	if (Parameter)
 	{
 		const float* InstanceValue = ScalarParameterValues.Find(ParameterName);
-		OutParameterValue = InstanceValue != NULL ? *InstanceValue : Parameter->DefaultValue;
+		OutParameterValue = InstanceValue != nullptr ? *InstanceValue : Parameter->DefaultValue;
 		return true;
 	}
 
@@ -596,7 +596,7 @@ bool UMaterialParameterCollectionInstance::GetVectorParameterValue(FName Paramet
 	if (Parameter)
 	{
 		const FLinearColor* InstanceValue = VectorParameterValues.Find(ParameterName);
-		OutParameterValue = InstanceValue != NULL ? *InstanceValue : Parameter->DefaultValue;
+		OutParameterValue = InstanceValue != nullptr ? *InstanceValue : Parameter->DefaultValue;
 		return true;
 	}
 
@@ -612,6 +612,7 @@ void UMaterialParameterCollectionInstance::UpdateRenderState()
 	}
 
 	bNeedsRenderStateUpdate = true;
+	World->SetMaterialParameterCollectionInstanceNeedsUpdate();
 
 	if (!GDeferUpdateRenderStates)
 	{
@@ -672,7 +673,7 @@ void UMaterialParameterCollectionInstance::FinishDestroy()
 	if (Resource)
 	{
 		Resource->GameThread_Destroy();
-		Resource = NULL;
+		Resource = nullptr;
 	}
 
 	Super::FinishDestroy();
