[NO NON-PII COMMIT MESSAGE]

diff --git a/Engine/Source/Runtime/Engine/Private/InstancedStaticMesh.cpp b/Engine/Source/Runtime/Engine/Private/InstancedStaticMesh.cpp
--- a/Engine/Source/Runtime/Engine/Private/InstancedStaticMesh.cpp
+++ b/Engine/Source/Runtime/Engine/Private/InstancedStaticMesh.cpp
@@ -514,8 +514,8 @@ void FInstancedStaticMeshRenderData::InitStaticMeshVertexFactories(
 }
 
 FPerInstanceRenderData::FPerInstanceRenderData(FStaticMeshInstanceData& Other, ERHIFeatureLevel::Type InFeaureLevel, bool InRequireCPUAccess)
-	: InstanceBuffer(InFeaureLevel, InRequireCPUAccess)
-	, ResourceSize(Other.GetResourceSize()) // 2x when with CPU access?
+	: ResourceSize(Other.GetResourceSize()) // 2x when with CPU access?
+	, InstanceBuffer(InFeaureLevel, InRequireCPUAccess)
 {
 	InstanceBuffer.InitFromPreallocatedData(Other);
 	InstanceBuffer_GameThread = InstanceBuffer.InstanceData;
