Fix issue where curve extractor was using wrong list for getting curves. We want to give ALL list and then filter by what the source animation has originally.

diff --git a/Engine/Source/Runtime/Engine/Private/Animation/PoseAsset.cpp b/Engine/Source/Runtime/Engine/Private/Animation/PoseAsset.cpp
--- a/Engine/Source/Runtime/Engine/Private/Animation/PoseAsset.cpp
+++ b/Engine/Source/Runtime/Engine/Private/Animation/PoseAsset.cpp
@@ -1040,7 +1040,7 @@ void UPoseAsset::CreatePoseFromAnimation(class UAnimSequence* AnimSequence, cons
 						// have to do iterate over time
 						// support curve
 						FBlendedCurve SourceCurve;
-						SourceCurve.InitFrom(&UIDList);
+						SourceCurve.InitFrom(&TargetSkeleton->GetDefaultCurveUIDList());
 						AnimSequence->EvaluateCurveData(SourceCurve, PoseIndex*IntervalBetweenKeys, true);
 
 						// copy back to CurveData
