Fix blending on landscapes with exactly 3 layers on mobile

diff --git a/Engine/Shaders/Public/ShaderVersion.ush b/Engine/Shaders/Public/ShaderVersion.ush
--- a/Engine/Shaders/Public/ShaderVersion.ush
+++ b/Engine/Shaders/Public/ShaderVersion.ush
@@ -2,4 +2,4 @@
 // Each time the console command is executed it generates a new GUID. As this file is included
 // in Platform.ush (which should be included in any shader) it allows to invalidate the shader DDC.
 // 
-// GUID = 20965343E6BD444B9A1D010A80DFF469
\ No newline at end of file
+// GUID = CA5419C0DF2F4DFCA153F9192B61D2F7
\ No newline at end of file
diff --git a/Engine/Source/Runtime/Engine/Private/Materials/HLSLMaterialTranslator.h b/Engine/Source/Runtime/Engine/Private/Materials/HLSLMaterialTranslator.h
--- a/Engine/Source/Runtime/Engine/Private/Materials/HLSLMaterialTranslator.h
+++ b/Engine/Source/Runtime/Engine/Private/Materials/HLSLMaterialTranslator.h
@@ -4079,14 +4079,9 @@ protected:
 		FMaterialParameterInfo ParameterInfo = GetParameterAssociationInfo();
 		ParameterInfo.Name = ParameterName;
 
-		int32 NumActiveTerrainLayerWeightParameters = 0;
 		for(int32 ParameterIndex = 0;ParameterIndex < StaticParameters.TerrainLayerWeightParameters.Num(); ++ParameterIndex)
 		{
 			const FStaticTerrainLayerWeightParameter& Parameter = StaticParameters.TerrainLayerWeightParameters[ParameterIndex];
-			if (Parameter.bOverride)
-			{
-				NumActiveTerrainLayerWeightParameters++;
-			}
 			if(Parameter.ParameterInfo == ParameterInfo)
 			{
 				WeightmapIndex = Parameter.WeightmapIndex;
@@ -4109,7 +4104,7 @@ protected:
 		else
 		{			
 			int32 WeightmapCode;
-			if (GetFeatureLevel() <= ERHIFeatureLevel::ES3_1 && NumActiveTerrainLayerWeightParameters <= 3 && bAtLeastOneWeightBasedBlend)
+			if (GetFeatureLevel() <= ERHIFeatureLevel::ES3_1 && StaticParameters.TerrainLayerWeightParameters.Num() <= 3 && bAtLeastOneWeightBasedBlend)
 			{
 				// Mobile can pack 3 layers into the normal map texture B and A channels, implying the 3rd using weight based blending
 				// Layer texture is sampled into Parameters.LayerWeights in LandscapeVertexFactory.ush
diff --git a/Engine/Source/Runtime/Landscape/Private/Landscape.cpp b/Engine/Source/Runtime/Landscape/Private/Landscape.cpp
--- a/Engine/Source/Runtime/Landscape/Private/Landscape.cpp
+++ b/Engine/Source/Runtime/Landscape/Private/Landscape.cpp
@@ -90,6 +90,9 @@ namespace LandscapeCookStats
 // Set this to 0 to disable landscape cooking and thus disable it on device.
 #define ENABLE_LANDSCAPE_COOKING 1
 
+// Increment this to regenerate mobile landscape data
+#define LANDSCAPE_MOBILE_COOK_VERSION 0
+
 #define LOCTEXT_NAMESPACE "Landscape"
 
 static void PrintNumLandscapeShadows()
@@ -215,6 +218,10 @@ void ULandscapeComponent::CheckGenerateLandscapePlatformData(bool bIsCooking, co
 
 	FBufferArchive ComponentStateAr;
 	SerializeStateHashes(ComponentStateAr);
+
+	// Serialize the version number as part of the hash so we can invalidate DDC data if needed
+	int32 Version = LANDSCAPE_MOBILE_COOK_VERSION;
+	ComponentStateAr << Version;
 	
 	uint32 Hash[5];
 	FSHA1::HashBuffer(ComponentStateAr.GetData(), ComponentStateAr.Num(), (uint8*)Hash);
diff --git a/Engine/Source/Runtime/Landscape/Private/LandscapeEdit.cpp b/Engine/Source/Runtime/Landscape/Private/LandscapeEdit.cpp
--- a/Engine/Source/Runtime/Landscape/Private/LandscapeEdit.cpp
+++ b/Engine/Source/Runtime/Landscape/Private/LandscapeEdit.cpp
@@ -5045,13 +5045,15 @@ void ULandscapeComponent::GeneratePlatformPixelData()
 				// If we can pack into 2 channels with the 3rd implied, track the mask for the weight blendable layers
 				if (bAtLeastOneWeightBasedBlend && MobileWeightmapLayerAllocations.Num() <= 3)
 				{
+					MobileBlendableLayerMask |= (!Allocation.LayerInfo->bNoWeightBlend ? (1 << CurrentChannel) : 0);
+
 					// we don't need to create a new texture for the 3rd layer
 					if (RemainingChannels == 0)
 					{
+						Allocation.WeightmapTextureIndex = 0;
+						Allocation.WeightmapTextureChannel = 2; // not a valid texture channel, but used for the mask.
 						break;
 					}
-
-					MobileBlendableLayerMask |= (!Allocation.LayerInfo->bNoWeightBlend ? (1 << CurrentChannel) : 0);
 				}
 
 				if (RemainingChannels == 0)
