Prevent sub anim related crash when parameter types are changed and anim blueprints are in error state

diff --git a/Engine/Source/Runtime/Engine/Private/Animation/AnimNode_SubInstance.cpp b/Engine/Source/Runtime/Engine/Private/Animation/AnimNode_SubInstance.cpp
--- a/Engine/Source/Runtime/Engine/Private/Animation/AnimNode_SubInstance.cpp
+++ b/Engine/Source/Runtime/Engine/Private/Animation/AnimNode_SubInstance.cpp
@@ -47,10 +47,15 @@ void FAnimNode_SubInstance::Update_AnyThread(const FAnimationUpdateContext& Cont
 
 			check(CallerProperty && SubProperty);
 
-			uint8* SrcPtr = CallerProperty->ContainerPtrToValuePtr<uint8>(Context.AnimInstanceProxy->GetAnimInstanceObject());
-			uint8* DestPtr = SubProperty->ContainerPtrToValuePtr<uint8>(InstanceToRun);
+#if WITH_EDITOR
+			if (ensure(CallerProperty->SameType(SubProperty)))
+#endif
+			{
+				uint8* SrcPtr = CallerProperty->ContainerPtrToValuePtr<uint8>(Context.AnimInstanceProxy->GetAnimInstanceObject());
+				uint8* DestPtr = SubProperty->ContainerPtrToValuePtr<uint8>(InstanceToRun);
 
-			CallerProperty->CopyCompleteValue(DestPtr, SrcPtr);
+				CallerProperty->CopyCompleteValue(DestPtr, SrcPtr);
+			}
 		}
 	}
 }
@@ -164,7 +169,12 @@ void FAnimNode_SubInstance::OnInitializeAnimInstance(const FAnimInstanceProxy* I
 			UProperty* SourceProperty = FindField<UProperty>(SourceClass, SourceName);
 			UProperty* DestProperty = FindField<UProperty>(*InstanceClass, DestName);
 
-			if (SourceProperty && DestProperty)
+			if (SourceProperty && DestProperty
+#if WITH_EDITOR
+				// This type check can fail when anim blueprints are in an error state:
+				&& SourceProperty->SameType(DestProperty)
+#endif
+				)
 			{
 				InstanceProperties.Add(SourceProperty);
 				SubInstanceProperties.Add(DestProperty);
