Move bShouldBeLoaded back to ULevelStreaming for convenience and backwards compatibility, but increment severity of using SetShouldBeLoaded to warning
#rnx

diff --git a/Engine/Source/Runtime/Engine/Classes/Engine/LevelStreaming.h b/Engine/Source/Runtime/Engine/Classes/Engine/LevelStreaming.h
--- a/Engine/Source/Runtime/Engine/Classes/Engine/LevelStreaming.h
+++ b/Engine/Source/Runtime/Engine/Classes/Engine/LevelStreaming.h
@@ -161,6 +161,11 @@ private:
 	UPROPERTY(Category=LevelStreaming, BlueprintSetter=SetShouldBeVisible)
 	uint8 bShouldBeVisible:1;
 
+protected:
+	/** Whether the level should be loaded																						*/
+	UPROPERTY(Category=LevelStreaming, BlueprintSetter=SetShouldBeLoaded, BlueprintGetter=ShouldBeLoaded)
+	uint8 bShouldBeLoaded:1;
+
 public:
 
 	/** Whether this level is locked; that is, its actors are read-only. */
@@ -264,6 +269,7 @@ public:
 	 * Virtual that can be overriden to change whether a streaming level should be loaded.
 	 * Doesn't do anything at the base level as should be loaded defaults to true 
 	 */
+	UFUNCTION(BlueprintSetter)
 	virtual void SetShouldBeLoaded(bool bInShouldBeLoaded);
 
 	/** Returns the world composition level LOD index. */
@@ -325,7 +331,8 @@ public:
 	 *
 	 * @return true if level should be loaded/ streamed in, false otherwise
 	 */
-	virtual bool ShouldBeLoaded() const;
+	UFUNCTION(BlueprintGetter)
+	virtual bool ShouldBeLoaded() const { return true; }
 
 	/**
 	 * Return whether this level should be visible/ associated with the world if it is
@@ -524,5 +531,4 @@ private:
 	/** Friend classes to manipulate the streaming level more extensively */
 	friend class UEngine;
 	friend class UWorld;
-};
-
+};
\ No newline at end of file
diff --git a/Engine/Source/Runtime/Engine/Classes/Engine/LevelStreamingKismet.h b/Engine/Source/Runtime/Engine/Classes/Engine/LevelStreamingKismet.h
--- a/Engine/Source/Runtime/Engine/Classes/Engine/LevelStreamingKismet.h
+++ b/Engine/Source/Runtime/Engine/Classes/Engine/LevelStreamingKismet.h
@@ -55,15 +55,10 @@ class ENGINE_API ULevelStreamingKismet : public ULevelStreaming
 	virtual bool ShouldBeLoaded() const override { return bShouldBeLoaded; }
 	//~ End ULevelStreaming Interface
 
-	UFUNCTION(BlueprintSetter)
 	virtual void SetShouldBeLoaded(bool bShouldBeLoaded) override;
 
 private:
 
-	/** Whether the level should be loaded																						*/
-	UPROPERTY(Category=LevelStreaming, BlueprintSetter=SetShouldBeLoaded)
-	uint32 bShouldBeLoaded:1;
-
 	// Counter used by LoadLevelInstance to create unique level names
 	static int32 UniqueLevelInstanceId;
 
diff --git a/Engine/Source/Runtime/Engine/Private/LevelStreaming.cpp b/Engine/Source/Runtime/Engine/Private/LevelStreaming.cpp
--- a/Engine/Source/Runtime/Engine/Private/LevelStreaming.cpp
+++ b/Engine/Source/Runtime/Engine/Private/LevelStreaming.cpp
@@ -1419,12 +1419,7 @@ void ULevelStreaming::SetShouldBeVisible(const bool bInShouldBeVisible)
 
 void ULevelStreaming::SetShouldBeLoaded(const bool bInShouldBeLoaded)
 {
-	UE_LOG(LogStreaming, Log, TEXT("Calling SetShouldBeLoaded on an object (%s) for which it has no effect."), *GetFullName());
-}
-
-bool ULevelStreaming::ShouldBeLoaded() const
-{
-	return true;
+	UE_LOG(LogStreaming, Warning, TEXT("Calling SetShouldBeLoaded on an object (%s) for which it has no effect."), *GetFullName());
 }
 
 bool ULevelStreaming::ShouldBeVisible() const
