[NO NON-PII COMMIT MESSAGE]

diff --git a/Engine/Source/Runtime/Engine/Private/Attenuation.cpp b/Engine/Source/Runtime/Engine/Private/Attenuation.cpp
--- a/Engine/Source/Runtime/Engine/Private/Attenuation.cpp
+++ b/Engine/Source/Runtime/Engine/Private/Attenuation.cpp
@@ -72,10 +72,10 @@ float FBaseAttenuationSettings::AttenuationEval(const float Distance, const floa
 	// which could be 1.0 (and not 0.0f).
 
 	const float FalloffCopy = FMath::Max(Falloff, 1.0f);
-	const float DistanceCopy = FMath::Max(Distance * DistanceScale, 1.0f);
+	float DistanceCopy = Distance * DistanceScale;
 
 	float Result = 0.0f;
-	switch(DistanceAlgorithm)
+	switch (DistanceAlgorithm)
 	{
 		case EAttenuationDistanceModel::Linear:
 
@@ -83,12 +83,17 @@ float FBaseAttenuationSettings::AttenuationEval(const float Distance, const floa
 			break;
 
 		case EAttenuationDistanceModel::Logarithmic:
-			Result = 0.5f * - FMath::Loge(DistanceCopy / FalloffCopy);
+			{
+				DistanceCopy = FMath::Max(DistanceCopy, KINDA_SMALL_NUMBER);
+				Result = 0.5f * -FMath::Loge(DistanceCopy / FalloffCopy);
+			}
 			break;
 
 		case EAttenuationDistanceModel::Inverse:
-
-			Result = 0.02f / (DistanceCopy / FalloffCopy);
+			{
+				DistanceCopy = FMath::Max(DistanceCopy, KINDA_SMALL_NUMBER);
+				Result = 0.02f / (DistanceCopy / FalloffCopy);
+			}
 			break;
 
 		case EAttenuationDistanceModel::LogReverse:
@@ -99,7 +104,8 @@ float FBaseAttenuationSettings::AttenuationEval(const float Distance, const floa
 			}
 			else
 			{
-				Result = 1.0f + 0.5f * FMath::Loge(1.0f - (DistanceCopy / FalloffCopy));
+				const float Argument = FMath::Max(1.0f - (DistanceCopy / FalloffCopy), KINDA_SMALL_NUMBER);
+				Result = 1.0f + 0.5f * FMath::Loge(Argument);
 			}
 		}
 		break;
