Fix trashed indentation in RefreshBoneTransforms

diff --git a/Engine/Source/Runtime/Engine/Private/Components/SkeletalMeshComponent.cpp b/Engine/Source/Runtime/Engine/Private/Components/SkeletalMeshComponent.cpp
--- a/Engine/Source/Runtime/Engine/Private/Components/SkeletalMeshComponent.cpp
+++ b/Engine/Source/Runtime/Engine/Private/Components/SkeletalMeshComponent.cpp
@@ -1953,10 +1953,10 @@ void USkeletalMeshComponent::RefreshBoneTransforms(FActorComponentTickFunction*
 
 	const bool bDoPAE = !!CVarUseParallelAnimationEvaluation.GetValueOnGameThread() && FApp::ShouldUseThreadingForPerformance();
 
-const bool bMainInstanceValidForParallelWork = AnimScriptInstance == nullptr || AnimScriptInstance->CanRunParallelWork();
-const bool bPostInstanceValidForParallelWork = PostProcessAnimInstance == nullptr || PostProcessAnimInstance->CanRunParallelWork();
-const bool bHasValidInstanceForParallelWork = HasValidAnimationInstance() && bMainInstanceValidForParallelWork && bPostInstanceValidForParallelWork;
-const bool bDoParallelEvaluation = bHasValidInstanceForParallelWork && bDoPAE && bShouldDoEvaluation && TickFunction && (TickFunction->GetActualTickGroup() == TickFunction->TickGroup) && TickFunction->IsCompletionHandleValid();
+	const bool bMainInstanceValidForParallelWork = AnimScriptInstance == nullptr || AnimScriptInstance->CanRunParallelWork();
+	const bool bPostInstanceValidForParallelWork = PostProcessAnimInstance == nullptr || PostProcessAnimInstance->CanRunParallelWork();
+	const bool bHasValidInstanceForParallelWork = HasValidAnimationInstance() && bMainInstanceValidForParallelWork && bPostInstanceValidForParallelWork;
+	const bool bDoParallelEvaluation = bHasValidInstanceForParallelWork && bDoPAE && bShouldDoEvaluation && TickFunction && (TickFunction->GetActualTickGroup() == TickFunction->TickGroup) && TickFunction->IsCompletionHandleValid();
 	const bool bBlockOnTask = !bDoParallelEvaluation;  // If we aren't trying to do parallel evaluation then we
 															// will need to wait on an existing task.
 
@@ -2003,44 +2003,44 @@ const bool bDoParallelEvaluation = bHasValidInstanceForParallelWork && bDoPAE &&
 
 	if (bShouldDoEvaluation)
 	{
-	// If we need to eval the graph, and we're not going to update it.
-	// make sure it's been ticked at least once!
-	{
-		bool bShouldTickAnimation = false;		
-		if (AnimScriptInstance && !AnimScriptInstance->NeedsUpdate())
+		// If we need to eval the graph, and we're not going to update it.
+		// make sure it's been ticked at least once!
 		{
-			bShouldTickAnimation = bShouldTickAnimation || !AnimScriptInstance->GetUpdateCounter().HasEverBeenUpdated();
-			for (const UAnimInstance* SubInstance : SubInstances)
+			bool bShouldTickAnimation = false;		
+			if (AnimScriptInstance && !AnimScriptInstance->NeedsUpdate())
 			{
-				bShouldTickAnimation = bShouldTickAnimation || (SubInstance && !SubInstance->GetUpdateCounter().HasEverBeenUpdated());
+				bShouldTickAnimation = bShouldTickAnimation || !AnimScriptInstance->GetUpdateCounter().HasEverBeenUpdated();
+				for (const UAnimInstance* SubInstance : SubInstances)
+				{
+					bShouldTickAnimation = bShouldTickAnimation || (SubInstance && !SubInstance->GetUpdateCounter().HasEverBeenUpdated());
+				}
 			}
-		}
 
-		bShouldTickAnimation = bShouldTickAnimation || (ShouldPostUpdatePostProcessInstance() && !PostProcessAnimInstance->GetUpdateCounter().HasEverBeenUpdated());
+			bShouldTickAnimation = bShouldTickAnimation || (ShouldPostUpdatePostProcessInstance() && !PostProcessAnimInstance->GetUpdateCounter().HasEverBeenUpdated());
 
-		if (bShouldTickAnimation)
-		{
-			// We bypass TickPose() and call TickAnimation directly, so URO doesn't intercept us.
-			TickAnimation(0.f, false);
+			if (bShouldTickAnimation)
+			{
+				// We bypass TickPose() and call TickAnimation directly, so URO doesn't intercept us.
+				TickAnimation(0.f, false);
+			}
 		}
-	}
 
 		// If we're going to evaluate animation, call PreEvaluateAnimation()
 		{
-	if(AnimScriptInstance)
-	{
-		AnimScriptInstance->PreEvaluateAnimation();
+			if(AnimScriptInstance)
+			{
+				AnimScriptInstance->PreEvaluateAnimation();
 
-		for(UAnimInstance* SubInstance : SubInstances)
-		{
-			SubInstance->PreEvaluateAnimation();
-		}
-	}
+				for(UAnimInstance* SubInstance : SubInstances)
+				{
+					SubInstance->PreEvaluateAnimation();
+				}
+			}
 
-	if(ShouldEvaluatePostProcessInstance())
-	{
-		PostProcessAnimInstance->PreEvaluateAnimation();
-	}
+			if(ShouldEvaluatePostProcessInstance())
+			{
+				PostProcessAnimInstance->PreEvaluateAnimation();
+			}
 		}
 	}
 
