Debug safe zone won't try to draw if no safe zone is set.

diff --git a/Engine/Source/Runtime/Engine/Classes/Engine/GameViewportClient.h b/Engine/Source/Runtime/Engine/Classes/Engine/GameViewportClient.h
--- a/Engine/Source/Runtime/Engine/Classes/Engine/GameViewportClient.h
+++ b/Engine/Source/Runtime/Engine/Classes/Engine/GameViewportClient.h
@@ -70,9 +70,6 @@ public:
 	UPROPERTY()
 	TArray<struct FDebugDisplayProperty> DebugProperties;
 
-	/** border of safe area */
-	struct FTitleSafeZoneArea TitleSafeZone;
-
 	/** Array of the screen data needed for all the different splitscreen configurations */
 	TArray<struct FSplitscreenData> SplitscreenInfo;
 
diff --git a/Engine/Source/Runtime/Engine/Private/GameViewportClient.cpp b/Engine/Source/Runtime/Engine/Private/GameViewportClient.cpp
--- a/Engine/Source/Runtime/Engine/Private/GameViewportClient.cpp
+++ b/Engine/Source/Runtime/Engine/Private/GameViewportClient.cpp
@@ -156,6 +156,7 @@ UGameViewportClient::UGameViewportClient(const FObjectInitializer& ObjectInitial
 	, bIsMouseOverClient(false)
 #if WITH_EDITOR
 	, bUseMouseForTouchInEditor(false)
+	, bShowTitleSafeZone(true)
 #endif
 {
 
@@ -218,6 +219,9 @@ UGameViewportClient::UGameViewportClient(FVTableHelper& Helper)
 	, MouseLockMode(EMouseLockMode::LockOnCapture)
 	, AudioDeviceHandle(INDEX_NONE)
 	, bHasAudioFocus(false)
+#if WITH_EDITOR
+	, bShowTitleSafeZone(true)
+#endif
 {
 
 }
@@ -1803,12 +1807,12 @@ bool UGameViewportClient::IsOrtho() const
 
 void UGameViewportClient::PostRender(UCanvas* Canvas)
 {
-#if !WITH_EDITOR
-	if( bShowTitleSafeZone)
-#endif
+#if WITH_EDITOR
+	if(bShowTitleSafeZone)
 	{
 		DrawTitleSafeArea(Canvas);
 	}
+#endif
 
 	// Draw the transition screen.
 	DrawTransition(Canvas);
@@ -2219,7 +2223,9 @@ void UGameViewportClient::DrawTitleSafeArea( UCanvas* Canvas )
 	const FLinearColor UnsafeZoneColor(1.0f, 0.0f, 0.0f, 0.25f);
 	FCanvasTileItem TileItem(FVector2D::ZeroVector, GWhiteTexture, UnsafeZoneColor);
 	TileItem.BlendMode = SE_BLEND_Translucent;
-	if (PlayInSettings->DeviceToEmulate.IsEmpty())
+	
+	// CalculateSafeZoneValues() can be slow, so we only want to run it if we have boundaries to draw
+	if (FDisplayMetrics::GetDebugTitleSafeZoneRatio() < 1.f)
 	{
 		CalculateSafeZoneValues(SafeZone, Canvas, 0, false);
 		const float HeightOfSides = Height - SafeZone.GetTotalSpaceAlong<Orient_Vertical>();
@@ -2243,7 +2249,7 @@ void UGameViewportClient::DrawTitleSafeArea( UCanvas* Canvas )
 		TileItem.Size = FVector2D(SafeZone.Right, HeightOfSides);
 		Canvas->DrawItem(TileItem);
 	}
-	else
+	else if (!FSlateApplication::Get().GetCustomSafeZone().GetDesiredSize().IsZero())
 	{
 		ULevelEditorPlaySettings* PlaySettings = GetMutableDefault<ULevelEditorPlaySettings>();
 		PlaySettings->CalculateCustomUnsafeZones(PlaySettings->CustomUnsafeZoneStarts, PlaySettings->CustomUnsafeZoneDimensions, PlaySettings->DeviceToEmulate, FVector2D(Width, Height));
diff --git a/Engine/Source/Runtime/SlateCore/Public/Application/SlateApplicationBase.h b/Engine/Source/Runtime/SlateCore/Public/Application/SlateApplicationBase.h
--- a/Engine/Source/Runtime/SlateCore/Public/Application/SlateApplicationBase.h
+++ b/Engine/Source/Runtime/SlateCore/Public/Application/SlateApplicationBase.h
@@ -536,6 +536,7 @@ public:
 	}
 #if WITH_EDITOR
 	void ResetCustomSafeZone() { CustomSafeZoneRatio = FMargin(); }
+	const FMargin& GetCustomSafeZone() { return CustomSafeZoneRatio; }
 #endif
 
 #if WITH_EDITORONLY_DATA
