Fixing variable order in the constructor

diff --git a/Engine/Source/Runtime/Engine/Private/GameViewportClient.cpp b/Engine/Source/Runtime/Engine/Private/GameViewportClient.cpp
--- a/Engine/Source/Runtime/Engine/Private/GameViewportClient.cpp
+++ b/Engine/Source/Runtime/Engine/Private/GameViewportClient.cpp
@@ -152,11 +152,13 @@ UGameViewportClient::UGameViewportClient(const FObjectInitializer& ObjectInitial
 	, bHideCursorDuringCapture(false)
 	, MouseLockMode(EMouseLockMode::LockOnCapture)
 	, AudioDeviceHandle(INDEX_NONE)
+#if WITH_EDITOR
+	, bShowTitleSafeZone(true)
+#endif
 	, bHasAudioFocus(false)
 	, bIsMouseOverClient(false)
 #if WITH_EDITOR
 	, bUseMouseForTouchInEditor(false)
-	, bShowTitleSafeZone(true)
 #endif
 {
 
