Differentiate the name of the two load level instance functions

diff --git a/Engine/Source/Runtime/Engine/Classes/Engine/LevelStreamingKismet.h b/Engine/Source/Runtime/Engine/Classes/Engine/LevelStreamingKismet.h
--- a/Engine/Source/Runtime/Engine/Classes/Engine/LevelStreamingKismet.h
+++ b/Engine/Source/Runtime/Engine/Classes/Engine/LevelStreamingKismet.h
@@ -41,10 +41,10 @@ class ENGINE_API ULevelStreamingKismet : public ULevelStreaming
 	* @param bOutSuccess - Whether operation was successful (map was found and added to the sub-levels list)
  	* @return Streaming level object for a level instance
  	*/ 
- 	UFUNCTION(BlueprintCallable, Category = LevelStreaming, meta=(WorldContext="WorldContextObject"))
+ 	UFUNCTION(BlueprintCallable, Category = LevelStreaming, meta=(DisplayName = "Load Level Instance (by Name)", WorldContext="WorldContextObject"))
  	static ULevelStreamingKismet* LoadLevelInstance(UObject* WorldContextObject, FString LevelName, FVector Location, FRotator Rotation, bool& bOutSuccess);
 
- 	UFUNCTION(BlueprintCallable, Category = LevelStreaming, meta=(DisplayName = "Load Level Instance", WorldContext="WorldContextObject"))
+ 	UFUNCTION(BlueprintCallable, Category = LevelStreaming, meta=(DisplayName = "Load Level Instance (by Object Reference)", WorldContext="WorldContextObject"))
  	static ULevelStreamingKismet* LoadLevelInstanceBySoftObjectPtr(UObject* WorldContextObject, TSoftObjectPtr<UWorld> Level, FVector Location, FRotator Rotation, bool& bOutSuccess);
  	
 	//~ Begin UObject Interface
