Fix case where demo net driver reported a playback failure even though it was successful, which could corrupt the ShooterGame UI

#rnx

diff --git a/Engine/Source/Runtime/Engine/Private/DemoNetDriver.cpp b/Engine/Source/Runtime/Engine/Private/DemoNetDriver.cpp
--- a/Engine/Source/Runtime/Engine/Private/DemoNetDriver.cpp
+++ b/Engine/Source/Runtime/Engine/Private/DemoNetDriver.cpp
@@ -3474,6 +3474,9 @@ void UDemoNetDriver::ReplayStreamingReady( const FStartStreamingResult& Result )
 			return;
 		}
 
+		// InitConnectInternal calls ResetDemoState which will reset this, so restore the value
+		bWasStartStreamingSuccessful = Result.WasSuccessful();
+
 		const TCHAR* const SkipToLevelIndexOption = DemoURL.GetOption(TEXT("SkipToLevelIndex="), nullptr);
 		if (SkipToLevelIndexOption)
 		{
