Use FTextStringHelper::ReadFromString to determine if FText values have actually been changed from the AutogeneratedDefaultValue

diff --git a/Engine/Source/Runtime/Engine/Private/EdGraph/EdGraphSchema.cpp b/Engine/Source/Runtime/Engine/Private/EdGraph/EdGraphSchema.cpp
--- a/Engine/Source/Runtime/Engine/Private/EdGraph/EdGraphSchema.cpp
+++ b/Engine/Source/Runtime/Engine/Private/EdGraph/EdGraphSchema.cpp
@@ -3,6 +3,7 @@
 #include "EdGraph/EdGraphSchema.h"
 #include "UObject/MetaData.h"
 #include "UObject/UnrealType.h"
+#include "UObject/TextProperty.h"
 #include "Widgets/DeclarativeSyntaxSupport.h"
 #include "Widgets/SBoxPanel.h"
 #include "Styling/CoreStyle.h"
@@ -500,7 +501,15 @@ bool UEdGraphSchema::DoesDefaultValueMatchAutogenerated(const UEdGraphPin& InPin
 {
 	// Same logic as UEdGraphPin::DoesDefaultValueMatchAutogenerated
 	// Ignoring case on purpose to match default behavior
-	return InPin.GetDefaultAsString().Equals(InPin.AutogeneratedDefaultValue, ESearchCase::IgnoreCase);
+	FText AutogeneratedDefaultValueAsText;
+	if (!InPin.DefaultTextValue.IsEmpty() && FTextStringHelper::ReadFromString(*InPin.AutogeneratedDefaultValue, AutogeneratedDefaultValueAsText))
+	{
+		return UTextProperty::Identical_Implementation(AutogeneratedDefaultValueAsText, InPin.DefaultTextValue, 0);
+	}
+	else
+	{
+		return InPin.GetDefaultAsString().Equals(InPin.AutogeneratedDefaultValue, ESearchCase::IgnoreCase);
+	}
 }
 
 void UEdGraphSchema::BreakNodeLinks(UEdGraphNode& TargetNode) const
