Fix bReset comment for Activate and SetActive
#rnx

diff --git a/Engine/Source/Runtime/Engine/Classes/Components/ActorComponent.h b/Engine/Source/Runtime/Engine/Classes/Components/ActorComponent.h
--- a/Engine/Source/Runtime/Engine/Classes/Components/ActorComponent.h
+++ b/Engine/Source/Runtime/Engine/Classes/Components/ActorComponent.h
@@ -315,7 +315,7 @@ public:
 
 	/**
 	 * Activates the SceneComponent
-	 * @param bReset - The value to assign to HiddenGame.
+	 * @param bReset - Whether the activation should be forced even if ShouldActivate returns false.
 	 */
 	UFUNCTION(BlueprintCallable, Category="Components|Activation", meta=(UnsafeDuringActorConstruction="true"))
 	virtual void Activate(bool bReset=false);
@@ -329,6 +329,7 @@ public:
 	/**
 	 * Sets whether the component is active or not
 	 * @param bNewActive - The new active state of the component
+	 * @param bReset - Whether the activation should be forced even if ShouldActivate returns false.
 	 */
 	UFUNCTION(BlueprintCallable, Category="Components|Activation", meta=(UnsafeDuringActorConstruction="true"))
 	virtual void SetActive(bool bNewActive, bool bReset=false);
