Speculative compilation fix  #UE4

#rnx

diff --git a/Engine/Source/Runtime/Engine/Classes/AI/NavigationSystemBase.h b/Engine/Source/Runtime/Engine/Classes/AI/NavigationSystemBase.h
--- a/Engine/Source/Runtime/Engine/Classes/AI/NavigationSystemBase.h
+++ b/Engine/Source/Runtime/Engine/Classes/AI/NavigationSystemBase.h
@@ -3,7 +3,6 @@
 #pragma once
 
 #include "UObject/Object.h"
-#include "Engine/Engine.h"
 #include "AI/Navigation/NavigationTypes.h"
 #include "AI/NavigationModifier.h"
 #include "Engine/World.h"
