Further fixes to blending on landscapes with exactly 3 layers on mobile

diff --git a/Engine/Shaders/Public/ShaderVersion.ush b/Engine/Shaders/Public/ShaderVersion.ush
--- a/Engine/Shaders/Public/ShaderVersion.ush
+++ b/Engine/Shaders/Public/ShaderVersion.ush
@@ -2,4 +2,4 @@
 // Each time the console command is executed it generates a new GUID. As this file is included
 // in Platform.ush (which should be included in any shader) it allows to invalidate the shader DDC.
 // 
-// GUID = CA5419C0DF2F4DFCA153F9192B61D2F7
\ No newline at end of file
+// GUID = 070F3DFA44CE4B03BC5811709E951D69
\ No newline at end of file
diff --git a/Engine/Source/Runtime/Engine/Private/Materials/HLSLMaterialTranslator.h b/Engine/Source/Runtime/Engine/Private/Materials/HLSLMaterialTranslator.h
--- a/Engine/Source/Runtime/Engine/Private/Materials/HLSLMaterialTranslator.h
+++ b/Engine/Source/Runtime/Engine/Private/Materials/HLSLMaterialTranslator.h
@@ -4079,9 +4079,14 @@ protected:
 		FMaterialParameterInfo ParameterInfo = GetParameterAssociationInfo();
 		ParameterInfo.Name = ParameterName;
 
+		int32 NumActiveTerrainLayerWeightParameters = 0;
 		for(int32 ParameterIndex = 0;ParameterIndex < StaticParameters.TerrainLayerWeightParameters.Num(); ++ParameterIndex)
 		{
 			const FStaticTerrainLayerWeightParameter& Parameter = StaticParameters.TerrainLayerWeightParameters[ParameterIndex];
+			if (Parameter.WeightmapIndex != INDEX_NONE)
+			{
+				NumActiveTerrainLayerWeightParameters++;
+			}
 			if(Parameter.ParameterInfo == ParameterInfo)
 			{
 				WeightmapIndex = Parameter.WeightmapIndex;
@@ -4104,7 +4109,7 @@ protected:
 		else
 		{			
 			int32 WeightmapCode;
-			if (GetFeatureLevel() <= ERHIFeatureLevel::ES3_1 && StaticParameters.TerrainLayerWeightParameters.Num() <= 3 && bAtLeastOneWeightBasedBlend)
+			if (GetFeatureLevel() <= ERHIFeatureLevel::ES3_1 && NumActiveTerrainLayerWeightParameters <= 3 && bAtLeastOneWeightBasedBlend)
 			{
 				// Mobile can pack 3 layers into the normal map texture B and A channels, implying the 3rd using weight based blending
 				// Layer texture is sampled into Parameters.LayerWeights in LandscapeVertexFactory.ush
diff --git a/Engine/Source/Runtime/Landscape/Private/LandscapeEdit.cpp b/Engine/Source/Runtime/Landscape/Private/LandscapeEdit.cpp
--- a/Engine/Source/Runtime/Landscape/Private/LandscapeEdit.cpp
+++ b/Engine/Source/Runtime/Landscape/Private/LandscapeEdit.cpp
@@ -5018,12 +5018,11 @@ void ULandscapeComponent::GeneratePlatformPixelData()
 
 	int32 WeightmapSize = (SubsectionSizeQuads + 1) * NumSubsections;
 
+	MobileWeightmapTextures.Empty();
+
 	UTexture2D* MobileWeightNormalmapTexture = GetLandscapeProxy()->CreateLandscapeTexture(WeightmapSize, WeightmapSize, TEXTUREGROUP_Terrain_Weightmap, TSF_BGRA8);
 	CreateEmptyTextureMips(MobileWeightNormalmapTexture);
 
-	MobileWeightmapTextures.Empty();
-	MobileWeightmapTextures.Add(MobileWeightNormalmapTexture);
-
 	{
 		FLandscapeTextureDataInterface LandscapeData;
 
