Redo of speculative binary compilation fix #UE4

#rnx

diff --git a/Engine/Source/Runtime/Engine/Classes/AI/NavigationSystemBase.h b/Engine/Source/Runtime/Engine/Classes/AI/NavigationSystemBase.h
--- a/Engine/Source/Runtime/Engine/Classes/AI/NavigationSystemBase.h
+++ b/Engine/Source/Runtime/Engine/Classes/AI/NavigationSystemBase.h
@@ -3,7 +3,6 @@
 #pragma once
 
 #include "UObject/Object.h"
-#include "Engine/Engine.h"
 #include "AI/Navigation/NavigationTypes.h"
 #include "AI/NavigationModifier.h"
 #include "Engine/World.h"
diff --git a/Engine/Source/Runtime/Engine/Private/AI/NavigationSystemBase.cpp b/Engine/Source/Runtime/Engine/Private/AI/NavigationSystemBase.cpp
--- a/Engine/Source/Runtime/Engine/Private/AI/NavigationSystemBase.cpp
+++ b/Engine/Source/Runtime/Engine/Private/AI/NavigationSystemBase.cpp
@@ -1,6 +1,7 @@
 // Copyright 1998-2018 Epic Games, Inc. All Rights Reserved.
 
 #include "AI/NavigationSystemBase.h"
+#include "Engine/Engine.h"
 #include "Engine/World.h"
 #include "GameFramework/WorldSettings.h"
 #include "AI/Navigation/NavigationTypes.h"
