[NO NON-PII COMMIT MESSAGE]

diff --git a/Engine/Source/Runtime/Engine/Private/HierarchicalInstancedStaticMesh.cpp b/Engine/Source/Runtime/Engine/Private/HierarchicalInstancedStaticMesh.cpp
--- a/Engine/Source/Runtime/Engine/Private/HierarchicalInstancedStaticMesh.cpp
+++ b/Engine/Source/Runtime/Engine/Private/HierarchicalInstancedStaticMesh.cpp
@@ -2228,7 +2228,36 @@ void UHierarchicalInstancedStaticMeshComponent::ClearInstances()
 		DEC_DWORD_STAT_BY(STAT_FoliageInstanceBuffers, ProxySize);
 	}
 
-	Super::ClearInstances();
+	InstanceUpdateCmdBuffer.Reset();
+
+	// Hide all instance until the build tree is completed
+	int32 NumInstances = PerInstanceSMData.Num();
+
+	for (int32 Index = 0; Index < NumInstances; ++Index)
+	{
+		int32 RenderIndex = InstanceReorderTable.IsValidIndex(Index) ? InstanceReorderTable[Index] : Index;
+		if (RenderIndex == INDEX_NONE)
+		{
+			// could be skipped by density settings
+			continue;
+		}
+
+		InstanceUpdateCmdBuffer.HideInstance(RenderIndex);
+	}
+
+	// Clear all the per-instance data
+	PerInstanceSMData.Empty();
+	InstanceReorderTable.Empty();
+	InstanceDataBuffers.Reset();
+
+	ProxySize = 0;
+
+	// Release any physics representations
+	ClearAllInstanceBodies();
+
+	MarkRenderStateDirty();
+
+	FNavigationSystem::UpdateComponentData(*this);
 }
 
 bool UHierarchicalInstancedStaticMeshComponent::ShouldCreatePhysicsState() const
