CIS fix #UE4

#rnx

diff --git a/Engine/Source/Runtime/Engine/Private/VectorField.cpp b/Engine/Source/Runtime/Engine/Private/VectorField.cpp
--- a/Engine/Source/Runtime/Engine/Private/VectorField.cpp
+++ b/Engine/Source/Runtime/Engine/Private/VectorField.cpp
@@ -27,6 +27,7 @@
 #include "Components/VectorFieldComponent.h"
 #include "PrimitiveSceneProxy.h"
 #include "Materials/Material.h"
+#include "Engine/Engine.h"
 
 #if WITH_EDITORONLY_DATA
 	#include "EditorFramework/AssetImportData.h"
