REL - By request of Alexis Matte
Fix recompute tangent on meshes that have more than one texcoord

diff --git a/Engine/Shaders/Private/RecomputeTangentsCommon.ush b/Engine/Shaders/Private/RecomputeTangentsCommon.ush
--- a/Engine/Shaders/Private/RecomputeTangentsCommon.ush
+++ b/Engine/Shaders/Private/RecomputeTangentsCommon.ush
@@ -29,6 +29,8 @@ RWBuffer<float4> TangentBufferUAV;
 
 // in vertices
 uint InputStreamStart;
+uint NumTexCoords;
+uint SelectedTexCoord;
 
 // index input ----------------------------
 
diff --git a/Engine/Shaders/Private/RecomputeTangentsPerTrianglePass.usf b/Engine/Shaders/Private/RecomputeTangentsPerTrianglePass.usf
--- a/Engine/Shaders/Private/RecomputeTangentsPerTrianglePass.usf
+++ b/Engine/Shaders/Private/RecomputeTangentsPerTrianglePass.usf
@@ -60,7 +60,7 @@ FVertexData GetVertex(uint VertexIndex)
 	{
 		uint InputOffset = VertexIndex + InputStreamStart;
 
-		Ret.UV = UVsInputBuffer[InputOffset];
+		Ret.UV = UVsInputBuffer[InputOffset * NumTexCoords + SelectedTexCoord];
 	}
 
 	return Ret;
diff --git a/Engine/Source/Runtime/Engine/Private/GPUSkinCache.cpp b/Engine/Source/Runtime/Engine/Private/GPUSkinCache.cpp
--- a/Engine/Source/Runtime/Engine/Private/GPUSkinCache.cpp
+++ b/Engine/Source/Runtime/Engine/Private/GPUSkinCache.cpp
@@ -169,76 +169,60 @@ public:
 	{
 		FGPUSkinCache::FRWBufferTracker PositionTracker;
 
-		FGPUBaseSkinVertexFactory* SourceVertexFactory;
-		FGPUSkinPassthroughVertexFactory* TargetVertexFactory;
+		FGPUBaseSkinVertexFactory* SourceVertexFactory = nullptr;
+		FGPUSkinPassthroughVertexFactory* TargetVertexFactory = nullptr;
 
 		// triangle index buffer (input for the RecomputeSkinTangents, might need special index buffer unique to position and normal, not considering UV/vertex color)
-		FShaderResourceViewRHIParamRef IndexBuffer;
+		FShaderResourceViewRHIParamRef IndexBuffer = nullptr;
 
-		const FSkelMeshRenderSection* Section;
+		const FSkelMeshRenderSection* Section = nullptr;
 
 		// for debugging / draw events, -1 if not set
-		uint32 SectionIndex;
+		uint32 SectionIndex = -1;
 
 		// 0:normal, 1:with morph target, 2:with APEX cloth (not yet implemented)
-		uint32 SkinType;
+		uint32 SkinType = 0;
 		//
-		bool bExtraBoneInfluences;
+		bool bExtraBoneInfluences = false;
 
 		// in floats (4 bytes)
-		uint32 OutputStreamStart;
-		uint32 NumVertices;
+		uint32 OutputStreamStart = 0;
+		uint32 NumVertices = 0;
 
 		// in vertices
-		uint32 InputStreamStart;
+		uint32 InputStreamStart = 0;
+		uint32 NumTexCoords = 1;
+		uint32 SelectedTexCoord = 0;
 
-		FShaderResourceViewRHIRef TangentBufferSRV;
-		FShaderResourceViewRHIRef UVsBufferSRV;
-		FShaderResourceViewRHIRef PositionBufferSRV;
+		FShaderResourceViewRHIRef TangentBufferSRV = nullptr;
+		FShaderResourceViewRHIRef UVsBufferSRV = nullptr;
+		FShaderResourceViewRHIRef PositionBufferSRV = nullptr;
 
 		// skin weight input
-		uint32 InputWeightStart;
+		uint32 InputWeightStart = 0;
 
 		// morph input
-		uint32 MorphBufferOffset;
+		uint32 MorphBufferOffset = 0;
 
         // cloth input
-        float ClothBlendWeight;
+        float ClothBlendWeight = 0.0f;
 
-        FMatrix ClothLocalToWorld;
-        FMatrix ClothWorldToLocal;
+        FMatrix ClothLocalToWorld = FMatrix::Identity;
+        FMatrix ClothWorldToLocal = FMatrix::Identity;
 
 		// triangle index buffer (input for the RecomputeSkinTangents, might need special index buffer unique to position and normal, not considering UV/vertex color)
-		uint32 IndexBufferOffsetValue;
-		uint32 NumTriangles;
+		uint32 IndexBufferOffsetValue = 0;
+		uint32 NumTriangles = 0;
 
-		FRWBuffer* TangentBuffer;
-		FRWBuffer* PositionBuffer;
-		FRWBuffer* PreviousPositionBuffer;
+		FRWBuffer* TangentBuffer = nullptr;
+		FRWBuffer* PositionBuffer = nullptr;
+		FRWBuffer* PreviousPositionBuffer = nullptr;
 
         // Handle duplicates
-        FShaderResourceViewRHIRef DuplicatedIndicesIndices;
-        FShaderResourceViewRHIRef DuplicatedIndices;
-
-		FSectionDispatchData()
-			: SourceVertexFactory(nullptr)
-			, TargetVertexFactory(nullptr)
-			, IndexBuffer(nullptr)
-			, Section(nullptr)
-			, SectionIndex(-1)
-			, SkinType(0)
-			, bExtraBoneInfluences(false)
-			, OutputStreamStart(0)
-			, NumVertices(0)
-			, InputStreamStart(0)
-			, MorphBufferOffset(0)
-			, IndexBufferOffsetValue(0)
-			, NumTriangles(0)
-			, TangentBuffer(nullptr)
-			, PositionBuffer(nullptr)
-			, PreviousPositionBuffer(nullptr)
-		{
-		}
+        FShaderResourceViewRHIRef DuplicatedIndicesIndices = nullptr;
+        FShaderResourceViewRHIRef DuplicatedIndices = nullptr;
+
+		FSectionDispatchData() = default;
 
 		inline FRWBuffer* GetPreviousPositionRWBuffer()
 		{
@@ -324,6 +308,7 @@ public:
 
 		Data.TangentBufferSRV = InSourceVertexFactory->GetTangentsSRV();
 		Data.UVsBufferSRV = InSourceVertexFactory->GetTextureCoordinatesSRV();
+		Data.NumTexCoords = InSourceVertexFactory->GetNumTexCoords();
 		Data.PositionBufferSRV = InSourceVertexFactory->GetPositionsSRV();
 
 		Data.bExtraBoneInfluences = InSourceVertexFactory->UsesExtraBoneInfluences();
@@ -634,6 +619,8 @@ public:
 	FShaderResourceParameter IndexBuffer;
 	FShaderParameter IndexBufferOffset;
 	FShaderParameter InputStreamStart;
+	FShaderParameter NumTexCoords;
+	FShaderParameter SelectedTexCoord;
 	FShaderResourceParameter TangentInputBuffer;
 	FShaderResourceParameter UVsInputBuffer;
 	FShaderResourceParameter DuplicatedIndices;
@@ -654,6 +641,8 @@ public:
 		IndexBufferOffset.Bind(Initializer.ParameterMap, TEXT("IndexBufferOffset"));
 
 		InputStreamStart.Bind(Initializer.ParameterMap, TEXT("InputStreamStart"));
+		NumTexCoords.Bind(Initializer.ParameterMap, TEXT("NumTexCoords"));
+		SelectedTexCoord.Bind(Initializer.ParameterMap, TEXT("SelectedTexCoord"));
 		TangentInputBuffer.Bind(Initializer.ParameterMap, TEXT("TangentInputBuffer"));
 		UVsInputBuffer.Bind(Initializer.ParameterMap, TEXT("UVsInputBuffer"));
 
@@ -679,6 +668,8 @@ public:
 		SetShaderValue(RHICmdList, ShaderRHI, IndexBufferOffset, DispatchData.IndexBufferOffsetValue);
 		
 		SetShaderValue(RHICmdList, ShaderRHI, InputStreamStart, DispatchData.InputStreamStart);
+		SetShaderValue(RHICmdList, ShaderRHI, NumTexCoords, DispatchData.NumTexCoords);
+		SetShaderValue(RHICmdList, ShaderRHI, SelectedTexCoord, DispatchData.SelectedTexCoord);
 		SetSRVParameter(RHICmdList, ShaderRHI, TangentInputBuffer, DispatchData.TangentBufferSRV);
 		SetSRVParameter(RHICmdList, ShaderRHI, TangentInputBuffer, DispatchData.UVsBufferSRV);
 
@@ -703,7 +694,7 @@ public:
 	{
 		bool bShaderHasOutdatedParameters = FGlobalShader::Serialize(Ar);
 		Ar << IntermediateAccumBufferUAV << NumTriangles << GPUPositionCacheBuffer << GPUTangentCacheBuffer << SkinCacheStart << IndexBuffer << IndexBufferOffset
-			<< InputStreamStart << TangentInputBuffer << UVsInputBuffer;
+			<< InputStreamStart << NumTexCoords << SelectedTexCoord << TangentInputBuffer << UVsInputBuffer;
         Ar << DuplicatedIndices << DuplicatedIndicesIndices;
 		return bShaderHasOutdatedParameters;
 	}
diff --git a/Engine/Source/Runtime/Engine/Public/GPUSkinVertexFactory.h b/Engine/Source/Runtime/Engine/Public/GPUSkinVertexFactory.h
--- a/Engine/Source/Runtime/Engine/Public/GPUSkinVertexFactory.h
+++ b/Engine/Source/Runtime/Engine/Public/GPUSkinVertexFactory.h
@@ -369,6 +369,7 @@ public:
 	virtual const FShaderResourceViewRHIRef GetTangentsSRV() const = 0;
 	virtual const FShaderResourceViewRHIRef GetTextureCoordinatesSRV() const = 0;
 	virtual const FShaderResourceViewRHIRef GetColorComponentsSRV() const = 0;
+	virtual uint32 GetNumTexCoords() const = 0;
 	virtual const uint32 GetColorIndexMask() const = 0;
 
 	inline const FVertexStreamComponent& GetTangentStreamComponent(int Index)
@@ -471,6 +472,11 @@ public:
 		return Data.TextureCoordinatesSRV;
 	}
 
+	uint32 GetNumTexCoords() const override
+	{
+		return Data.NumTexCoords;
+	}
+
 	const FShaderResourceViewRHIRef GetColorComponentsSRV() const override
 	{
 		return Data.ColorComponentsSRV;
@@ -623,6 +629,11 @@ public:
 		return MorphData.TextureCoordinatesSRV;
 	}
 
+	uint32 GetNumTexCoords() const override
+	{
+		return MorphData.NumTexCoords;
+	}
+
 	const FShaderResourceViewRHIRef GetColorComponentsSRV() const override
 	{
 		return MorphData.ColorComponentsSRV;
@@ -930,6 +941,36 @@ public:
 		return this;
 	}
 
+	const FShaderResourceViewRHIRef GetPositionsSRV() const override
+	{
+		return MeshMappingData.PositionComponentSRV;
+	}
+
+	const FShaderResourceViewRHIRef GetTangentsSRV() const override
+	{
+		return MeshMappingData.TangentsSRV;
+	}
+
+	const FShaderResourceViewRHIRef GetTextureCoordinatesSRV() const override
+	{
+		return MeshMappingData.TextureCoordinatesSRV;
+	}
+
+	uint32 GetNumTexCoords() const override
+	{
+		return MeshMappingData.NumTexCoords;
+	}
+
+	const FShaderResourceViewRHIRef GetColorComponentsSRV() const override
+	{
+		return MeshMappingData.ColorComponentsSRV;
+	}
+
+	const uint32 GetColorIndexMask() const override
+	{
+		return MeshMappingData.ColorIndexMask;
+	}
+
 	// FRenderResource interface.
 
 	/**
