Workaround for static draw lists not being updated when a volumetric lightmap is added to the scene during Undo.  Undo causes the ULevel and UActorComponent to be processed in a different order than loading.

diff --git a/Engine/Source/Runtime/Engine/Private/Level.cpp b/Engine/Source/Runtime/Engine/Private/Level.cpp
--- a/Engine/Source/Runtime/Engine/Private/Level.cpp
+++ b/Engine/Source/Runtime/Engine/Private/Level.cpp
@@ -56,6 +56,7 @@ Level.cpp: Level-related functions
 #include "Components/ModelComponent.h"
 #include "Engine/LevelActorContainer.h"
 #include "Engine/StaticMeshActor.h"
+#include "ComponentRecreateRenderStateContext.h"
 
 DEFINE_LOG_CATEGORY(LogLevel);
 
@@ -1435,6 +1436,11 @@ void ULevel::PostEditUndo()
 		if (bIsStreamingLevelVisible)
 		{
 			InitializeRenderingResources();
+
+			// Hack: FScene::AddPrecomputedVolumetricLightmap does not cause static draw lists to be updated - force an update so the correct base pass shader is selected in ProcessBasePassMesh.  
+			// With the normal load order, the level rendering resources are always initialized before the components that are in the level, so this is not an issue. 
+			// During undo, PostEditUndo on the component and ULevel are called in an arbitrary order.
+			MarkLevelComponentsRenderStateDirty();
 		}
 	}
 
