Fix ensure when switching editor viewport feature level.
Set null render target during FPlanarReflectionSceneProxy constructor as this can cause an invalid state within FPlanarReflectionParameters::SetParameters during UpdatePlanarReflectionContents with multiple planar captures.

diff --git a/Engine/Source/Runtime/Engine/Private/Components/SceneCaptureComponent.cpp b/Engine/Source/Runtime/Engine/Private/Components/SceneCaptureComponent.cpp
--- a/Engine/Source/Runtime/Engine/Private/Components/SceneCaptureComponent.cpp
+++ b/Engine/Source/Runtime/Engine/Private/Components/SceneCaptureComponent.cpp
@@ -762,7 +762,7 @@ void UPlanarReflectionComponent::CreateRenderState_Concurrent()
 
 	if (ShouldComponentAddToScene() && ShouldRender())
 	{
-		SceneProxy = new FPlanarReflectionSceneProxy(this, RenderTarget);
+		SceneProxy = new FPlanarReflectionSceneProxy(this);
 		GetWorld()->Scene->AddPlanarReflection(this);
 	}
 }
diff --git a/Engine/Source/Runtime/Engine/Private/PlanarReflectionSceneProxy.cpp b/Engine/Source/Runtime/Engine/Private/PlanarReflectionSceneProxy.cpp
--- a/Engine/Source/Runtime/Engine/Private/PlanarReflectionSceneProxy.cpp
+++ b/Engine/Source/Runtime/Engine/Private/PlanarReflectionSceneProxy.cpp
@@ -8,11 +8,9 @@
 #include "GameFramework/Actor.h"
 #include "Components/PlanarReflectionComponent.h"
 
-FPlanarReflectionSceneProxy::FPlanarReflectionSceneProxy(UPlanarReflectionComponent* Component, FPlanarReflectionRenderTarget* InRenderTarget) :
-	bIsStereo(false)
+FPlanarReflectionSceneProxy::FPlanarReflectionSceneProxy(UPlanarReflectionComponent* Component) :
+	bIsStereo(false), RenderTarget(nullptr)
 {
-	RenderTarget = InRenderTarget;
-
 	float ClampedFadeStart = FMath::Max(Component->DistanceFromPlaneFadeoutStart, 0.0f);
 	float ClampedFadeEnd = FMath::Max(Component->DistanceFromPlaneFadeoutEnd, 0.0f);
 
diff --git a/Engine/Source/Runtime/Engine/Public/PlanarReflectionSceneProxy.h b/Engine/Source/Runtime/Engine/Public/PlanarReflectionSceneProxy.h
--- a/Engine/Source/Runtime/Engine/Public/PlanarReflectionSceneProxy.h
+++ b/Engine/Source/Runtime/Engine/Public/PlanarReflectionSceneProxy.h
@@ -86,7 +86,7 @@ class FPlanarReflectionSceneProxy
 {
 public:
 
-	FPlanarReflectionSceneProxy(UPlanarReflectionComponent* Component, FPlanarReflectionRenderTarget* InRenderTarget);
+	FPlanarReflectionSceneProxy(UPlanarReflectionComponent* Component);
 
 	void UpdateTransform(const FMatrix& NewTransform)
 	{
