Remove warning from ULevelStreaming about ShouldBeLoaded having no effect as a number of systems iterate and set all streaming levels state regardless of whether they can be manipulated or not.
#rnx

diff --git a/Engine/Source/Runtime/Engine/Private/LevelStreaming.cpp b/Engine/Source/Runtime/Engine/Private/LevelStreaming.cpp
--- a/Engine/Source/Runtime/Engine/Private/LevelStreaming.cpp
+++ b/Engine/Source/Runtime/Engine/Private/LevelStreaming.cpp
@@ -1419,7 +1419,6 @@ void ULevelStreaming::SetShouldBeVisible(const bool bInShouldBeVisible)
 
 void ULevelStreaming::SetShouldBeLoaded(const bool bInShouldBeLoaded)
 {
-	UE_LOG(LogStreaming, Warning, TEXT("Calling SetShouldBeLoaded on an object (%s) for which it has no effect."), *GetFullName());
 }
 
 bool ULevelStreaming::ShouldBeVisible() const
