If the property class of the object property is null then go back to the old behavior where we set the object property regarldess as we don't really understand what the PropertyClass being null represents
#rnx

diff --git a/Engine/Source/Runtime/Engine/Private/SCS_Node.cpp b/Engine/Source/Runtime/Engine/Private/SCS_Node.cpp
--- a/Engine/Source/Runtime/Engine/Private/SCS_Node.cpp
+++ b/Engine/Source/Runtime/Engine/Private/SCS_Node.cpp
@@ -164,7 +164,8 @@ UActorComponent* USCS_Node::ExecuteNodeOnActor(AActor* Actor, USceneComponent* P
 			UClass* ActorClass = Actor->GetClass();
 			if (UObjectPropertyBase* Prop = FindField<UObjectPropertyBase>(ActorClass, VarName))
 			{
-				if (NewActorComp->IsA(Prop->PropertyClass))
+				// If it is null we don't really know what's going on, but make it behave as it did before the bug fix
+				if (Prop->PropertyClass == nullptr || NewActorComp->IsA(Prop->PropertyClass))
 				{
 					Prop->SetObjectPropertyValue_InContainer(Actor, NewActorComp);
 				}
