More data to help track a material crash

diff --git a/Engine/Source/Runtime/Engine/Public/MaterialShared.h b/Engine/Source/Runtime/Engine/Public/MaterialShared.h
--- a/Engine/Source/Runtime/Engine/Public/MaterialShared.h
+++ b/Engine/Source/Runtime/Engine/Public/MaterialShared.h
@@ -862,9 +862,11 @@ public:
 	uint32 GetNumUsedCustomInterpolatorScalars() const { return MaterialCompilationOutput.NumUsedCustomInterpolatorScalars; }
 	void GetEstimatedNumTextureSamples(uint32& VSSamples, uint32& PSSamples) const { VSSamples = MaterialCompilationOutput.EstimatedNumTextureSamplesVS; PSSamples = MaterialCompilationOutput.EstimatedNumTextureSamplesPS; }
 
-	bool IsValidForRendering() const
+	bool IsValidForRendering(bool bFailOnInvalid = false) const
 	{
-		return bCompilationFinalized && bCompiledSuccessfully && !bDeletedThroughDeferredCleanup;
+		const bool bValid = bCompilationFinalized && bCompiledSuccessfully && !bDeletedThroughDeferredCleanup;
+		checkf(bValid || !bFailOnInvalid, TEXT("FMaterialShaderMap %s invalid for rendering: bCompilationFinalized: %i, bCompiledSuccessfully: %i, bDeletedThroughDeferredCleanup: %i"), *GetFriendlyName(), bCompilationFinalized, bCompiledSuccessfully, bDeletedThroughDeferredCleanup ? 1 : 0);
+		return bValid;
 	}
 
 	const FUniformExpressionSet& GetUniformExpressionSet() const { return MaterialCompilationOutput.UniformExpressionSet; }
diff --git a/Engine/Source/Runtime/Renderer/Private/ShaderBaseClasses.cpp b/Engine/Source/Runtime/Renderer/Private/ShaderBaseClasses.cpp
--- a/Engine/Source/Runtime/Renderer/Private/ShaderBaseClasses.cpp
+++ b/Engine/Source/Runtime/Renderer/Private/ShaderBaseClasses.cpp
@@ -177,7 +177,7 @@ void FMaterialShader::SetParametersInner(
 
 	ERHIFeatureLevel::Type FeatureLevel = View.GetFeatureLevel();
 	checkf(Material.GetRenderingThreadShaderMap(), TEXT("RenderingThreadShaderMap: %i"), Material.GetRenderingThreadShaderMap() ? 1 : 0);
-	checkf(Material.GetRenderingThreadShaderMap()->IsValidForRendering() && Material.GetFeatureLevel() == FeatureLevel, TEXT("IsValid:%i, MaterialFeatureLevel:%i, FeatureLevel:%i"), Material.GetRenderingThreadShaderMap()->IsValidForRendering() ? 1 : 0, Material.GetFeatureLevel(), FeatureLevel);
+	checkf(Material.GetRenderingThreadShaderMap()->IsValidForRendering(true) && Material.GetFeatureLevel() == FeatureLevel, TEXT("IsValid:%i, MaterialFeatureLevel:%i, FeatureLevel:%i"), Material.GetRenderingThreadShaderMap()->IsValidForRendering() ? 1 : 0, Material.GetFeatureLevel(), FeatureLevel);
 
 	FUniformExpressionCache* UniformExpressionCache = &MaterialRenderProxy->UniformExpressionCache[FeatureLevel];
 	bool bUniformExpressionCacheNeedsDelete = false;
