Changing assert to log error.

diff --git a/Engine/Source/Runtime/Engine/Private/ADPCMAudioInfo.cpp b/Engine/Source/Runtime/Engine/Private/ADPCMAudioInfo.cpp
--- a/Engine/Source/Runtime/Engine/Private/ADPCMAudioInfo.cpp
+++ b/Engine/Source/Runtime/Engine/Private/ADPCMAudioInfo.cpp
@@ -73,7 +73,13 @@ void FADPCMAudioInfo::SeekToTime(const float SeekTime)
 
 bool FADPCMAudioInfo::ReadCompressedInfo(const uint8* InSrcBufferData, uint32 InSrcBufferDataSize, struct FSoundQualityInfo* QualityInfo)
 {
-	check(InSrcBufferData);
+	if (!InSrcBufferData)
+	{
+		FString Name = QualityInfo ? QualityInfo->DebugName : TEXT("Unknown");
+		UE_LOG(LogAudio, Warning, TEXT("Failed to read compressed ADPCM audio from ('%s') because there was no resource data."), *Name);
+
+		return false;
+	}
 
 	SrcBufferData = InSrcBufferData;
 	SrcBufferDataSize = InSrcBufferDataSize;
