Fix the undo remove LOD

diff --git a/Engine/Source/Runtime/Engine/Private/StaticMesh.cpp b/Engine/Source/Runtime/Engine/Private/StaticMesh.cpp
--- a/Engine/Source/Runtime/Engine/Private/StaticMesh.cpp
+++ b/Engine/Source/Runtime/Engine/Private/StaticMesh.cpp
@@ -2734,44 +2734,63 @@ static FStaticMeshRenderData& GetPlatformStaticMeshRenderData(UStaticMesh* Mesh,
 
 void UStaticMesh::LoadMeshDescriptions()
 {
-	// Do nothing if already loaded
 	if (MeshDescriptions)
 	{
-		return;
+		//Sync the already loaded MeshDescription
+		MeshDescriptions->SetNum(SourceModels.Num());
+		for (int32 LodIndex = 0; LodIndex < MeshDescriptions->Num(); ++LodIndex)
+		{
+			//Get the missing MeshDescription to create them from the FRawMesh
+			if (MeshDescriptions->Get(LodIndex) == nullptr && !SourceModels[LodIndex].IsRawMeshEmpty())
+			{
+				// If the MeshDescriptions are out of sync with the SourceModels RawMesh, perform a conversion here.
+				// @todo: once all tools are ported, we can replace this with a check() instead.
+				FMeshDescription* MeshDescription = MeshDescriptions->Create(LodIndex);
+				SourceModels[LodIndex].OriginalMeshDescription = MeshDescription;
+				RegisterMeshAttributes(*MeshDescription);
+
+				FRawMesh LodRawMesh;
+				SourceModels[LodIndex].LoadRawMesh(LodRawMesh);
+				TMap<int32, FName> MaterialMap;
+				FillMaterialName(StaticMaterials, MaterialMap);
+				FMeshDescriptionOperations::ConvertFromRawMesh(LodRawMesh, *MeshDescription, MaterialMap);
+			}
+		}
 	}
-
-	MeshDescriptions = NewObject<UStaticMeshDescriptions>(GetTransientPackage());
-
-	// For the moment, this comes from the DDC. Eventually it will load the UObject from the same package as the static mesh
-	// from a soft object path.
-	FString MeshDataKey;
-	if (GetMeshDataKey(MeshDataKey))
+	else
 	{
-		TArray<uint8> DerivedData;
-		if (GetDerivedDataCacheRef().GetSynchronous(*MeshDataKey, DerivedData))
+		MeshDescriptions = NewObject<UStaticMeshDescriptions>(GetTransientPackage());
+
+		// For the moment, this comes from the DDC. Eventually it will load the UObject from the same package as the static mesh
+		// from a soft object path.
+		FString MeshDataKey;
+		if (GetMeshDataKey(MeshDataKey))
 		{
-			// Load from the DDC
-			const bool bIsPersistent = true;
-			FMemoryReader Ar(DerivedData, bIsPersistent);
-			MeshDescriptions->Serialize(Ar);
+			TArray<uint8> DerivedData;
+			if (GetDerivedDataCacheRef().GetSynchronous(*MeshDataKey, DerivedData))
+			{
+				// Load from the DDC
+				const bool bIsPersistent = true;
+				FMemoryReader Ar(DerivedData, bIsPersistent);
+				MeshDescriptions->Serialize(Ar);
+			}
+			else
+			{
+				// Nothing cached in the DDC; create a blank one
+				MeshDescriptions->SetNum(SourceModels.Num());
+			}
 		}
 		else
 		{
-			// Nothing cached in the DDC; create a blank one
-			MeshDescriptions->SetNum(SourceModels.Num());
+			// If we get here, it's because there are no SourceModels.
+			// At this point we just have an empty UStaticMeshDescriptions object.
+		}
+		// Assign the pointer in the individual FStaticMeshSourceModels
+		check(MeshDescriptions->Num() == SourceModels.Num());
+		for (int32 Index = 0; Index < SourceModels.Num(); ++Index)
+		{
+			SourceModels[Index].OriginalMeshDescription = MeshDescriptions->Get(Index);
 		}
-	}
-	else
-	{
-		// If we get here, it's because there are no SourceModels.
-		// At this point we just have an empty UStaticMeshDescriptions object.
-	}
-
-	// Assign the pointer in the individual FStaticMeshSourceModels
-	check(MeshDescriptions->Num() == SourceModels.Num());
-	for (int32 Index = 0; Index < SourceModels.Num(); ++Index)
-	{
-		SourceModels[Index].OriginalMeshDescription = MeshDescriptions->Get(Index);
 	}
 }
 
