Cook foliage for HTML5 using full precision instance transforms
Remove previous shader fix for this issue

diff --git a/Engine/Shaders/Private/LocalVertexFactory.ush b/Engine/Shaders/Private/LocalVertexFactory.ush
--- a/Engine/Shaders/Private/LocalVertexFactory.ush
+++ b/Engine/Shaders/Private/LocalVertexFactory.ush
@@ -636,15 +636,9 @@ FVertexFactoryIntermediates GetVertexFactoryIntermediates(FVertexFactoryInput In
 	Intermediates.InstanceOrigin = VertexFetch_InstanceOriginBuffer[(Input.InstanceId + InstanceOffset)];
 	Intermediates.InstanceLightmapAndShadowMapUVBias = VertexFetch_InstanceLightmapBuffer[(Input.InstanceId + InstanceOffset)];
 #elif USE_INSTANCING && !USE_INSTANCING_EMULATED
-	#if WEBGL
-		Intermediates.InstanceTransform1 = f16tof32(Input.InstanceTransform1);
-		Intermediates.InstanceTransform2 = f16tof32(Input.InstanceTransform2);
-		Intermediates.InstanceTransform3 = f16tof32(Input.InstanceTransform3);
-	#else
-		Intermediates.InstanceTransform1 = Input.InstanceTransform1;
-		Intermediates.InstanceTransform2 = Input.InstanceTransform2;
-		Intermediates.InstanceTransform3 = Input.InstanceTransform3;
-	#endif
+	Intermediates.InstanceTransform1 = Input.InstanceTransform1;
+	Intermediates.InstanceTransform2 = Input.InstanceTransform2;
+	Intermediates.InstanceTransform3 = Input.InstanceTransform3;
 	Intermediates.InstanceOrigin = Input.InstanceOrigin;
 	Intermediates.InstanceLightmapAndShadowMapUVBias = Input.InstanceLightmapAndShadowMapUVBias;
 #endif
diff --git a/Engine/Source/Developer/HTML5/HTML5TargetPlatform/Private/HTML5TargetPlatform.h b/Engine/Source/Developer/HTML5/HTML5TargetPlatform/Private/HTML5TargetPlatform.h
--- a/Engine/Source/Developer/HTML5/HTML5TargetPlatform/Private/HTML5TargetPlatform.h
+++ b/Engine/Source/Developer/HTML5/HTML5TargetPlatform/Private/HTML5TargetPlatform.h
@@ -63,6 +63,8 @@ public:
 			return true;
 		case ETargetPlatformFeatures::DeferredRendering:
 			return false;
+		case ETargetPlatformFeatures::HalfFloatVertexFormat:
+			return false;
 		}
 
 		return TTargetPlatformBase<FHTML5PlatformProperties>::SupportsFeature(Feature);
diff --git a/Engine/Source/Developer/TargetPlatform/Public/Common/TargetPlatformBase.h b/Engine/Source/Developer/TargetPlatform/Public/Common/TargetPlatformBase.h
--- a/Engine/Source/Developer/TargetPlatform/Public/Common/TargetPlatformBase.h
+++ b/Engine/Source/Developer/TargetPlatform/Public/Common/TargetPlatformBase.h
@@ -226,6 +226,9 @@ public:
 
 		case ETargetPlatformFeatures::ShouldSplitPaksIntoSmallerSizes :
 			return false;
+
+		case ETargetPlatformFeatures::HalfFloatVertexFormat:
+			return true;
 		}
 
 		return false;
diff --git a/Engine/Source/Developer/TargetPlatform/Public/Interfaces/ITargetPlatform.h b/Engine/Source/Developer/TargetPlatform/Public/Interfaces/ITargetPlatform.h
--- a/Engine/Source/Developer/TargetPlatform/Public/Interfaces/ITargetPlatform.h
+++ b/Engine/Source/Developer/TargetPlatform/Public/Interfaces/ITargetPlatform.h
@@ -62,6 +62,9 @@ enum class ETargetPlatformFeatures
 
 	/* The platform uses software rasterization of the scene for primitive occlusion */
 	SoftwareOcclusion,
+
+	/* The platform supports half float vertex format */
+	HalfFloatVertexFormat,
 };
 
 
diff --git a/Engine/Source/Runtime/Engine/Private/InstancedStaticMesh.cpp b/Engine/Source/Runtime/Engine/Private/InstancedStaticMesh.cpp
--- a/Engine/Source/Runtime/Engine/Private/InstancedStaticMesh.cpp
+++ b/Engine/Source/Runtime/Engine/Private/InstancedStaticMesh.cpp
@@ -323,6 +323,72 @@ void FStaticMeshInstanceBuffer::BindInstanceVertexBuffer(const class FVertexFact
 }
 
 
+void FStaticMeshInstanceData::Serialize(FArchive& Ar)
+{
+	// HTML5 doesn't support half float so we need to convert at cook time.
+	const bool bCookConvertTransformsToFullFloat = Ar.IsCooking() && bUseHalfFloat && !Ar.CookingTarget()->SupportsFeature(ETargetPlatformFeatures::HalfFloatVertexFormat);
+
+	if (Ar.IsLoading())
+	{
+		AllocateBuffers(NumInstances);
+	}
+
+	if (bCookConvertTransformsToFullFloat)
+	{
+		bool bSaveUseHalfFloat = false;
+		Ar << bSaveUseHalfFloat;
+	}
+	else
+	{
+		Ar << bUseHalfFloat;
+	}
+
+	Ar << NumInstances;
+
+	InstanceOriginData->Serialize(Ar);
+	InstanceLightmapData->Serialize(Ar);
+
+	if (bCookConvertTransformsToFullFloat)
+	{
+		TStaticMeshVertexData<FInstanceTransformMatrix<float>> FullInstanceTransformData;
+		FullInstanceTransformData.ResizeBuffer(NumInstances);
+
+		FInstanceTransformMatrix<FFloat16>* Src = (FInstanceTransformMatrix<FFloat16>*)InstanceTransformData->GetDataPointer();
+		FInstanceTransformMatrix<float>* Dest = (FInstanceTransformMatrix<float>*)FullInstanceTransformData.GetDataPointer();
+		for (int32 Idx = 0; Idx < NumInstances; Idx++)
+		{
+			Dest->InstanceTransform1[0] = Src->InstanceTransform1[0];
+			Dest->InstanceTransform1[1] = Src->InstanceTransform1[1];
+			Dest->InstanceTransform1[2] = Src->InstanceTransform1[2];
+			Dest->InstanceTransform1[3] = Src->InstanceTransform1[3];
+			Dest->InstanceTransform2[0] = Src->InstanceTransform2[0];
+			Dest->InstanceTransform2[1] = Src->InstanceTransform2[1];
+			Dest->InstanceTransform2[2] = Src->InstanceTransform2[2];
+			Dest->InstanceTransform2[3] = Src->InstanceTransform2[3];
+			Dest->InstanceTransform3[0] = Src->InstanceTransform3[0];
+			Dest->InstanceTransform3[1] = Src->InstanceTransform3[1];
+			Dest->InstanceTransform3[2] = Src->InstanceTransform3[2];
+			Dest->InstanceTransform3[3] = Src->InstanceTransform3[3];
+			Src++;
+			Dest++;
+		}
+
+		FullInstanceTransformData.Serialize(Ar);
+	}
+	else
+	{
+		InstanceTransformData->Serialize(Ar);
+	}
+
+	if (Ar.IsLoading())
+	{
+		InstanceOriginDataPtr = InstanceOriginData->GetDataPointer();
+		InstanceLightmapDataPtr = InstanceLightmapData->GetDataPointer();
+		InstanceTransformDataPtr = InstanceTransformData->GetDataPointer();
+	}
+}
+
+
 /**
  * Should we cache the material's shadertype on this platform with this vertex factory? 
  */
diff --git a/Engine/Source/Runtime/Engine/Private/InstancedStaticMesh.h b/Engine/Source/Runtime/Engine/Private/InstancedStaticMesh.h
--- a/Engine/Source/Runtime/Engine/Private/InstancedStaticMesh.h
+++ b/Engine/Source/Runtime/Engine/Private/InstancedStaticMesh.h
@@ -65,9 +65,6 @@ public:
 	ENGINE_API void InitFromPreallocatedData(FStaticMeshInstanceData& Other);
 	ENGINE_API void UpdateFromCommandBuffer_Concurrent(FInstanceUpdateCmdBuffer& CmdBuffer);
 
-	/** Serializer. */
-	friend FArchive& operator<<(FArchive& Ar, FStaticMeshInstanceBuffer& VertexBuffer);
-
 	/**
 	 * Specialized assignment operator, only used when importing LOD's. 
 	 */
diff --git a/Engine/Source/Runtime/Engine/Public/StaticMeshResources.h b/Engine/Source/Runtime/Engine/Public/StaticMeshResources.h
--- a/Engine/Source/Runtime/Engine/Public/StaticMeshResources.h
+++ b/Engine/Source/Runtime/Engine/Public/StaticMeshResources.h
@@ -802,7 +802,7 @@ public:
 	 * @param bInUseHalfFloat - true if device has support for half float in vertex arrays
 	 */
 	FStaticMeshInstanceData(bool bInUseHalfFloat)
-		: bUseHalfFloat(PLATFORM_BUILTIN_VERTEX_HALF_FLOAT || bInUseHalfFloat)
+	:	bUseHalfFloat(PLATFORM_BUILTIN_VERTEX_HALF_FLOAT || bInUseHalfFloat)
 	{
 		AllocateBuffers(0);
 	}
@@ -814,28 +814,8 @@ public:
 		delete InstanceTransformData;
 	}
 
-	void Serialize(FArchive& Ar)
-	{
-		Ar << bUseHalfFloat;
-		Ar << NumInstances;
-		
-		if (Ar.IsLoading())
-		{
-			AllocateBuffers(NumInstances);
-		}
-
-		InstanceOriginData->Serialize(Ar);
-		InstanceLightmapData->Serialize(Ar);
-		InstanceTransformData->Serialize(Ar);
-		
-		if (Ar.IsLoading())
-		{
-			InstanceOriginDataPtr = InstanceOriginData->GetDataPointer();
-			InstanceLightmapDataPtr = InstanceLightmapData->GetDataPointer();
-			InstanceTransformDataPtr = InstanceTransformData->GetDataPointer();
-		}
-	}
-
+	void Serialize(FArchive& Ar);
+	
 	void AllocateInstances(int32 InNumInstances, EResizeBufferFlags BufferFlags, bool DestroyExistingInstances)
 	{
 		NumInstances = InNumInstances;
