Fix any key not triggering actions
#rnx

diff --git a/Engine/Source/Runtime/Engine/Classes/Components/InputComponent.h b/Engine/Source/Runtime/Engine/Classes/Components/InputComponent.h
--- a/Engine/Source/Runtime/Engine/Classes/Components/InputComponent.h
+++ b/Engine/Source/Runtime/Engine/Classes/Components/InputComponent.h
@@ -608,6 +608,9 @@ struct FCachedKeyToActionInfo
 	/** Reverse lookup map to speed up evaluation of action bindings. Will be rebuilt via ConditionalBuildKeyMap when action bindings or key mappings are changed. */
 	TMap<FKey, TArray<TSharedPtr<FInputActionBinding>>> KeyToActionMap;
 
+	/** Keep the AnyKey to action map separately as we don't want to have query the map to find it every time. */
+	TArray<TSharedPtr<FInputActionBinding>> AnyKeyToActionMap;
+
 	FCachedKeyToActionInfo()
 		: PlayerInput(nullptr)
 		, KeyMapBuiltForIndex(0)
diff --git a/Engine/Source/Runtime/Engine/Private/Components/InputComponent.cpp b/Engine/Source/Runtime/Engine/Private/Components/InputComponent.cpp
--- a/Engine/Source/Runtime/Engine/Private/Components/InputComponent.cpp
+++ b/Engine/Source/Runtime/Engine/Private/Components/InputComponent.cpp
@@ -52,11 +52,12 @@ void UInputComponent::ConditionalBuildKeyMap(UPlayerInput* PlayerInput)
 		CachedInfoToPopulate->PlayerInput = PlayerInput;
 	}
 
-	// Reset the map
+	// Reset the map and AnyKey array
 	for (TPair<FKey, TArray<TSharedPtr<FInputActionBinding>>>& KeyBindingPair : CachedInfoToPopulate->KeyToActionMap)
 	{
 		KeyBindingPair.Value.Reset();
 	}
+	CachedInfoToPopulate->AnyKeyToActionMap.Reset();
 
 	for (const TSharedPtr<FInputActionBinding>& ActionBinding : ActionBindings)
 	{
@@ -64,7 +65,14 @@ void UInputComponent::ConditionalBuildKeyMap(UPlayerInput* PlayerInput)
 
 		for (const FInputActionKeyMapping& KeyMapping : KeysForAction)
 		{
-			CachedInfoToPopulate->KeyToActionMap.FindOrAdd(KeyMapping.Key).Add(ActionBinding);
+			if (KeyMapping.Key != EKeys::AnyKey)
+			{
+				CachedInfoToPopulate->KeyToActionMap.FindOrAdd(KeyMapping.Key).Add(ActionBinding);
+			}
+			else
+			{
+				CachedInfoToPopulate->AnyKeyToActionMap.Add(ActionBinding);
+			}
 		}
 	}
 
@@ -84,6 +92,10 @@ void UInputComponent::GetActionsBoundToKey(UPlayerInput* PlayerInput, const FKey
 					Actions.AddUnique(ActionForKey);
 				}
 			}
+			for (const TSharedPtr<FInputActionBinding>& ActionForKey : CachedInfo.AnyKeyToActionMap)
+			{
+				Actions.AddUnique(ActionForKey);
+			}
 			return;
 		}
 	}
