Anim graph root node is now hooked up at class Link() time

This ensures that the index/ptr is self-consistent with the compiled graph properties.
Added a new anim graph node and anim node that StateResult node suse now. No funcitonal change, but a different inheritance/data hierarchy. Added some mismatched tag handling for serialization to deal with upgrading old content.
Added a Dev-Anim object version header.

diff --git a/Engine/Source/Editor/AnimGraph/Classes/AnimGraphNode_StateResult.h b/Engine/Source/Editor/AnimGraph/Classes/AnimGraphNode_StateResult.h
--- a/Engine/Source/Editor/AnimGraph/Classes/AnimGraphNode_StateResult.h
+++ b/Engine/Source/Editor/AnimGraph/Classes/AnimGraphNode_StateResult.h
@@ -5,11 +5,32 @@
 
 #include "CoreMinimal.h"
 #include "UObject/ObjectMacros.h"
-#include "AnimGraphNode_Root.h"
+#include "AnimGraphNode_Base.h"
+#include "AnimNodes/AnimNode_StateResult.h"
 #include "AnimGraphNode_StateResult.generated.h"
 
 UCLASS(MinimalAPI)
-class UAnimGraphNode_StateResult : public UAnimGraphNode_Root
+class UAnimGraphNode_StateResult : public UAnimGraphNode_Base
 {
 	GENERATED_UCLASS_BODY()
+
+	UPROPERTY(EditAnywhere, Category=Settings)
+	FAnimNode_StateResult Node;
+
+	//~ Begin UEdGraphNode Interface.
+	virtual FLinearColor GetNodeTitleColor() const override;
+	virtual FText GetTooltipText() const override;
+	virtual FText GetNodeTitle(ENodeTitleType::Type TitleType) const override;
+	virtual bool CanUserDeleteNode() const override { return false; }
+	virtual bool CanDuplicateNode() const override { return false; }
+	virtual void GetMenuActions(FBlueprintActionDatabaseRegistrar& ActionRegistrar) const override;
+	//~ End UEdGraphNode Interface.
+
+	//~ Begin UAnimGraphNode_Base Interface
+	virtual bool IsSinkNode() const override;
+
+	// Get the link to the documentation
+	virtual FString GetDocumentationLink() const override;
+
+	//~ End UAnimGraphNode_Base Interface
 };
diff --git a/Engine/Source/Editor/AnimGraph/Private/AnimGraphNode_StateResult.cpp b/Engine/Source/Editor/AnimGraph/Private/AnimGraphNode_StateResult.cpp
--- a/Engine/Source/Editor/AnimGraph/Private/AnimGraphNode_StateResult.cpp
+++ b/Engine/Source/Editor/AnimGraph/Private/AnimGraphNode_StateResult.cpp
@@ -1,6 +1,9 @@
 // Copyright 1998-2018 Epic Games, Inc. All Rights Reserved.
 
 #include "AnimGraphNode_StateResult.h"
+#include "GraphEditorSettings.h"
+
+#define LOCTEXT_NAMESPACE "A3Nodes"
 
 /////////////////////////////////////////////////////
 // UAnimGraphNode_StateResult
@@ -9,3 +12,35 @@ UAnimGraphNode_StateResult::UAnimGraphNode_StateResult(const FObjectInitializer&
 	: Super(ObjectInitializer)
 {
 }
+
+FLinearColor UAnimGraphNode_StateResult::GetNodeTitleColor() const
+{
+	return GetDefault<UGraphEditorSettings>()->ResultNodeTitleColor;
+}
+
+FText UAnimGraphNode_StateResult::GetNodeTitle(ENodeTitleType::Type TitleType) const
+{
+	return LOCTEXT("AnimGraphNodeStateResult_Title", "Output Animation Pose");
+}
+
+FText UAnimGraphNode_StateResult::GetTooltipText() const
+{
+	return LOCTEXT("AnimGraphNodeStateResult_Tooltip", "This is the output of this animation state");
+}
+
+bool UAnimGraphNode_StateResult::IsSinkNode() const
+{
+	return true;
+}
+
+void UAnimGraphNode_StateResult::GetMenuActions(FBlueprintActionDatabaseRegistrar& ActionRegistrar) const
+{
+	// Intentionally empty. This node is auto-generated when a new graph is created.
+}
+
+FString UAnimGraphNode_StateResult::GetDocumentationLink() const
+{
+	return TEXT("Shared/GraphNodes/AnimationStateMachine");
+}
+
+#undef LOCTEXT_NAMESPACE
diff --git a/Engine/Source/Editor/KismetCompiler/Private/AnimBlueprintCompiler.cpp b/Engine/Source/Editor/KismetCompiler/Private/AnimBlueprintCompiler.cpp
--- a/Engine/Source/Editor/KismetCompiler/Private/AnimBlueprintCompiler.cpp
+++ b/Engine/Source/Editor/KismetCompiler/Private/AnimBlueprintCompiler.cpp
@@ -1140,7 +1140,6 @@ void FAnimBlueprintCompilerContext::ProcessAllAnimationNodes()
 	TArray<UAnimGraphNode_Base*> RootSet;
 
 	AllocateNodeIndexCounter = 0;
-	NewAnimBlueprintClass->RootAnimNodeIndex = 0;//INDEX_NONE;
 
 	for (auto SourceNodeIt = AnimNodeList.CreateIterator(); SourceNodeIt; ++SourceNodeIt)
 	{
@@ -1194,8 +1193,6 @@ void FAnimBlueprintCompilerContext::ProcessAllAnimationNodes()
 		{
 			AutoWireAnimGetter(Getter, nullptr);
 		}
-
-		NewAnimBlueprintClass->RootAnimNodeIndex = GetAllocationIndexOfNode(PrePhysicsRoot);
 	}
 	else
 	{
@@ -2055,8 +2052,6 @@ void FAnimBlueprintCompilerContext::CleanAndSanitizeClass(UBlueprintGeneratedCla
 	//@TODO: Move this into PurgeClass
 	NewAnimBlueprintClass->BakedStateMachines.Empty();
 	NewAnimBlueprintClass->AnimNotifies.Empty();
-
-	NewAnimBlueprintClass->RootAnimNodeIndex = INDEX_NONE;
 	NewAnimBlueprintClass->RootAnimNodeProperty = NULL;
 	NewAnimBlueprintClass->OrderedSavedPoseIndices.Empty();
 	NewAnimBlueprintClass->AnimNodeProperties.Empty();
@@ -2074,7 +2069,6 @@ void FAnimBlueprintCompilerContext::CleanAndSanitizeClass(UBlueprintGeneratedCla
 
 		NewAnimBlueprintClass->BakedStateMachines.Append(RootAnimClass->BakedStateMachines);
 		NewAnimBlueprintClass->AnimNotifies.Append(RootAnimClass->AnimNotifies);
-		NewAnimBlueprintClass->RootAnimNodeIndex = RootAnimClass->RootAnimNodeIndex;
 		NewAnimBlueprintClass->OrderedSavedPoseIndices = RootAnimClass->OrderedSavedPoseIndices;
 	}
 }
diff --git a/Engine/Source/Runtime/AnimGraphRuntime/Public/AnimNodes/AnimNode_StateResult.cpp b/Engine/Source/Runtime/AnimGraphRuntime/Public/AnimNodes/AnimNode_StateResult.cpp
new file mode 100644
--- /dev/null
+++ b/Engine/Source/Runtime/AnimGraphRuntime/Public/AnimNodes/AnimNode_StateResult.cpp
@@ -0,0 +1,17 @@
+// Copyright 1998-2018 Epic Games, Inc. All Rights Reserved.
+
+#include "AnimNode_StateResult.h"
+
+bool FAnimNode_StateResult::SerializeFromMismatchedTag(const FPropertyTag& Tag, FArchive& Ar)
+{
+	if(Tag.Type == NAME_StructProperty && Tag.StructName == FAnimNode_Root::StaticStruct()->GetFName())
+	{
+		FAnimNode_Root OldValue;
+		FAnimNode_Root::StaticStruct()->SerializeItem(Ar, &OldValue, nullptr);
+		*static_cast<FAnimNode_Root*>(this) = OldValue;
+
+		return true;
+	}
+
+	return false;
+}
\ No newline at end of file
diff --git a/Engine/Source/Runtime/AnimGraphRuntime/Public/AnimNodes/AnimNode_StateResult.h b/Engine/Source/Runtime/AnimGraphRuntime/Public/AnimNodes/AnimNode_StateResult.h
new file mode 100644
--- /dev/null
+++ b/Engine/Source/Runtime/AnimGraphRuntime/Public/AnimNodes/AnimNode_StateResult.h
@@ -0,0 +1,29 @@
+// Copyright 1998-2018 Epic Games, Inc. All Rights Reserved.
+
+#pragma once
+
+#include "CoreMinimal.h"
+#include "UObject/ObjectMacros.h"
+#include "AnimNode_Root.h"
+#include "AnimNode_StateResult.generated.h"
+
+// Root node of an state machine state (sink node).
+// We dont use AnimNode_Root to let us distinguish these nodes in the property list at link time.
+USTRUCT(BlueprintInternalUseOnly)
+struct ANIMGRAPHRUNTIME_API FAnimNode_StateResult : public FAnimNode_Root
+{
+	GENERATED_USTRUCT_BODY()
+
+	/** Used to upgrade old FAnimNode_Roots to FAnimNode_StateResult */
+	bool SerializeFromMismatchedTag(const FPropertyTag& Tag, FArchive& Ar);
+};
+
+template<>
+struct TStructOpsTypeTraits<FAnimNode_StateResult>
+	: public TStructOpsTypeTraitsBase2<FAnimNode_StateResult>
+{
+	enum
+	{
+		WithSerializeFromMismatchedTag = true,
+	};
+};
diff --git a/Engine/Source/Runtime/Core/Private/UObject/DevObjectVersion.cpp b/Engine/Source/Runtime/Core/Private/UObject/DevObjectVersion.cpp
--- a/Engine/Source/Runtime/Core/Private/UObject/DevObjectVersion.cpp
+++ b/Engine/Source/Runtime/Core/Private/UObject/DevObjectVersion.cpp
@@ -17,6 +17,7 @@
 #include "UObject/VRObjectVersion.h"
 #include "UObject/GeometryObjectVersion.h"
 #include "UObject/AnimPhysObjectVersion.h"
+#include "UObject/AnimObjectVersion.h"
 #include "UObject/FortniteMainBranchObjectVersion.h"
 #include "UObject/ReflectionCaptureObjectVersion.h"
 #include "UObject/LoadTimesObjectVersion.h"
@@ -126,6 +127,11 @@ const FGuid FAnimPhysObjectVersion::GUID(0x29E575DD, 0xE0A34627, 0x9D10D276, 0x2
 // Register AnimPhys custom version with Core
 FDevVersionRegistration GRegisterAnimPhysObjectVersion(FAnimPhysObjectVersion::GUID, FAnimPhysObjectVersion::LatestVersion, TEXT("Dev-AnimPhys"));
 
+// Unique Anim Object version id
+const FGuid FAnimObjectVersion::GUID(0xAF43A65D, 0x7FD34947, 0x98733E8E, 0xD9C1BB05);
+// Register AnimPhys custom version with Core
+FDevVersionRegistration GRegisterAnimObjectVersion(FAnimObjectVersion::GUID, FAnimObjectVersion::LatestVersion, TEXT("Dev-Anim"));
+
 // Unique ReflectionCapture Object version id
 const FGuid FReflectionCaptureObjectVersion::GUID(0x6B266CEC, 0x1EC74B8F, 0xA30BE4D9, 0x0942FC07);
 // Register Rendering custom version with Core
diff --git a/Engine/Source/Runtime/Core/Public/UObject/AnimObjectVersion.h b/Engine/Source/Runtime/Core/Public/UObject/AnimObjectVersion.h
new file mode 100644
--- /dev/null
+++ b/Engine/Source/Runtime/Core/Public/UObject/AnimObjectVersion.h
@@ -0,0 +1,28 @@
+// Copyright 1998-2018 Epic Games, Inc. All Rights Reserved.
+#pragma once
+
+#include "CoreTypes.h"
+#include "Misc/Guid.h"
+
+// Custom serialization version for changes made in Dev-Anim stream
+struct CORE_API FAnimObjectVersion
+{
+	enum Type
+	{
+		// Before any version changes were made
+		BeforeCustomVersionWasAdded,
+
+		// Reworked how anim blueprint root nodes are recovered
+		LinkTimeAnimBlueprintRootDiscovery,
+
+		// -----<new versions can be added above this line>-------------------------------------------------
+		VersionPlusOne,
+		LatestVersion = VersionPlusOne - 1
+	};
+
+	// The GUID for this custom version number
+	const static FGuid GUID;
+
+private:
+	FAnimObjectVersion() {}
+};
diff --git a/Engine/Source/Runtime/Engine/Classes/Animation/AnimBlueprintGeneratedClass.h b/Engine/Source/Runtime/Engine/Classes/Animation/AnimBlueprintGeneratedClass.h
--- a/Engine/Source/Runtime/Engine/Classes/Animation/AnimBlueprintGeneratedClass.h
+++ b/Engine/Source/Runtime/Engine/Classes/Animation/AnimBlueprintGeneratedClass.h
@@ -189,8 +189,7 @@ class ENGINE_API UAnimBlueprintGeneratedClass : public UBlueprintGeneratedClass,
 	UPROPERTY()
 	TArray<FAnimNotifyEvent> AnimNotifies;
 
-	// The index of the root node in the animation tree
-	UPROPERTY()
+	// The index of the root node in the animation tree (created during link)
 	int32 RootAnimNodeIndex;
 
 	// Indices for each of the saved pose nodes that require updating, in the order they need to get updates.
@@ -334,6 +333,10 @@ public:
 
 #endif
 
+	// UObject interface
+	virtual void Serialize(FArchive& Ar) override;
+	// End of UObject interface
+
 	// UStruct interface
 	virtual void Link(FArchive& Ar, bool bRelinkExistingProperties) override;
 	// End of UStruct interface
diff --git a/Engine/Source/Runtime/Engine/Private/Animation/AnimBlueprintGeneratedClass.cpp b/Engine/Source/Runtime/Engine/Private/Animation/AnimBlueprintGeneratedClass.cpp
--- a/Engine/Source/Runtime/Engine/Private/Animation/AnimBlueprintGeneratedClass.cpp
+++ b/Engine/Source/Runtime/Engine/Private/Animation/AnimBlueprintGeneratedClass.cpp
@@ -10,6 +10,7 @@
 #include "BonePose.h"
 #include "Animation/AnimNodeBase.h"
 #include "Animation/AnimInstance.h"
+#include "UObject/AnimObjectVersion.h"
 
 /////////////////////////////////////////////////////
 // FStateMachineDebugData
@@ -234,10 +235,20 @@ UAnimBlueprintGeneratedClass::UAnimBlueprintGeneratedClass(const FObjectInitiali
 	RootAnimNodeIndex = INDEX_NONE;
 }
 
+void UAnimBlueprintGeneratedClass::Serialize(FArchive& Ar)
+{
+	Ar.UsingCustomVersion(FAnimObjectVersion::GUID);
+
+	Super::Serialize(Ar);
+}
+
 void UAnimBlueprintGeneratedClass::Link(FArchive& Ar, bool bRelinkExistingProperties)
 {
 	Super::Link(Ar, bRelinkExistingProperties);
 
+	// We cant reference FAnimNode_Root directly as it is in AnimGraphRuntime, so we just hard-code the name instead.
+	static const FName NAME_AnimNode_Root(TEXT("AnimNode_Root"));
+
 	// @TODO: Shouldn't be necessary to clear these, but currently the class gets linked twice during compilation
 	AnimNodeProperties.Empty();
 	RootAnimNodeProperty = NULL;
@@ -249,6 +260,19 @@ void UAnimBlueprintGeneratedClass::Link(FArchive& Ar, bool bRelinkExistingProper
 		{
 			if (StructProp->Struct->IsChildOf(FAnimNode_Base::StaticStruct()))
 			{
+				if(StructProp->Struct->GetFName() == NAME_AnimNode_Root)
+				{
+					// If we are loading from a newer version, we must verify that there is only one root here.
+					// When linking an older version on load we may find multiple roots until the class is recompiled or saved
+#if DO_CHECK
+					if(Ar.IsLoading() && Ar.CustomVer(FAnimObjectVersion::GUID) >= FAnimObjectVersion::LinkTimeAnimBlueprintRootDiscovery)
+					{
+						check(RootAnimNodeProperty == nullptr);
+					}
+#endif
+					RootAnimNodeProperty = StructProp;
+					RootAnimNodeIndex = AnimNodeProperties.Num();
+				}
 				AnimNodeProperties.Add(StructProp);
 			}
 		}
@@ -273,11 +297,7 @@ void UAnimBlueprintGeneratedClass::Link(FArchive& Ar, bool bRelinkExistingProper
 	if (AnimNodeProperties.Num() > 0)
 	{
 		const bool bValidRootIndex = (RootAnimNodeIndex >= 0) && (RootAnimNodeIndex < AnimNodeProperties.Num());
-		if (bValidRootIndex)
-		{
-			RootAnimNodeProperty = AnimNodeProperties[AnimNodeProperties.Num() - 1 - RootAnimNodeIndex];
-		}
-		else
+		if (!bValidRootIndex)
 		{
 			UE_LOG(LogAnimation, Warning, TEXT("Invalid animation root node index %d on '%s' (only %d nodes)"), RootAnimNodeIndex, *GetPathName(), AnimNodeProperties.Num());
 			AnimNodeProperties.Empty();
