Adding more info for the case where getting bone transform fails due to out of bound index

diff --git a/Engine/Source/Runtime/Engine/Private/Components/SkeletalMeshComponent.cpp b/Engine/Source/Runtime/Engine/Private/Components/SkeletalMeshComponent.cpp
--- a/Engine/Source/Runtime/Engine/Private/Components/SkeletalMeshComponent.cpp
+++ b/Engine/Source/Runtime/Engine/Private/Components/SkeletalMeshComponent.cpp
@@ -3443,7 +3443,25 @@ void USkeletalMeshComponent::FinalizeBoneTransform()
 
 void USkeletalMeshComponent::GetCurrentRefToLocalMatrices(TArray<FMatrix>& OutRefToLocals, int32 InLodIdx)
 {
-	UpdateRefToLocalMatrices(OutRefToLocals, this, SkeletalMesh->GetResourceForRendering(), InLodIdx, nullptr);
+	if(SkeletalMesh)
+	{
+		FSkeletalMeshRenderData* RenderData = SkeletalMesh->GetResourceForRendering();
+		if (ensureMsgf(RenderData->LODRenderData.IsValidIndex(InLodIdx),
+			TEXT("GetCurrentRefToLocalMatrices (SkelMesh :%s) input LODIndex (%d) doesn't match with render data size (%d)."),
+			*SkeletalMesh->GetPathName(), InLodIdx, RenderData->LODRenderData.Num()))
+		{
+			UpdateRefToLocalMatrices(OutRefToLocals, this, RenderData, InLodIdx, nullptr);
+		}
+		else
+		{
+			const FReferenceSkeleton& RefSkeleton = SkeletalMesh->RefSkeleton;
+			OutRefToLocals.AddUninitialized(RefSkeleton.GetNum());
+			for (int32 Index = 0; Index < OutRefToLocals.Num(); ++Index)
+			{
+				OutRefToLocals[Index] = FMatrix::Identity;
+			}
+		}
+	}
 }
 
 bool USkeletalMeshComponent::ShouldUpdatePostProcessInstance() const
