Fix negative play rate montages

diff --git a/Engine/Source/Runtime/Engine/Private/Animation/AnimMontage.cpp b/Engine/Source/Runtime/Engine/Private/Animation/AnimMontage.cpp
--- a/Engine/Source/Runtime/Engine/Private/Animation/AnimMontage.cpp
+++ b/Engine/Source/Runtime/Engine/Private/Animation/AnimMontage.cpp
@@ -2268,21 +2268,27 @@ void FAnimMontageInstance::Advance(float DeltaTime, struct FRootMotionMovementPa
 							, *RootMotion.GetTranslation().ToCompactString()
 							, *OutRootMotionParams->GetRootMotionTransform().GetTranslation().ToCompactString()
 							, bBlendRootMotion
-							);
+						);
 					}
+				}
 
-					// Delegate has to be called last in this loop
-					// so that if this changes position, the new position will be applied in the next loop
-					// first need to have event handler to handle it
-					// Save off position before triggering events, in case they cause a jump to another position
-					const float PositionBeforeFiringEvents = Position;
+				// Delegate has to be called last in this loop
+				// so that if this changes position, the new position will be applied in the next loop
+				// first need to have event handler to handle it
+				// Save off position before triggering events, in case they cause a jump to another position
+				const float PositionBeforeFiringEvents = Position;
 
+				if(bHaveMoved)
+				{
 					// Save position before firing events.
 					if (!bInterrupted)
 					{
 						HandleEvents(PreviousSubStepPosition, Position, BranchingPointMarker);
 					}
+				}
 
+				if(MontageSubStepper.HasTimeRemaining())
+				{
 					// if we reached end of section, and we were not processing a branching point, and no events has messed with out current position..
 					// .. Move to next section.
 					// (this also handles looping, the same as jumping to a different section).
@@ -2299,6 +2305,7 @@ void FAnimMontageInstance::Advance(float DeltaTime, struct FRootMotionMovementPa
 							// Jump to next section's appropriate starting point (start or end).
 							const float EndOffset = KINDA_SMALL_NUMBER / 2.f; //KINDA_SMALL_NUMBER/2 because we use KINDA_SMALL_NUMBER to offset notifies for triggering and SMALL_NUMBER is too small
 							Position = bPlayingForward ? LatestNextSectionStartTime : (LatestNextSectionEndTime - EndOffset);
+							SubStepResult = EMontageSubStepResult::Moved;
 						}
 						else
 						{
