Fix build without PhysX.

diff --git a/Engine/Source/Runtime/Engine/Classes/PhysicsEngine/BodyInstance.h b/Engine/Source/Runtime/Engine/Classes/PhysicsEngine/BodyInstance.h
--- a/Engine/Source/Runtime/Engine/Classes/PhysicsEngine/BodyInstance.h
+++ b/Engine/Source/Runtime/Engine/Classes/PhysicsEngine/BodyInstance.h
@@ -450,9 +450,7 @@ public:
 	// Internal physics representation of our body instance
 	FPhysicsActorHandle ActorHandle;
 
-#if WITH_PHYSX
 	TSharedPtr<TArray<ANSICHAR>> CharDebugName;
-#endif	//WITH_PHYSX
 
 	/** PrimitiveComponent containing this body.   */
 	TWeakObjectPtr<class UPrimitiveComponent> OwnerComponent;
diff --git a/Engine/Source/Runtime/Engine/Classes/PhysicsEngine/BodySetup.h b/Engine/Source/Runtime/Engine/Classes/PhysicsEngine/BodySetup.h
--- a/Engine/Source/Runtime/Engine/Classes/PhysicsEngine/BodySetup.h
+++ b/Engine/Source/Runtime/Engine/Classes/PhysicsEngine/BodySetup.h
@@ -410,7 +410,6 @@ public:
 	ENGINE_API virtual void CopyBodySetupProperty(const UBodySetup* Other);
 #endif // WITH_EDITOR
 
-#if WITH_PHYSX
 	/** 
 	 *   Add the shapes defined by this body setup to the supplied PxRigidBody. 
 	 */
@@ -424,7 +423,6 @@ public:
 		const FTransform& RelativeTM = FTransform::Identity, 
 		TArray<FPhysicsShapeHandle>* NewShapes = NULL, 
 		bool bShapeSharing = false);
-#endif // WITH_PHYSX
 
 	friend struct FIterateBodySetupHelper;
 
diff --git a/Engine/Source/Runtime/Engine/Classes/PhysicsEngine/ConstraintInstance.h b/Engine/Source/Runtime/Engine/Classes/PhysicsEngine/ConstraintInstance.h
--- a/Engine/Source/Runtime/Engine/Classes/PhysicsEngine/ConstraintInstance.h
+++ b/Engine/Source/Runtime/Engine/Classes/PhysicsEngine/ConstraintInstance.h
@@ -83,8 +83,6 @@ struct ENGINE_API FConstraintProfileProperties
 
 	FConstraintProfileProperties();
 
-#if WITH_PHYSX
-
 	/** Updates physx joint properties from unreal properties (limits, drives, flags, etc...) */
 	void Update_AssumesLocked(const FPhysicsConstraintHandle& InConstraintRef, float AverageMass, float UseScale) const;
 
@@ -93,7 +91,6 @@ struct ENGINE_API FConstraintProfileProperties
 
 	/** Updates physx joint flag based on profile properties */
 	void UpdateConstraintFlags_AssumesLocked(const FPhysicsConstraintHandle& InConstraintRef) const;
-#endif
 
 #if WITH_EDITOR
 	void SyncChangedConstraintProperties(struct FPropertyChangedChainEvent& PropertyChangedEvent);
@@ -185,9 +182,7 @@ public:
 	/** Get underlying physics engine constraint */
 	const FPhysicsConstraintHandle& GetPhysicsConstraintRef() const;
 
-#if WITH_PHYSX
 	FPhysxUserData PhysxUserData;
-#endif
 
 private:
 	/** The component scale passed in during initialization*/
diff --git a/Engine/Source/Runtime/Engine/Classes/PhysicsEngine/ConstraintTypes.h b/Engine/Source/Runtime/Engine/Classes/PhysicsEngine/ConstraintTypes.h
--- a/Engine/Source/Runtime/Engine/Classes/PhysicsEngine/ConstraintTypes.h
+++ b/Engine/Source/Runtime/Engine/Classes/PhysicsEngine/ConstraintTypes.h
@@ -110,10 +110,8 @@ struct ENGINE_API FConeConstraint : public FConstraintBaseParams
 
 	FConeConstraint();
 
-#if WITH_PHYSX
-	/** Updates physx cone limit from unreal data. */
+	/** Updates cone limit from unreal data. */
 	void UpdateConeLimit_AssumesLocked(const FPhysicsConstraintHandle& InConstraintRef, float AverageMass) const;
-#endif
 };
 
 /** Angular roll constraint */
@@ -132,10 +130,8 @@ struct ENGINE_API FTwistConstraint : public FConstraintBaseParams
 
 	FTwistConstraint();
 
-#if WITH_PHYSX
-	/** Updates physx twist limit from unreal data. */
+	/** Updates twist limit from unreal data. */
 	void UpdateTwistLimit_AssumesLocked(const FPhysicsConstraintHandle& InConstraintRef, float AverageMass) const;
-#endif
 };
 
 #define		RB_MinSizeToLockDOF				(0.1)
diff --git a/Engine/Source/Runtime/Engine/Private/PhysicsEngine/BodyInstance.cpp b/Engine/Source/Runtime/Engine/Private/PhysicsEngine/BodyInstance.cpp
--- a/Engine/Source/Runtime/Engine/Private/PhysicsEngine/BodyInstance.cpp
+++ b/Engine/Source/Runtime/Engine/Private/PhysicsEngine/BodyInstance.cpp
@@ -15,6 +15,7 @@
 #include "SceneManagement.h"
 #include "Collision.h"
 #include "Physics/PhysicsInterfaceCore.h"
+#include "Physics/PhysicsFiltering.h"
 #include "PhysicsEngine/ConstraintInstance.h"
 #include "PhysicsEngine/ShapeElem.h"
 #include "PhysicsEngine/ConvexElem.h"
@@ -924,8 +925,6 @@ void GetSimulatingAndBlendWeight(const USkeletalMeshComponent* SkelMeshComp, con
 	}
 }
 
-#if WITH_PHYSX
-
 struct FPhysicsActorCreatePayload_PhysX
 {
 	bool bUseAsync;
@@ -1225,8 +1224,6 @@ struct FInitBodiesHelper
 		}
 	}
 };
-#endif // WITH_PHYSX
-
 
 FBodyInstance::FInitBodySpawnParams::FInitBodySpawnParams(const UPrimitiveComponent* PrimComp)
 {
@@ -3402,7 +3399,6 @@ bool FBodyInstance::ValidateTransform(const FTransform &Transform, const FString
 	return true;
 }
 
-#if WITH_PHYSX
 void FBodyInstance::InitDynamicProperties_AssumesLocked()
 {
 	if(!BodySetup.IsValid())
@@ -3588,7 +3584,6 @@ void FBodyInstance::BuildBodyFilterData(FBodyCollisionFilterData& OutFilterData)
 		OutFilterData.QueryComplexFilter = ComplexQueryData;
 	}
 }
-#endif // WITH_PHYSX
 
 void FBodyInstance::InitStaticBodies(const TArray<FBodyInstance*>& Bodies, const TArray<FTransform>& Transforms, UBodySetup* BodySetup, UPrimitiveComponent* PrimitiveComp, FPhysScene* InRBScene)
 {
diff --git a/Engine/Source/Runtime/Engine/Private/PhysicsEngine/BodySetup.cpp b/Engine/Source/Runtime/Engine/Private/PhysicsEngine/BodySetup.cpp
--- a/Engine/Source/Runtime/Engine/Private/PhysicsEngine/BodySetup.cpp
+++ b/Engine/Source/Runtime/Engine/Private/PhysicsEngine/BodySetup.cpp
@@ -186,7 +186,9 @@ UBodySetup::UBodySetup(const FObjectInitializer& ObjectInitializer)
 	SetFlags(RF_Transactional);
 	bSharedCookedData = false;
 	CookedFormatDataOverride = nullptr;
+#if WITH_PHYSX
 	CurrentCookHelper = nullptr;
+#endif
 }
 
 void UBodySetup::CopyBodyPropertiesFrom(const UBodySetup* FromSetup)
@@ -558,9 +560,11 @@ void UBodySetup::FinishCreatingPhysicsMeshes(const TArray<PxConvexMesh*>& Convex
 void UBodySetup::CreatePhysicsMeshesAsync(FOnAsyncPhysicsCookFinished OnAsyncPhysicsCookFinished)
 {
 	check(IsInGameThread());
-	
+
+#if WITH_PHYSX
 	// Don't start another cook cycle if one's already in progress
 	check(CurrentCookHelper == nullptr);
+#endif
 
 #if WITH_PHYSX_COOKING
 	UActorComponent* OwningComp = Cast<UActorComponent>(GetOuter());
@@ -678,8 +682,6 @@ void UBodySetup::ClearPhysicsMeshes()
 	AggGeom.FreeRenderInfo();
 }
 
-#if WITH_PHYSX
-
 void UBodySetup::AddShapesToRigidActor_AssumesLocked(
 	FBodyInstance* OwningInstance, 
 	EPhysicsSceneType SceneType, 
@@ -716,13 +718,13 @@ void UBodySetup::AddShapesToRigidActor_AssumesLocked(
 	AddParams.ComplexMaterials = TArrayView<UPhysicalMaterial*>(ComplexMaterials);
 	AddParams.LocalTransform = RelativeTM;
 	AddParams.Geometry = &AggGeom;
+#if WITH_PHYSX
 	AddParams.TriMeshes = TArrayView<PxTriangleMesh*>(TriMeshes);
+#endif
 
 	FPhysicsInterface::AddGeometry(OwningInstance->ActorHandle, AddParams, NewShapes);
 }
 
-#endif // WITH_PHYSX
-
 void UBodySetup::RemoveSimpleCollision()
 {
 	AggGeom.EmptyElements();
diff --git a/Engine/Source/Runtime/Engine/Private/PhysicsEngine/CollisionQueryFilterCallback.cpp b/Engine/Source/Runtime/Engine/Private/PhysicsEngine/CollisionQueryFilterCallback.cpp
--- a/Engine/Source/Runtime/Engine/Private/PhysicsEngine/CollisionQueryFilterCallback.cpp
+++ b/Engine/Source/Runtime/Engine/Private/PhysicsEngine/CollisionQueryFilterCallback.cpp
@@ -42,11 +42,6 @@ ECollisionQueryHitType FCollisionQueryFilterCallback::CalcQueryHitType(const FCo
 	}
 	else
 	{
-		// Then see if the channel wants to be blocked
-		// @todo delete once we fix up object/trace APIs to work separate
-		PxU32 ShapeFlags = ShapeFilter.Word3 & 0xFFFFFF;
-		bool bStaticShape = ((ShapeFlags & EPDF_StaticShape) != 0);
-
 		// if query channel is Touch All, then just return touch
 		if (QuerierChannel == ECC_OverlapAll_Deprecated)
 		{
diff --git a/Engine/Source/Runtime/Engine/Private/PhysicsEngine/ConstraintInstance.cpp b/Engine/Source/Runtime/Engine/Private/PhysicsEngine/ConstraintInstance.cpp
--- a/Engine/Source/Runtime/Engine/Private/PhysicsEngine/ConstraintInstance.cpp
+++ b/Engine/Source/Runtime/Engine/Private/PhysicsEngine/ConstraintInstance.cpp
@@ -291,8 +291,6 @@ float ComputeAverageMass_AssumesLocked(const FPhysicsActorHandle& InActor1, cons
 	return AverageMass;
 }
 
-#if WITH_PHYSX
-
 bool GetActorRefs(FBodyInstance* Body1, FBodyInstance* Body2, FPhysicsActorHandle& OutActorRef1, FPhysicsActorHandle& OutActorRef2, UObject* DebugOwner)
 {
 	FPhysicsActorHandle ActorRef1 = Body1 ? Body1->ActorHandle : FPhysicsActorHandle();
@@ -417,8 +415,6 @@ void EnsureSleepingActorsStaySleeping_AssumesLocked(const FPhysicsActorHandle& I
 	}
 }
 
-#endif
-
 /** 
  *	Create physics engine constraint.
  */
diff --git a/Engine/Source/Runtime/Engine/Private/PhysicsEngine/Experimental/PhysInterface_Apeiron.cpp b/Engine/Source/Runtime/Engine/Private/PhysicsEngine/Experimental/PhysInterface_Apeiron.cpp
--- a/Engine/Source/Runtime/Engine/Private/PhysicsEngine/Experimental/PhysInterface_Apeiron.cpp
+++ b/Engine/Source/Runtime/Engine/Private/PhysicsEngine/Experimental/PhysInterface_Apeiron.cpp
@@ -1164,6 +1164,8 @@ TArray<Apeiron::TVector<int32, 3>> GetMeshElements(const PHYSX_MESH* PhysXMesh)
 	check(false);
 }
 
+#if WITH_PHYSX
+
 template<>
 TArray<Apeiron::TVector<int32, 3>> GetMeshElements(const physx::PxConvexMesh* PhysXMesh)
 {
@@ -1249,6 +1251,8 @@ TUniquePtr<Apeiron::TImplicitObject<float, 3>> ConvertPhysXMeshToLevelset(const
 #endif
 }
 
+#endif
+
 FPhysicsShapeHandle FPhysInterface_Apeiron::CreateShape(physx::PxGeometry* InGeom, bool bSimulation, bool bQuery, UPhysicalMaterial* InSimpleMaterial, TArray<UPhysicalMaterial*>* InComplexMaterials, bool bShared)
 {
     // @todo(mlentine): Should we be doing anything with the InGeom here?
@@ -1315,18 +1319,22 @@ void FPhysInterface_Apeiron::AddGeometry(const FPhysicsActorHandle& InActor, con
 			}
 		}
 #endif
+#if WITH_PHYSX
 		for (uint32 i = 0; i < static_cast<uint32>(InParams.Geometry->ConvexElems.Num()); ++i)
 		{
 			const auto& CollisionBody = InParams.Geometry->ConvexElems[i];
 			Objects.Add(ConvertPhysXMeshToLevelset(CollisionBody.GetConvexMesh(), Scale));
 		}
+#endif
 	}
 	else
 	{
+#if WITH_PHYSX
 		for (const auto& PhysXMesh : InParams.TriMeshes)
 		{
 			Objects.Add(ConvertPhysXMeshToLevelset(PhysXMesh, Scale));
 		}
+#endif
 	}
 	if (Objects.Num() == 0) return;
     TArray<RigidBodyId> BodiesToModify = { InActor.First };
@@ -1492,6 +1500,8 @@ void FinishSceneStat(uint32 Scene)
 {
 }
 
+#if WITH_PHYSX
+
 void FPhysInterface_Apeiron::CalculateMassPropertiesFromShapeCollection(physx::PxMassProperties& OutProperties, const TArray<FPhysicsShapeHandle>& InShapes, float InDensityKGPerCM)
 {
     // What does it mean when if there is more than one collision object?
@@ -1511,6 +1521,8 @@ void FPhysInterface_Apeiron::CalculateMassPropertiesFromShapeCollection(physx::P
     }
 }
 
+#endif
+
 bool FPhysInterface_Apeiron::LineTrace_Geom(FHitResult& OutHit, const FBodyInstance* InInstance, const FVector& InStart, const FVector& InEnd, bool bTraceComplex, bool bExtractPhysMaterial)
 {
 	OutHit.TraceStart = InStart;
diff --git a/Engine/Source/Runtime/Engine/Private/PhysicsEngine/PhysicsGeometryPhysX.cpp b/Engine/Source/Runtime/Engine/Private/PhysicsEngine/PhysicsGeometryPhysX.cpp
--- a/Engine/Source/Runtime/Engine/Private/PhysicsEngine/PhysicsGeometryPhysX.cpp
+++ b/Engine/Source/Runtime/Engine/Private/PhysicsEngine/PhysicsGeometryPhysX.cpp
@@ -1,16 +1,18 @@
 // Copyright 1998-2018 Epic Games, Inc. All Rights Reserved.
 
+#if WITH_PHYSX
+
 #include "Physics/PhysicsGeometryPhysX.h"
 
 #include "PhysicsEngine/AggregateGeom.h"
 
 //#include "PhysXPublic.h"
 #include "PhysicsEngine/PhysXSupport.h"
-#include "PxSphereGeometry.h"
 #include "PhysicsEngine/PhysicsSettings.h"
 #include "Engine/Engine.h"
 #include "PhysicalMaterials/PhysicalMaterial.h"
 #include "PhysicsEngine/BodySetup.h"
+#include "PxSphereGeometry.h"
 
 using namespace physx;
 
@@ -293,3 +295,5 @@ template void FBodySetupShapeIterator::ForEachShape(const TArrayView<FKBoxElem>&
 template void FBodySetupShapeIterator::ForEachShape(const TArrayView<FKSphylElem>&, TFunctionRef<void(const FKSphylElem&, const physx::PxCapsuleGeometry&, const physx::PxTransform&, float, float)>) const;
 template void FBodySetupShapeIterator::ForEachShape(const TArrayView<FKConvexElem>&, TFunctionRef<void(const FKConvexElem&, const physx::PxConvexMeshGeometry&, const physx::PxTransform&, float, float)>) const;
 template void FBodySetupShapeIterator::ForEachShape(const TArrayView<physx::PxTriangleMesh*>&, TFunctionRef<void(physx::PxTriangleMesh* const &, const physx::PxTriangleMeshGeometry&, const physx::PxTransform&, float, float)>) const;
+
+#endif
\ No newline at end of file
diff --git a/Engine/Source/Runtime/Engine/Private/PhysicsEngine/PhysicsInterfaceUtils.cpp b/Engine/Source/Runtime/Engine/Private/PhysicsEngine/PhysicsInterfaceUtils.cpp
--- a/Engine/Source/Runtime/Engine/Private/PhysicsEngine/PhysicsInterfaceUtils.cpp
+++ b/Engine/Source/Runtime/Engine/Private/PhysicsEngine/PhysicsInterfaceUtils.cpp
@@ -6,6 +6,8 @@
 #include "Physics/PhysicsFiltering.h"
 #include "Physics/PhysicsInterfaceTypes.h"
 
+#if WITH_PHYSX
+
 PxShapeFlags BuildPhysXShapeFlags(FBodyCollisionFlags BodyCollisionFlags, bool bPhysicsStatic, bool bIsSync, bool bIsTriangleMesh)
 {
 	PxShapeFlags ShapeFlags;
@@ -164,9 +166,6 @@ uint32 FindFaceIndex(const PxSweepHit& PHit, const PxVec3& unitDir)
 	return PHit.faceIndex;	//If no custom logic just return whatever face index they initially had
 }
 
-
-
-
 FPhysXShapeAdaptor::FPhysXShapeAdaptor(const FQuat& Rot, const FCollisionShape& CollisionShape)
 	: Rotation(physx::PxIdentity)
 {
@@ -341,4 +340,6 @@ PxFilterData CreateQueryFilterData(const uint8 MyChannel, const bool bTraceCompl
 	{
 		return CreateTraceQueryFilterData(MyChannel, bTraceComplex, InCollisionResponseContainer, QueryParam);
 	}
-}
\ No newline at end of file
+}
+
+#endif
\ No newline at end of file
diff --git a/Engine/Source/Runtime/Engine/Private/PhysicsEngine/PxQueryFilterCallback.cpp b/Engine/Source/Runtime/Engine/Private/PhysicsEngine/PxQueryFilterCallback.cpp
--- a/Engine/Source/Runtime/Engine/Private/PhysicsEngine/PxQueryFilterCallback.cpp
+++ b/Engine/Source/Runtime/Engine/Private/PhysicsEngine/PxQueryFilterCallback.cpp
@@ -6,6 +6,8 @@
 #include "Physics/PhysicsInterfaceUtils.h"
 #include "Collision.h"
 
+#if WITH_PHYSX
+
 PxQueryHitType::Enum FPxQueryFilterCallback::preFilter(const PxFilterData& filterData, const PxShape* shape, const PxRigidActor* actor, PxHitFlags& queryFlags)
 {
 	SCOPE_CYCLE_COUNTER(STAT_Collision_PreFilter);
@@ -51,4 +53,6 @@ PxQueryHitType::Enum FPxQueryFilterCallback::postFilter(const PxFilterData& filt
 
 	const ECollisionQueryHitType HitType = PostFilter(FilterData, bIsOverlap);
 	return HitTypeToPxQueryHitType(HitType);
-}
\ No newline at end of file
+}
+
+#endif
\ No newline at end of file
diff --git a/Engine/Source/Runtime/Engine/Public/Physics/Experimental/PhysInterface_Apeiron.h b/Engine/Source/Runtime/Engine/Public/Physics/Experimental/PhysInterface_Apeiron.h
--- a/Engine/Source/Runtime/Engine/Public/Physics/Experimental/PhysInterface_Apeiron.h
+++ b/Engine/Source/Runtime/Engine/Public/Physics/Experimental/PhysInterface_Apeiron.h
@@ -239,9 +239,6 @@ public:
     ENGINE_API void SyncBodies();
 
     // Interface needed for interface
-    static ENGINE_API physx::PxRigidActor* GetPxRigidActorFromScene_AssumesLocked(const FPhysicsActorReference_Apeiron& InActorRef, int32 SceneType = -1) { return nullptr; }
-    static ENGINE_API physx::PxRigidActor* GetPxRigidActor_AssumesLocked(const FPhysicsActorReference_Apeiron& InRef) { return nullptr; }
-
 	static ENGINE_API FPhysicsActorHandle CreateActor(const FActorCreationParams& Params);
 	static ENGINE_API void ReleaseActor(FPhysicsActorReference_Apeiron& InActorReference, FPhysScene* InScene = nullptr, bool bNeverDeferRelease=false);
 
@@ -526,9 +523,11 @@ public:
 	{
 		return InActorReference.Second;
 	}
-	
+
+#if WITH_PHYSX
     static void CalculateMassPropertiesFromShapeCollection(physx::PxMassProperties& OutProperties, const TArray<FPhysicsShapeHandle>& InShapes, float InDensityKGPerCM);
-    
+#endif
+
     static const Apeiron::TPBDRigidParticles<float, 3>& GetParticlesAndIndex(const FPhysicsActorReference_Apeiron& InActorReference, uint32& Index);
     static const TArray<Apeiron::TVector<int32, 2>>& GetConstraintArrayAndIndex(const FPhysicsConstraintReference_Apeiron& InActorReference, uint32& Index);
 
diff --git a/Engine/Source/Runtime/Engine/Public/Physics/ImmediatePhysics/ImmediatePhysicsActorHandle.h b/Engine/Source/Runtime/Engine/Public/Physics/ImmediatePhysics/ImmediatePhysicsActorHandle.h
--- a/Engine/Source/Runtime/Engine/Public/Physics/ImmediatePhysics/ImmediatePhysicsActorHandle.h
+++ b/Engine/Source/Runtime/Engine/Public/Physics/ImmediatePhysics/ImmediatePhysicsActorHandle.h
@@ -268,11 +268,13 @@ public:
 		return &OwningSimulation.Actors[ActorDataIndex];
 	}
 
+#if WITH_PHYSX
 	immediate::PxRigidBodyData* GetSimulationRigidBodyData() const
 	{
 		immediate::PxRigidBodyData& Data = OwningSimulation.GetLowLevelBody(ActorDataIndex);
 		return &Data;
 	}
+#endif
 
 private:
 	/** Converts from actor space (i.e. the transform in world space as the client gives us) to body space (body with its origin at the COM and oriented to inertia tensor) */
diff --git a/Engine/Source/Runtime/Engine/Public/Physics/ImmediatePhysics/ImmediatePhysicsSimulation.h b/Engine/Source/Runtime/Engine/Public/Physics/ImmediatePhysics/ImmediatePhysicsSimulation.h
--- a/Engine/Source/Runtime/Engine/Public/Physics/ImmediatePhysics/ImmediatePhysicsSimulation.h
+++ b/Engine/Source/Runtime/Engine/Public/Physics/ImmediatePhysics/ImmediatePhysicsSimulation.h
@@ -35,7 +35,9 @@ namespace ImmediatePhysics
 {
 	struct FActorData
 	{
+#if WITH_PHYSX
 		immediate::PxRigidBodyData RigidBodyData;
+#endif
 		FTransform InitialTransform;
 
 		bool bStatic;
@@ -122,8 +124,10 @@ public:
 
 	TArray<FActorHandle*>& GetActorHandles() { return ActorHandles; }
 	const TArray<FActorHandle*>& GetActorHandles() const { return ActorHandles; }
+#if WITH_PHYSX
 	TArray<immediate::PxRigidBodyData>& GetRigidBodyData() { return RigidBodiesData; }
 	const TArray<immediate::PxRigidBodyData>& GetRigidBodyData() const { return RigidBodiesData; }
+#endif
 
 	int32 NumActors() const { return ActorHandles.Num(); }
 
diff --git a/Engine/Source/Runtime/Engine/Public/PhysxUserData.h b/Engine/Source/Runtime/Engine/Public/PhysxUserData.h
--- a/Engine/Source/Runtime/Engine/Public/PhysxUserData.h
+++ b/Engine/Source/Runtime/Engine/Public/PhysxUserData.h
@@ -12,8 +12,6 @@ struct FConstraintInstance;
 struct FKShapeElem;
 class FPhysScene_PhysX;
 
-#if WITH_PHYSX
-
 /** Forward declarations */
 struct FKShapeElem;
 struct FCustomPhysXPayload;
@@ -72,7 +70,4 @@ template <> FORCEINLINE void FPhysxUserData::Set(void* UserData, FPhysScene* Pay
 template <> FORCEINLINE void FPhysxUserData::Set(void* UserData, FConstraintInstance* Payload)		{ check(UserData); ((FPhysxUserData*)UserData)->Type = EPhysxUserDataType::ConstraintInstance; ((FPhysxUserData*)UserData)->Payload = Payload; }
 template <> FORCEINLINE void FPhysxUserData::Set(void* UserData, UPrimitiveComponent* Payload)		{ check(UserData); ((FPhysxUserData*)UserData)->Type = EPhysxUserDataType::PrimitiveComponent; ((FPhysxUserData*)UserData)->Payload = Payload; }
 template <> FORCEINLINE void FPhysxUserData::Set(void* UserData, FKShapeElem* Payload)	{ check(UserData); ((FPhysxUserData*)UserData)->Type = EPhysxUserDataType::AggShape; ((FPhysxUserData*)UserData)->Payload = Payload; }
-template <> FORCEINLINE void FPhysxUserData::Set(void* UserData, FCustomPhysXPayload* Payload) { check(UserData); ((FPhysxUserData*)UserData)->Type = EPhysxUserDataType::CustomPayload; ((FPhysxUserData*)UserData)->Payload = Payload; }
-
-
-#endif // WITH_PHYSICS
+template <> FORCEINLINE void FPhysxUserData::Set(void* UserData, FCustomPhysXPayload* Payload) { check(UserData); ((FPhysxUserData*)UserData)->Type = EPhysxUserDataType::CustomPayload; ((FPhysxUserData*)UserData)->Payload = Payload; }
\ No newline at end of file
