Fix linking error for NoPhysX

diff --git a/Engine/Source/Runtime/Engine/Private/PhysicsEngine/BodyInstance.cpp b/Engine/Source/Runtime/Engine/Private/PhysicsEngine/BodyInstance.cpp
--- a/Engine/Source/Runtime/Engine/Private/PhysicsEngine/BodyInstance.cpp
+++ b/Engine/Source/Runtime/Engine/Private/PhysicsEngine/BodyInstance.cpp
@@ -394,7 +394,6 @@ FBodyInstance* FindParentBodyInstance(FName BodyName, USkeletalMeshComponent* Sk
 	return NULL;
 }
 
-#if WITH_PHYSX
 //Determine that the shape is associated with this subbody (or root body)
 bool FBodyInstance::IsShapeBoundToBody(const FPhysicsShapeHandle& Shape) const
 {
@@ -414,8 +413,6 @@ int32 FBodyInstance::GetAllShapes_AssumesLocked(TArray<FPhysicsShapeHandle>& Out
 	return FPhysicsInterface::GetAllShapes_AssumedLocked(ActorHandle, OutShapes, SceneType);
 }
 
-#endif
-
 void FBodyInstance::UpdateTriMeshVertices(const TArray<FVector> & NewPositions)
 {
 #if WITH_APEIRON || WITH_IMMEDIATE_PHYSX
