Missed this file for fixing PhysX build.

diff --git a/Engine/Source/Runtime/Engine/Private/PhysicsEngine/ImmediatePhysics/ImmediatePhysicsSimulation.cpp b/Engine/Source/Runtime/Engine/Private/PhysicsEngine/ImmediatePhysics/ImmediatePhysicsSimulation.cpp
--- a/Engine/Source/Runtime/Engine/Private/PhysicsEngine/ImmediatePhysics/ImmediatePhysicsSimulation.cpp
+++ b/Engine/Source/Runtime/Engine/Private/PhysicsEngine/ImmediatePhysics/ImmediatePhysicsSimulation.cpp
@@ -1063,6 +1063,7 @@ ImmediatePhysics::FActorData CreateActorData(const FActorCreationParams& InParam
 {
 	FActorData OutData;
 
+#if WITH_PHYSX
 	immediate::PxRigidBodyData& RigidData = OutData.RigidBodyData;
 
 	RigidData.linearVelocity = PxVec3(PxZero);
@@ -1076,6 +1077,7 @@ ImmediatePhysics::FActorData CreateActorData(const FActorCreationParams& InParam
 	RigidData.maxLinearVelocitySq = PX_MAX_F32;
 	RigidData.maxAngularVelocitySq = PX_MAX_F32;
 	RigidData.body2World = U2PTransform(InParams.InitialTM);
+#endif
 
 	OutData.InitialTransform = InParams.InitialTM;
 	OutData.bStatic = InParams.bStatic;
