Add option to maintain original scale of root motion

PR #4940

diff --git a/Engine/Source/Runtime/Engine/Classes/Animation/AnimSequence.h b/Engine/Source/Runtime/Engine/Classes/Animation/AnimSequence.h
--- a/Engine/Source/Runtime/Engine/Classes/Animation/AnimSequence.h
+++ b/Engine/Source/Runtime/Engine/Classes/Animation/AnimSequence.h
@@ -514,6 +514,10 @@ public:
 	UPROPERTY(EditAnywhere, Category = RootMotion)
 	bool bForceRootLock;
 
+	/** If this is on, it will use a normalized scale value for the root motion extracted: FVector(1.0, 1.0, 1.0) **/
+	UPROPERTY(EditAnywhere, AssetRegistrySearchable, Category = RootMotion, meta = (DisplayName = "Use Normalized Root Motion Scale"))
+	bool bUseNormalizedRootMotionScale;
+
 	/** Have we copied root motion settings from an owning montage */
 	UPROPERTY()
 	bool bRootMotionSettingsCopiedFromMontage;
diff --git a/Engine/Source/Runtime/Engine/Private/Animation/AnimSequence.cpp b/Engine/Source/Runtime/Engine/Private/Animation/AnimSequence.cpp
--- a/Engine/Source/Runtime/Engine/Private/Animation/AnimSequence.cpp
+++ b/Engine/Source/Runtime/Engine/Private/Animation/AnimSequence.cpp
@@ -129,6 +129,7 @@ UAnimSequence::UAnimSequence(const FObjectInitializer& ObjectInitializer)
 	: Super(ObjectInitializer)
 	, Interpolation(EAnimInterpolationType::Linear)
 	, bEnableRootMotion(false)
+	, bUseNormalizedRootMotionScale(true)
 	, RootMotionRootLock(ERootMotionRootLock::RefPose)
 	, bRootMotionSettingsCopiedFromMontage(false)
 	, bUseRawDataOnly(!FPlatformProperties::RequiresCookedData())
@@ -1008,9 +1009,21 @@ FTransform UAnimSequence::ExtractRootMotionFromRange(float StartTrackPosition, f
 	FTransform StartTransform = ExtractRootTrackTransform(StartTrackPosition, NULL);
 	FTransform EndTransform = ExtractRootTrackTransform(EndTrackPosition, NULL);
 
-	//Clear scale as it will muck up GetRelativeTransform
-	StartTransform.SetScale3D(FVector(1.f));
-	EndTransform.SetScale3D(FVector(1.f));
+	// Use old calculation if needed.
+	if (bUseNormalizedRootMotionScale)
+	{
+		//Clear scale as it will muck up GetRelativeTransform
+		StartTransform.SetScale3D(FVector(1.f));
+		EndTransform.SetScale3D(FVector(1.f));
+	}
+	else
+	{
+		if (IsValidAdditive())
+		{
+			StartTransform.SetScale3D(StartTransform.GetScale3D() + DefaultScale);
+			EndTransform.SetScale3D(EndTransform.GetScale3D() + DefaultScale);
+		}
+	}
 
 	// Transform to Component Space Rotation (inverse root transform from first frame)
 	const FTransform RootToComponentRot = FTransform(InitialTransform.GetRotation().Inverse());
