Native notifies for dynamics and cloth management

diff --git a/Engine/Source/Runtime/Engine/Classes/Animation/AnimNotifies/AnimNotify_PauseClothingSimulation.h b/Engine/Source/Runtime/Engine/Classes/Animation/AnimNotifies/AnimNotify_PauseClothingSimulation.h
new file mode 100644
--- /dev/null
+++ b/Engine/Source/Runtime/Engine/Classes/Animation/AnimNotifies/AnimNotify_PauseClothingSimulation.h
@@ -0,0 +1,26 @@
+// Copyright 1998-2018 Epic Games, Inc. All Rights Reserved.
+
+#pragma once
+
+#include "AnimNotify.h"
+
+#include "AnimNotify_PauseClothingSimulation.generated.h"
+
+UCLASS(const, hidecategories=Object, collapsecategories, meta=(DisplayName="Pause Clothing Simulation"))
+class ENGINE_API UAnimNotify_PauseClothingSimulation : public UAnimNotify
+{
+	GENERATED_BODY()
+
+public:
+
+	UAnimNotify_PauseClothingSimulation();
+
+	// Begin UAnimNotify interface
+	virtual FString GetNotifyName_Implementation() const override;
+	virtual void Notify(USkeletalMeshComponent* MeshComp, UAnimSequenceBase* Animation) override;
+	// End UAnimNotify interface
+
+};
+
+
+
diff --git a/Engine/Source/Runtime/Engine/Classes/Animation/AnimNotifies/AnimNotify_ResetDynamics.h b/Engine/Source/Runtime/Engine/Classes/Animation/AnimNotifies/AnimNotify_ResetDynamics.h
new file mode 100644
--- /dev/null
+++ b/Engine/Source/Runtime/Engine/Classes/Animation/AnimNotifies/AnimNotify_ResetDynamics.h
@@ -0,0 +1,26 @@
+// Copyright 1998-2018 Epic Games, Inc. All Rights Reserved.
+
+#pragma once
+
+#include "AnimNotify.h"
+
+#include "AnimNotify_ResetDynamics.generated.h"
+
+UCLASS(const, hidecategories=Object, collapsecategories, meta=(DisplayName="Reset Dynamics"))
+class ENGINE_API UAnimNotify_ResetDynamics : public UAnimNotify
+{
+	GENERATED_BODY()
+
+public:
+
+	UAnimNotify_ResetDynamics();
+
+	// Begin UAnimNotify interface
+	virtual FString GetNotifyName_Implementation() const override;
+	virtual void Notify(USkeletalMeshComponent* MeshComp, UAnimSequenceBase* Animation) override;
+	// End UAnimNotify interface
+
+};
+
+
+
diff --git a/Engine/Source/Runtime/Engine/Classes/Animation/AnimNotifies/AnimNotify_ResumeClothingSimulation.h b/Engine/Source/Runtime/Engine/Classes/Animation/AnimNotifies/AnimNotify_ResumeClothingSimulation.h
new file mode 100644
--- /dev/null
+++ b/Engine/Source/Runtime/Engine/Classes/Animation/AnimNotifies/AnimNotify_ResumeClothingSimulation.h
@@ -0,0 +1,26 @@
+// Copyright 1998-2018 Epic Games, Inc. All Rights Reserved.
+
+#pragma once
+
+#include "AnimNotify.h"
+
+#include "AnimNotify_ResumeClothingSimulation.generated.h"
+
+UCLASS(const, hidecategories=Object, collapsecategories, meta=(DisplayName="Resume Clothing Simulation"))
+class ENGINE_API UAnimNotify_ResumeClothingSimulation : public UAnimNotify
+{
+	GENERATED_BODY()
+
+public:
+
+	UAnimNotify_ResumeClothingSimulation();
+
+	// Begin UAnimNotify interface
+	virtual FString GetNotifyName_Implementation() const override;
+	virtual void Notify(USkeletalMeshComponent* MeshComp, UAnimSequenceBase* Animation) override;
+	// End UAnimNotify interface
+
+};
+
+
+
diff --git a/Engine/Source/Runtime/Engine/Private/Animation/AnimNotify_PauseClothingSimulation.cpp b/Engine/Source/Runtime/Engine/Private/Animation/AnimNotify_PauseClothingSimulation.cpp
new file mode 100644
--- /dev/null
+++ b/Engine/Source/Runtime/Engine/Private/Animation/AnimNotify_PauseClothingSimulation.cpp
@@ -0,0 +1,22 @@
+// Copyright 1998-2018 Epic Games, Inc. All Rights Reserved.
+
+#include "Animation/AnimNotifies/AnimNotify_PauseClothingSimulation.h"
+#include "Components/SkeletalMeshComponent.h"
+
+UAnimNotify_PauseClothingSimulation::UAnimNotify_PauseClothingSimulation()
+	: Super()
+{
+#if WITH_EDITORONLY_DATA
+	NotifyColor = FColor(90, 220, 255, 255);
+#endif // WITH_EDITORONLY_DATA
+}
+
+void UAnimNotify_PauseClothingSimulation::Notify(class USkeletalMeshComponent* MeshComp, class UAnimSequenceBase* Animation)
+{
+	MeshComp->SuspendClothingSimulation();
+}
+
+FString UAnimNotify_PauseClothingSimulation::GetNotifyName_Implementation() const
+{
+	return TEXT("Pause Clothing Sim");
+}
\ No newline at end of file
diff --git a/Engine/Source/Runtime/Engine/Private/Animation/AnimNotify_ResetDynamics.cpp b/Engine/Source/Runtime/Engine/Private/Animation/AnimNotify_ResetDynamics.cpp
new file mode 100644
--- /dev/null
+++ b/Engine/Source/Runtime/Engine/Private/Animation/AnimNotify_ResetDynamics.cpp
@@ -0,0 +1,22 @@
+// Copyright 1998-2018 Epic Games, Inc. All Rights Reserved.
+
+#include "Animation/AnimNotifies/AnimNotify_ResetDynamics.h"
+#include "Components/SkeletalMeshComponent.h"
+
+UAnimNotify_ResetDynamics::UAnimNotify_ResetDynamics()
+	: Super()
+{
+#if WITH_EDITORONLY_DATA
+	NotifyColor = FColor(90, 220, 255, 255);
+#endif // WITH_EDITORONLY_DATA
+}
+
+void UAnimNotify_ResetDynamics::Notify(class USkeletalMeshComponent* MeshComp, class UAnimSequenceBase* Animation)
+{
+	MeshComp->ResetAnimInstanceDynamics();
+}
+
+FString UAnimNotify_ResetDynamics::GetNotifyName_Implementation() const
+{
+	return TEXT("Reset Dynamics");
+}
\ No newline at end of file
diff --git a/Engine/Source/Runtime/Engine/Private/Animation/AnimNotify_ResumeClothingSimulation.cpp b/Engine/Source/Runtime/Engine/Private/Animation/AnimNotify_ResumeClothingSimulation.cpp
new file mode 100644
--- /dev/null
+++ b/Engine/Source/Runtime/Engine/Private/Animation/AnimNotify_ResumeClothingSimulation.cpp
@@ -0,0 +1,22 @@
+// Copyright 1998-2018 Epic Games, Inc. All Rights Reserved.
+
+#include "Animation/AnimNotifies/AnimNotify_ResumeClothingSimulation.h"
+#include "Components/SkeletalMeshComponent.h"
+
+UAnimNotify_ResumeClothingSimulation::UAnimNotify_ResumeClothingSimulation()
+	: Super()
+{
+#if WITH_EDITORONLY_DATA
+	NotifyColor = FColor(90, 220, 255, 255);
+#endif // WITH_EDITORONLY_DATA
+}
+
+void UAnimNotify_ResumeClothingSimulation::Notify(class USkeletalMeshComponent* MeshComp, class UAnimSequenceBase* Animation)
+{
+	MeshComp->ResumeClothingSimulation();
+}
+
+FString UAnimNotify_ResumeClothingSimulation::GetNotifyName_Implementation() const
+{
+	return TEXT("Resume Clothing Sim");
+}
\ No newline at end of file
