Add animation sequence framerate to editor viewport and content browser tooltip

diff --git a/Engine/Source/Editor/Persona/Private/AnimationEditorViewportClient.cpp b/Engine/Source/Editor/Persona/Private/AnimationEditorViewportClient.cpp
--- a/Engine/Source/Editor/Persona/Private/AnimationEditorViewportClient.cpp
+++ b/Engine/Source/Editor/Persona/Private/AnimationEditorViewportClient.cpp
@@ -946,6 +946,11 @@ FText FAnimationViewportClient::GetDisplayInfo(bool bDisplayAllInfo) const
 		TextValue = ConcatenateLine(TextValue, LOCTEXT("MeshMaterialHiddenWarning", "Mesh Materials Hidden"));
 	}
 
+	if (const UAnimSequence* AnimSequence = Cast<UAnimSequence>(GetAnimPreviewScene()->GetPreviewAnimationAsset()))
+	{
+		TextValue = ConcatenateLine(TextValue, FText::Format(LOCTEXT("FramerateFormat", "Framerate: {0}"), FText::AsNumber(AnimSequence->GetFrameRate())));
+	}
+
 	return TextValue;
 }
 void FAnimationViewportClient::DrawNodeDebugLines(TArray<FText>& Lines, FCanvas* Canvas, FSceneView* View)
diff --git a/Engine/Source/Runtime/Engine/Classes/Animation/AnimSequence.h b/Engine/Source/Runtime/Engine/Classes/Animation/AnimSequence.h
--- a/Engine/Source/Runtime/Engine/Classes/Animation/AnimSequence.h
+++ b/Engine/Source/Runtime/Engine/Classes/Animation/AnimSequence.h
@@ -610,6 +610,9 @@ public:
 #endif
 	//~ End UAnimSequenceBase Interface
 
+	// Returns the framerate of the animation
+	float GetFrameRate() const { return (float)(FMath::Max(NumFrames - 1, 1)) / (SequenceLength > 0.f ? SequenceLength : 1.f); }
+
 	// Extract Root Motion transform from the animation
 	FTransform ExtractRootMotion(float StartTime, float DeltaTime, bool bAllowLooping) const;
 
diff --git a/Engine/Source/Runtime/Engine/Private/Animation/AnimSequence.cpp b/Engine/Source/Runtime/Engine/Private/Animation/AnimSequence.cpp
--- a/Engine/Source/Runtime/Engine/Private/Animation/AnimSequence.cpp
+++ b/Engine/Source/Runtime/Engine/Private/Animation/AnimSequence.cpp
@@ -162,7 +162,7 @@ void UAnimSequence::GetAssetRegistryTags(TArray<FAssetRegistryTag>& OutTags) con
 
 	OutTags.Add(FAssetRegistryTag(TEXT("Compression Ratio"), FString::Printf(TEXT("%.03f"), (float)GetApproxCompressedSize() / (float)GetUncompressedRawSize()), FAssetRegistryTag::TT_Numerical));
 	OutTags.Add(FAssetRegistryTag(TEXT("Compressed Size (KB)"), FString::Printf(TEXT("%.02f"), (float)GetApproxCompressedSize() / 1024.0f), FAssetRegistryTag::TT_Numerical));
-
+	OutTags.Add(FAssetRegistryTag(TEXT("FrameRate"), FString::Printf(TEXT("%.2f"), GetFrameRate()), FAssetRegistryTag::TT_Numerical));
 	Super::GetAssetRegistryTags(OutTags);
 }
 
