TBA: Convert a few uobject serializers to structured archive format for example purposes

diff --git a/Engine/Source/Runtime/Core/Public/Math/Vector2D.h b/Engine/Source/Runtime/Core/Public/Math/Vector2D.h
--- a/Engine/Source/Runtime/Core/Public/Math/Vector2D.h
+++ b/Engine/Source/Runtime/Core/Public/Math/Vector2D.h
@@ -487,12 +487,27 @@ public:
 		return Ar << V.X << V.Y;
 	}
 
+	friend void operator<<(FStructuredArchive::FSlot Slot, FVector2D& V)
+	{
+		// @warning BulkSerialize: FVector2D is serialized as memory dump
+		// See TArray::BulkSerialize for detailed description of implied limitations.
+		FStructuredArchive::FStream Stream = Slot.EnterStream();
+		Stream.EnterElement() << V.X;
+		Stream.EnterElement() << V.Y;
+	}
+
 	bool Serialize(FArchive& Ar)
 	{
 		Ar << *this;
 		return true;
 	}
 
+	bool Serialize(FStructuredArchive::FSlot Slot)
+	{
+		Slot << *this;
+		return true;
+	}
+
 #if ENABLE_NAN_DIAGNOSTIC
 	FORCEINLINE void DiagnosticCheckNaN()
 	{
diff --git a/Engine/Source/Runtime/Engine/Classes/Distributions/DistributionFloat.h b/Engine/Source/Runtime/Engine/Classes/Distributions/DistributionFloat.h
--- a/Engine/Source/Runtime/Engine/Classes/Distributions/DistributionFloat.h
+++ b/Engine/Source/Runtime/Engine/Classes/Distributions/DistributionFloat.h
@@ -180,7 +180,7 @@ public:
 	virtual bool NeedsLoadForClient() const override;
 	virtual bool NeedsLoadForServer() const override;
 	virtual bool NeedsLoadForEditorGame() const override;
-	virtual void Serialize(FArchive& Ar) override;
+	virtual void Serialize(FStructuredArchive::FRecord Record) override;
 	/** End UObject interface */
 };
 
diff --git a/Engine/Source/Runtime/Engine/Classes/Distributions/DistributionVector.h b/Engine/Source/Runtime/Engine/Classes/Distributions/DistributionVector.h
--- a/Engine/Source/Runtime/Engine/Classes/Distributions/DistributionVector.h
+++ b/Engine/Source/Runtime/Engine/Classes/Distributions/DistributionVector.h
@@ -225,7 +225,7 @@ public:
 	virtual bool NeedsLoadForClient() const override;
 	virtual bool NeedsLoadForServer() const override;
 	virtual bool NeedsLoadForEditorGame() const override;
-	virtual void Serialize(FArchive& Ar) override;
+	virtual void Serialize(FStructuredArchive::FRecord Record) override;
 	/** End UObject interface */
 
 };
diff --git a/Engine/Source/Runtime/Engine/Classes/Engine/DataAsset.h b/Engine/Source/Runtime/Engine/Classes/Engine/DataAsset.h
--- a/Engine/Source/Runtime/Engine/Classes/Engine/DataAsset.h
+++ b/Engine/Source/Runtime/Engine/Classes/Engine/DataAsset.h
@@ -19,7 +19,7 @@ class UDataAsset : public UObject
 public:
 	// UObject interface
 #if WITH_EDITORONLY_DATA
-	ENGINE_API virtual void Serialize(FArchive& Ar) override;
+	ENGINE_API virtual void Serialize(FStructuredArchive::FRecord Record) override;
 #endif
 
 private:
diff --git a/Engine/Source/Runtime/Engine/Classes/Particles/ParticleModuleRequired.h b/Engine/Source/Runtime/Engine/Classes/Particles/ParticleModuleRequired.h
--- a/Engine/Source/Runtime/Engine/Classes/Particles/ParticleModuleRequired.h
+++ b/Engine/Source/Runtime/Engine/Classes/Particles/ParticleModuleRequired.h
@@ -392,7 +392,7 @@ class UParticleModuleRequired : public UParticleModule
 	//~ Begin UObject Interface
 	virtual void PostLoad() override;
 	virtual void PostInitProperties() override;
-	virtual void Serialize(FArchive& Ar) override;
+	virtual void Serialize(FStructuredArchive::FRecord Record) override;
 #if WITH_EDITOR
 	virtual void PreEditChange(UProperty* PropertyAboutToChange) override;
 	virtual void PostEditChangeProperty(FPropertyChangedEvent& PropertyChangedEvent) override;
diff --git a/Engine/Source/Runtime/Engine/Classes/Particles/SubUVAnimation.h b/Engine/Source/Runtime/Engine/Classes/Particles/SubUVAnimation.h
--- a/Engine/Source/Runtime/Engine/Classes/Particles/SubUVAnimation.h
+++ b/Engine/Source/Runtime/Engine/Classes/Particles/SubUVAnimation.h
@@ -58,7 +58,7 @@ public:
 	TArray<FVector2D> BoundingGeometry;
 
 	static FString GetDDCKeyString(const FGuid& StateId, int32 SizeX, int32 SizeY, int32 Mode, float AlphaThreshold, int32 OpacitySourceMode);
-	void Serialize(FArchive& Ar);
+	void Serialize(FStructuredArchive::FSlot Slot);
 	void Build(UTexture2D* SubUVTexture, int32 SubImages_Horizontal, int32 SubImages_Vertical, ESubUVBoundingVertexCount BoundingMode, float AlphaThreshold, EOpacitySourceMode OpacitySourceMode);
 };
 
@@ -200,7 +200,7 @@ public:
 
 	//~ Begin UObject Interface.
     virtual void PostInitProperties() override;
-	virtual void Serialize(FArchive& Ar) override;
+	virtual void Serialize(FStructuredArchive::FRecord Record) override;
 	virtual void PostLoad() override;
 #if WITH_EDITOR
 	virtual void PreEditChange(UProperty* PropertyAboutToChange) override;
diff --git a/Engine/Source/Runtime/Engine/Private/DataAsset.cpp b/Engine/Source/Runtime/Engine/Private/DataAsset.cpp
--- a/Engine/Source/Runtime/Engine/Private/DataAsset.cpp
+++ b/Engine/Source/Runtime/Engine/Private/DataAsset.cpp
@@ -11,11 +11,12 @@ UDataAsset::UDataAsset(const FObjectInitializer& ObjectInitializer)
 }
 
 #if WITH_EDITORONLY_DATA
-void UDataAsset::Serialize(FArchive& Ar)
+void UDataAsset::Serialize(FStructuredArchive::FRecord Record)
 {
-	Super::Serialize(Ar);
+	FArchive& UnderlyingArchive = Record.GetUnderlyingArchive();
+	Super::Serialize(Record);
 
-	if (Ar.IsLoading() && (Ar.UE4Ver() < VER_UE4_ADD_TRANSACTIONAL_TO_DATA_ASSETS))
+	if (UnderlyingArchive.IsLoading() && (UnderlyingArchive.UE4Ver() < VER_UE4_ADD_TRANSACTIONAL_TO_DATA_ASSETS))
 	{
 		SetFlags(RF_Transactional);
 	}
diff --git a/Engine/Source/Runtime/Engine/Private/Distributions.cpp b/Engine/Source/Runtime/Engine/Private/Distributions.cpp
--- a/Engine/Source/Runtime/Engine/Private/Distributions.cpp
+++ b/Engine/Source/Runtime/Engine/Private/Distributions.cpp
@@ -1053,16 +1053,18 @@ bool HasBakedDistributionDataHelper(const UDistribution* GivenDistribution)
 }
 #endif
 
-void UDistributionFloat::Serialize(FArchive& Ar)
+void UDistributionFloat::Serialize(FStructuredArchive::FRecord Record)
 {
 #if WITH_EDITOR
-	if (Ar.IsCooking() || Ar.IsSaving())
+	FArchive& UnderlyingArchive = Record.GetUnderlyingArchive();
+
+	if (UnderlyingArchive.IsCooking() || UnderlyingArchive.IsSaving())
 	{
 		bBakedDataSuccesfully = HasBakedDistributionDataHelper<FRawDistributionFloat>(this);
 	}
 #endif
 
-	Super::Serialize(Ar);
+	Super::Serialize(Record);
 }
 
 
@@ -1261,16 +1263,18 @@ void FRawDistributionVector::InitLookupTable()
 #endif
 }
 
-void UDistributionVector::Serialize(FArchive& Ar)
+void UDistributionVector::Serialize(FStructuredArchive::FRecord Record)
 {
 #if WITH_EDITOR
-	if (Ar.IsCooking() || Ar.IsSaving())
+	FArchive& UnderlyingArchive = Record.GetUnderlyingArchive();
+
+	if (UnderlyingArchive.IsCooking() || UnderlyingArchive.IsSaving())
 	{
 		bBakedDataSuccesfully = HasBakedDistributionDataHelper<FRawDistributionVector>(this);
 	}
 #endif
 
-	Super::Serialize(Ar);
+	Super::Serialize(Record);
 }
 
 bool UDistributionVector::NeedsLoadForClient() const
diff --git a/Engine/Source/Runtime/Engine/Private/Particles/ParticleModules.cpp b/Engine/Source/Runtime/Engine/Private/Particles/ParticleModules.cpp
--- a/Engine/Source/Runtime/Engine/Private/Particles/ParticleModules.cpp
+++ b/Engine/Source/Runtime/Engine/Private/Particles/ParticleModules.cpp
@@ -1093,27 +1093,29 @@ void UParticleModuleRequired::PostInitProperties()
 	}
 }
 
-void UParticleModuleRequired::Serialize(FArchive& Ar)
+void UParticleModuleRequired::Serialize(FStructuredArchive::FRecord Record)
 {
-	Super::Serialize(Ar);
+	FArchive& UnderlyingArchive = Record.GetUnderlyingArchive();
+
+	Super::Serialize(Record);
 
-	Ar.UsingCustomVersion(FRenderingObjectVersion::GUID);
-	if (Ar.CustomVer(FRenderingObjectVersion::GUID) >= FRenderingObjectVersion::MovedParticleCutoutsToRequiredModule)
+	UnderlyingArchive.UsingCustomVersion(FRenderingObjectVersion::GUID);
+	if (UnderlyingArchive.CustomVer(FRenderingObjectVersion::GUID) >= FRenderingObjectVersion::MovedParticleCutoutsToRequiredModule)
 	{
-		bool bCooked = Ar.IsCooking();
+		bool bCooked = UnderlyingArchive.IsCooking();
 
 		// Save a bool indicating whether this is cooked data
 		// This is needed when loading cooked data, to know to serialize differently
-		Ar << bCooked;
+		Record << NAMED_FIELD(bCooked);
 
-		if (FPlatformProperties::RequiresCookedData() && !bCooked && Ar.IsLoading())
+		if (FPlatformProperties::RequiresCookedData() && !bCooked && UnderlyingArchive.IsLoading())
 		{
 			UE_LOG(LogParticles, Fatal, TEXT("This platform requires cooked packages, and this SubUV animation does not contain cooked data %s."), *GetName());
 		}
 
 		if (bCooked)
 		{
-			DerivedData.Serialize(Ar);
+			DerivedData.Serialize(Record.EnterField(FIELD_NAME_TEXT("DerivedData")));
 		}
 	}
 }
diff --git a/Engine/Source/Runtime/Engine/Private/Particles/SubUVAnimation.cpp b/Engine/Source/Runtime/Engine/Private/Particles/SubUVAnimation.cpp
--- a/Engine/Source/Runtime/Engine/Private/Particles/SubUVAnimation.cpp
+++ b/Engine/Source/Runtime/Engine/Private/Particles/SubUVAnimation.cpp
@@ -30,9 +30,9 @@ FString FSubUVDerivedData::GetDDCKeyString(const FGuid& StateId, int32 SizeX, in
 	return FDerivedDataCacheInterface::BuildCacheKey(TEXT("SUBUV_"), SUBUV_DERIVEDDATA_VER, *KeyString);
 }
 
-void FSubUVDerivedData::Serialize(FArchive& Ar)
+void FSubUVDerivedData::Serialize(FStructuredArchive::FSlot Slot)
 {
-	Ar << BoundingGeometry;
+	Slot << BoundingGeometry;
 }
 
 void FSubUVBoundingGeometryBuffer::InitRHI()
@@ -66,24 +66,26 @@ void USubUVAnimation::PostInitProperties()
 	}
 }
 
-void USubUVAnimation::Serialize(FArchive& Ar)
+void USubUVAnimation::Serialize(FStructuredArchive::FRecord Record)
 {
-	Super::Serialize(Ar);
+	FArchive& UnderlyingArchive = Record.GetUnderlyingArchive();
 
-	bool bCooked = Ar.IsCooking();
+	Super::Serialize(Record);
+
+	bool bCooked = UnderlyingArchive.IsCooking();
 
 	// Save a bool indicating whether this is cooked data
 	// This is needed when loading cooked data, to know to serialize differently
-	Ar << bCooked;
+	Record << NAMED_FIELD(bCooked);
 
-	if (FPlatformProperties::RequiresCookedData() && !bCooked && Ar.IsLoading())
+	if (FPlatformProperties::RequiresCookedData() && !bCooked && UnderlyingArchive.IsLoading())
 	{
 		UE_LOG(LogParticles, Fatal, TEXT("This platform requires cooked packages, and this SubUV animation does not contain cooked data %s."), *GetName());
 	}
 
 	if (bCooked)
 	{
-		DerivedData.Serialize(Ar);
+		DerivedData.Serialize(Record.EnterField(FIELD_NAME_TEXT("DerivedData")));
 	}
 }
 
