Added a couple of extensions to VisLog #UE4

Including:
- functionality to log arrows, including a macro for it, UE_VLOG_BOX
- a function for mas-logging boxes, FVisualLogEntry::AddBoxes

diff --git a/Engine/Source/Developer/LogVisualizer/Private/VisualLoggerRenderingActor.cpp b/Engine/Source/Developer/LogVisualizer/Private/VisualLoggerRenderingActor.cpp
--- a/Engine/Source/Developer/LogVisualizer/Private/VisualLoggerRenderingActor.cpp
+++ b/Engine/Source/Developer/LogVisualizer/Private/VisualLoggerRenderingActor.cpp
@@ -78,6 +78,7 @@ FPrimitiveSceneProxy* UVisualLoggerRenderingComponent::CreateSceneProxy()
 		VLogSceneProxy->Cones.Append(CurrentShapes.Value.Cones);
 		VLogSceneProxy->Texts.Append(CurrentShapes.Value.Texts);
 		VLogSceneProxy->Cylinders.Append(CurrentShapes.Value.Cylinders);
+		VLogSceneProxy->ArrowLines.Append(CurrentShapes.Value.Arrows);
 		VLogSceneProxy->Capsles.Append(CurrentShapes.Value.Capsules);
 	}
 
@@ -89,6 +90,7 @@ FPrimitiveSceneProxy* UVisualLoggerRenderingComponent::CreateSceneProxy()
 		VLogSceneProxy->Cones.Append(RenderingActor->TestDebugShapes.Cones);
 		VLogSceneProxy->Texts.Append(RenderingActor->TestDebugShapes.Texts);
 		VLogSceneProxy->Cylinders.Append(RenderingActor->TestDebugShapes.Cylinders);
+		VLogSceneProxy->ArrowLines.Append(RenderingActor->TestDebugShapes.Arrows);
 		VLogSceneProxy->Capsles.Append(RenderingActor->TestDebugShapes.Capsules);
 	}
 
@@ -634,7 +636,9 @@ void AVisualLoggerRenderingActor::GetDebugShapes(const FVisualLogDevice::FVisual
 		case EVisualLoggerShapeElement::NavAreaMesh:
 		{
 			if (ElementToDraw->Points.Num() == 0)
+			{
 				continue;
+			}
 
 			struct FHeaderData
 			{
@@ -695,6 +699,30 @@ void AVisualLoggerRenderingActor::GetDebugShapes(const FVisualLogDevice::FVisual
 					);
 			}
 
+		}
+			break;
+
+		case EVisualLoggerShapeElement::Arrow:
+		{
+			const bool bDrawLabels = ElementToDraw->Description.IsEmpty() == false && ElementToDraw->Points.Num() > 2;
+			const FVector* Location = ElementToDraw->Points.GetData();
+
+			if (bDrawLabels)
+			{
+				for (int32 Index = 0; Index + 1 < ElementToDraw->Points.Num(); Index += 2, Location += 2)
+				{
+					DebugShapes.Arrows.Add(FDebugRenderSceneProxy::FArrowLine(*Location, *(Location + 1), Color));
+					const FString PrintString = FString::Printf(TEXT("%s_%d"), *ElementToDraw->Description, Index);
+					DebugShapes.Texts.Add(FDebugRenderSceneProxy::FText3d(PrintString, (*Location + (*(Location + 1) - *Location) / 2), Color));
+				}
+			}
+			else
+			{
+				for (int32 Index = 0; Index + 1 < ElementToDraw->Points.Num(); Index += 2, Location += 2)
+				{
+					DebugShapes.Arrows.Add(FDebugRenderSceneProxy::FArrowLine(*Location, *(Location + 1), Color));
+				}
+			}
 		}
 			break;
 		}
diff --git a/Engine/Source/Developer/LogVisualizer/Private/VisualLoggerRenderingActor.h b/Engine/Source/Developer/LogVisualizer/Private/VisualLoggerRenderingActor.h
--- a/Engine/Source/Developer/LogVisualizer/Private/VisualLoggerRenderingActor.h
+++ b/Engine/Source/Developer/LogVisualizer/Private/VisualLoggerRenderingActor.h
@@ -31,6 +31,7 @@ class LOGVISUALIZER_API AVisualLoggerRenderingActor : public AActor
 		TArray<FDebugRenderSceneProxy::FText3d> Texts;
 		TArray<FDebugRenderSceneProxy::FWireCylinder> Cylinders;
 		TArray<FDebugRenderSceneProxy::FCapsule> Capsules;
+		TArray<FDebugRenderSceneProxy::FArrowLine> Arrows;
 		TArray<FVector> LogEntriesPath;
 
 		void Reset()
@@ -43,6 +44,7 @@ class LOGVISUALIZER_API AVisualLoggerRenderingActor : public AActor
 			Texts.Reset();
 			Cylinders.Reset();
 			Capsules.Reset();
+			Arrows.Reset();
 			LogEntriesPath.Reset();
 		}
 	};
diff --git a/Engine/Source/Runtime/Engine/Private/VisualLogger/VisualLoggerKismetLibrary.cpp b/Engine/Source/Runtime/Engine/Private/VisualLogger/VisualLoggerKismetLibrary.cpp
--- a/Engine/Source/Runtime/Engine/Private/VisualLogger/VisualLoggerKismetLibrary.cpp
+++ b/Engine/Source/Runtime/Engine/Private/VisualLogger/VisualLoggerKismetLibrary.cpp
@@ -62,7 +62,7 @@ void UVisualLoggerKismetLibrary::LogBox(UObject* WorldContextObject, FBox Box, F
 {
 #if ENABLE_VISUAL_LOG
 	const ELogVerbosity::Type DefaultVerbosity = ELogVerbosity::Log;
-	FVisualLogger::GeometryShapeLogf(WorldContextObject, CategoryName, DefaultVerbosity
+	FVisualLogger::GeometryBoxLogf(WorldContextObject, CategoryName, DefaultVerbosity
 		, Box, FMatrix::Identity, ObjectColor.ToFColor(true), TEXT("%s"), *Text);
 #endif
 	if (bAddToMessageLog)
diff --git a/Engine/Source/Runtime/Engine/Private/VisualLogger/VisualLoggerTypes.cpp b/Engine/Source/Runtime/Engine/Private/VisualLogger/VisualLoggerTypes.cpp
--- a/Engine/Source/Runtime/Engine/Private/VisualLogger/VisualLoggerTypes.cpp
+++ b/Engine/Source/Runtime/Engine/Private/VisualLogger/VisualLoggerTypes.cpp
@@ -182,6 +182,19 @@ void FVisualLogEntry::AddElement(const FVector& Start, const FVector& End, const
 	ElementsToDraw.Add(Element);
 }
 
+void FVisualLogEntry::AddArrow(const FVector& Start, const FVector& End, const FName& CategoryName, ELogVerbosity::Type Verbosity, const FColor& Color, const FString& Description)
+{
+	FVisualLogShapeElement Element(EVisualLoggerShapeElement::Arrow);
+	Element.Category = CategoryName;
+	Element.SetColor(Color);
+	Element.Description = Description;
+	Element.Points.Reserve(2);
+	Element.Points.Add(Start);
+	Element.Points.Add(End);
+	Element.Verbosity = Verbosity;
+	ElementsToDraw.Add(Element);
+}
+
 void FVisualLogEntry::AddElement(const FBox& Box, const FMatrix& Matrix, const FName& CategoryName, ELogVerbosity::Type Verbosity, const FColor& Color, const FString& Description, uint16 Thickness)
 {
 	FVisualLogShapeElement Element(Description, Color, Thickness, CategoryName);
@@ -194,6 +207,18 @@ void FVisualLogEntry::AddElement(const FBox& Box, const FMatrix& Matrix, const F
 	ElementsToDraw.Add(Element);
 }
 
+void FVisualLogEntry::AddBoxes(const TArray<FBox>& Boxes, const FName& CategoryName, ELogVerbosity::Type Verbosity, const FColor& Color, const TArray<FString>& Descriptions)
+{
+	FVisualLogShapeElement& Element = ElementsToDraw[ElementsToDraw.Add(FVisualLogShapeElement(EVisualLoggerShapeElement::Box))];
+	Element.Points.Reserve(2 * Boxes.Num());
+	for (const FBox& Box : Boxes)
+	{
+		Element.Points.Add(Box.Min);
+		Element.Points.Add(Box.Max);
+	}
+	Element.Verbosity = Verbosity;
+}
+
 void FVisualLogEntry::AddElement(const FVector& Orgin, const FVector& Direction, float Length, float AngleWidth, float AngleHeight, const FName& CategoryName, ELogVerbosity::Type Verbosity, const FColor& Color, const FString& Description, uint16 Thickness)
 {
 	FVisualLogShapeElement Element(Description, Color, Thickness, CategoryName);
diff --git a/Engine/Source/Runtime/Engine/Public/VisualLogger/VisualLogger.h b/Engine/Source/Runtime/Engine/Public/VisualLogger/VisualLogger.h
--- a/Engine/Source/Runtime/Engine/Public/VisualLogger/VisualLogger.h
+++ b/Engine/Source/Runtime/Engine/Public/VisualLogger/VisualLogger.h
@@ -36,10 +36,10 @@
 #define UE_VLOG_LOCATION(LogOwner, CategoryName, Verbosity, Location, Radius, Color, Format, ...) if(FVisualLogger::IsRecording()) FVisualLogger::GeometryShapeLogf(LogOwner, CategoryName, ELogVerbosity::Verbosity, Location, Radius, Color, Format, ##__VA_ARGS__)
 #define UE_CVLOG_LOCATION(Condition, LogOwner, CategoryName, Verbosity, Location, Radius, Color, Format, ...)  if(FVisualLogger::IsRecording() && Condition) {UE_VLOG_LOCATION(LogOwner, CategoryName, Verbosity, Location, Radius, Color, Format, ##__VA_ARGS__);} 
 // Box shape
-#define UE_VLOG_BOX(LogOwner, CategoryName, Verbosity, Box, Color, Format, ...) if(FVisualLogger::IsRecording()) FVisualLogger::GeometryShapeLogf(LogOwner, CategoryName, ELogVerbosity::Verbosity, Box, FMatrix::Identity, Color, Format, ##__VA_ARGS__)
-#define UE_CVLOG_BOX(Condition, LogOwner, CategoryName, Verbosity, Box, Color, Format, ...)  if(FVisualLogger::IsRecording() && Condition) {UE_VLOG_BOX(LogOwner, CategoryName, Verbosity, Box, FMatrix::Identity, Color, Format, ##__VA_ARGS__);} 
+#define UE_VLOG_BOX(LogOwner, CategoryName, Verbosity, Box, Color, Format, ...) if(FVisualLogger::IsRecording()) FVisualLogger::GeometryBoxLogf(LogOwner, CategoryName, ELogVerbosity::Verbosity, Box, FMatrix::Identity, Color, Format, ##__VA_ARGS__)
+#define UE_CVLOG_BOX(Condition, LogOwner, CategoryName, Verbosity, Box, Color, Format, ...)  if(FVisualLogger::IsRecording() && Condition) {UE_VLOG_BOX(LogOwner, CategoryName, Verbosity, Box, Color, Format, ##__VA_ARGS__);} 
 // Oriented box shape
-#define UE_VLOG_OBOX(LogOwner, CategoryName, Verbosity, Box, Matrix, Color, Format, ...) if(FVisualLogger::IsRecording()) FVisualLogger::GeometryShapeLogf(LogOwner, CategoryName, ELogVerbosity::Verbosity, Box, Matrix, Color, Format, ##__VA_ARGS__)
+#define UE_VLOG_OBOX(LogOwner, CategoryName, Verbosity, Box, Matrix, Color, Format, ...) if(FVisualLogger::IsRecording()) FVisualLogger::GeometryBoxLogf(LogOwner, CategoryName, ELogVerbosity::Verbosity, Box, Matrix, Color, Format, ##__VA_ARGS__)
 #define UE_CVLOG_OBOX(Condition, LogOwner, CategoryName, Verbosity, Box, Matrix, Color, Format, ...)  if(FVisualLogger::IsRecording() && Condition) {UE_VLOG_OBOX(LogOwner, CategoryName, Verbosity, Box, Matrix, Color, Format, ##__VA_ARGS__);} 
 // Cone shape
 #define UE_VLOG_CONE(LogOwner, CategoryName, Verbosity, Orgin, Direction, Length, Angle, Color, Format, ...) if(FVisualLogger::IsRecording()) FVisualLogger::GeometryShapeLogf(LogOwner, CategoryName, ELogVerbosity::Verbosity, Orgin, Direction, Length, Angle, Color, Format, ##__VA_ARGS__)
@@ -62,7 +62,9 @@
 // 2d convex poly shape
 #define UE_VLOG_CONVEXPOLY(LogOwner, CategoryName, Verbosity, Points, Color, Format, ...) if(FVisualLogger::IsRecording()) FVisualLogger::GeometryConvexLogf(LogOwner, CategoryName, ELogVerbosity::Verbosity, Points, Color, Format, ##__VA_ARGS__)
 #define UE_CVLOG_CONVEXPOLY(Condition, LogOwner, CategoryName, Verbosity, Points, Color, Format, ...)  if(FVisualLogger::IsRecording() && Condition) {UE_VLOG_CONVEXPOLY(LogOwner, CategoryName, Verbosity, Points, Color, Format, ##__VA_ARGS__);}
-
+// Segment with an arrowhead
+#define UE_VLOG_ARROW(LogOwner, CategoryName, Verbosity, SegmentStart, SegmentEnd, Color, Format, ...) if(FVisualLogger::IsRecording()) FVisualLogger::ArrowLogf(LogOwner, CategoryName, ELogVerbosity::Verbosity, SegmentStart, SegmentEnd, Color, Format, ##__VA_ARGS__)
+#define UE_CVLOG_ARROW(Condition, LogOwner, CategoryName, Verbosity, SegmentStart, SegmentEnd, Color, Format, ...) if(FVisualLogger::IsRecording() && Condition) {UE_VLOG_ARROW(LogOwner, CategoryName, Verbosity, SegmentStart, SegmentEnd, Color, Format, ##__VA_ARGS__);} 
 
 #define DECLARE_VLOG_EVENT(EventName) extern FVisualLogEventBase EventName;
 #define DEFINE_VLOG_EVENT(EventName, Verbosity, UserFriendlyDesc) FVisualLogEventBase EventName(TEXT(#EventName), TEXT(UserFriendlyDesc), ELogVerbosity::Verbosity); 
@@ -107,7 +109,9 @@
 #define UE_VLOG_MESH(LogOwner, CategoryName, Verbosity, Vertices, Indexes, Color, Format, ...) 
 #define UE_CVLOG_MESH(Condition, LogOwner, CategoryName, Verbosity, Vertices, Indexes, Color, Format, ...) 
 #define UE_VLOG_CONVEXPOLY(LogOwner, CategoryName, Verbosity, Points, Color, Format, ...) 
-#define UE_CVLOG_CONVEXPOLY(Condition, LogOwner, CategoryName, Verbosity, Points, Color, Format, ...) 
+#define UE_CVLOG_CONVEXPOLY(Condition, LogOwner, CategoryName, Verbosity, Points, Color, Format, ...)
+#define UE_VLOG_ARROW(LogOwner, CategoryName, Verbosity, SegmentStart, SegmentEnd, Color, Format, ...) 
+#define UE_CVLOG_ARROW(Condition, LogOwner, CategoryName, Verbosity, SegmentStart, SegmentEnd, Color, Format, ...) 
 
 #define DECLARE_VLOG_EVENT(EventName)
 #define DEFINE_VLOG_EVENT(EventName, Verbosity, UserFriendlyDesc)
@@ -150,8 +154,8 @@ class ENGINE_API FVisualLogger : public FOutputDevice
 	static void GeometryShapeLogfImpl (const UObject* LogOwner, const FName& CategoryName, ELogVerbosity::Type Verbosity, const FVector& Start, const FVector& End, const FColor& Color, const uint16 Thickness, const TCHAR* Fmt, ...);
 	static void GeometryShapeLogfImpl (const UObject* LogOwner, const FLogCategoryBase& Category, ELogVerbosity::Type Verbosity, const FVector& Location, float Radius, const FColor& Color, const TCHAR* Fmt, ...);
 	static void GeometryShapeLogfImpl (const UObject* LogOwner, const FName& CategoryName, ELogVerbosity::Type Verbosity, const FVector& Location, float Radius, const FColor& Color, const TCHAR* Fmt, ...);
-	static void GeometryShapeLogfImpl (const UObject* LogOwner, const FLogCategoryBase& Category, ELogVerbosity::Type Verbosity, const FBox& Box, const FMatrix& Matrix, const FColor& Color, const TCHAR* Fmt, ...);
-	static void GeometryShapeLogfImpl (const UObject* LogOwner, const FName& CategoryName, ELogVerbosity::Type Verbosity, const FBox& Box, const FMatrix& Matrix, const FColor& Color, const TCHAR* Fmt, ...);
+	static void GeometryBoxLogfImpl(const UObject* LogOwner, const FLogCategoryBase& Category, ELogVerbosity::Type Verbosity, const FBox& Box, const FMatrix& Matrix, const FColor& Color, const TCHAR* Fmt, ...);
+	static void GeometryBoxLogfImpl(const UObject* LogOwner, const FName& CategoryName, ELogVerbosity::Type Verbosity, const FBox& Box, const FMatrix& Matrix, const FColor& Color, const TCHAR* Fmt, ...);
 	static void GeometryShapeLogfImpl (const UObject* LogOwner, const FLogCategoryBase& Category, ELogVerbosity::Type Verbosity, const FVector& Origin, const FVector& Direction, const float Length, const float Angle, const FColor& Color, const TCHAR* Fmt, ...);
 	static void GeometryShapeLogfImpl (const UObject* LogOwner, const FName& CategoryName, ELogVerbosity::Type Verbosity, const FVector& Origin, const FVector& Direction, const float Length, const float Angle, const FColor& Color, const TCHAR* Fmt, ...);
 	static void GeometryShapeLogfImpl (const UObject* LogOwner, const FLogCategoryBase& Category, ELogVerbosity::Type Verbosity, const FVector& Start, const FVector& End, const float Radius, const FColor& Color, const TCHAR* Fmt, ...);
@@ -166,6 +170,7 @@ class ENGINE_API FVisualLogger : public FOutputDevice
 	static void GeometryConvexLogfImpl(const UObject* LogOwner, const FName& CategoryName, ELogVerbosity::Type Verbosity, const TArray<FVector>& Points, const FColor& Color, const TCHAR* Fmt, ...);
 	static void HistogramDataLogfImpl (const UObject* LogOwner, const FLogCategoryBase& Category, ELogVerbosity::Type Verbosity, FName GraphName, FName DataName, const FVector2D& Data, const FColor& Color, const TCHAR* Fmt, ...);
 	static void HistogramDataLogfImpl (const UObject* LogOwner, const FName& CategoryName, ELogVerbosity::Type Verbosity, FName GraphName, FName DataName, const FVector2D& Data, const FColor& Color, const TCHAR* Fmt, ...);
+	static void ArrowLogfImpl(const UObject* LogOwner, const FName& CategoryName, ELogVerbosity::Type Verbosity, const FVector& Start, const FVector& End, const FColor& Color, const TCHAR* Fmt, ...);
 
 public:
 	// Regular text log
@@ -204,6 +209,16 @@ public:
 		GeometryShapeLogfImpl(LogOwner, CategoryName, Verbosity, Start, End, Color, Thickness, Fmt, Args...);
 	}
 
+	// Arrow
+	template <typename FmtType, typename... Types>
+	static void ArrowLogf(const UObject* LogOwner, const FName& CategoryName, ELogVerbosity::Type Verbosity, const FVector& Start, const FVector& End, const FColor& Color, const FmtType& Fmt, Types... Args)
+	{
+		static_assert(TIsArrayOrRefOfType<FmtType, TCHAR>::Value, "Formatting string must be a TCHAR array.");
+		static_assert(TAnd<TIsValidVariadicFunctionArg<Types>...>::Value, "Invalid argument(s) passed to FVisualLogger::ArrowLogf");
+
+		ArrowLogfImpl(LogOwner, CategoryName, Verbosity, Start, End, Color, Fmt, Args...);
+	}
+
 	// Location/Sphere log
 	template <typename FmtType, typename... Types>
 	static void GeometryShapeLogf(const UObject* LogOwner, const FLogCategoryBase& Category, ELogVerbosity::Type Verbosity, const FVector& Location, float Radius, const FColor& Color, const FmtType& Fmt, Types... Args)
@@ -224,20 +239,20 @@ public:
 
 	// Box log
 	template <typename FmtType, typename... Types>
-	static void GeometryShapeLogf(const UObject* LogOwner, const FLogCategoryBase& Category, ELogVerbosity::Type Verbosity, const FBox& Box, const FMatrix& Matrix, const FColor& Color, const FmtType& Fmt, Types... Args)
+	static void GeometryBoxLogf(const UObject* LogOwner, const FLogCategoryBase& Category, ELogVerbosity::Type Verbosity, const FBox& Box, const FMatrix& Matrix, const FColor& Color, const FmtType& Fmt, Types... Args)
 	{
 		static_assert(TIsArrayOrRefOfType<FmtType, TCHAR>::Value, "Formatting string must be a TCHAR array.");
-		static_assert(TAnd<TIsValidVariadicFunctionArg<Types>...>::Value, "Invalid argument(s) passed to FVisualLogger::GeometryShapeLogf");
+		static_assert(TAnd<TIsValidVariadicFunctionArg<Types>...>::Value, "Invalid argument(s) passed to FVisualLogger::GeometryBoxLogf");
 
-		GeometryShapeLogfImpl(LogOwner, Category, Verbosity, Box, Matrix, Color, Fmt, Args...);
+		GeometryBoxLogfImpl(LogOwner, Category, Verbosity, Box, Matrix, Color, Fmt, Args...);
 	}
 	template <typename FmtType, typename... Types>
-	static void GeometryShapeLogf(const UObject* LogOwner, const FName& CategoryName, ELogVerbosity::Type Verbosity, const FBox& Box, const FMatrix& Matrix, const FColor& Color, const FmtType& Fmt, Types... Args)
+	static void GeometryBoxLogf(const UObject* LogOwner, const FName& CategoryName, ELogVerbosity::Type Verbosity, const FBox& Box, const FMatrix& Matrix, const FColor& Color, const FmtType& Fmt, Types... Args)
 	{
 		static_assert(TIsArrayOrRefOfType<FmtType, TCHAR>::Value, "Formatting string must be a TCHAR array.");
-		static_assert(TAnd<TIsValidVariadicFunctionArg<Types>...>::Value, "Invalid argument(s) passed to FVisualLogger::GeometryShapeLogf");
+		static_assert(TAnd<TIsValidVariadicFunctionArg<Types>...>::Value, "Invalid argument(s) passed to FVisualLogger::GeometryBoxLogf");
 
-		GeometryShapeLogfImpl(LogOwner, CategoryName, Verbosity, Box, Matrix, Color, Fmt, Args...);
+		GeometryBoxLogfImpl(LogOwner, CategoryName, Verbosity, Box, Matrix, Color, Fmt, Args...);
 	}
 
 	// Cone log
@@ -604,6 +619,13 @@ inline void FVisualLogger::GeometryShapeLogfImpl(const UObject* Object, const FN
 	);
 }
 
+inline void FVisualLogger::ArrowLogfImpl(const UObject* Object, const FName& CategoryName, ELogVerbosity::Type Verbosity, const FVector& Start, const FVector& End, const FColor& Color, const TCHAR* Fmt, ...)
+{
+	COLLAPSED_LOGF(
+		CurrentEntry->AddArrow(Start, End, CategoryName, Verbosity, Color, Buffer);
+	);
+}
+
 inline void FVisualLogger::GeometryShapeLogfImpl(const UObject* Object, const FLogCategoryBase& Category, ELogVerbosity::Type Verbosity, const FVector& Location, float Radius, const FColor& Color, const TCHAR* Fmt, ...)
 {
 	const FName CategoryName = Category.GetCategoryName(); 
@@ -618,14 +640,14 @@ inline void FVisualLogger::GeometryShapeLogfImpl(const UObject* Object, const FN
 	);
 }
 
-inline void FVisualLogger::GeometryShapeLogfImpl(const UObject* Object, const FLogCategoryBase& Category, ELogVerbosity::Type Verbosity, const FBox& Box, const FMatrix& Matrix, const FColor& Color, const TCHAR* Fmt, ...)
+inline void FVisualLogger::GeometryBoxLogfImpl(const UObject* Object, const FLogCategoryBase& Category, ELogVerbosity::Type Verbosity, const FBox& Box, const FMatrix& Matrix, const FColor& Color, const TCHAR* Fmt, ...)
 {
 	const FName CategoryName = Category.GetCategoryName(); 
 	COLLAPSED_LOGF(
 		CurrentEntry->AddElement(Box, Matrix, CategoryName, Verbosity, Color, Buffer);
 	);
 }
-inline void FVisualLogger::GeometryShapeLogfImpl(const UObject* Object, const FName& CategoryName, ELogVerbosity::Type Verbosity, const FBox& Box, const FMatrix& Matrix, const FColor& Color, const TCHAR* Fmt, ...)
+inline void FVisualLogger::GeometryBoxLogfImpl(const UObject* Object, const FName& CategoryName, ELogVerbosity::Type Verbosity, const FBox& Box, const FMatrix& Matrix, const FColor& Color, const TCHAR* Fmt, ...)
 {
 	COLLAPSED_LOGF(
 		CurrentEntry->AddElement(Box, Matrix, CategoryName, Verbosity, Color, Buffer);
diff --git a/Engine/Source/Runtime/Engine/Public/VisualLogger/VisualLoggerTypes.h b/Engine/Source/Runtime/Engine/Public/VisualLogger/VisualLoggerTypes.h
--- a/Engine/Source/Runtime/Engine/Public/VisualLogger/VisualLoggerTypes.h
+++ b/Engine/Source/Runtime/Engine/Public/VisualLogger/VisualLoggerTypes.h
@@ -61,6 +61,7 @@ enum class EVisualLoggerShapeElement : uint8
 	Polygon,
 	Mesh,
 	NavAreaMesh, // convex based mesh with min and max Z values
+	Arrow, 
 	// note that in order to remain backward compatibility in terms of log
 	// serialization new enum values need to be added at the end
 };
@@ -218,6 +219,11 @@ struct ENGINE_API FVisualLogEntry
 	void AddConvexElement(const TArray<FVector>& Points, const FName& CategoryName, ELogVerbosity::Type Verbosity, const FColor& Color = FColor::White, const FString& Description = TEXT(""));
 	// histogram sample
 	void AddHistogramData(const FVector2D& DataSample, const FName& CategoryName, ELogVerbosity::Type Verbosity, const FName& GraphName, const FName& DataName);
+	// arrow
+	void AddArrow(const FVector& Start, const FVector& End, const FName& CategoryName, ELogVerbosity::Type Verbosity, const FColor& Color = FColor::White, const FString& Description = TEXT(""));
+	// boxes
+	void AddBoxes(const TArray<FBox>& Boxes, const FName& CategoryName, ELogVerbosity::Type Verbosity, const FColor& Color = FColor::White, const TArray<FString>& Descriptions = TArray<FString>());
+
 	// Custom data block
 	FVisualLogDataBlock& AddDataBlock(const FString& TagName, const TArray<uint8>& BlobDataArray, const FName& CategoryName, ELogVerbosity::Type Verbosity);
 	// Event
