VisLog's boxes logging function fix #UE4

diff --git a/Engine/Source/Runtime/Engine/Private/VisualLogger/VisualLoggerTypes.cpp b/Engine/Source/Runtime/Engine/Private/VisualLogger/VisualLoggerTypes.cpp
--- a/Engine/Source/Runtime/Engine/Private/VisualLogger/VisualLoggerTypes.cpp
+++ b/Engine/Source/Runtime/Engine/Private/VisualLogger/VisualLoggerTypes.cpp
@@ -207,9 +207,11 @@ void FVisualLogEntry::AddElement(const FBox& Box, const FMatrix& Matrix, const F
 	ElementsToDraw.Add(Element);
 }
 
-void FVisualLogEntry::AddBoxes(const TArray<FBox>& Boxes, const FName& CategoryName, ELogVerbosity::Type Verbosity, const FColor& Color, const TArray<FString>& Descriptions)
+void FVisualLogEntry::AddBoxes(const TArray<FBox>& Boxes, const FName& CategoryName, ELogVerbosity::Type Verbosity, const FColor& Color)
 {
 	FVisualLogShapeElement& Element = ElementsToDraw[ElementsToDraw.Add(FVisualLogShapeElement(EVisualLoggerShapeElement::Box))];
+	Element.Category = CategoryName;
+	Element.Verbosity = Verbosity;
 	Element.Points.Reserve(2 * Boxes.Num());
 	for (const FBox& Box : Boxes)
 	{
diff --git a/Engine/Source/Runtime/Engine/Public/VisualLogger/VisualLoggerTypes.h b/Engine/Source/Runtime/Engine/Public/VisualLogger/VisualLoggerTypes.h
--- a/Engine/Source/Runtime/Engine/Public/VisualLogger/VisualLoggerTypes.h
+++ b/Engine/Source/Runtime/Engine/Public/VisualLogger/VisualLoggerTypes.h
@@ -222,7 +222,7 @@ struct ENGINE_API FVisualLogEntry
 	// arrow
 	void AddArrow(const FVector& Start, const FVector& End, const FName& CategoryName, ELogVerbosity::Type Verbosity, const FColor& Color = FColor::White, const FString& Description = TEXT(""));
 	// boxes
-	void AddBoxes(const TArray<FBox>& Boxes, const FName& CategoryName, ELogVerbosity::Type Verbosity, const FColor& Color = FColor::White, const TArray<FString>& Descriptions = TArray<FString>());
+	void AddBoxes(const TArray<FBox>& Boxes, const FName& CategoryName, ELogVerbosity::Type Verbosity, const FColor& Color = FColor::White);
 
 	// Custom data block
 	FVisualLogDataBlock& AddDataBlock(const FString& TagName, const TArray<uint8>& BlobDataArray, const FName& CategoryName, ELogVerbosity::Type Verbosity);
