#ue4 -(Integrate 4154124 from Dan.Youhon) - Fix bWantsToCrouch not set for simulated proxies, which could cause premature crouch/uncrouch issues if proxies use root motion that runs through the full PerformMovemnt->UpdateCharacterStateBeforeMovement() flow.

diff --git a/Engine/Source/Runtime/Engine/Private/Character.cpp b/Engine/Source/Runtime/Engine/Private/Character.cpp
--- a/Engine/Source/Runtime/Engine/Private/Character.cpp
+++ b/Engine/Source/Runtime/Engine/Private/Character.cpp
@@ -347,10 +347,12 @@ void ACharacter::OnRep_IsCrouched()
 	{
 		if (bIsCrouched)
 		{
+			CharacterMovement->bWantsToCrouch = true;
 			CharacterMovement->Crouch(true);
 		}
 		else
 		{
+			CharacterMovement->bWantsToCrouch = false;
 			CharacterMovement->UnCrouch(true);
 		}
 	}
