PR #4873: Convert gameplay static functions to be pure (Contributed by [REDACTED_PERSON])
#rnx

diff --git a/Engine/Source/Runtime/Engine/Classes/Kismet/GameplayStatics.h b/Engine/Source/Runtime/Engine/Classes/Kismet/GameplayStatics.h
--- a/Engine/Source/Runtime/Engine/Classes/Kismet/GameplayStatics.h
+++ b/Engine/Source/Runtime/Engine/Classes/Kismet/GameplayStatics.h
@@ -629,7 +629,7 @@ public:
 	 * Returns whether or not subtitles are currently enabled.
 	 * @return true if subtitles are enabled.
 	 */
-	UFUNCTION(BlueprintCallable, Category = "Audio|Subtitles")
+	UFUNCTION(BlueprintPure, Category = "Audio|Subtitles")
 	static bool AreSubtitlesEnabled();
 
 	// --- Audio Functions ----------------------------
@@ -885,7 +885,7 @@ public:
 	/**
 	 * Returns the string name of the current platform, to perform different behavior based on platform. 
 	 * (Platform names include Windows, Mac, IOS, Android, PS4, XboxOne, HTML5, Linux) */
-	UFUNCTION(BlueprintCallable, Category="Game")
+	UFUNCTION(BlueprintPure, Category="Game")
 	static FString GetPlatformName();
 
 	/**
