Allow TInlineComponentArray specify it would like to include from child actors

diff --git a/Engine/Source/Runtime/Engine/Classes/GameFramework/Actor.h b/Engine/Source/Runtime/Engine/Classes/GameFramework/Actor.h
--- a/Engine/Source/Runtime/Engine/Classes/GameFramework/Actor.h
+++ b/Engine/Source/Runtime/Engine/Classes/GameFramework/Actor.h
@@ -3064,11 +3064,11 @@ class TInlineComponentArray : public TArray<T, TInlineAllocator<NumElements>>
 
 public:
 	TInlineComponentArray() : Super() { }
-	TInlineComponentArray(const class AActor* Actor) : Super()
+	TInlineComponentArray(const class AActor* Actor, bool bIncludeFromChildActors = false) : Super()
 	{
 		if (Actor)
 		{
-			Actor->GetComponents(*this);
+			Actor->GetComponents(*this, bIncludeFromChildActors);
 		}
 	};
 };
