Make DiscardPlayerInput actually do that

diff --git a/Engine/Source/Runtime/Engine/Private/UserInterface/PlayerInput.cpp b/Engine/Source/Runtime/Engine/Private/UserInterface/PlayerInput.cpp
--- a/Engine/Source/Runtime/Engine/Private/UserInterface/PlayerInput.cpp
+++ b/Engine/Source/Runtime/Engine/Private/UserInterface/PlayerInput.cpp
@@ -1,4 +1,4 @@
-﻿// Copyright 1998-2018 Epic Games, Inc. All Rights Reserved.
+// Copyright 1998-2018 Epic Games, Inc. All Rights Reserved.
 
 /*=============================================================================
 	PlayerInput.cpp: Unreal input system.
@@ -1305,7 +1305,18 @@ void UPlayerInput::ProcessInputStack(const TArray<UInputComponent*>& InputCompon
 
 void UPlayerInput::DiscardPlayerInput()
 {
-	FinishProcessingPlayerInput();
+	// discard any accumulated player input that hasn't yet been processed
+	for (TMap<FKey, FKeyState>::TIterator It(KeyStateMap); It; ++It)
+	{
+		FKeyState& KeyState = It.Value();
+		KeyState.RawValueAccumulator = FVector(0.f,0.f,0.f);
+		KeyState.SampleCountAccumulator = 0;
+		for (uint8 EventIndex = 0; EventIndex < IE_MAX; ++EventIndex)
+		{
+			KeyState.EventAccumulator[EventIndex].Reset();
+		}
+	}
+	EventCount = 0;
 }
 
 void UPlayerInput::FinishProcessingPlayerInput()
