Fix immutable structs with non-editable properties corrupting later properties in the component instance data cache byte stream

diff --git a/Engine/Source/Runtime/Engine/Private/ComponentInstanceDataCache.cpp b/Engine/Source/Runtime/Engine/Private/ComponentInstanceDataCache.cpp
--- a/Engine/Source/Runtime/Engine/Private/ComponentInstanceDataCache.cpp
+++ b/Engine/Source/Runtime/Engine/Private/ComponentInstanceDataCache.cpp
@@ -53,9 +53,16 @@ public:
 
 	virtual bool ShouldSkipProperty(const UProperty* InProperty) const override
 	{
-		return (	InProperty->HasAnyPropertyFlags(CPF_Transient)
-			|| !InProperty->HasAnyPropertyFlags(CPF_Edit | CPF_Interp)
-			|| PropertiesToSkip.Contains(InProperty));
+		// Immutable structs expect to serialize all properties so don't skip regardless of other conditions
+		UScriptStruct* ScriptStruct = Cast<UScriptStruct>(InProperty->GetOuter());
+		const bool bPropertyInImmutableStruct = ScriptStruct && ((ScriptStruct->StructFlags & STRUCT_Immutable) != 0);
+		
+		return (!bPropertyInImmutableStruct
+		        && (   InProperty->HasAnyPropertyFlags(CPF_Transient)
+		            || !InProperty->HasAnyPropertyFlags(CPF_Edit | CPF_Interp)
+		            || PropertiesToSkip.Contains(InProperty)
+		           )
+		       );
 	}
 
 
