PR #4897: Extend Splitscreen types with 2 new more types (Contributed by [REDACTED_PERSON])

diff --git a/Engine/Source/Runtime/Engine/Classes/Engine/ViewportSplitScreen.h b/Engine/Source/Runtime/Engine/Classes/Engine/ViewportSplitScreen.h
--- a/Engine/Source/Runtime/Engine/Classes/Engine/ViewportSplitScreen.h
+++ b/Engine/Source/Runtime/Engine/Classes/Engine/ViewportSplitScreen.h
@@ -29,11 +29,19 @@ namespace ESplitScreenType
 		ThreePlayer_FavorTop,
 		// 3 Player split with 1 player on bottom and 2 on top
 		ThreePlayer_FavorBottom,
-		// 4 Player split
-		FourPlayer,
+		//3 Player vertical split
+		ThreePlayer_Vertical,
+		// 4 Player grid split
+		FourPlayer_Grid,
+		// 4 Player vertical split
+		FourPlayer_Vertical,
 
 		SplitTypeCount
 	};
+
+	// Deprecated old FourPlayer grid enum value
+	DEPRECATED(4.21, "FourPlayer is now FourPlayer_Grid")
+	const Type FourPlayer = FourPlayer_Grid;
 }
 
 /** Structure to store splitscreen data. */
diff --git a/Engine/Source/Runtime/Engine/Private/GameViewportClient.cpp b/Engine/Source/Runtime/Engine/Private/GameViewportClient.cpp
--- a/Engine/Source/Runtime/Engine/Private/GameViewportClient.cpp
+++ b/Engine/Source/Runtime/Engine/Private/GameViewportClient.cpp
@@ -183,10 +183,19 @@ UGameViewportClient::UGameViewportClient(const FObjectInitializer& ObjectInitial
 	SplitscreenInfo[ESplitScreenType::ThreePlayer_FavorBottom].PlayerData.Add(FPerPlayerSplitscreenData(0.5f, 0.5f, 0.5f, 0.0f));
 	SplitscreenInfo[ESplitScreenType::ThreePlayer_FavorBottom].PlayerData.Add(FPerPlayerSplitscreenData(1.0f, 0.5f, 0.0f, 0.5f));
 
-	SplitscreenInfo[ESplitScreenType::FourPlayer].PlayerData.Add(FPerPlayerSplitscreenData(0.5f, 0.5f, 0.0f, 0.0f));
-	SplitscreenInfo[ESplitScreenType::FourPlayer].PlayerData.Add(FPerPlayerSplitscreenData(0.5f, 0.5f, 0.5f, 0.0f));
-	SplitscreenInfo[ESplitScreenType::FourPlayer].PlayerData.Add(FPerPlayerSplitscreenData(0.5f, 0.5f, 0.0f, 0.5f));
-	SplitscreenInfo[ESplitScreenType::FourPlayer].PlayerData.Add(FPerPlayerSplitscreenData(0.5f, 0.5f, 0.5f, 0.5f));
+	SplitscreenInfo[ESplitScreenType::ThreePlayer_Vertical].PlayerData.Add(FPerPlayerSplitscreenData(0.333f, 1.0f, 0.0f, 0.0f));
+	SplitscreenInfo[ESplitScreenType::ThreePlayer_Vertical].PlayerData.Add(FPerPlayerSplitscreenData(0.333f, 1.0f, 0.333f, 0.0f));
+	SplitscreenInfo[ESplitScreenType::ThreePlayer_Vertical].PlayerData.Add(FPerPlayerSplitscreenData(0.333f, 1.0f, 0.666f, 0.0f));
+
+	SplitscreenInfo[ESplitScreenType::FourPlayer_Grid].PlayerData.Add(FPerPlayerSplitscreenData(0.5f, 0.5f, 0.0f, 0.0f));
+	SplitscreenInfo[ESplitScreenType::FourPlayer_Grid].PlayerData.Add(FPerPlayerSplitscreenData(0.5f, 0.5f, 0.5f, 0.0f));
+	SplitscreenInfo[ESplitScreenType::FourPlayer_Grid].PlayerData.Add(FPerPlayerSplitscreenData(0.5f, 0.5f, 0.0f, 0.5f));
+	SplitscreenInfo[ESplitScreenType::FourPlayer_Grid].PlayerData.Add(FPerPlayerSplitscreenData(0.5f, 0.5f, 0.5f, 0.5f));
+
+	SplitscreenInfo[ESplitScreenType::FourPlayer_Vertical].PlayerData.Add(FPerPlayerSplitscreenData(0.25f, 1.0f, 0.0f, 0.0f));
+	SplitscreenInfo[ESplitScreenType::FourPlayer_Vertical].PlayerData.Add(FPerPlayerSplitscreenData(0.25f, 1.0f, 0.25f, 0.0f));
+	SplitscreenInfo[ESplitScreenType::FourPlayer_Vertical].PlayerData.Add(FPerPlayerSplitscreenData(0.25f, 1.0f, 0.5f, 0.0f));
+	SplitscreenInfo[ESplitScreenType::FourPlayer_Vertical].PlayerData.Add(FPerPlayerSplitscreenData(0.25f, 1.0f, 0.75f, 0.0f));
 
 	MaxSplitscreenPlayers = 4;
 	bSuppressTransitionMessage = true;
@@ -1945,6 +1954,10 @@ void UGameViewportClient::UpdateActiveSplitscreenType()
 				SplitType = ESplitScreenType::ThreePlayer_FavorBottom;
 				break;
 
+			case EThreePlayerSplitScreenType::Vertical:
+				SplitType = ESplitScreenType::ThreePlayer_Vertical;
+				break;
+
 			default:
 				check(0);
 			}
@@ -1952,7 +1965,19 @@ void UGameViewportClient::UpdateActiveSplitscreenType()
 
 		default:
 			ensure(NumPlayers == 4);
-			SplitType = ESplitScreenType::FourPlayer;
+			switch (Settings->FourPlayerSplitscreenLayout)
+			{
+			case EFourPlayerSplitScreenType::Grid:
+				SplitType = ESplitScreenType::FourPlayer_Grid;
+				break;
+
+			case EFourPlayerSplitScreenType::Vertical:
+				SplitType = ESplitScreenType::FourPlayer_Vertical;
+				break;
+
+			default:
+				check(0);
+			}
 			break;
 		}
 	}
@@ -2015,15 +2040,17 @@ bool UGameViewportClient::HasTopSafeZone( int32 LocalPlayerIndex )
 	{
 	case ESplitScreenType::None:
 	case ESplitScreenType::TwoPlayer_Vertical:
+	case ESplitScreenType::ThreePlayer_Vertical:
+	case ESplitScreenType::FourPlayer_Vertical:
 		return true;
 
 	case ESplitScreenType::TwoPlayer_Horizontal:
 	case ESplitScreenType::ThreePlayer_FavorTop:
-		return (LocalPlayerIndex == 0) ? true : false;
+		return (LocalPlayerIndex == 0);
 
 	case ESplitScreenType::ThreePlayer_FavorBottom:
-	case ESplitScreenType::FourPlayer:
-		return (LocalPlayerIndex < 2) ? true : false;
+	case ESplitScreenType::FourPlayer_Grid:
+		return (LocalPlayerIndex < 2);
 	}
 
 	return false;
@@ -2035,15 +2062,17 @@ bool UGameViewportClient::HasBottomSafeZone( int32 LocalPlayerIndex )
 	{
 	case ESplitScreenType::None:
 	case ESplitScreenType::TwoPlayer_Vertical:
+	case ESplitScreenType::ThreePlayer_Vertical:
+	case ESplitScreenType::FourPlayer_Vertical:
 		return true;
 
 	case ESplitScreenType::TwoPlayer_Horizontal:
 	case ESplitScreenType::ThreePlayer_FavorTop:
-		return (LocalPlayerIndex == 0) ? false : true;
+		return (LocalPlayerIndex > 0);
 
 	case ESplitScreenType::ThreePlayer_FavorBottom:
-	case ESplitScreenType::FourPlayer:
-		return (LocalPlayerIndex > 1) ? true : false;
+	case ESplitScreenType::FourPlayer_Grid:
+		return (LocalPlayerIndex > 1);
 	}
 
 	return false;
@@ -2058,14 +2087,16 @@ bool UGameViewportClient::HasLeftSafeZone( int32 LocalPlayerIndex )
 		return true;
 
 	case ESplitScreenType::TwoPlayer_Vertical:
-		return (LocalPlayerIndex == 0) ? true : false;
+	case ESplitScreenType::ThreePlayer_Vertical:
+	case ESplitScreenType::FourPlayer_Vertical:
+		return (LocalPlayerIndex == 0);
 
 	case ESplitScreenType::ThreePlayer_FavorTop:
 		return (LocalPlayerIndex < 2) ? true : false;
 
 	case ESplitScreenType::ThreePlayer_FavorBottom:
-	case ESplitScreenType::FourPlayer:
-		return (LocalPlayerIndex == 0 || LocalPlayerIndex == 2) ? true : false;
+	case ESplitScreenType::FourPlayer_Grid:
+		return (LocalPlayerIndex == 0 || LocalPlayerIndex == 2);
 	}
 
 	return false;
@@ -2081,13 +2112,19 @@ bool UGameViewportClient::HasRightSafeZone( int32 LocalPlayerIndex )
 
 	case ESplitScreenType::TwoPlayer_Vertical:
 	case ESplitScreenType::ThreePlayer_FavorBottom:
-		return (LocalPlayerIndex > 0) ? true : false;
+		return (LocalPlayerIndex > 0);
 
 	case ESplitScreenType::ThreePlayer_FavorTop:
-		return (LocalPlayerIndex == 1) ? false : true;
+		return (LocalPlayerIndex != 1);
+
+	case ESplitScreenType::ThreePlayer_Vertical:
+		return (LocalPlayerIndex == 2);
 
-	case ESplitScreenType::FourPlayer:
-		return (LocalPlayerIndex == 0 || LocalPlayerIndex == 2) ? false : true;
+	case ESplitScreenType::FourPlayer_Vertical:
+		return (LocalPlayerIndex == 3);
+
+	case ESplitScreenType::FourPlayer_Grid:
+		return (LocalPlayerIndex == 1 || LocalPlayerIndex == 3);
 	}
 
 	return false;
@@ -2132,10 +2169,17 @@ void UGameViewportClient::GetPixelSizeOfScreen( float& Width, float& Height, UCa
 		}
 		Height = Canvas->ClipY * 2;
 		return;
-	case ESplitScreenType::FourPlayer:
+	case ESplitScreenType::ThreePlayer_Vertical:
+		Width = Canvas->ClipX * 3;
+		Height = Canvas->ClipY;
+		return;
+	case ESplitScreenType::FourPlayer_Grid:
 		Width = Canvas->ClipX * 2;
 		Height = Canvas->ClipY * 2;
 		return;
+	case ESplitScreenType::FourPlayer_Vertical:
+		Width = Canvas->ClipX * 4;
+		Height = Canvas->ClipY;
 	}
 }
 
diff --git a/Engine/Source/Runtime/EngineSettings/Classes/GameMapsSettings.h b/Engine/Source/Runtime/EngineSettings/Classes/GameMapsSettings.h
--- a/Engine/Source/Runtime/EngineSettings/Classes/GameMapsSettings.h
+++ b/Engine/Source/Runtime/EngineSettings/Classes/GameMapsSettings.h
@@ -27,10 +27,18 @@ namespace EThreePlayerSplitScreenType
 	enum Type
 	{
 		FavorTop,
-		FavorBottom
+		FavorBottom,
+		Vertical
 	};
 }
 
+UENUM()
+enum class EFourPlayerSplitScreenType : uint8
+{
+	Grid,
+	Vertical
+};
+
 /** Helper structure, used to associate GameModes with shortcut names. */
 USTRUCT()
 struct FGameModeName
@@ -125,6 +133,10 @@ public:
 	UPROPERTY(config, EditAnywhere, Category=LocalMultiplayer, meta=(editcondition="bUseSplitScreen"))
 	TEnumAsByte<EThreePlayerSplitScreenType::Type> ThreePlayerSplitscreenLayout;
 
+	/** The viewport layout to use if the screen should be split and there are three local players */
+	UPROPERTY(config, EditAnywhere, Category=LocalMultiplayer, meta=(editcondition="bUseSplitScreen"))
+	EFourPlayerSplitScreenType FourPlayerSplitscreenLayout;
+
 	/**
 	* If enabled, this will make so that gamepads start being assigned to the second controller ID in local multiplayer games.
 	* In PIE sessions with multiple windows, this has the same effect as enabling "Route 1st Gamepad to 2nd Client"
