[NO NON-PII COMMIT MESSAGE]

diff --git a/Engine/Source/Runtime/Engine/Private/LevelStreaming.cpp b/Engine/Source/Runtime/Engine/Private/LevelStreaming.cpp
--- a/Engine/Source/Runtime/Engine/Private/LevelStreaming.cpp
+++ b/Engine/Source/Runtime/Engine/Private/LevelStreaming.cpp
@@ -945,6 +945,10 @@ bool ULevelStreaming::RequestLevel(UWorld* PersistentWorld, bool bAllowLevelLoad
 				}
 #endif
 
+				// Level already exists but may have the wrong type due to being inactive before, so copy data over
+				World->WorldType = PersistentWorld->WorldType;
+				World->PersistentLevel->OwningWorld = PersistentWorld;
+
 				SetLoadedLevel(World->PersistentLevel);
 				// Broadcast level loaded event to blueprints
 				OnLevelLoaded.Broadcast();
