(4.20.2) Cleanup all loaded worlds, not just those that have been made visible
#rnx

diff --git a/Engine/Source/Runtime/Engine/Private/World.cpp b/Engine/Source/Runtime/Engine/Private/World.cpp
--- a/Engine/Source/Runtime/Engine/Private/World.cpp
+++ b/Engine/Source/Runtime/Engine/Private/World.cpp
@@ -1462,7 +1462,7 @@ void UWorld::DestroyWorld( bool bInformEngineOfWorld, UWorld* NewWorld )
 {
 	// Clean up existing world and remove it from root set so it can be garbage collected.
 	bIsLevelStreamingFrozen = false;
-	bShouldForceUnloadStreamingLevels = true;
+	SetShouldForceUnloadStreamingLevels(true);
 	FlushLevelStreaming();
 	CleanupWorld(true, true, NewWorld);
 
@@ -3800,11 +3800,11 @@ void UWorld::CleanupWorld(bool bSessionEnded, bool bCleanupResources, UWorld* Ne
 	}
 #endif //WITH_EDITOR
 
-	for (int32 LevelIndex=0; LevelIndex < GetNumLevels(); ++LevelIndex)
+	for (ULevelStreaming* StreamingLevel : GetStreamingLevels())
 	{
-		UWorld* World = CastChecked<UWorld>(GetLevel(LevelIndex)->GetOuter());
-		if (World != this)
+		if (ULevel* Level = StreamingLevel->GetLoadedLevel())
 		{
+			UWorld* World = CastChecked<UWorld>(Level->GetOuter());
 			World->CleanupWorld(bSessionEnded, bCleanupResources, NewWorld);
 		}
 	}
