PR #4872: UE-61614: Only update OnscreenDebugMessage times once a frame (Contributed by [REDACTED_PERSON])

diff --git a/Engine/Source/Runtime/Engine/Private/UnrealEngine.cpp b/Engine/Source/Runtime/Engine/Private/UnrealEngine.cpp
--- a/Engine/Source/Runtime/Engine/Private/UnrealEngine.cpp
+++ b/Engine/Source/Runtime/Engine/Private/UnrealEngine.cpp
@@ -9502,6 +9502,8 @@ float DrawMapWarnings(UWorld* World, FViewport* Viewport, FCanvas* Canvas, UCanv
 */
 float DrawOnscreenDebugMessages(UWorld* World, FViewport* Viewport, FCanvas* Canvas, UCanvas* CanvasObject, float MessageX, float MessageY)
 {
+	static TFrameValue<bool> HasUpdatedScreenDebugMessages;
+
 	int32 YPos = MessageY;
 	const int32 MaxYPos = CanvasObject ? CanvasObject->SizeY : 700;
 	if (GEngine->PriorityScreenMessages.Num() > 0)
@@ -9518,10 +9520,13 @@ float DrawOnscreenDebugMessages(UWorld* World, FViewport* Viewport, FCanvas* Can
 				Canvas->DrawItem(MessageTextItem, FVector2D(MessageX, YPos));
 				YPos += MessageTextItem.DrawnSize.Y * 1.15f;
 			}
-			Message.CurrentTimeDisplayed += World->GetDeltaSeconds();
-			if (Message.CurrentTimeDisplayed >= Message.TimeToDisplay)
+			if (!HasUpdatedScreenDebugMessages.IsSet())
 			{
-				GEngine->PriorityScreenMessages.RemoveAt(PrioIndex);
+				Message.CurrentTimeDisplayed += World->GetDeltaSeconds();
+				if (Message.CurrentTimeDisplayed >= Message.TimeToDisplay)
+				{
+					GEngine->PriorityScreenMessages.RemoveAt(PrioIndex);
+				}
 			}
 		}
 	}
@@ -9541,14 +9546,20 @@ float DrawOnscreenDebugMessages(UWorld* World, FViewport* Viewport, FCanvas* Can
 				Canvas->DrawItem(MessageTextItem, FVector2D(MessageX, YPos));
 				YPos += MessageTextItem.DrawnSize.Y * 1.15f;
 			}
-			Message.CurrentTimeDisplayed += World->GetDeltaSeconds();
-			if (Message.CurrentTimeDisplayed >= Message.TimeToDisplay)
+			if (!HasUpdatedScreenDebugMessages.IsSet())
 			{
-				MsgIt.RemoveCurrent();
+				Message.CurrentTimeDisplayed += World->GetDeltaSeconds();
+				if (Message.CurrentTimeDisplayed >= Message.TimeToDisplay)
+				{
+					MsgIt.RemoveCurrent();
+				}
 			}
 		}
 	}
 
+	// Flag variable that the update has already been done this frame
+	HasUpdatedScreenDebugMessages = true;
+
 	return MessageY;
 }
 
@@ -9576,7 +9587,6 @@ void DrawStatsHUD( UWorld* World, FViewport* Viewport, FCanvas* Canvas, UCanvas*
 		return;
 	}
 
-
 	float DPIScale = Canvas->GetDPIScale();
 
 	const FVector2D ScaledViewportSize = FVector2D(Viewport->GetSizeXY()) / DPIScale;
