Interpreter Optimization 5 of N: Move UberGraphFramePointerProperty into UClass to avoid VCalls in blueprint intrepreter

diff --git a/Engine/Source/Developer/BlueprintCompilerCppBackend/Private/BlueprintCompilerCppBackendBase.cpp b/Engine/Source/Developer/BlueprintCompilerCppBackend/Private/BlueprintCompilerCppBackendBase.cpp
--- a/Engine/Source/Developer/BlueprintCompilerCppBackend/Private/BlueprintCompilerCppBackendBase.cpp
+++ b/Engine/Source/Developer/BlueprintCompilerCppBackend/Private/BlueprintCompilerCppBackendBase.cpp
@@ -1209,10 +1209,12 @@ FString FBlueprintCompilerCppBackendBase::GenerateWrapperForClass(UClass* Source
 		// PROPERTIES:
 		for (auto Property : TFieldRange<UProperty>(SourceClass, EFieldIteratorFlags::ExcludeSuper))
 		{
+#if USE_UBER_GRAPH_PERSISTENT_FRAME
 			if (BPGC && (Property == BPGC->UberGraphFramePointerProperty))
 			{
 				continue;
 			}
+#endif //USE_UBER_GRAPH_PERSISTENT_FRAME
 
 			if (Cast<UAnimBlueprintGeneratedClass>(BPGC))
 			{
diff --git a/Engine/Source/Editor/Kismet/Private/BlueprintCompilationManager.cpp b/Engine/Source/Editor/Kismet/Private/BlueprintCompilationManager.cpp
--- a/Engine/Source/Editor/Kismet/Private/BlueprintCompilationManager.cpp
+++ b/Engine/Source/Editor/Kismet/Private/BlueprintCompilationManager.cpp
@@ -418,8 +418,10 @@ void FBlueprintCompilationManagerImpl::FlushCompilationQueueImpl(TArray<UObject*
 				{
 					bDefaultComponentMustBeAdded = true;
 				}
-				
+
+#if USE_UBER_GRAPH_PERSISTENT_FRAME
 				bHasPendingUberGraphFrame = BPGC->UberGraphFramePointerProperty || BPGC->UberGraphFunction;
+#endif//USE_UBER_GRAPH_PERSISTENT_FRAME
 			}
 			
 			if( FBlueprintEditorUtils::IsDataOnlyBlueprint(QueuedBP) && !QueuedBP->bHasBeenRegenerated && QueuedBP->GetLinker() && !bDefaultComponentMustBeAdded && !bHasPendingUberGraphFrame)
diff --git a/Engine/Source/Editor/KismetCompiler/Private/KismetCompilerMisc.cpp b/Engine/Source/Editor/KismetCompiler/Private/KismetCompilerMisc.cpp
--- a/Engine/Source/Editor/KismetCompiler/Private/KismetCompilerMisc.cpp
+++ b/Engine/Source/Editor/KismetCompiler/Private/KismetCompilerMisc.cpp
@@ -1096,11 +1096,14 @@ UProperty* FKismetCompilerUtilities::CreatePropertyOnScope(UStruct* Scope, const
 	// Check to see if there's already a object on this scope with the same name, and throw an internal compiler error if so
 	// If this happens, it breaks the property link, which causes stack corruption and hard-to-track errors, so better to fail at this point
 	{
+#if !USE_UBER_GRAPH_PERSISTENT_FRAME
+	#error "Without the uber graph frame we will intentionally create properties with conflicting names on the same scope - disable this error at your own risk"
+#else
 		if (UObject* ExistingObject = CheckPropertyNameOnScope(Scope, PropertyName))
 		{
 			const FString ScopeName((Scope != nullptr) ? Scope->GetName() : FString(TEXT("None")));
 			const FString ExistingTypeAndPath(ExistingObject->GetFullName(Scope));
-			MessageLog.Error(*FString::Printf(TEXT("Internal Compiler Error: Tried to create a property %s in scope %s, but another object (%s) already already exists there."), *PropertyName.ToString(), *ScopeName, *ExistingTypeAndPath));
+			MessageLog.Error(*FString::Printf(TEXT("Internal Compiler Error: Tried to create a property %s in scope %s, but another object (%s) already exists there."), *PropertyName.ToString(), *ScopeName, *ExistingTypeAndPath));
 
 			// Find a free name, so we can still create the property to make it easier to spot the duplicates, and avoid crashing
 			uint32 Counter = 0;
@@ -1114,6 +1117,7 @@ UProperty* FKismetCompilerUtilities::CreatePropertyOnScope(UStruct* Scope, const
 
 			ValidatedPropertyName = TestName;
 		}
+#endif
 	}
 
 	UProperty* NewProperty = nullptr;
diff --git a/Engine/Source/Editor/UnrealEd/Private/Kismet2/BlueprintEditorUtils.cpp b/Engine/Source/Editor/UnrealEd/Private/Kismet2/BlueprintEditorUtils.cpp
--- a/Engine/Source/Editor/UnrealEd/Private/Kismet2/BlueprintEditorUtils.cpp
+++ b/Engine/Source/Editor/UnrealEd/Private/Kismet2/BlueprintEditorUtils.cpp
@@ -1348,8 +1348,12 @@ UClass* FBlueprintEditorUtils::RegenerateBlueprintClass(UBlueprint* Blueprint, U
 			&& (!ClassToRegenerate || (Blueprint->ParentClass != ClassToRegenerate->GetSuperClass()));
 
 		UBlueprintGeneratedClass* BPGClassToRegenerate = Cast<UBlueprintGeneratedClass>(ClassToRegenerate);
+#if USE_UBER_GRAPH_PERSISTENT_FRAME
 		const bool bHasPendingUberGraphFrame = BPGClassToRegenerate
 			&& (BPGClassToRegenerate->UberGraphFramePointerProperty || BPGClassToRegenerate->UberGraphFunction);
+#else 
+		const bool bHasPendingUberGraphFrame = false;
+#endif //USE_UBER_GRAPH_PERSISTENT_FRAME
 
 		const bool bDefaultComponentMustBeAdded = !bHasCode 
 			&& BPGClassToRegenerate
diff --git a/Engine/Source/Editor/UnrealEd/Private/Kismet2/KismetDebugUtilities.cpp b/Engine/Source/Editor/UnrealEd/Private/Kismet2/KismetDebugUtilities.cpp
--- a/Engine/Source/Editor/UnrealEd/Private/Kismet2/KismetDebugUtilities.cpp
+++ b/Engine/Source/Editor/UnrealEd/Private/Kismet2/KismetDebugUtilities.cpp
@@ -1215,10 +1215,10 @@ FKismetDebugUtilities::EWatchTextResult FKismetDebugUtilities::FindDebuggingData
 			UFunction* OuterFunction = Cast<UFunction>(Property->GetOuter());
 			if (!PropertyBase && OuterFunction)
 			{
-				UBlueprintGeneratedClass* BPGC = Cast<UBlueprintGeneratedClass>(Blueprint->GeneratedClass);
+				UBlueprintGeneratedClass* BPGC = Cast<UBlueprintGeneratedClass>(OuterFunction->GetOuter());
 				if (BPGC && ActiveObject->IsA(BPGC))
 				{
-					PropertyBase = BPGC->GetPersistentUberGraphFrame(ActiveObject, OuterFunction);
+					PropertyBase = GetPersistentUberGraphFrame( OuterFunction, ActiveObject );
 				}
 			}
 #endif // USE_UBER_GRAPH_PERSISTENT_FRAME
diff --git a/Engine/Source/Runtime/CoreUObject/Private/UObject/Class.cpp b/Engine/Source/Runtime/CoreUObject/Private/UObject/Class.cpp
--- a/Engine/Source/Runtime/CoreUObject/Private/UObject/Class.cpp
+++ b/Engine/Source/Runtime/CoreUObject/Private/UObject/Class.cpp
@@ -3974,6 +3974,10 @@ void UClass::PurgeClass(bool bRecompilingOnLoad)
 	}
 #endif
 
+#if USE_UBER_GRAPH_PERSISTENT_FRAME
+	UberGraphFramePointerProperty = NULL;
+#endif//USE_UBER_GRAPH_PERSISTENT_FRAME
+
 	ClassDefaultObject = NULL;
 
 	Interfaces.Empty();
@@ -4054,6 +4058,9 @@ UClass::UClass(const FObjectInitializer& ObjectInitializer)
 ,	ClassCastFlags(CASTCLASS_None)
 ,	ClassWithin( UObject::StaticClass() )
 ,	ClassGeneratedBy(nullptr)
+#if USE_UBER_GRAPH_PERSISTENT_FRAME
+,	UberGraphFramePointerProperty(nullptr)
+#endif // USE_UBER_GRAPH_PERSISTENT_FRAME
 ,	ClassDefaultObject(nullptr)
 {
 	// If you add properties here, please update the other constructors and PurgeClass()
@@ -4072,6 +4079,9 @@ UClass::UClass(const FObjectInitializer& ObjectInitializer, UClass* InBaseClass)
 ,	ClassCastFlags(CASTCLASS_None)
 ,	ClassWithin(UObject::StaticClass())
 ,	ClassGeneratedBy(nullptr)
+#if USE_UBER_GRAPH_PERSISTENT_FRAME
+,	UberGraphFramePointerProperty(nullptr)
+#endif // USE_UBER_GRAPH_PERSISTENT_FRAME
 ,	ClassDefaultObject(nullptr)
 {
 	// If you add properties here, please update the other constructors and PurgeClass()
diff --git a/Engine/Source/Runtime/CoreUObject/Private/UObject/ScriptCore.cpp b/Engine/Source/Runtime/CoreUObject/Private/UObject/ScriptCore.cpp
--- a/Engine/Source/Runtime/CoreUObject/Private/UObject/ScriptCore.cpp
+++ b/Engine/Source/Runtime/CoreUObject/Private/UObject/ScriptCore.cpp
@@ -50,6 +50,18 @@ static FAutoConsoleVariableRef CVarVerboseScriptStats(
 	ECVF_Default
 );
 
+#if USE_UBER_GRAPH_PERSISTENT_FRAME
+// Mirror definition of FPointerToUberGraphFrame, it is a UStruct and
+// we cannot easily generate its reflection data here in CoreUObject. The
+// builtins pattern we use for FVector, FQuat etc cannot be used because
+// our only member is a raw pointer and it cannot be a UProperty. This
+// creates difficulty in determining the correct size for the UStruct
+struct FPointerToUberGraphFrameCoreUObject
+{
+	uint8* RawPointer;
+};
+#endif //USE_UBER_GRAPH_PERSISTENT_FRAME
+
 /*-----------------------------------------------------------------------------
 	Globals.
 -----------------------------------------------------------------------------*/
@@ -310,6 +322,33 @@ FString UnicodeToCPPIdentifier(const FString& InName, bool bDeprecated, const TC
 	return bDeprecated ? Ret + TEXT("_DEPRECATED") : Ret;
 }
 
+#if USE_UBER_GRAPH_PERSISTENT_FRAME
+/** Returns memory used to store temporary data on an instance, used by blueprints */
+static uint8* GetPersistentUberGraphFrameUnchecked(const UFunction* ForFn, UObject* Obj)
+{
+	const UClass* FromClass = ForFn->GetOuterUClassUnchecked();
+	checkSlow(ForFn->HasAnyFunctionFlags(FUNC_UbergraphFunction));
+	checkSlow(Obj->IsA(FromClass));
+	checkSlow(FromClass->UberGraphFramePointerProperty);
+	FPointerToUberGraphFrameCoreUObject* PointerToUberGraphFrame =
+		FromClass->UberGraphFramePointerProperty->ContainerPtrToValuePtr<FPointerToUberGraphFrameCoreUObject>(
+			(void*)Obj
+		);
+	checkSlow(PointerToUberGraphFrame);
+	checkSlow(PointerToUberGraphFrame->RawPointer);
+	return PointerToUberGraphFrame->RawPointer;
+}
+
+uint8* GetPersistentUberGraphFrame(const UFunction* ForFn, UObject* Obj)
+{
+	if (ForFn->HasAnyFunctionFlags(FUNC_UbergraphFunction))
+	{
+		return GetPersistentUberGraphFrameUnchecked(ForFn, Obj);
+	}
+	return nullptr;
+}
+#endif //USE_UBER_GRAPH_PERSISTENT_FRAME
+
 /*-----------------------------------------------------------------------------
 	FFrame implementation.
 -----------------------------------------------------------------------------*/
@@ -783,7 +822,7 @@ void UObject::CallFunction( FFrame& Stack, RESULT_DECL, UFunction* Function )
 	{
 		uint8* Frame = NULL;
 #if USE_UBER_GRAPH_PERSISTENT_FRAME
-		Frame = Function->GetOuterUClassUnchecked()->GetPersistentUberGraphFrame(this, Function);
+		Frame = GetPersistentUberGraphFrame(Function, this);
 #endif
 		const bool bUsePersistentFrame = (NULL != Frame);
 		if (!bUsePersistentFrame)
@@ -1361,7 +1400,7 @@ void UObject::ProcessEvent( UFunction* Function, void* Parms )
 	{
 		uint8* Frame = NULL;
 #if USE_UBER_GRAPH_PERSISTENT_FRAME
-		Frame = Function->GetOuterUClassUnchecked()->GetPersistentUberGraphFrame(this, Function);
+		Frame = GetPersistentUberGraphFrame(Function, this);
 #endif
 		const bool bUsePersistentFrame = (NULL != Frame);
 		if (!bUsePersistentFrame)
@@ -1901,7 +1940,7 @@ DEFINE_FUNCTION(UObject::execLetValueOnPersistentFrame)
 	checkSlow(DestProperty);
 	UFunction* UberGraphFunction = CastChecked<UFunction>(DestProperty->GetOwnerStruct());
 	checkSlow(Stack.Object->GetClass()->IsChildOf(UberGraphFunction->GetOuterUClassUnchecked()));
-	uint8* FrameBase = UberGraphFunction->GetOuterUClassUnchecked()->GetPersistentUberGraphFrame(Stack.Object, UberGraphFunction);
+	uint8* FrameBase = GetPersistentUberGraphFrameUnchecked(UberGraphFunction, Stack.Object);
 	checkSlow(FrameBase);
 	uint8* DestAddress = DestProperty->ContainerPtrToValuePtr<uint8>(FrameBase);
 
diff --git a/Engine/Source/Runtime/CoreUObject/Public/UObject/Class.h b/Engine/Source/Runtime/CoreUObject/Public/UObject/Class.h
--- a/Engine/Source/Runtime/CoreUObject/Public/UObject/Class.h
+++ b/Engine/Source/Runtime/CoreUObject/Public/UObject/Class.h
@@ -2107,6 +2107,16 @@ public:
 	/** This is the blueprint that caused the generation of this class, or null if it is a native compiled-in class */
 	UObject* ClassGeneratedBy;
 
+#if USE_UBER_GRAPH_PERSISTENT_FRAME
+	/**
+	 * Property that points to the ubergraph frame, this is a blueprint specific structure that has been hoisted
+	 * to UClass so that the interpreter (ScriptCore.cpp) can access it efficiently. The uber graph frame is a struct
+	 * owned by a UObject but allocated separately that has a layout that corresponds to a specific UFunction (the
+	 * UberGraphFunction) in a blueprint.
+	 */
+	UStructProperty* UberGraphFramePointerProperty;
+#endif //USE_UBER_GRAPH_PERSISTENT_FRAME
+
 #if WITH_EDITOR
 	/**
 	 * Conditionally recompiles the class after loading, in case any dependencies were also newly loaded
@@ -2401,12 +2411,6 @@ public:
 	 */
 	FName GetDefaultObjectName();
 
-	/** Returns memory used to store temporary data on an instance, used by blueprints */
-	virtual uint8* GetPersistentUberGraphFrame(UObject* Obj, UFunction* FuncToCheck) const
-	{
-		return nullptr;
-	}
-
 	/** Creates memory to store temporary data */
 	virtual void CreatePersistentUberGraphFrame(UObject* Obj, bool bCreateOnlyIfEmpty = false, bool bSkipSuperClass = false, UClass* OldClass = nullptr) const
 	{
diff --git a/Engine/Source/Runtime/CoreUObject/Public/UObject/Script.h b/Engine/Source/Runtime/CoreUObject/Public/UObject/Script.h
--- a/Engine/Source/Runtime/CoreUObject/Public/UObject/Script.h
+++ b/Engine/Source/Runtime/CoreUObject/Public/UObject/Script.h
@@ -121,7 +121,7 @@ enum EFunctionFlags : uint32
 	FUNC_NetResponse		= 0x00001000,   // Function response from a net service
 	FUNC_Static				= 0x00002000,   // Static function.
 	FUNC_NetMulticast		= 0x00004000,	// Function is networked multicast Server -> All Clients
-	// FUNC_				= 0x00008000,   // unused.
+	FUNC_UbergraphFunction	= 0x00008000,   // Function is used as the merge 'ubergraph' for a blueprint, only assigned when using the persistent 'ubergraph' frame
 	FUNC_MulticastDelegate	= 0x00010000,	// Function is a multi-cast delegate signature (also requires FUNC_Delegate to be set!)
 	FUNC_Public				= 0x00020000,	// Function is accessible in all classes (if overridden, parameters must remain unchanged).
 	FUNC_Private			= 0x00040000,	// Function is accessible only in the class it is defined in (cannot be overridden, but function name may be reused in subclasses.  IOW: if overridden, parameters don't need to match, and Super.Func() cannot be accessed since it's private.)
@@ -552,3 +552,13 @@ COREUOBJECT_API FString ToValidCPPIdentifierChars(TCHAR Char);
 */
 COREUOBJECT_API FString UnicodeToCPPIdentifier(const FString& InName, bool bDeprecated, const TCHAR* Prefix);
 
+/**
+	@param ForFn Function to search for a corresponding uber graph frame
+	@param Obj owning the uber graph frame, Obj must be an instance based on ForFn->GetOuterUClass()
+	@return A pointer to the start of the objects persistent uber graph
+		frame if ForFn is an ubergraph function, else nullptr. The uber graph
+		frame is a struct that has a matching layout to a blueprint's UberGraphFunction's
+		parameters. The Uber Graph frame is an optimization because it avoids a large
+		allocation and corresponding initialization when calling the UberGraphFunction
+*/
+COREUOBJECT_API uint8* GetPersistentUberGraphFrame(const UFunction* ForFn, UObject* Obj);
\ No newline at end of file
diff --git a/Engine/Source/Runtime/Engine/Classes/Engine/BlueprintGeneratedClass.h b/Engine/Source/Runtime/Engine/Classes/Engine/BlueprintGeneratedClass.h
--- a/Engine/Source/Runtime/Engine/Classes/Engine/BlueprintGeneratedClass.h
+++ b/Engine/Source/Runtime/Engine/Classes/Engine/BlueprintGeneratedClass.h
@@ -76,16 +76,21 @@ public:
 	}
 };
 
-
+//////////////////////////////////////////////////////////////////////////
+// WARNING: Following struct layout definition repeated in ScriptCore.cpp as
+// FPointerToUberGraphFrameCoreUObject to work around reflection generation issues:
 USTRUCT()
 struct FPointerToUberGraphFrame
 {
 	GENERATED_USTRUCT_BODY()
 
 public:
+	//////////////////////////////////////////////////////////////////////////
+	// WARNING: This struct layout definition repeated in ScriptCore.cpp as
+	// FPointerToUberGraphFrameCoreUObject to work around reflection generation issues:
 	uint8* RawPointer;
 
-	FPointerToUberGraphFrame() 
+	FPointerToUberGraphFrame()
 		: RawPointer(nullptr)
 	{}
 
@@ -94,7 +99,9 @@ public:
 		check(!RawPointer);
 	}
 };
-
+// WARNING: Preceding struct layout definition repeated in ScriptCore.cpp as
+// FPointerToUberGraphFrameCoreUObject to work around reflection generation issues!
+//////////////////////////////////////////////////////////////////////////
 
 template<>
 struct TStructOpsTypeTraits<FPointerToUberGraphFrame> : public TStructOpsTypeTraitsBase2<FPointerToUberGraphFrame>
@@ -617,9 +624,6 @@ public:
 	UPROPERTY()
 	class UInheritableComponentHandler* InheritableComponentHandler;
 
-	UPROPERTY()
-	UStructProperty* UberGraphFramePointerProperty;
-
 	UPROPERTY()
 	UFunction* UberGraphFunction;
 
@@ -687,7 +691,6 @@ public:
 	virtual void SerializeDefaultObject(UObject* Object, FArchive& Ar) override;
 	virtual void PostLoadDefaultObject(UObject* Object) override;
 	virtual bool IsFunctionImplementedInBlueprint(FName InFunctionName) const override;
-	virtual uint8* GetPersistentUberGraphFrame(UObject* Obj, UFunction* FuncToCheck) const override;
 	virtual void CreatePersistentUberGraphFrame(UObject* Obj, bool bCreateOnlyIfEmpty = false, bool bSkipSuperClass = false, UClass* OldClass = nullptr) const override;
 	virtual void DestroyPersistentUberGraphFrame(UObject* Obj, bool bSkipSuperClass = false) const override;
 	virtual void Link(FArchive& Ar, bool bRelinkExistingProperties) override;
diff --git a/Engine/Source/Runtime/Engine/Private/BlueprintGeneratedClass.cpp b/Engine/Source/Runtime/Engine/Private/BlueprintGeneratedClass.cpp
--- a/Engine/Source/Runtime/Engine/Private/BlueprintGeneratedClass.cpp
+++ b/Engine/Source/Runtime/Engine/Private/BlueprintGeneratedClass.cpp
@@ -1189,25 +1189,9 @@ void UBlueprintGeneratedClass::CheckAndApplyComponentTemplateOverrides(AActor* A
 	}
 }
 
-uint8* UBlueprintGeneratedClass::GetPersistentUberGraphFrame(UObject* Obj, UFunction* FuncToCheck) const
-{
-	if (Obj && UsePersistentUberGraphFrame() && UberGraphFramePointerProperty && UberGraphFunction)
-	{
-		if (UberGraphFunction == FuncToCheck)
-		{
-			FPointerToUberGraphFrame* PointerToUberGraphFrame = UberGraphFramePointerProperty->ContainerPtrToValuePtr<FPointerToUberGraphFrame>(Obj);
-			checkSlow(PointerToUberGraphFrame);
-			ensure(PointerToUberGraphFrame->RawPointer);
-			return PointerToUberGraphFrame->RawPointer;
-		}
-	}
-	UClass* ParentClass = GetSuperClass();
-	checkSlow(ParentClass);
-	return ParentClass->GetPersistentUberGraphFrame(Obj, FuncToCheck);
-}
-
 void UBlueprintGeneratedClass::CreatePersistentUberGraphFrame(UObject* Obj, bool bCreateOnlyIfEmpty, bool bSkipSuperClass, UClass* OldClass) const
 {
+#if USE_UBER_GRAPH_PERSISTENT_FRAME
 	/** Macros should not create uber graph frames as they have no uber graph. If UBlueprints are cooked out the macro class probably does not exist as well */
 	UBlueprint* Blueprint = Cast<UBlueprint>(ClassGeneratedBy);
 	if (Blueprint && Blueprint->BlueprintType == BPTYPE_MacroLibrary)
@@ -1261,10 +1245,12 @@ void UBlueprintGeneratedClass::CreatePersistentUberGraphFrame(UObject* Obj, bool
 		checkSlow(ParentClass);
 		ParentClass->CreatePersistentUberGraphFrame(Obj, bCreateOnlyIfEmpty);
 	}
+#endif // USE_UBER_GRAPH_PERSISTENT_FRAME
 }
 
 void UBlueprintGeneratedClass::DestroyPersistentUberGraphFrame(UObject* Obj, bool bSkipSuperClass) const
 {
+#if USE_UBER_GRAPH_PERSISTENT_FRAME
 	ensure(!UberGraphFramePointerProperty == !UberGraphFunction);
 	if (Obj && UsePersistentUberGraphFrame() && UberGraphFramePointerProperty && UberGraphFunction)
 	{
@@ -1293,6 +1279,7 @@ void UBlueprintGeneratedClass::DestroyPersistentUberGraphFrame(UObject* Obj, boo
 		checkSlow(ParentClass);
 		ParentClass->DestroyPersistentUberGraphFrame(Obj);
 	}
+#endif // USE_UBER_GRAPH_PERSISTENT_FRAME
 }
 
 void UBlueprintGeneratedClass::GetPreloadDependencies(TArray<UObject*>& OutDeps)
@@ -1388,6 +1375,7 @@ void UBlueprintGeneratedClass::Link(FArchive& Ar, bool bRelinkExistingProperties
 {
 	Super::Link(Ar, bRelinkExistingProperties);
 
+#if USE_UBER_GRAPH_PERSISTENT_FRAME
 	if (UsePersistentUberGraphFrame() && UberGraphFunction)
 	{
 		Ar.Preload(UberGraphFunction);
@@ -1402,6 +1390,7 @@ void UBlueprintGeneratedClass::Link(FArchive& Ar, bool bRelinkExistingProperties
 		}
 		checkSlow(UberGraphFramePointerProperty);
 	}
+#endif
 
 	AssembleReferenceTokenStream(true);
 }
@@ -1410,7 +1399,6 @@ void UBlueprintGeneratedClass::PurgeClass(bool bRecompilingOnLoad)
 {
 	Super::PurgeClass(bRecompilingOnLoad);
 
-	UberGraphFramePointerProperty = NULL;
 	UberGraphFunction = NULL;
 #if WITH_EDITORONLY_DATA
 	OverridenArchetypeForCDO = NULL;
@@ -1438,6 +1426,7 @@ void UBlueprintGeneratedClass::AddReferencedObjectsInUbergraphFrame(UObject* InT
 	{
 		if (UBlueprintGeneratedClass* BPGC = Cast<UBlueprintGeneratedClass>(CurrentClass))
 		{
+#if USE_UBER_GRAPH_PERSISTENT_FRAME
 			if (BPGC->UberGraphFramePointerProperty)
 			{
 				FPointerToUberGraphFrame* PointerToUberGraphFrame = BPGC->UberGraphFramePointerProperty->ContainerPtrToValuePtr<FPointerToUberGraphFrame>(InThis);
@@ -1454,6 +1443,7 @@ void UBlueprintGeneratedClass::AddReferencedObjectsInUbergraphFrame(UObject* InT
 					BPGC->UberGraphFunction->SerializeBin(CollectorScope.GetArchive(), PointerToUberGraphFrame->RawPointer);
 				}
 			}
+#endif // USE_UBER_GRAPH_PERSISTENT_FRAME
 		}
 		else if (CurrentClass->HasAllClassFlags(CLASS_Native))
 		{
