Const correct USceneComponent::IsAttachedTo
#rnx

diff --git a/Engine/Source/Runtime/Engine/Classes/Components/SceneComponent.h b/Engine/Source/Runtime/Engine/Classes/Components/SceneComponent.h
--- a/Engine/Source/Runtime/Engine/Classes/Components/SceneComponent.h
+++ b/Engine/Source/Runtime/Engine/Classes/Components/SceneComponent.h
@@ -1,4 +1,4 @@
-﻿// Copyright 1998-2018 Epic Games, Inc. All Rights Reserved.
+// Copyright 1998-2018 Epic Games, Inc. All Rights Reserved.
 
 #pragma once
 
@@ -1237,7 +1237,7 @@ public:
 	AActor* GetAttachmentRootActor() const;
 
 	/** Walks up the attachment chain to see if this component is attached to the supplied component. If TestComp == this, returns false.*/
-	bool IsAttachedTo(USceneComponent* TestComp) const;
+	bool IsAttachedTo(const USceneComponent* TestComp) const;
 
 	/**
 	 * Find the world-space location and rotation of the given named socket.
diff --git a/Engine/Source/Runtime/Engine/Private/Components/SceneComponent.cpp b/Engine/Source/Runtime/Engine/Private/Components/SceneComponent.cpp
--- a/Engine/Source/Runtime/Engine/Private/Components/SceneComponent.cpp
+++ b/Engine/Source/Runtime/Engine/Private/Components/SceneComponent.cpp
@@ -2166,7 +2166,7 @@ AActor* USceneComponent::GetAttachmentRootActor() const
 	return AttachmentRootComponent ? AttachmentRootComponent->GetOwner() : nullptr;
 }
 
-bool USceneComponent::IsAttachedTo(class USceneComponent* TestComp) const
+bool USceneComponent::IsAttachedTo(const USceneComponent* TestComp) const
 {
 	if(TestComp != nullptr)
 	{
