Fix crash trying to create an object when one with that name already exists for child actor template.
Per: https://udn.unrealengine.com/questions/448897/crash-when-loading-blueprint-class.html

diff --git a/Engine/Source/Runtime/Engine/Private/Components/ChildActorComponent.cpp b/Engine/Source/Runtime/Engine/Private/Components/ChildActorComponent.cpp
--- a/Engine/Source/Runtime/Engine/Private/Components/ChildActorComponent.cpp
+++ b/Engine/Source/Runtime/Engine/Private/Components/ChildActorComponent.cpp
@@ -118,7 +118,14 @@ void UChildActorComponent::Serialize(FArchive& Ar)
 			if (ChildActorTemplate->GetOuter() != this)
 			{
 				const FString TemplateName = FString::Printf(TEXT("%s_%s_CAT"), *GetName(), *ChildActorClass->GetName());
-				ChildActorTemplate = CastChecked<AActor>(StaticDuplicateObject(ChildActorTemplate, this, *TemplateName));
+				if (UObject* ExistingTemplate = StaticFindObject(nullptr, this, *TemplateName))
+				{
+					ChildActorTemplate = CastChecked<AActor>(ExistingTemplate);
+				}
+				else
+				{
+					ChildActorTemplate = CastChecked<AActor>(StaticDuplicateObject(ChildActorTemplate, this, *TemplateName));
+				}
 			}
 		}
 		else
