Use GetComponents() directly
Convert to weak object pointer once for AddUnique and Remove
#rnx

diff --git a/Engine/Source/Runtime/Engine/Private/Components/SceneCaptureComponent.cpp b/Engine/Source/Runtime/Engine/Private/Components/SceneCaptureComponent.cpp
--- a/Engine/Source/Runtime/Engine/Private/Components/SceneCaptureComponent.cpp
+++ b/Engine/Source/Runtime/Engine/Private/Components/SceneCaptureComponent.cpp
@@ -212,7 +212,8 @@ void USceneCaptureComponent::HideComponent(UPrimitiveComponent* InComponent)
 {
 	if (InComponent)
 	{
-		HiddenComponents.AddUnique(InComponent);
+		TWeakObjectPtr<UPrimitiveComponent> WeakComponent(InComponent);
+		HiddenComponents.AddUnique(WeakComponent);
 	}
 }
 
@@ -220,11 +221,13 @@ void USceneCaptureComponent::HideActorComponents(AActor* InActor)
 {
 	if (InActor)
 	{
-		TInlineComponentArray<UPrimitiveComponent*> PrimitiveComponents;
-		InActor->GetComponents(PrimitiveComponents);
-		for (int32 ComponentIndex = 0, NumComponents = PrimitiveComponents.Num(); ComponentIndex < NumComponents; ++ComponentIndex)
+		for (UActorComponent* Component : InActor->GetComponents())
 		{
-			HiddenComponents.AddUnique(PrimitiveComponents[ComponentIndex]);
+			if (UPrimitiveComponent* PrimComp = Cast<UPrimitiveComponent>(Component))
+			{
+				TWeakObjectPtr<UPrimitiveComponent> WeakComponent(PrimComp);
+				HiddenComponents.AddUnique(WeakComponent);
+			}
 		}
 	}
 }
@@ -246,29 +249,33 @@ void USceneCaptureComponent::ShowOnlyActorComponents(AActor* InActor)
 		// Backward compatibility - set PrimitiveRenderMode to PRM_UseShowOnlyList if BP / game code tries to add a ShowOnlyComponent
 		PrimitiveRenderMode = ESceneCapturePrimitiveRenderMode::PRM_UseShowOnlyList;
 
-		TInlineComponentArray<UPrimitiveComponent*> PrimitiveComponents;
-		InActor->GetComponents(PrimitiveComponents);
-		for (int32 ComponentIndex = 0, NumComponents = PrimitiveComponents.Num(); ComponentIndex < NumComponents; ++ComponentIndex)
+		for (UActorComponent* Component : InActor->GetComponents())
 		{
-			ShowOnlyComponents.Add(PrimitiveComponents[ComponentIndex]);
+			if (UPrimitiveComponent* PrimComp = Cast<UPrimitiveComponent>(Component))
+			{
+				ShowOnlyComponents.Add(PrimComp);
+			}
 		}
 	}
 }
 
 void USceneCaptureComponent::RemoveShowOnlyComponent(UPrimitiveComponent* InComponent)
 {
-	ShowOnlyComponents.Remove(InComponent);
+	TWeakObjectPtr<UPrimitiveComponent> WeakComponent(InComponent);
+	ShowOnlyComponents.Remove(WeakComponent);
 }
 
 void USceneCaptureComponent::RemoveShowOnlyActorComponents(AActor* InActor)
 {
 	if (InActor)
 	{
-		TInlineComponentArray<UPrimitiveComponent*> PrimitiveComponents;
-		InActor->GetComponents(PrimitiveComponents);
-		for (int32 ComponentIndex = 0, NumComponents = PrimitiveComponents.Num(); ComponentIndex < NumComponents; ++ComponentIndex)
+		for (UActorComponent* Component : InActor->GetComponents())
 		{
-			ShowOnlyComponents.Remove(PrimitiveComponents[ComponentIndex]);
+			if (UPrimitiveComponent* PrimComp = Cast<UPrimitiveComponent>(Component))
+			{
+				TWeakObjectPtr<UPrimitiveComponent> WeakComponent(PrimComp);
+				ShowOnlyComponents.Remove(WeakComponent);
+			}
 		}
 	}
 }
