Use components set directly rather than incurring overhead of filtering and copying in to external storage and then using it.
#rnx

diff --git a/Engine/Plugins/2D/Paper2D/Source/Paper2DEditor/Private/GroupedSprites/PaperGroupedSpriteUtilities.cpp b/Engine/Plugins/2D/Paper2D/Source/Paper2DEditor/Private/GroupedSprites/PaperGroupedSpriteUtilities.cpp
--- a/Engine/Plugins/2D/Paper2D/Source/Paper2DEditor/Private/GroupedSprites/PaperGroupedSpriteUtilities.cpp
+++ b/Engine/Plugins/2D/Paper2D/Source/Paper2DEditor/Private/GroupedSprites/PaperGroupedSpriteUtilities.cpp
@@ -34,13 +34,10 @@ void FPaperGroupedSpriteUtilities::BuildHarvestList(const TArray<UObject*>& Obje
 
 		if (AActor* SelectedActor = Cast<AActor>(Object))
 		{
-			TArray<UActorComponent*> Components;
-			SelectedActor->GetComponents(Components);
-
 			bool bHadHarvestableComponents = false;
-			for (UActorComponent* Component : Components)
+			for (UActorComponent* Component : SelectedActor->GetComponents())
 			{
-				if (Component->IsA(HarvestClassType) && !Component->IsEditorOnly())
+				if (Component && Component->IsA(HarvestClassType) && !Component->IsEditorOnly())
 				{
 					bHadHarvestableComponents = true;
 					OutComponentsToHarvest.Add(Component);
diff --git a/Engine/Plugins/2D/Paper2D/Source/Paper2DEditor/Private/LevelEditorMenuExtensions/Paper2DLevelEditorExtensions.cpp b/Engine/Plugins/2D/Paper2D/Source/Paper2DEditor/Private/LevelEditorMenuExtensions/Paper2DLevelEditorExtensions.cpp
--- a/Engine/Plugins/2D/Paper2D/Source/Paper2DEditor/Private/LevelEditorMenuExtensions/Paper2DLevelEditorExtensions.cpp
+++ b/Engine/Plugins/2D/Paper2D/Source/Paper2DEditor/Private/LevelEditorMenuExtensions/Paper2DLevelEditorExtensions.cpp
@@ -92,18 +92,18 @@ public:
 
 		for (AActor* Actor : SelectedActors)
 		{
-			TInlineComponentArray<UActorComponent*> ActorComponents;
-			Actor->GetComponents(ActorComponents);
-
-			for (UActorComponent* Component : ActorComponents)
+			for (UActorComponent* Component : Actor->GetComponents())
 			{
-				if (Component->IsA(UPaperSpriteComponent::StaticClass()))
-				{
-					bCanMergeSprites = true;
-				}
-				else if (Component->IsA(UPaperGroupedSpriteComponent::StaticClass()))
+				if (Component)
 				{
-					bCanSplitSprites = true;
+					if (Component->IsA(UPaperSpriteComponent::StaticClass()))
+					{
+						bCanMergeSprites = true;
+					}
+					else if (Component->IsA(UPaperGroupedSpriteComponent::StaticClass()))
+					{
+						bCanSplitSprites = true;
+					}
 				}
 			}
 
diff --git a/Engine/Plugins/Experimental/GeometryCache/Source/GeometryCacheSequencer/Classes/GeometryCacheTrackEditor.cpp b/Engine/Plugins/Experimental/GeometryCache/Source/GeometryCacheSequencer/Classes/GeometryCacheTrackEditor.cpp
--- a/Engine/Plugins/Experimental/GeometryCache/Source/GeometryCacheSequencer/Classes/GeometryCacheTrackEditor.cpp
+++ b/Engine/Plugins/Experimental/GeometryCache/Source/GeometryCacheSequencer/Classes/GeometryCacheTrackEditor.cpp
@@ -43,13 +43,9 @@ static UGeometryCacheComponent* AcquireGeometryCacheFromObjectGuid(const FGuid&
 
 	if (AActor* Actor = Cast<AActor>(BoundObject))
 	{
-		TInlineComponentArray<UGeometryCacheComponent*> GeometryMeshComponents;
-		Actor->GetComponents(GeometryMeshComponents);
-
-		for (int32 j = 0; j <GeometryMeshComponents.Num(); ++j)
+		for (UActorComponent* Component : Actor->GetComponents())
 		{
-			UGeometryCacheComponent* GeometryMeshComp = GeometryMeshComponents[j];
-			if (GeometryMeshComp)
+			if (UGeometryCacheComponent* GeometryMeshComp = Cast<UGeometryCacheComponent>(Component))
 			{
 				return GeometryMeshComp;
 			}
diff --git a/Engine/Plugins/Runtime/MixedRealityCaptureFramework/Source/MixedRealityCaptureFramework/Private/MrcUtilLibrary.cpp b/Engine/Plugins/Runtime/MixedRealityCaptureFramework/Source/MixedRealityCaptureFramework/Private/MrcUtilLibrary.cpp
--- a/Engine/Plugins/Runtime/MixedRealityCaptureFramework/Source/MixedRealityCaptureFramework/Private/MrcUtilLibrary.cpp
+++ b/Engine/Plugins/Runtime/MixedRealityCaptureFramework/Source/MixedRealityCaptureFramework/Private/MrcUtilLibrary.cpp
@@ -136,15 +136,15 @@ UCameraComponent* UMrcUtilLibrary::GetHMDCameraComponent(const APawn* PlayerPawn
 	UCameraComponent* HMDCam = nullptr;
 	if (PlayerPawn)
 	{
-		TArray<UCameraComponent*> CameraComponents;
-		PlayerPawn->GetComponents(CameraComponents);
-
-		for (UCameraComponent* Camera : CameraComponents)
+		for (UActorComponent* Component : PlayerPawn->GetComponents())
 		{
-			if (Camera->bLockToHmd)
+			if (UCameraComponent* Camera = Cast<UCameraComponent>(Component))
 			{
-				HMDCam = Camera;
-				break;
+				if (Camera->bLockToHmd)
+				{
+					HMDCam = Camera;
+					break;
+				}
 			}
 		}
 
diff --git a/Engine/Source/Editor/DetailCustomizations/Private/SkinnedMeshComponentDetails.cpp b/Engine/Source/Editor/DetailCustomizations/Private/SkinnedMeshComponentDetails.cpp
--- a/Engine/Source/Editor/DetailCustomizations/Private/SkinnedMeshComponentDetails.cpp
+++ b/Engine/Source/Editor/DetailCustomizations/Private/SkinnedMeshComponentDetails.cpp
@@ -70,27 +70,25 @@ void FSkinnedMeshComponentDetails::CreateActuallyUsedPhysicsAssetWidget(FDetailW
 
 bool FSkinnedMeshComponentDetails::FindUniqueUsedPhysicsAsset(IDetailLayoutBuilder* DetailBuilder, UPhysicsAsset*& OutFoundPhysicsAsset) const
 {
-	int UsedPhysicsAssetCount = 0;
-	OutFoundPhysicsAsset = NULL;
-	const TArray< TWeakObjectPtr<UObject> >& SelectedObjectsList = DetailBuilder->GetSelectedObjects();
-	for (auto SelectionIt = SelectedObjectsList.CreateConstIterator(); SelectionIt; ++SelectionIt)
+	int32 UsedPhysicsAssetCount = 0;
+	OutFoundPhysicsAsset = nullptr;
+	
+	for (const TWeakObjectPtr<UObject>& SelectedObject : DetailBuilder->GetSelectedObjects())
 	{
-		if (AActor* Actor = Cast<AActor>(SelectionIt->Get()))
+		if (AActor* Actor = Cast<AActor>(SelectedObject.Get()))
 		{
-			TInlineComponentArray<USkinnedMeshComponent*> SkinnedMeshComponents;
-			Actor->GetComponents(SkinnedMeshComponents);
-
-			for (int32 CompIdx = 0; CompIdx < SkinnedMeshComponents.Num(); CompIdx++)
+			for (UActorComponent* Component : Actor->GetComponents())
 			{
-				USkinnedMeshComponent* SkinnedMeshComp = SkinnedMeshComponents[CompIdx];
-
-				// Only use registered and visible primitive components when calculating bounds
-				if (UsedPhysicsAssetCount > 0)
+				if (USkinnedMeshComponent* SkinnedMeshComp = Cast<USkinnedMeshComponent>(Component))
 				{
-					return false;
+					// Only use registered and visible primitive components when calculating bounds
+					if (UsedPhysicsAssetCount > 0)
+					{
+						return false;
+					}
+					++UsedPhysicsAssetCount;
+					OutFoundPhysicsAsset = SkinnedMeshComp->GetPhysicsAsset();
 				}
-				++UsedPhysicsAssetCount;
-				OutFoundPhysicsAsset = SkinnedMeshComp->GetPhysicsAsset();
 			}
 		}
 	}
diff --git a/Engine/Source/Editor/Kismet/Private/BlueprintEditor.cpp b/Engine/Source/Editor/Kismet/Private/BlueprintEditor.cpp
--- a/Engine/Source/Editor/Kismet/Private/BlueprintEditor.cpp
+++ b/Engine/Source/Editor/Kismet/Private/BlueprintEditor.cpp
@@ -894,12 +894,12 @@ void FBlueprintEditor::OnSelectionUpdated(const TArray<FSCSEditorTreeNodePtrType
 	AActor* EditorActorInstance = Blueprint->SimpleConstructionScript->GetComponentEditorActorInstance();
 	if (EditorActorInstance != nullptr)
 	{
-		TInlineComponentArray<UPrimitiveComponent*> PrimitiveComponents;
-		EditorActorInstance->GetComponents(PrimitiveComponents, true);
-
-		for (UPrimitiveComponent* PrimitiveComponent : PrimitiveComponents)
+		for (UActorComponent* Component : EditorActorInstance->GetComponents())
 		{
-			PrimitiveComponent->PushSelectionToProxy();
+			if (UPrimitiveComponent* PrimitiveComponent = Cast<UPrimitiveComponent>(Component))
+			{
+				PrimitiveComponent->PushSelectionToProxy();
+			}
 		}
 	}
 
diff --git a/Engine/Source/Editor/Kismet/Private/SCSEditorViewportClient.cpp b/Engine/Source/Editor/Kismet/Private/SCSEditorViewportClient.cpp
--- a/Engine/Source/Editor/Kismet/Private/SCSEditorViewportClient.cpp
+++ b/Engine/Source/Editor/Kismet/Private/SCSEditorViewportClient.cpp
@@ -169,14 +169,14 @@ void FSCSEditorViewportClient::Tick(float DeltaSeconds)
 	AActor* PreviewActor = GetPreviewActor();
 	if (PreviewActor != nullptr)
 	{
-		TInlineComponentArray<UPrimitiveComponent*> PrimitiveComponents;
-		PreviewActor->GetComponents(PrimitiveComponents, true);
-
-		for (UPrimitiveComponent* PrimComponent : PrimitiveComponents)
+		for (UActorComponent* Component : PreviewActor->GetComponents())
 		{
-			if (!PrimComponent->SelectionOverrideDelegate.IsBound())
+			if (UPrimitiveComponent* PrimComponent = Cast<UPrimitiveComponent>(Component))
 			{
-				SCSEditor->SetSelectionOverride(PrimComponent);
+				if (!PrimComponent->SelectionOverrideDelegate.IsBound())
+				{
+					SCSEditor->SetSelectionOverride(PrimComponent);
+				}
 			}
 		}
 	}
@@ -959,17 +959,16 @@ void FSCSEditorViewportClient::RefreshPreviewBounds()
 	if(PreviewActor)
 	{
 		// Compute actor bounds as the sum of its visible parts
-		TInlineComponentArray<UPrimitiveComponent*> PrimitiveComponents;
-		PreviewActor->GetComponents(PrimitiveComponents);
-
 		PreviewActorBounds = FBoxSphereBounds(ForceInitToZero);
-		for(auto CompIt = PrimitiveComponents.CreateConstIterator(); CompIt; ++CompIt)
+		for (UActorComponent* Component : PreviewActor->GetComponents())
 		{
 			// Aggregate primitive components that either have collision enabled or are otherwise visible components in-game
-			UPrimitiveComponent* PrimComp = *CompIt;
-			if(PrimComp->IsRegistered() && (!PrimComp->bHiddenInGame || PrimComp->IsCollisionEnabled()) && PrimComp->Bounds.SphereRadius < HALF_WORLD_MAX)
+			if (UPrimitiveComponent* PrimComp = Cast<UPrimitiveComponent>(Component))
 			{
-				PreviewActorBounds = PreviewActorBounds + PrimComp->Bounds;
+				if (PrimComp->IsRegistered() && (!PrimComp->bHiddenInGame || PrimComp->IsCollisionEnabled()) && PrimComp->Bounds.SphereRadius < HALF_WORLD_MAX)
+				{
+					PreviewActorBounds = PreviewActorBounds + PrimComp->Bounds;
+				}
 			}
 		}
 	}
diff --git a/Engine/Source/Editor/Matinee/Private/MatineeTrackHelpers.cpp b/Engine/Source/Editor/Matinee/Private/MatineeTrackHelpers.cpp
--- a/Engine/Source/Editor/Matinee/Private/MatineeTrackHelpers.cpp
+++ b/Engine/Source/Editor/Matinee/Private/MatineeTrackHelpers.cpp
@@ -251,17 +251,16 @@ bool UMatineeTrackAnimControlHelper::PreCreateKeyframe( UInterpTrack *Track, flo
 
 	USkeletalMeshComponent * SkelMeshComp = NULL;
 
-	TInlineComponentArray<USkeletalMeshComponent*> SkeletalMeshComponents;
-	Actor->GetComponents(SkeletalMeshComponents);
-
-	for (int32 I=0; I<SkeletalMeshComponents.Num(); ++I)
+	for (UActorComponent* Component : Actor->GetComponents())
 	{
-		USkeletalMeshComponent * CurSkelMeshComp = SkeletalMeshComponents[I];
-		// if qualified to play animation, break
-		if (CurSkelMeshComp->SkeletalMesh && CurSkelMeshComp->SkeletalMesh->Skeleton)
+		if (USkeletalMeshComponent* CurSkelMeshComp = Cast<USkeletalMeshComponent>(Component))
 		{
-			SkelMeshComp = CurSkelMeshComp;
-			break;
+			// if qualified to play animation, break
+			if (CurSkelMeshComp->SkeletalMesh && CurSkelMeshComp->SkeletalMesh->Skeleton)
+			{
+				SkelMeshComp = CurSkelMeshComp;
+				break;
+			}
 		}
 	}
 
diff --git a/Engine/Source/Editor/MovieSceneTools/Private/TrackEditors/PathTrackEditor.cpp b/Engine/Source/Editor/MovieSceneTools/Private/TrackEditors/PathTrackEditor.cpp
--- a/Engine/Source/Editor/MovieSceneTools/Private/TrackEditors/PathTrackEditor.cpp
+++ b/Engine/Source/Editor/MovieSceneTools/Private/TrackEditors/PathTrackEditor.cpp
@@ -163,11 +163,12 @@ bool F3DPathTrackEditor::IsActorPickable(const AActor* const ParentActor, FGuid
 		!ParentActor->IsA( AWorldSettings::StaticClass() ) &&
 		!ParentActor->IsPendingKill())
 	{			
-		TArray<USplineComponent*> SplineComponents;
-		ParentActor->GetComponents(SplineComponents);
-		if (SplineComponents.Num())
+		for (UActorComponent* Component : ParentActor->GetComponents())
 		{
-			return true;
+			if (Cast<USplineComponent>(Component))
+			{
+				return true;
+			}
 		}
 	}
 	return false;
diff --git a/Engine/Source/Editor/MovieSceneTools/Private/TrackEditors/SkeletalAnimationTrackEditor.cpp b/Engine/Source/Editor/MovieSceneTools/Private/TrackEditors/SkeletalAnimationTrackEditor.cpp
--- a/Engine/Source/Editor/MovieSceneTools/Private/TrackEditors/SkeletalAnimationTrackEditor.cpp
+++ b/Engine/Source/Editor/MovieSceneTools/Private/TrackEditors/SkeletalAnimationTrackEditor.cpp
@@ -55,13 +55,9 @@ USkeletalMeshComponent* AcquireSkeletalMeshFromObjectGuid(const FGuid& Guid, TSh
 
 	if (AActor* Actor = Cast<AActor>(BoundObject))
 	{
-		TInlineComponentArray<USkeletalMeshComponent*> SkeletalMeshComponents;
-		Actor->GetComponents(SkeletalMeshComponents);
-
-		for (int32 j = 0; j <SkeletalMeshComponents.Num(); ++j)
+		for (UActorComponent* Component : Actor->GetComponents())
 		{
-			USkeletalMeshComponent* SkeletalMeshComp = SkeletalMeshComponents[j];
-			if (SkeletalMeshComp)
+			if (USkeletalMeshComponent* SkeletalMeshComp = Cast<USkeletalMeshComponent>(Component))
 			{
 				return SkeletalMeshComp;
 			}
@@ -84,13 +80,10 @@ USkeleton* AcquireSkeletonFromObjectGuid(const FGuid& Guid, TSharedPtr<ISequence
 
 	if (AActor* Actor = Cast<AActor>(BoundObject))
 	{
-		TInlineComponentArray<USkeletalMeshComponent*> SkeletalMeshComponents;
-		Actor->GetComponents(SkeletalMeshComponents);
-
-		for (int32 j = 0; j <SkeletalMeshComponents.Num(); ++j)
+		for (UActorComponent* Component : Actor->GetComponents())
 		{
-			USkeletalMeshComponent* SkeletalMeshComp = SkeletalMeshComponents[j];
-			if (SkeletalMeshComp->SkeletalMesh && SkeletalMeshComp->SkeletalMesh->Skeleton)
+			USkeletalMeshComponent* SkeletalMeshComp = Cast<USkeletalMeshComponent>(Component);
+			if (SkeletalMeshComp && SkeletalMeshComp->SkeletalMesh && SkeletalMeshComp->SkeletalMesh->Skeleton)
 			{
 				// @todo Multiple actors, multiple components
 				return SkeletalMeshComp->SkeletalMesh->Skeleton;
diff --git a/Engine/Source/Editor/SequenceRecorder/Private/ActorRecording.cpp b/Engine/Source/Editor/SequenceRecorder/Private/ActorRecording.cpp
--- a/Engine/Source/Editor/SequenceRecorder/Private/ActorRecording.cpp
+++ b/Engine/Source/Editor/SequenceRecorder/Private/ActorRecording.cpp
@@ -450,14 +450,15 @@ void UActorRecording::StartRecordingActorProperties(ULevelSequence* CurrentSeque
 
 				if (ObjectTemplate)
 				{
-					TInlineComponentArray<USkeletalMeshComponent*> SkeletalMeshComponents;
-					ObjectTemplate->GetComponents(SkeletalMeshComponents);
-					for (USkeletalMeshComponent* SkeletalMeshComponent : SkeletalMeshComponents)
+					for (UActorComponent* Component : ObjectTemplate->GetComponents())
 					{
-						SkeletalMeshComponent->SetAnimationMode(EAnimationMode::AnimationSingleNode);
-						SkeletalMeshComponent->bEnableUpdateRateOptimizations = false;
-						SkeletalMeshComponent->MeshComponentUpdateFlag = EMeshComponentUpdateFlag::AlwaysTickPoseAndRefreshBones;
-						SkeletalMeshComponent->ForcedLodModel = 1;
+						if (USkeletalMeshComponent* SkeletalMeshComponent = Cast<USkeletalMeshComponent>(Component))
+						{
+							SkeletalMeshComponent->SetAnimationMode(EAnimationMode::AnimationSingleNode);
+							SkeletalMeshComponent->bEnableUpdateRateOptimizations = false;
+							SkeletalMeshComponent->MeshComponentUpdateFlag = EMeshComponentUpdateFlag::AlwaysTickPoseAndRefreshBones;
+							SkeletalMeshComponent->ForcedLodModel = 1;
+						}
 					}
 
 					// Disable possession of pawns otherwise the recorded character will auto possess the player
@@ -1023,16 +1024,16 @@ void UActorRecording::StartRecordingNewComponents(ULevelSequence* CurrentSequenc
 						{
 							FName AttachName = SceneComponent->GetAttachParent()->GetFName();
 
-							TInlineComponentArray<USceneComponent*> AllChildren;
-							ObjectTemplate->GetComponents(AllChildren);
-
-							for (USceneComponent* Child : AllChildren)
+							for (UActorComponent* Component : ObjectTemplate->GetComponents())
 							{
-								CA_SUPPRESS(28182); // Dereferencing NULL pointer. 'Child' contains the same NULL value as 'AttachToComponent' did.
-								if (Child->GetFName() == AttachName)
+								if (USceneComponent* Child = Cast<USceneComponent>(Component))
 								{
-									AttachToComponent = Child;
-									break;
+									CA_SUPPRESS(28182); // Dereferencing NULL pointer. 'Child' contains the same NULL value as 'AttachToComponent' did.
+									if (Child->GetFName() == AttachName)
+									{
+										AttachToComponent = Child;
+										break;
+									}
 								}
 							}
 						}
diff --git a/Engine/Source/Editor/UnrealEd/Private/AssetSelection.cpp b/Engine/Source/Editor/UnrealEd/Private/AssetSelection.cpp
--- a/Engine/Source/Editor/UnrealEd/Private/AssetSelection.cpp
+++ b/Engine/Source/Editor/UnrealEd/Private/AssetSelection.cpp
@@ -210,10 +210,7 @@ namespace AssetSelectionUtils
 						ActorInfo.bHaveAttachedActor = true;
 					}
 
-					TInlineComponentArray<UActorComponent*> ActorComponents;
-					CurrentActor->GetComponents(ActorComponents);
-
-					for( UActorComponent* Component : ActorComponents )
+					for (UActorComponent* Component : CurrentActor->GetComponents())
 					{
 						if( UStaticMeshComponent* SMComp = Cast<UStaticMeshComponent>(Component) )
 						{
@@ -749,10 +746,14 @@ bool FActorFactoryAssetProxy::IsActorValidForMaterialApplication( AActor* Target
 	// a material applied to it. Otherwise, it cannot.
 	if ( TargetActor )
 	{
-		TInlineComponentArray<UMeshComponent*> MeshComponents;
-		TargetActor->GetComponents(MeshComponents);
-
-		bIsValid = (MeshComponents.Num() > 0);
+		for (UActorComponent* Component : TargetActor->GetComponents())
+		{
+			if (Cast<UMeshComponent>(Component))
+			{
+				bIsValid = true;
+				break;
+			}
+		}
 	}
 
 	return bIsValid;
@@ -792,15 +793,12 @@ bool FActorFactoryAssetProxy::ApplyMaterialToActor( AActor* TargetActor, UMateri
 			TArray<USceneComponent*> FoundMeshComponents;
 
 			// Find which mesh the user clicked on first.
-			TInlineComponentArray<USceneComponent*> SceneComponents;
-			TargetActor->GetComponents(SceneComponents);
-
-			for ( int32 ComponentIdx=0; ComponentIdx < SceneComponents.Num(); ComponentIdx++ )
+			for (UActorComponent* Component : TargetActor->GetComponents())
 			{
-				USceneComponent* SceneComp = SceneComponents[ComponentIdx];
+				USceneComponent* SceneComp = Cast<USceneComponent>(Component);
 				
 				// Only apply the material to editable components.  Components which are not exposed are not intended to be changed.
-				if( EditableComponents.Contains( SceneComp ) )
+				if (SceneComp && EditableComponents.Contains( SceneComp ) )
 				{
 					UMeshComponent* MeshComponent = Cast<UMeshComponent>(SceneComp);
 
diff --git a/Engine/Source/Editor/UnrealEd/Private/DragTool_BoxSelect.cpp b/Engine/Source/Editor/UnrealEd/Private/DragTool_BoxSelect.cpp
--- a/Engine/Source/Editor/UnrealEd/Private/DragTool_BoxSelect.cpp
+++ b/Engine/Source/Editor/UnrealEd/Private/DragTool_BoxSelect.cpp
@@ -304,13 +304,10 @@ bool FDragTool_ActorBoxSelect::IntersectsBox( AActor& InActor, const FBox& InBox
 	if( !bActorIsHiddenByShowFlags && !InActor.IsHiddenEd() && !FActorEditorUtils::IsABuilderBrush(&InActor) && bActorRecentlyRendered )
 	{
 		// Iterate over all actor components, selecting out primitive components
-		TInlineComponentArray<UPrimitiveComponent*> PrimitiveComponents;
-		InActor.GetComponents(PrimitiveComponents);
-
-		for (const UPrimitiveComponent* PrimitiveComponent : PrimitiveComponents)
+		for (UActorComponent* Component : InActor.GetComponents())
 		{
-			check(PrimitiveComponent != nullptr);
-			if (PrimitiveComponent->IsRegistered() && PrimitiveComponent->IsVisibleInEditor())
+			UPrimitiveComponent* PrimitiveComponent = Cast<UPrimitiveComponent>(Component);
+			if (PrimitiveComponent && PrimitiveComponent->IsRegistered() && PrimitiveComponent->IsVisibleInEditor())
 			{
 				if (PrimitiveComponent->ComponentIsTouchingSelectionBox(InBox, LevelViewportClient->EngineShowFlags, bGeometryMode, bUseStrictSelection))
 				{
diff --git a/Engine/Source/Editor/UnrealEd/Private/DragTool_FrustumSelect.cpp b/Engine/Source/Editor/UnrealEd/Private/DragTool_FrustumSelect.cpp
--- a/Engine/Source/Editor/UnrealEd/Private/DragTool_FrustumSelect.cpp
+++ b/Engine/Source/Editor/UnrealEd/Private/DragTool_FrustumSelect.cpp
@@ -201,13 +201,10 @@ bool FDragTool_ActorFrustumSelect::IntersectsFrustum( AActor& InActor, const FCo
 	if( !bActorIsHiddenByShowFlags && !InActor.IsHiddenEd() && !FActorEditorUtils::IsABuilderBrush(&InActor) )
 	{
 		// Iterate over all actor components, selecting out primitive components
-		TInlineComponentArray<UPrimitiveComponent*> Components;
-		InActor.GetComponents(Components);
-
-		for (const UPrimitiveComponent* PrimitiveComponent : Components)
+		for (UActorComponent* Component : InActor.GetComponents())
 		{
-			check(PrimitiveComponent != nullptr);
-			if ( PrimitiveComponent->IsRegistered() && PrimitiveComponent->IsVisibleInEditor() )
+			UPrimitiveComponent* PrimitiveComponent = Cast<UPrimitiveComponent>(Component);
+			if (PrimitiveComponent && PrimitiveComponent->IsRegistered() && PrimitiveComponent->IsVisibleInEditor() )
 			{
 				if (PrimitiveComponent->ComponentIsTouchingSelectionFrustum(InFrustum, LevelViewportClient->EngineShowFlags, bGeometryMode, bUseStrictSelection))
 				{
diff --git a/Engine/Source/Editor/UnrealEd/Private/EditorActor.cpp b/Engine/Source/Editor/UnrealEd/Private/EditorActor.cpp
--- a/Engine/Source/Editor/UnrealEd/Private/EditorActor.cpp
+++ b/Engine/Source/Editor/UnrealEd/Private/EditorActor.cpp
@@ -2208,16 +2208,15 @@ void UUnrealEdEngine::edactSelectMatchingMaterial()
 		if( CurrentActor )
 		{
 			// Find the materials by iterating over every primitive component.
-			TInlineComponentArray<UPrimitiveComponent*> PrimitiveComponents;
-			CurrentActor->GetComponents(PrimitiveComponents);
-
-			for (int32 ComponentIdx = 0; ComponentIdx < PrimitiveComponents.Num(); ComponentIdx++)
+			for (UActorComponent* Component : CurrentActor->GetComponents())
 			{
-				UPrimitiveComponent* CurrentComponent = PrimitiveComponents[ComponentIdx];
-				TArray<UMaterialInterface*> UsedMaterials;
-				CurrentComponent->GetUsedMaterials( UsedMaterials );
-				MaterialsInSelection.Append( UsedMaterials );
-				SelectedWorlds.AddUnique( CurrentActor->GetWorld() );
+				if (UPrimitiveComponent* CurrentComponent = Cast<UPrimitiveComponent>(Component))
+				{
+					TArray<UMaterialInterface*> UsedMaterials;
+					CurrentComponent->GetUsedMaterials(UsedMaterials);
+					MaterialsInSelection.Append(UsedMaterials);
+					SelectedWorlds.AddUnique(CurrentActor->GetWorld());
+				}
 			}
 		}
 	}
@@ -2337,14 +2336,11 @@ void UUnrealEdEngine::edactSelectRelevantLights( UWorld* InWorld )
 
 		if (Actor->GetLevel()->IsCurrentLevel() )
 		{
-			TInlineComponentArray<UPrimitiveComponent*> PrimitiveComponents;
-			Actor->GetComponents(PrimitiveComponents);
-
 			// Gather static lighting info from each of the actor's components.
-			for(int32 ComponentIndex = 0;ComponentIndex < PrimitiveComponents.Num();ComponentIndex++)
+			for (UActorComponent* Component : Actor->GetComponents())
 			{
-				UPrimitiveComponent* Primitive = PrimitiveComponents[ComponentIndex];
-				if (Primitive->IsRegistered())
+				UPrimitiveComponent* Primitive = Cast<UPrimitiveComponent>(Component);
+				if (Primitive && Primitive->IsRegistered())
 				{
 					TArray<const ULightComponent*> RelevantLightComponents;
 					InWorld->Scene->GetRelevantLights(Primitive, &RelevantLightComponents);
diff --git a/Engine/Source/Editor/UnrealEd/Private/EditorEngine.cpp b/Engine/Source/Editor/UnrealEd/Private/EditorEngine.cpp
--- a/Engine/Source/Editor/UnrealEd/Private/EditorEngine.cpp
+++ b/Engine/Source/Editor/UnrealEd/Private/EditorEngine.cpp
@@ -5250,14 +5250,9 @@ static void CopyLightComponentProperties( const AActor& InOldActor, AActor& InNe
 	UActorComponent* LightComponentToCopy = NULL;
 
 	// Go through the old actor's components and look for a light component to copy.
-	TInlineComponentArray<UActorComponent*> OldActorComponents;
-	InOldActor.GetComponents(OldActorComponents);
-
-	for( int32 CompToCopyIdx = 0; CompToCopyIdx < OldActorComponents.Num(); ++CompToCopyIdx )
+	for (UActorComponent* Component : InOldActor.GetComponents())
 	{
-		UActorComponent* Component = OldActorComponents[CompToCopyIdx];
-
-		if( Component->IsRegistered() && Component->IsA( CopyableComponentClass ) ) 
+		if (Component && Component->IsRegistered() && Component->IsA( CopyableComponentClass ) ) 
 		{
 			// A light component has been found. 
 			CompToCopyClass = Component->GetClass();
@@ -5271,17 +5266,13 @@ static void CopyLightComponentProperties( const AActor& InOldActor, AActor& InNe
 	// The class of the new actors light component
 	const UClass* CommonLightComponentClass = NULL;
 
-	// Dont do anything if there is no valid light component to copy from
+	// Don't do anything if there is no valid light component to copy from
 	if( LightComponentToCopy )
 	{
-		TInlineComponentArray<UActorComponent*> NewActorComponents;
-		InNewActor.GetComponents(NewActorComponents);
-
 		// Find a light component to overwrite in the new actor
-		for( int32 NewCompIdx = 0; NewCompIdx < NewActorComponents.Num(); ++NewCompIdx )
+		for (UActorComponent* Component : InNewActor.GetComponents())
 		{
-			UActorComponent* Component = NewActorComponents[ NewCompIdx ];
-			if(Component->IsRegistered())
+			if (Component && Component->IsRegistered())
 			{
 				// Find a common component class between the new and old actor.   
 				// This needs to be done so we can copy as many properties as possible. 
@@ -5520,37 +5511,31 @@ void CopyActorComponentProperties( const AActor* SourceActor, AActor* DestActor,
 
 		// Construct a mapping from the default actor of its relevant component names to its actual components. Here relevant component
 		// names are those that match a name provided as a parameter.
-		TInlineComponentArray<UActorComponent*> CDOComponents;
-		SrcActorDefaultActor->GetComponents(CDOComponents);
-
 		TMap<FString, const UActorComponent*> NameToDefaultComponentMap; 
-		for ( TInlineComponentArray<UActorComponent*>::TConstIterator CompIter( CDOComponents ); CompIter; ++CompIter )
+		for (UActorComponent* CurComp : SrcActorDefaultActor->GetComponents())
 		{
-			const UActorComponent* CurComp = *CompIter;
-			check( CurComp );
-
-			const FString CurCompName = CurComp->GetName();
-			if ( ComponentNames.Contains( CurCompName ) )
+			if (CurComp)
 			{
-				NameToDefaultComponentMap.Add( CurCompName, CurComp );
+				FString CurCompName = CurComp->GetName();
+				if (ComponentNames.Contains(CurCompName))
+				{
+					NameToDefaultComponentMap.Add(MoveTemp(CurCompName), CurComp);
+				}
 			}
 		}
 
 		// Construct a mapping from the source actor of its relevant component names to its actual components. Here relevant component names
 		// are those that match a name provided as a parameter.
-		TInlineComponentArray<UActorComponent*> SourceComponents;
-		SourceActor->GetComponents(SourceComponents);
-
 		TMap<FString, const UActorComponent*> NameToSourceComponentMap;
-		for ( TInlineComponentArray<UActorComponent*>::TConstIterator CompIter( SourceComponents ); CompIter; ++CompIter )
+		for (UActorComponent* CurComp : SourceActor->GetComponents())
 		{
-			const UActorComponent* CurComp = *CompIter;
-			check( CurComp );
-
-			const FString CurCompName = CurComp->GetName();
-			if ( ComponentNames.Contains( CurCompName ) )
+			if (CurComp)
 			{
-				NameToSourceComponentMap.Add( CurCompName, CurComp );
+				FString CurCompName = CurComp->GetName();
+				if (ComponentNames.Contains(CurCompName))
+				{
+					NameToSourceComponentMap.Add(MoveTemp(CurCompName), CurComp);
+				}
 			}
 		}
 
diff --git a/Engine/Source/Editor/UnrealEd/Private/EditorExporters.cpp b/Engine/Source/Editor/UnrealEd/Private/EditorExporters.cpp
--- a/Engine/Source/Editor/UnrealEd/Private/EditorExporters.cpp
+++ b/Engine/Source/Editor/UnrealEd/Private/EditorExporters.cpp
@@ -1483,15 +1483,13 @@ ULevelExporterFBX::ULevelExporterFBX(const FObjectInitializer& ObjectInitializer
 bool IsSomethingToExport(AActor* Actor)
 {
 	check(Actor);
-	TInlineComponentArray<USceneComponent*> SceneComponents;
-	Actor->GetComponents(SceneComponents);
 
 	TInlineComponentArray<USceneComponent*> ComponentsToExport;
-	for (int32 ComponentIndex = 0; ComponentIndex < SceneComponents.Num(); ++ComponentIndex)
+	for (UActorComponent* ActorComp : Actor->GetComponents())
 	{
-		USceneComponent* Component = SceneComponents[ComponentIndex];
+		USceneComponent* Component = Cast<USceneComponent>(ActorComp);
 
-		if (Component == nullptr || (Component && Component->bHiddenInGame))
+		if (Component == nullptr || Component->bHiddenInGame)
 		{
 			//Skip hidden component like camera mesh or other editor helper
 			continue;
diff --git a/Engine/Source/Editor/UnrealEd/Private/EditorServer.cpp b/Engine/Source/Editor/UnrealEd/Private/EditorServer.cpp
--- a/Engine/Source/Editor/UnrealEd/Private/EditorServer.cpp
+++ b/Engine/Source/Editor/UnrealEd/Private/EditorServer.cpp
@@ -5156,34 +5156,34 @@ void UEditorEngine::ExecFile( UWorld* InWorld, const TCHAR* InFilename, FOutputD
 
 void UEditorEngine::AssignReplacementComponentsByActors(TArray<AActor*>& ActorsToReplace, AActor* Replacement, UClass* ClassToReplace)
 {
-	// the code will use this to find the best possible component, in the priority listed here
-	// (ie it will first look for a mesh component, then a particle, and finally a sprite)
-	UClass* PossibleReplacementClass[] = 
-	{
-		UMeshComponent::StaticClass(),
-		UParticleSystemComponent::StaticClass(),
-		UBillboardComponent::StaticClass(),
-	};
-
 	// look for a mesh component to replace with
-	UPrimitiveComponent* ReplacementComponent = NULL;
+	UPrimitiveComponent* ReplacementComponent = nullptr;
 
-	// loop over the clases until a component is found
-	for (int32 ClassIndex = 0; ClassIndex < ARRAY_COUNT(PossibleReplacementClass); ClassIndex++)
+	// if we are clearing the replacement, then we don't need to find a component
+	if (Replacement)
 	{
-		// use ClassToReplace of UMeshComponent if not specified
-		UClass* ReplacementComponentClass = ClassToReplace ? ClassToReplace : PossibleReplacementClass[ClassIndex];
+		// the code will use this to find the best possible component, in the priority listed here
+		// (ie it will first look for a mesh component, then a particle, and finally a sprite)
+		TArray<UClass*, TInlineAllocator<3>> PossibleReplacementClasses;
 
-		// if we are clearing the replacement, then we don't need to find a component
-		if (Replacement)
+		if (ClassToReplace)
 		{
-			TInlineComponentArray<UPrimitiveComponent*> Components;
-			Replacement->GetComponents(Components);
+			PossibleReplacementClasses.Emplace(ClassToReplace);
+		}
+		else
+		{
+			PossibleReplacementClasses.Emplace(UMeshComponent::StaticClass());
+			PossibleReplacementClasses.Emplace(UParticleSystemComponent::StaticClass());
+			PossibleReplacementClasses.Emplace(UBillboardComponent::StaticClass());
+		}
 
-			for (int32 ComponentIndex = 0; ComponentIndex < Components.Num(); ComponentIndex++)
+		// loop over the clases until a component is found
+		for (UClass* ReplacementComponentClass : PossibleReplacementClasses)
+		{
+			for (UActorComponent* Component : Replacement->GetComponents())
 			{
-				UPrimitiveComponent* PrimitiveComponent = Components[ComponentIndex];
-				if (PrimitiveComponent->IsA(ReplacementComponentClass))
+				UPrimitiveComponent* PrimitiveComponent = Cast<UPrimitiveComponent>(Component);
+				if (PrimitiveComponent && PrimitiveComponent->IsA(ReplacementComponentClass))
 				{
 					ReplacementComponent = PrimitiveComponent;
 					goto FoundComponent;
@@ -5191,6 +5191,7 @@ void UEditorEngine::AssignReplacementComponentsByActors(TArray<AActor*>& ActorsT
 			}
 		}
 	}
+
 FoundComponent:
 
 	// attempt to set replacement component for all selected actors
@@ -5205,8 +5206,8 @@ FoundComponent:
 		{
 			UPrimitiveComponent* PrimitiveComponent = Components[ComponentIndex];
 			// if the primitive component matches the class we are looking for (if specified)
-			// then set it's replacement component
-			if (ClassToReplace == NULL || PrimitiveComponent->IsA(ClassToReplace))
+			// then set its replacement component
+			if (ClassToReplace == nullptr || PrimitiveComponent->IsA(ClassToReplace))
 			{
 				// need to reregister the component
 				FComponentReregisterContext ComponentReattch(PrimitiveComponent);
@@ -5214,7 +5215,7 @@ FoundComponent:
 				// set the replacement
 				PrimitiveComponent->SetLODParentPrimitive(ReplacementComponent);
 
-				// makr the package as dirty now that we've modified it
+				// mark the package as dirty now that we've modified it
 				Actor->MarkPackageDirty();
 			}
 		}
@@ -6112,17 +6113,16 @@ bool UEditorEngine::HandleSetReplacementCommand( const TCHAR* Str, FOutputDevice
 	// attempt to set replacement component for all selected actors
 	for( FSelectedActorIterator It(InWorld); It; ++It )
 	{
-		TInlineComponentArray<UPrimitiveComponent*> Components;
-		It->GetComponents(Components);
-
-		for (int32 ComponentIndex = 0; ComponentIndex < Components.Num(); ComponentIndex++)
+		for (UActorComponent* Component : It->GetComponents())
 		{
-			UPrimitiveComponent* PrimitiveComponent = Components[ComponentIndex];
-			// if the primitive component matches the class we are looking for (if specified)
-			// then set it's replacement component
-			if (ClassToReplace == NULL || PrimitiveComponent->IsA(ClassToReplace))
+			if (UPrimitiveComponent* PrimitiveComponent = Cast<UPrimitiveComponent>(Component))
 			{
-				PrimitiveComponent->SetLODParentPrimitive(ReplacementComponent);
+				// if the primitive component matches the class we are looking for (if specified)
+				// then set it's replacement component
+				if (ClassToReplace == NULL || PrimitiveComponent->IsA(ClassToReplace))
+				{
+					PrimitiveComponent->SetLODParentPrimitive(ReplacementComponent);
+				}
 			}
 		}
 	}
@@ -6358,14 +6358,11 @@ bool UEditorEngine::HandleSetDetailModeCommand( const TCHAR* Str, FOutputDevice&
 			AActor* Actor = static_cast<AActor*>( *It );
 			checkSlow( Actor->IsA(AActor::StaticClass()) );
 
-			TInlineComponentArray<UPrimitiveComponent*> Components;
-			Actor->GetComponents(Components);
-
-			for(int32 ComponentIndex = 0;ComponentIndex < Components.Num();ComponentIndex++)
+			for (UActorComponent* Component : Actor->GetComponents())
 			{
-				UPrimitiveComponent* primComp = Components[ComponentIndex];
+				UPrimitiveComponent* primComp = Cast<UPrimitiveComponent>(Component);
 
-				if( primComp->DetailMode != ParsedDetailMode )
+				if (primComp && primComp->DetailMode != ParsedDetailMode )
 				{
 					primComp->Modify();
 					primComp->DetailMode = EDetailMode(ParsedDetailMode);
diff --git a/Engine/Source/Editor/UnrealEd/Private/EditorViewportCommands.cpp b/Engine/Source/Editor/UnrealEd/Private/EditorViewportCommands.cpp
--- a/Engine/Source/Editor/UnrealEd/Private/EditorViewportCommands.cpp
+++ b/Engine/Source/Editor/UnrealEd/Private/EditorViewportCommands.cpp
@@ -165,9 +165,13 @@ FText GetViewModeOptionsMenuLabel(EViewModeIndex ViewModeIndex)
 			GEditor->GetSelectedActors()->GetSelectedObjects(SelectedActors);
 			for (AActor* Actor : SelectedActors)
 			{
-				TArray<UPrimitiveComponent*> ActorComponents;
-				Actor->GetComponents(ActorComponents);
-				SelectedComponents.Append(ActorComponents);
+				for (UActorComponent* Component : Actor->GetComponents())
+				{
+					if (UPrimitiveComponent* PrimComp = Cast<UPrimitiveComponent>(Component))
+					{
+						SelectedComponents.Add(PrimComp);
+					}
+				}
 			}
 
 			for (const UPrimitiveComponent* SelectedComponent : SelectedComponents)
@@ -226,9 +230,13 @@ static void GetSelectedMaterials(TArray<const UMaterialInterface*>& SelectedMate
 		GEditor->GetSelectedActors()->GetSelectedObjects(SelectedActors);
 		for (AActor* Actor : SelectedActors)
 		{
-			TArray<UPrimitiveComponent*> ActorComponents;
-			Actor->GetComponents(ActorComponents);
-			SelectedComponents.Append(ActorComponents);
+			for (UActorComponent* Component : Actor->GetComponents())
+			{
+				if (UPrimitiveComponent* PrimComp = Cast<UPrimitiveComponent>(Component))
+				{
+					SelectedComponents.Add(PrimComp);
+				}
+			}
 		}
 
 		for (const UPrimitiveComponent* SelectedComponent : SelectedComponents)
diff --git a/Engine/Source/Editor/UnrealEd/Private/Factories/ActorFactory.cpp b/Engine/Source/Editor/UnrealEd/Private/Factories/ActorFactory.cpp
--- a/Engine/Source/Editor/UnrealEd/Private/Factories/ActorFactory.cpp
+++ b/Engine/Source/Editor/UnrealEd/Private/Factories/ActorFactory.cpp
@@ -529,13 +529,14 @@ void UActorFactoryDeferredDecal::PostSpawnActor(UObject* Asset, AActor* NewActor
 	if (Material != NULL )
 	{
 		// Change properties
-		TInlineComponentArray<UDecalComponent*> DecalComponents;
-		NewActor->GetComponents(DecalComponents);
-
-		UDecalComponent* DecalComponent = NULL;
-		for (int32 Idx = 0; Idx < DecalComponents.Num() && DecalComponent == NULL; Idx++)
+		UDecalComponent* DecalComponent = nullptr;
+		for (UActorComponent* Component : NewActor->GetComponents())
+		{
+			if (UDecalComponent* DecalComp = Cast<UDecalComponent>(Component))
 			{
-			DecalComponent = DecalComponents[Idx];
+				DecalComponent = DecalComp;
+				break;
+			}
 		}
 
 		check(DecalComponent);
@@ -557,13 +558,14 @@ void UActorFactoryDeferredDecal::PostCreateBlueprint( UObject* Asset, AActor* CD
 
 		if (Material != NULL)
 		{
-			TInlineComponentArray<UDecalComponent*> DecalComponents;
-			CDO->GetComponents(DecalComponents);
-
-			UDecalComponent* DecalComponent = NULL;
-			for (int32 Idx = 0; Idx < DecalComponents.Num() && DecalComponent == NULL; Idx++)
+			UDecalComponent* DecalComponent = nullptr;
+			for (UActorComponent* Component : CDO->GetComponents())
 			{
-				DecalComponent = DecalComponents[Idx];
+				if (UDecalComponent* DecalComp = Cast<UDecalComponent>(Component))
+				{
+					DecalComponent = DecalComp;
+					break;
+				}
 			}
 
 			check(DecalComponent);
diff --git a/Engine/Source/Editor/UnrealEd/Private/Fbx/FbxMainExport.cpp b/Engine/Source/Editor/UnrealEd/Private/Fbx/FbxMainExport.cpp
--- a/Engine/Source/Editor/UnrealEd/Private/Fbx/FbxMainExport.cpp
+++ b/Engine/Source/Editor/UnrealEd/Private/Fbx/FbxMainExport.cpp
@@ -1657,15 +1657,12 @@ FbxNode* FFbxExporter::ExportActor(AActor* Actor, bool bExportComponents, INodeN
 	
 		if( bExportComponents )
 		{
-			TInlineComponentArray<USceneComponent*> SceneComponents;
-			Actor->GetComponents(SceneComponents);
-
 			TInlineComponentArray<USceneComponent*> ComponentsToExport;
-			for( int32 ComponentIndex = 0; ComponentIndex < SceneComponents.Num(); ++ComponentIndex )
+			for (UActorComponent* ActorComp : Actor->GetComponents())
 			{
-				USceneComponent* Component = SceneComponents[ComponentIndex];
+				USceneComponent* Component = Cast<USceneComponent>(ActorComp);
 	
-				if (Component == nullptr || (Component && Component->bHiddenInGame))
+				if (Component == nullptr || Component->bHiddenInGame)
 				{
 					//Skip hidden component like camera mesh or other editor helper
 					continue;
diff --git a/Engine/Source/Editor/UnrealEd/Private/Kismet2/BlueprintEditorUtils.cpp b/Engine/Source/Editor/UnrealEd/Private/Kismet2/BlueprintEditorUtils.cpp
--- a/Engine/Source/Editor/UnrealEd/Private/Kismet2/BlueprintEditorUtils.cpp
+++ b/Engine/Source/Editor/UnrealEd/Private/Kismet2/BlueprintEditorUtils.cpp
@@ -6834,22 +6834,17 @@ bool FBlueprintEditorUtils::IsSCSComponentProperty(UObjectProperty* MemberProper
 		}
 		else if (const AActor* ActorCDO = GetDefault<AActor>(OwnerClass))
 		{
-			TInlineComponentArray<UActorComponent*> CDOComponents;
-			ActorCDO->GetComponents(CDOComponents);
-
 			const void* PropertyAddress = MemberProperty->ContainerPtrToValuePtr<void>(ActorCDO);
 			UObject* PropertyValue = MemberProperty->GetObjectPropertyValue(PropertyAddress);
 
-			for (UActorComponent* Component : CDOComponents)
+			for (UActorComponent* Component : ActorCDO->GetComponents())
 			{
-				if (!Component->GetClass()->IsChildOf(MemberProperty->PropertyClass))
-				{
-					continue;
-				}
-
-				if (PropertyValue == Component)
+				if (Component && Component->GetClass()->IsChildOf(MemberProperty->PropertyClass))
 				{
-					return true;
+					if (PropertyValue == Component)
+					{
+						return true;
+					}
 				}
 			}
 		}
diff --git a/Engine/Source/Editor/UnrealEd/Private/Layers/Layers.cpp b/Engine/Source/Editor/UnrealEd/Private/Layers/Layers.cpp
--- a/Engine/Source/Editor/UnrealEd/Private/Layers/Layers.cpp
+++ b/Engine/Source/Editor/UnrealEd/Private/Layers/Layers.cpp
@@ -698,10 +698,10 @@ void FLayers::RemoveViewFromActorViewVisibility( FLevelEditorViewportClient* Vie
 		TInlineComponentArray<UPrimitiveComponent*> Components;
 		Actor->GetComponents(Components);
 
-		for( int32 ComponentIdx = 0; ComponentIdx < Components.Num(); ++ComponentIdx )
+		for (UActorComponent* Component : Actor->GetComponents())
 		{
-			UPrimitiveComponent* PrimitiveComponent = Components[ComponentIdx];
-			if (PrimitiveComponent->IsRegistered())
+			UPrimitiveComponent* PrimitiveComponent = Cast<UPrimitiveComponent>(Component);
+			if (PrimitiveComponent && PrimitiveComponent->IsRegistered())
 			{
 				// Push visibility to the render thread
 				PrimitiveComponent->PushEditorVisibilityToProxy( Actor->HiddenEditorViews );
diff --git a/Engine/Source/Editor/UnrealEd/Private/LevelEditorViewport.cpp b/Engine/Source/Editor/UnrealEd/Private/LevelEditorViewport.cpp
--- a/Engine/Source/Editor/UnrealEd/Private/LevelEditorViewport.cpp
+++ b/Engine/Source/Editor/UnrealEd/Private/LevelEditorViewport.cpp
@@ -1380,12 +1380,12 @@ bool FLevelEditorViewportClient::DropObjectsAtCoordinates(int32 MouseX, int32 Mo
 					}
 
 					// Prevent future selection. This also prevents the hit proxy from interfering with placement logic.
-					TInlineComponentArray<UPrimitiveComponent*> PrimitiveComponents;
-					NewActor->GetComponents(PrimitiveComponents);
-
-					for ( auto CompIt = PrimitiveComponents.CreateConstIterator(); CompIt; ++CompIt )
+					for (UActorComponent* Component : NewActor->GetComponents())
 					{
-						(*CompIt)->bSelectable = false;
+						if (UPrimitiveComponent* PrimComp = Cast<UPrimitiveComponent>(Component))
+						{
+							PrimComp->bSelectable = false;
+						}
 					}
 				}
 
@@ -4128,13 +4128,10 @@ void FLevelEditorViewportClient::Draw(const FSceneView* View,FPrimitiveDrawInter
 				continue;
 			}
 
-			TInlineComponentArray<USceneComponent*> Components;
-			Actor->GetComponents(Components);
-
-			for (int32 ComponentIndex = 0 ; ComponentIndex < Components.Num(); ++ComponentIndex)
+			for (UActorComponent* Component : Actor->GetComponents())
 			{
-				USceneComponent* SceneComponent = Components[ComponentIndex];
-				if (SceneComponent->HasAnySockets())
+				USceneComponent* SceneComponent = Cast<USceneComponent>(Component);
+				if (SceneComponent && SceneComponent->HasAnySockets())
 				{
 					TArray<FComponentSocketDescription> Sockets;
 					SceneComponent->QuerySupportedSockets(Sockets);
@@ -4564,13 +4561,10 @@ void FLevelEditorViewportClient::AddHoverEffect( const FViewportHoverTarget& InH
 
 	if( ActorUnderCursor != nullptr )
 	{
-		TInlineComponentArray<UPrimitiveComponent*> Components;
-		ActorUnderCursor->GetComponents(Components);
-
-		for(int32 ComponentIndex = 0;ComponentIndex < Components.Num();ComponentIndex++)
+		for (UActorComponent* Component : ActorUnderCursor->GetComponents())
 		{
-			UPrimitiveComponent* PrimitiveComponent = Components[ComponentIndex];
-			if (PrimitiveComponent->IsRegistered())
+			UPrimitiveComponent* PrimitiveComponent = Cast<UPrimitiveComponent>(Component);
+			if (PrimitiveComponent && PrimitiveComponent->IsRegistered())
 			{
 				PrimitiveComponent->PushHoveredToProxy( true );
 			}
@@ -4596,13 +4590,10 @@ void FLevelEditorViewportClient::RemoveHoverEffect( const FViewportHoverTarget&
 	AActor* CurHoveredActor = InHoverTarget.HoveredActor;
 	if( CurHoveredActor != nullptr )
 	{
-		TInlineComponentArray<UPrimitiveComponent*> Components;
-		CurHoveredActor->GetComponents(Components);
-
-		for(int32 ComponentIndex = 0;ComponentIndex < Components.Num();ComponentIndex++)
+		for (UActorComponent* Component : CurHoveredActor->GetComponents())
 		{
-			UPrimitiveComponent* PrimitiveComponent = Components[ComponentIndex];
-			if (PrimitiveComponent->IsRegistered())
+			UPrimitiveComponent* PrimitiveComponent = Cast<UPrimitiveComponent>(Component);
+			if (PrimitiveComponent && PrimitiveComponent->IsRegistered())
 			{
 				check(PrimitiveComponent->IsRegistered());
 				PrimitiveComponent->PushHoveredToProxy( false );
diff --git a/Engine/Source/Editor/UnrealEd/Private/Lightmass/Lightmass.cpp b/Engine/Source/Editor/UnrealEd/Private/Lightmass/Lightmass.cpp
--- a/Engine/Source/Editor/UnrealEd/Private/Lightmass/Lightmass.cpp
+++ b/Engine/Source/Editor/UnrealEd/Private/Lightmass/Lightmass.cpp
@@ -864,13 +864,10 @@ void FLightmassExporter::WriteVisibilityData( int32 Channel )
 				AActor* CurrentActor = Volume->OverrideVisibleActors[ActorIndex];
 				if (CurrentActor)
 				{
-					TInlineComponentArray<UPrimitiveComponent*> Components;
-					CurrentActor->GetComponents(Components);
-
-					for (int32 ComponentIndex = 0; ComponentIndex < Components.Num(); ComponentIndex++)
+					for (UActorComponent* Component : CurrentActor->GetComponents())
 					{
-						UPrimitiveComponent* CurrentComponent = Components[ComponentIndex];
-						if ((CurrentComponent->Mobility == EComponentMobility::Static) && CurrentComponent->VisibilityId != INDEX_NONE)
+						UPrimitiveComponent* CurrentComponent = Cast<UPrimitiveComponent>(Component);
+						if (CurrentComponent && (CurrentComponent->Mobility == EComponentMobility::Static) && CurrentComponent->VisibilityId != INDEX_NONE)
 						{
 							VisibilityIds.AddUnique(CurrentComponent->VisibilityId);
 						}
@@ -883,13 +880,10 @@ void FLightmassExporter::WriteVisibilityData( int32 Channel )
 				AActor* RemoveActor = Volume->OverrideInvisibleActors[RemoveActorIndex];
 				if (RemoveActor)
 				{
-					TInlineComponentArray<UPrimitiveComponent*> Components;
-					RemoveActor->GetComponents(Components);
-
-					for (int32 ComponentIndex = 0; ComponentIndex < Components.Num(); ComponentIndex++)
+					for (UActorComponent* Component : RemoveActor->GetComponents())
 					{
-						UPrimitiveComponent* RemoveComponent = Components[ComponentIndex];
-						if ((RemoveComponent->Mobility == EComponentMobility::Static) && RemoveComponent->VisibilityId != INDEX_NONE)
+						UPrimitiveComponent* RemoveComponent = Cast<UPrimitiveComponent>(Component);
+						if (RemoveComponent && (RemoveComponent->Mobility == EComponentMobility::Static) && RemoveComponent->VisibilityId != INDEX_NONE)
 						{
 							InvisibilityIds.AddUnique(RemoveComponent->VisibilityId);
 						}
@@ -907,13 +901,10 @@ void FLightmassExporter::WriteVisibilityData( int32 Channel )
 						AActor* RemoveActor = Level->Actors[RemoveActorIndex];
 						if (RemoveActor)
 						{
-							TInlineComponentArray<UPrimitiveComponent*> PrimitiveComponents;
-							RemoveActor->GetComponents(PrimitiveComponents);
-
-							for (int32 ComponentIndex = 0; ComponentIndex < PrimitiveComponents.Num(); ComponentIndex++)
+							for (UActorComponent* Component : RemoveActor->GetComponents())
 							{
-								UPrimitiveComponent* RemoveComponent = PrimitiveComponents[ComponentIndex];
-								if ((RemoveComponent->Mobility == EComponentMobility::Static) && RemoveComponent->VisibilityId != INDEX_NONE)
+								UPrimitiveComponent* RemoveComponent = Cast<UPrimitiveComponent>(Component);
+								if (RemoveComponent && (RemoveComponent->Mobility == EComponentMobility::Static) && RemoveComponent->VisibilityId != INDEX_NONE)
 								{
 									InvisibilityIds.AddUnique(RemoveComponent->VisibilityId);
 								}
@@ -2452,20 +2443,19 @@ void FLightmassExporter::WriteDebugInput( Lightmass::FDebugLightingInputData& In
 	{
 		for (FSelectedActorIterator It(World); It; ++It)
 		{
-			TInlineComponentArray<UPrimitiveComponent*> Components;
-			It->GetComponents(Components);
-
-			for (int32 ComponentIndex = 0; ComponentIndex < Components.Num(); ComponentIndex++)
+			for (UActorComponent* Comp : It->GetComponents())
 			{
-				UPrimitiveComponent* Component = Components[ComponentIndex];
-				if (DebugVisibilityId == INDEX_NONE)
-				{
-					DebugVisibilityId = Component->VisibilityId;
-				}
-				else if (DebugVisibilityId != Component->VisibilityId)
+				if (UPrimitiveComponent* Component = Cast<UPrimitiveComponent>(Comp))
 				{
-					UE_LOG(LogLightmassSolver, Warning, TEXT("Not debugging visibility for component %s with vis id %u, as it was not the first component on the selected actor."),
-						*Component->GetPathName(), Component->VisibilityId);
+					if (DebugVisibilityId == INDEX_NONE)
+					{
+						DebugVisibilityId = Component->VisibilityId;
+					}
+					else if (DebugVisibilityId != Component->VisibilityId)
+					{
+						UE_LOG(LogLightmassSolver, Warning, TEXT("Not debugging visibility for component %s with vis id %u, as it was not the first component on the selected actor."),
+							*Component->GetPathName(), Component->VisibilityId);
+					}
 				}
 			}
 		}
diff --git a/Engine/Source/Editor/UnrealEd/Private/ThumbnailRendering/BlueprintThumbnailRenderer.cpp b/Engine/Source/Editor/UnrealEd/Private/ThumbnailRendering/BlueprintThumbnailRenderer.cpp
--- a/Engine/Source/Editor/UnrealEd/Private/ThumbnailRendering/BlueprintThumbnailRenderer.cpp
+++ b/Engine/Source/Editor/UnrealEd/Private/ThumbnailRendering/BlueprintThumbnailRenderer.cpp
@@ -22,12 +22,9 @@ bool UBlueprintThumbnailRenderer::CanVisualizeAsset(UObject* Object)
 		// Try to find any visible primitive components in the native class' CDO
 		AActor* CDO = Blueprint->GeneratedClass->GetDefaultObject<AActor>();
 
-		TInlineComponentArray<UActorComponent*> Components;
-		CDO->GetComponents(Components);
-
-		for (auto CompIt = Components.CreateConstIterator(); CompIt; ++CompIt)
+		for (UActorComponent* Component : CDO->GetComponents())
 		{
-			if (FBlueprintThumbnailScene::IsValidComponentForVisualization(*CompIt))
+			if (FBlueprintThumbnailScene::IsValidComponentForVisualization(Component))
 			{
 				return true;
 			}
diff --git a/Engine/Source/Editor/UnrealEd/Private/ThumbnailRendering/ClassThumbnailRenderer.cpp b/Engine/Source/Editor/UnrealEd/Private/ThumbnailRendering/ClassThumbnailRenderer.cpp
--- a/Engine/Source/Editor/UnrealEd/Private/ThumbnailRendering/ClassThumbnailRenderer.cpp
+++ b/Engine/Source/Editor/UnrealEd/Private/ThumbnailRendering/ClassThumbnailRenderer.cpp
@@ -20,12 +20,9 @@ bool UClassThumbnailRenderer::CanVisualizeAsset(UObject* Object)
 		// Try to find any visible primitive components in the class' CDO
 		AActor* CDO = Class->GetDefaultObject<AActor>();
 
-		TInlineComponentArray<UActorComponent*> Components;
-		CDO->GetComponents(Components);
-
-		for (auto CompIt = Components.CreateConstIterator(); CompIt; ++CompIt)
+		for (UActorComponent* Component : CDO->GetComponents())
 		{
-			if (FClassThumbnailScene::IsValidComponentForVisualization(*CompIt))
+			if (FClassThumbnailScene::IsValidComponentForVisualization(Component))
 			{
 				return true;
 			}
diff --git a/Engine/Source/Editor/UnrealEd/Private/UnrealEdSrv.cpp b/Engine/Source/Editor/UnrealEd/Private/UnrealEdSrv.cpp
--- a/Engine/Source/Editor/UnrealEd/Private/UnrealEdSrv.cpp
+++ b/Engine/Source/Editor/UnrealEd/Private/UnrealEdSrv.cpp
@@ -1823,14 +1823,11 @@ TArray<FPoly*> GetSelectedPolygons()
 		checkSlow( Actor->IsA(AActor::StaticClass()) );
 		FTransform ActorToWorld = Actor->ActorToWorld();
 		
-		TInlineComponentArray<UStaticMeshComponent*> StaticMeshComponents;
-		Actor->GetComponents(StaticMeshComponents);
-
-		for(int32 j=0; j<StaticMeshComponents.Num(); j++)
+		for (UActorComponent* Component : Actor->GetComponents())
 		{
 			// If its a static mesh component, with a static mesh
-			UStaticMeshComponent* SMComp = StaticMeshComponents[j];
-			if(SMComp->IsRegistered() && SMComp->GetStaticMesh())
+			UStaticMeshComponent* SMComp = Cast<UStaticMeshComponent>(Component);
+			if (SMComp && SMComp->IsRegistered() && SMComp->GetStaticMesh())
 			{
 				UStaticMesh* StaticMesh = SMComp->GetStaticMesh();
 				if ( StaticMesh )
diff --git a/Engine/Source/Editor/UnrealEd/Private/VertexSnapping.cpp b/Engine/Source/Editor/UnrealEd/Private/VertexSnapping.cpp
--- a/Engine/Source/Editor/UnrealEd/Private/VertexSnapping.cpp
+++ b/Engine/Source/Editor/UnrealEd/Private/VertexSnapping.cpp
@@ -272,13 +272,10 @@ void FVertexSnappingImpl::ClearSnappingHelpers( bool bClearImmediately )
 
 static void DrawSnapVertices( AActor* Actor, float PointSize, FPrimitiveDrawInterface* PDI )
 {
-	TInlineComponentArray<UActorComponent*> Components;
-	Actor->GetComponents(Components);
-
 	// Get the closest vertex on each component
-	for( int32 ComponentIndex  = 0; ComponentIndex < Components.Num(); ++ComponentIndex )
+	for (UActorComponent* Component : Actor->GetComponents())
 	{
-		TSharedPtr<FVertexIterator> VertexGetter = MakeVertexIterator( Cast<UPrimitiveComponent>( Components[ComponentIndex] ) );
+		TSharedPtr<FVertexIterator> VertexGetter = MakeVertexIterator( Cast<UPrimitiveComponent>(Component) );
 		if( VertexGetter.IsValid() )
 		{
 			FVertexIterator& VertexGetterRef = *VertexGetter;
@@ -472,13 +469,16 @@ FSnappingVertex FVertexSnappingImpl::GetClosestVertex( const TArray<FSnapActor>&
 		AActor* Actor = SnapActor.Actor;
 
 		// Get the closest vertex on each component
-		TInlineComponentArray<UPrimitiveComponent*> PrimitiveComponents;
-		Actor->GetComponents(PrimitiveComponents);
-
-		for( int32 ComponentIndex  = 0; ComponentIndex < PrimitiveComponents.Num(); ++ComponentIndex )
+		for (UActorComponent* Component : Actor->GetComponents())
 		{
+			UPrimitiveComponent* PrimComp = Cast<UPrimitiveComponent>(Component);
+			if (PrimComp == nullptr)
+			{
+				continue;
+			}
+
 			FSnappingVertex ClosestLocationOnComponent( CurrentLocation );
-			if( !GetClosestVertexOnComponent( SnapActor, PrimitiveComponents[ComponentIndex], InArgs, ClosestLocationOnComponent ) )
+			if( !GetClosestVertexOnComponent( SnapActor, PrimComp, InArgs, ClosestLocationOnComponent ) )
 			{
 				ClosestLocationOnComponent.Position = Actor->GetActorLocation();
 				ClosestLocationOnComponent.Normal = FVector::ZeroVector;
diff --git a/Engine/Source/Runtime/AIModule/Private/AIController.cpp b/Engine/Source/Runtime/AIModule/Private/AIController.cpp
--- a/Engine/Source/Runtime/AIModule/Private/AIController.cpp
+++ b/Engine/Source/Runtime/AIModule/Private/AIController.cpp
@@ -67,11 +67,12 @@ void AAIController::PostInitializeComponents()
 	}
 
 #if ENABLE_VISUAL_LOG
-	TInlineComponentArray<UActorComponent*> ComponentSet;
-	GetComponents(ComponentSet);
-	for (auto Component : ComponentSet)
+	for (UActorComponent* Component : GetComponents())
 	{
-		REDIRECT_OBJECT_TO_VLOG(Component, this);
+		if (Component)
+		{
+			REDIRECT_OBJECT_TO_VLOG(Component, this);
+		}
 	}
 #endif // ENABLE_VISUAL_LOG
 }
diff --git a/Engine/Source/Runtime/AIModule/Private/BehaviorTree/BlueprintNodeHelpers.cpp b/Engine/Source/Runtime/AIModule/Private/BehaviorTree/BlueprintNodeHelpers.cpp
--- a/Engine/Source/Runtime/AIModule/Private/BehaviorTree/BlueprintNodeHelpers.cpp
+++ b/Engine/Source/Runtime/AIModule/Private/BehaviorTree/BlueprintNodeHelpers.cpp
@@ -247,12 +247,9 @@ namespace BlueprintNodeHelpers
 
 		if (OwningActor && !bFound)
 		{
-			TInlineComponentArray<UBehaviorTreeComponent*> Components;
-			OwningActor->GetComponents(Components);
-			for (int32 ComponentIndex = 0; ComponentIndex < Components.Num(); ComponentIndex++)
+			for (UActorComponent* Component : OwningActor->GetComponents())
 			{
-				UBehaviorTreeComponent* BTComp = Components[ComponentIndex];
-				if (BTComp)
+				if (UBehaviorTreeComponent* BTComp = Cast<UBehaviorTreeComponent>(Component))
 				{
 					const int32 InstanceIdx = BTComp->FindInstanceContainingNode(Node);
 					if (InstanceIdx != INDEX_NONE)
diff --git a/Engine/Source/Runtime/Engine/Private/Actor.cpp b/Engine/Source/Runtime/Engine/Private/Actor.cpp
--- a/Engine/Source/Runtime/Engine/Private/Actor.cpp
+++ b/Engine/Source/Runtime/Engine/Private/Actor.cpp
@@ -1449,13 +1449,10 @@ void AActor::NotifyHit(class UPrimitiveComponent* MyComp, AActor* Other, class U
  */
 static void MarkOwnerRelevantComponentsDirty(AActor* TheActor)
 {
-	TInlineComponentArray<UPrimitiveComponent*> Components;
-	TheActor->GetComponents(Components);
-
-	for (int32 i = 0; i < Components.Num(); i++)
+	for (UActorComponent* Component : TheActor->GetComponents())
 	{
-		UPrimitiveComponent* Primitive = Components[i];
-		if (Primitive->IsRegistered() && (Primitive->bOnlyOwnerSee || Primitive->bOwnerNoSee))
+		UPrimitiveComponent* Primitive = Cast<UPrimitiveComponent>(Component);
+		if (Primitive && Primitive->IsRegistered() && (Primitive->bOnlyOwnerSee || Primitive->bOwnerNoSee))
 		{
 			Primitive->MarkRenderStateDirty();
 		}
@@ -1975,14 +1972,11 @@ void AActor::PrestreamTextures( float Seconds, bool bEnableStreaming, int32 Cine
 	}
 
 	// Iterate over all components of that actor
-	TInlineComponentArray<UMeshComponent*> Components;
-	GetComponents(Components);
-
-	for (int32 ComponentIndex=0; ComponentIndex < Components.Num(); ComponentIndex++)
+	for (UActorComponent* Component : GetComponents())
 	{
 		// If its a static mesh component, with a static mesh
-		UMeshComponent* MeshComponent = Components[ComponentIndex];
-		if ( MeshComponent->IsRegistered() )
+		UMeshComponent* MeshComponent = Cast<UMeshComponent>(Component);
+		if (MeshComponent && MeshComponent->IsRegistered() )
 		{
 			MeshComponent->PrestreamTextures( Duration, false, CinematicTextureGroups );
 		}
@@ -2054,14 +2048,11 @@ FVector AActor::GetPlacementExtent() const
 	FVector Extent(0.f);
 	if( (RootComponent && GetRootComponent()->ShouldCollideWhenPlacing()) && bCollideWhenPlacing) 
 	{
-		TInlineComponentArray<USceneComponent*> Components;
-		GetComponents(Components);
-
 		FBox ActorBox(ForceInit);
-		for (int32 ComponentID=0; ComponentID<Components.Num(); ++ComponentID)
+		for (UActorComponent* Component : GetComponents())
 		{
-			USceneComponent* SceneComp = Components[ComponentID];
-			if (SceneComp->ShouldCollideWhenPlacing() )
+			USceneComponent* SceneComp = Cast<USceneComponent>(Component);
+			if (SceneComp && SceneComp->ShouldCollideWhenPlacing())
 			{
 				ActorBox += SceneComp->GetPlacementExtent().GetBox();
 			}
@@ -2372,17 +2363,16 @@ void AActor::DisplayDebug(UCanvas* Canvas, const FDebugDisplayInfo& DebugDisplay
 	static FName NAME_Bones = FName(TEXT("Bones"));
 	if (DebugDisplay.IsDisplayOn(NAME_Animation) || DebugDisplay.IsDisplayOn(NAME_Bones))
 	{
-		TInlineComponentArray<USkeletalMeshComponent*> Components;
-		GetComponents(Components);
-
 		if (DebugDisplay.IsDisplayOn(NAME_Animation))
 		{
-			for (USkeletalMeshComponent* Comp : Components)
+			for (UActorComponent* Comp : GetComponents())
 			{
-				UAnimInstance* AnimInstance = Comp->GetAnimInstance();
-				if (AnimInstance)
+				if (USkeletalMeshComponent* SkelMeshComp = Cast<USkeletalMeshComponent>(Comp))
 				{
-					AnimInstance->DisplayDebug(Canvas, DebugDisplay, YL, YPos);
+					if (UAnimInstance* AnimInstance = SkelMeshComp->GetAnimInstance())
+					{
+						AnimInstance->DisplayDebug(Canvas, DebugDisplay, YL, YPos);
+					}
 				}
 			}
 		}
@@ -2715,12 +2705,12 @@ TArray<UActorComponent*> AActor::GetComponentsByTag(TSubclassOf<UActorComponent>
 
 void AActor::DisableComponentsSimulatePhysics()
 {
-	TInlineComponentArray<UPrimitiveComponent*> Components;
-	GetComponents(Components);
-
-	for (UPrimitiveComponent* Component : Components)
+	for (UActorComponent* Component : GetComponents())
 	{
-		Component->SetSimulatePhysics(false);
+		if (UPrimitiveComponent* PrimComp = Cast<UPrimitiveComponent>(Component))
+		{
+			PrimComp->SetSimulatePhysics(false);
+		}
 	}
 }
 
@@ -2740,7 +2730,7 @@ static void DispatchOnComponentsCreated(AActor* NewActor)
 
 	for (UActorComponent* ActorComp : Components)
 	{
-		if (ActorComp && !ActorComp->HasBeenCreated())
+		if (!ActorComp->HasBeenCreated())
 		{
 			ActorComp->OnComponentCreated();
 		}
@@ -4329,29 +4319,27 @@ void AActor::UninitializeComponents()
 void AActor::DrawDebugComponents(FColor const& BaseColor) const
 {
 #if ENABLE_DRAW_DEBUG
-	TInlineComponentArray<USceneComponent*> Components;
-	GetComponents(Components);
-
 	UWorld* MyWorld = GetWorld();
 
-	for(int32 ComponentIndex = 0; ComponentIndex < Components.Num(); ComponentIndex++)
+	for (UActorComponent* ActorComp : GetComponents())
 	{
-		USceneComponent const* const Component = Components[ComponentIndex]; 
+		if (USceneComponent const* const Component = Cast<USceneComponent>(ActorComp))
+		{
+			FVector const Loc = Component->GetComponentLocation();
+			FRotator const Rot = Component->GetComponentRotation();
 
-		FVector const Loc = Component->GetComponentLocation();
-		FRotator const Rot = Component->GetComponentRotation();
+			// draw coord system at component loc
+			DrawDebugCoordinateSystem(MyWorld, Loc, Rot, 10.f);
 
-		// draw coord system at component loc
-		DrawDebugCoordinateSystem(MyWorld, Loc, Rot, 10.f);
+			// draw line from me to my parent
+			if (Component->GetAttachParent())
+			{
+				DrawDebugLine(MyWorld, Component->GetAttachParent()->GetComponentLocation(), Loc, BaseColor);
+			}
 
-		// draw line from me to my parent
-		if (Component->GetAttachParent())
-		{
-			DrawDebugLine(MyWorld, Component->GetAttachParent()->GetComponentLocation(), Loc, BaseColor);
+			// draw component name
+			DrawDebugString(MyWorld, Loc+FVector(0,0,32), *Component->GetName());
 		}
-
-		// draw component name
-		DrawDebugString(MyWorld, Loc+FVector(0,0,32), *Component->GetName());
 	}
 #endif // ENABLE_DRAW_DEBUG
 }
@@ -4377,14 +4365,11 @@ bool AActor::ActorLineTraceSingle(struct FHitResult& OutHit, const FVector& Star
 	OutHit.TraceEnd = End;
 	bool bHasHit = false;
 	
-	TInlineComponentArray<UPrimitiveComponent*> Components;
-	GetComponents(Components);
-
-	for (int32 ComponentIndex=0; ComponentIndex<Components.Num(); ComponentIndex++)
+	for (UActorComponent* Component : GetComponents())
 	{
 		FHitResult HitResult;
-		UPrimitiveComponent* Primitive = Components[ComponentIndex];
-		if( Primitive->IsRegistered() && Primitive->IsCollisionEnabled() 
+		UPrimitiveComponent* Primitive = Cast<UPrimitiveComponent>(Component);
+		if (Primitive && Primitive->IsRegistered() && Primitive->IsCollisionEnabled() 
 			&& (Primitive->GetCollisionResponseToChannel(TraceChannel) == ECollisionResponse::ECR_Block) 
 			&& Primitive->LineTraceComponent(HitResult, Start, End, Params) )
 		{
diff --git a/Engine/Source/Runtime/Engine/Private/ActorEditor.cpp b/Engine/Source/Runtime/Engine/Private/ActorEditor.cpp
--- a/Engine/Source/Runtime/Engine/Private/ActorEditor.cpp
+++ b/Engine/Source/Runtime/Engine/Private/ActorEditor.cpp
@@ -818,13 +818,9 @@ void AActor::CheckForErrors()
 	}
 
 	// Route error checking to components.
-	TInlineComponentArray<UActorComponent*> Components;
-	GetComponents(Components);
-
-	for ( int32 ComponentIndex = 0 ; ComponentIndex < Components.Num() ; ++ComponentIndex )
+	for (UActorComponent* ActorComponent : GetComponents())
 	{
-		UActorComponent* ActorComponent = Components[ ComponentIndex ];
-		if (ActorComponent->IsRegistered())
+		if (ActorComponent && ActorComponent->IsRegistered())
 		{
 			ActorComponent->CheckForErrors();
 		}
diff --git a/Engine/Source/Runtime/Engine/Private/ActorEditorUtils.cpp b/Engine/Source/Runtime/Engine/Private/ActorEditorUtils.cpp
--- a/Engine/Source/Runtime/Engine/Private/ActorEditorUtils.cpp
+++ b/Engine/Source/Runtime/Engine/Private/ActorEditorUtils.cpp
@@ -45,25 +45,26 @@ namespace FActorEditorUtils
 
 	void GetEditableComponents( const AActor* InActor, TArray<UActorComponent*>& OutEditableComponents )
 	{
-		TInlineComponentArray<UActorComponent*> InstanceComponents;
-		InActor->GetComponents(InstanceComponents);
-		for (UActorComponent* Component : InstanceComponents)
+		for (UActorComponent* Component : InActor->GetComponents())
 		{
-#if WITH_EDITOR
-			if (Component->CreationMethod == EComponentCreationMethod::Native)
+			if (Component)
 			{
-				// Make sure it's an exposed native component
-				if (FComponentEditorUtils::CanEditNativeComponent(Component))
+#if WITH_EDITOR
+				if (Component->CreationMethod == EComponentCreationMethod::Native)
 				{
-					OutEditableComponents.Add(Component);
+					// Make sure it's an exposed native component
+					if (FComponentEditorUtils::CanEditNativeComponent(Component))
+					{
+						OutEditableComponents.Add(Component);
+					}
 				}
-			}
-			else if (Component->CreationMethod == EComponentCreationMethod::Instance)
+				else if (Component->CreationMethod == EComponentCreationMethod::Instance)
 #else
-			if (!Component->IsCreatedByConstructionScript())
+				if (!Component->IsCreatedByConstructionScript())
 #endif
-			{
-				OutEditableComponents.Add(Component);
+				{
+					OutEditableComponents.Add(Component);
+				}
 			}
 		}
 	}
diff --git a/Engine/Source/Runtime/Engine/Private/DebugCameraHUD.cpp b/Engine/Source/Runtime/Engine/Private/DebugCameraHUD.cpp
--- a/Engine/Source/Runtime/Engine/Private/DebugCameraHUD.cpp
+++ b/Engine/Source/Runtime/Engine/Private/DebugCameraHUD.cpp
@@ -141,13 +141,10 @@ void ADebugCameraHUD::PostRender()
 				}
 				else
 				{
-					TInlineComponentArray<UMeshComponent*> Components;
-					GetComponents(Components);
-
-					for ( int32 i=0; i<Components.Num(); i++ )
+					for (UActorComponent* Component : GetComponents())
 					{
-						UMeshComponent* MeshComp = Components[i];
-						if ( MeshComp->IsRegistered() )
+						UMeshComponent* MeshComp = Cast<UMeshComponent>(Component);
+						if (MeshComp && MeshComp->IsRegistered())
 						{
 							bFoundMaterial = bFoundMaterial || DisplayMaterials( X, yl, Y, MeshComp );	
 						}
diff --git a/Engine/Source/Runtime/Engine/Private/LODActor.cpp b/Engine/Source/Runtime/Engine/Private/LODActor.cpp
--- a/Engine/Source/Runtime/Engine/Private/LODActor.cpp
+++ b/Engine/Source/Runtime/Engine/Private/LODActor.cpp
@@ -747,23 +747,25 @@ const bool ALODActor::HasValidSubActors() const
 	{
 		if (SubActor)
 		{
-			TInlineComponentArray<UStaticMeshComponent*> Components;
-			SubActor->GetComponents(/*out*/ Components);
-
 #if WITH_EDITOR
 			FHierarchicalLODUtilitiesModule& Module = FModuleManager::LoadModuleChecked<FHierarchicalLODUtilitiesModule>("HierarchicalLODUtilities");
 			IHierarchicalLODUtilities* Utilities = Module.GetUtilities();
+#endif
 
-			for (UStaticMeshComponent* Component : Components)
+			for (UActorComponent* Comp : GetComponents())
 			{
-				if (!Component->bHiddenInGame && Component->ShouldGenerateAutoLOD(LODLevel - 1))
+				if (UStaticMeshComponent* Component = Cast<UStaticMeshComponent>(Comp))
 				{
+#if WITH_EDITOR
+					if (!Component->bHiddenInGame && Component->ShouldGenerateAutoLOD(LODLevel - 1))
+					{
+						++NumMeshes;
+					}
+#else
 					++NumMeshes;
+#endif
 				}
 			}
-#else
-			NumMeshes += Components.Num();
-#endif
 
 			if (NumMeshes > 0)
 			{
diff --git a/Engine/Source/Runtime/Engine/Private/LocalPlayer.cpp b/Engine/Source/Runtime/Engine/Private/LocalPlayer.cpp
--- a/Engine/Source/Runtime/Engine/Private/LocalPlayer.cpp
+++ b/Engine/Source/Runtime/Engine/Private/LocalPlayer.cpp
@@ -1237,16 +1237,14 @@ bool ULocalPlayer::HandleListPawnComponentsCommand( const TCHAR* Cmd, FOutputDev
 		APawn *Pawn = *It;
 		UE_LOG(LogPlayerManagement, Log, TEXT("Components for pawn: %s (collision component: %s)"),*Pawn->GetName(),*Pawn->GetRootComponent()->GetName());
 
-		TInlineComponentArray<UActorComponent*> Components;
-		Pawn->GetComponents(Components);
-
-		for (int32 CompIdx = 0; CompIdx < Components.Num(); CompIdx++)
+		int32 CompIdx = 0;
+		for (UActorComponent* Comp : Pawn->GetComponents())
 		{
-			UActorComponent *Comp = Components[CompIdx];
-			if (Comp->IsRegistered())
+			if (Comp && Comp->IsRegistered())
 			{
 				UE_LOG(LogPlayerManagement, Log, TEXT("  %d: %s"),CompIdx,*Comp->GetName());
 			}
+			++CompIdx;
 		}
 	}
 	return true;
diff --git a/Engine/Source/Runtime/Engine/Private/Particles/ParticleSystemRender.cpp b/Engine/Source/Runtime/Engine/Private/Particles/ParticleSystemRender.cpp
--- a/Engine/Source/Runtime/Engine/Private/Particles/ParticleSystemRender.cpp
+++ b/Engine/Source/Runtime/Engine/Private/Particles/ParticleSystemRender.cpp
@@ -7326,16 +7326,14 @@ void DrawParticleSystemHelpers(UParticleSystemComponent* InPSysComp, const FScen
 
 ENGINE_API void DrawParticleSystemHelpers(const FSceneView* View,FPrimitiveDrawInterface* PDI)
 {
-	TInlineComponentArray<UParticleSystemComponent*> PSCArray;
-
 	for (TObjectIterator<AActor> It; It; ++It)
 	{
-		PSCArray.Reset();
-		(*It)->GetComponents(PSCArray);
-
-		for (int PSCIndex = 0; PSCIndex < PSCArray.Num(); ++PSCIndex)
+		for (UActorComponent* Component : It->GetComponents())
 		{
-			DrawParticleSystemHelpers(PSCArray[PSCIndex], View, PDI);
+			if (UParticleSystemComponent* PSC = Cast<UParticleSystemComponent>(Component))
+			{
+				DrawParticleSystemHelpers(PSC, View, PDI);
+			}
 		}
 	}
 }
diff --git a/Engine/Source/Runtime/Engine/Private/SCS_Node.cpp b/Engine/Source/Runtime/Engine/Private/SCS_Node.cpp
--- a/Engine/Source/Runtime/Engine/Private/SCS_Node.cpp
+++ b/Engine/Source/Runtime/Engine/Private/SCS_Node.cpp
@@ -553,12 +553,10 @@ USceneComponent* USCS_Node::GetParentComponentTemplate(UBlueprint* InBlueprint)
 			if(CDO != nullptr)
 			{
 				// Find the component template in the CDO that matches the specified name
-				TInlineComponentArray<USceneComponent*> Components;
-				CDO->GetComponents(Components);
-
-				for (USceneComponent* CompTemplate : Components)
+				for (UActorComponent* ActorComp : CDO->GetComponents())
 				{
-					if(CompTemplate->GetFName() == ParentComponentOrVariableName)
+					USceneComponent* CompTemplate = Cast<USceneComponent>(ActorComp);
+					if (CompTemplate && CompTemplate->GetFName() == ParentComponentOrVariableName)
 					{
 						// Found a match; this is our parent, we're done
 						ParentComponentTemplate = CompTemplate;
diff --git a/Engine/Source/Runtime/Engine/Private/SimpleConstructionScript.cpp b/Engine/Source/Runtime/Engine/Private/SimpleConstructionScript.cpp
--- a/Engine/Source/Runtime/Engine/Private/SimpleConstructionScript.cpp
+++ b/Engine/Source/Runtime/Engine/Private/SimpleConstructionScript.cpp
@@ -510,13 +510,13 @@ void USimpleConstructionScript::FixupRootNodeParentReferences()
 				if(CDO != NULL)
 				{
 					// Look for the parent component in the CDO's components array
-					TInlineComponentArray<UActorComponent*> Components;
-					CDO->GetComponents(Components);
-
-					for (auto CompIter = Components.CreateConstIterator(); CompIter && !bWasFound; ++CompIter)
+					for (UActorComponent* ComponentTemplate : CDO->GetComponents())
 					{
-						UActorComponent* ComponentTemplate = *CompIter;
-						bWasFound = ComponentTemplate->GetFName() == RootNode->ParentComponentOrVariableName;
+						if (ComponentTemplate && ComponentTemplate->GetFName() == RootNode->ParentComponentOrVariableName)
+						{
+							bWasFound = true;
+							break;
+						}
 					}
 				}
 				else 
@@ -969,11 +969,13 @@ USceneComponent* USimpleConstructionScript::GetSceneRootComponentTemplate(USCS_N
 		RootComponentTemplate = CDO->GetRootComponent();
 		if(!RootComponentTemplate)
 		{
-			TInlineComponentArray<USceneComponent*> SceneComponents;
-			CDO->GetComponents(SceneComponents);
-			if(SceneComponents.Num() > 0)
+			for (UActorComponent* Component : CDO->GetComponents())
 			{
-				RootComponentTemplate = SceneComponents[0];
+				if (USceneComponent* SceneComp = Cast<USceneComponent>(Component))
+				{
+					RootComponentTemplate = SceneComp;
+					break;
+				}
 			}
 		}
 	}
diff --git a/Engine/Source/Runtime/Engine/Private/Streaming/StreamingManagerTexture.cpp b/Engine/Source/Runtime/Engine/Private/Streaming/StreamingManagerTexture.cpp
--- a/Engine/Source/Runtime/Engine/Private/Streaming/StreamingManagerTexture.cpp
+++ b/Engine/Source/Runtime/Engine/Private/Streaming/StreamingManagerTexture.cpp
@@ -463,13 +463,10 @@ void FStreamingManagerTexture::BoostTextures( AActor* Actor, float BoostFactor )
 		TArray<UTexture*> Textures;
 		Textures.Empty( 32 );
 
-		TInlineComponentArray<UPrimitiveComponent*> Components;
-		Actor->GetComponents(Components);
-
-		for(int32 ComponentIndex = 0;ComponentIndex < Components.Num();ComponentIndex++)
+		for (UActorComponent* Component : Actor->GetComponents())
 		{
-			UPrimitiveComponent* Primitive = Components[ComponentIndex];
-			if ( Primitive->IsRegistered() )
+			UPrimitiveComponent* Primitive = Cast<UPrimitiveComponent>(Component);
+			if (Primitive && Primitive->IsRegistered())
 			{
 				Textures.Reset();
 				Primitive->GetUsedTextures( Textures, EMaterialQualityLevel::Num );
@@ -623,7 +620,7 @@ void FStreamingManagerTexture::NotifyActorDestroyed( AActor* Actor )
 		{
 			LevelManager.RemoveActorReferences(Actor);
 			for (UPrimitiveComponent* Component : Components)
-			{
+				{
 				LevelManager.RemoveComponentReferences(Component, RemovedTextures);
 			}
 		}
diff --git a/Engine/Source/Runtime/Engine/Private/UnrealEngine.cpp b/Engine/Source/Runtime/Engine/Private/UnrealEngine.cpp
--- a/Engine/Source/Runtime/Engine/Private/UnrealEngine.cpp
+++ b/Engine/Source/Runtime/Engine/Private/UnrealEngine.cpp
@@ -13496,12 +13496,9 @@ void UEngine::CopyPropertiesForUnrelatedObjects(UObject* OldObject, UObject* New
 	AActor* NewActor = Cast<AActor>(NewObject);
 	if (NewActor != nullptr)
 	{
-		TInlineComponentArray<UActorComponent*> Components;
-		NewActor->GetComponents(Components);
-
-		for(int32 i=0; i<Components.Num(); i++)
+		for (UActorComponent* Component : NewActor->GetComponents())
 		{
-			ensure(!Components[i]->IsRegistered());
+			ensure(Component == nullptr || !Component->IsRegistered());
 		}
 	}
 
diff --git a/Engine/Source/Runtime/Foliage/Private/InstancedFoliage.cpp b/Engine/Source/Runtime/Foliage/Private/InstancedFoliage.cpp
--- a/Engine/Source/Runtime/Foliage/Private/InstancedFoliage.cpp
+++ b/Engine/Source/Runtime/Foliage/Private/InstancedFoliage.cpp
@@ -2686,11 +2686,12 @@ void AInstancedFoliageActor::Serialize(FArchive& Ar)
 	// Clean up any old cluster components and convert to hierarchical instanced foliage.
 	if (Ar.CustomVer(FFoliageCustomVersion::GUID) < FFoliageCustomVersion::FoliageUsingHierarchicalISMC)
 	{
-		TInlineComponentArray<UInstancedStaticMeshComponent*> ClusterComponents;
-		GetComponents(ClusterComponents);
-		for (UInstancedStaticMeshComponent* Component : ClusterComponents)
+		for (UActorComponent* Component : GetComponents())
 		{
-			Component->bAutoRegister = false;
+			if (Cast<UInstancedStaticMeshComponent>(Component))
+			{
+				Component->bAutoRegister = false;
+			}
 		}
 	}
 }
@@ -2970,12 +2971,12 @@ void AInstancedFoliageActor::OnLevelActorMoved(AActor* InActor)
 
 	if (!InWorld || !InWorld->IsGameWorld())
 	{
-		TInlineComponentArray<UActorComponent*> Components;
-		InActor->GetComponents(Components);
-
-		for (auto Component : Components)
+		for (UActorComponent* Component : GetComponents())
 		{
-			MoveInstancesForMovedComponent(Component);
+			if (Component)
+			{
+				MoveInstancesForMovedComponent(Component);
+			}
 		}
 	}
 }
@@ -2986,12 +2987,12 @@ void AInstancedFoliageActor::OnLevelActorDeleted(AActor* InActor)
 
 	if (!InWorld || !InWorld->IsGameWorld())
 	{
-		TInlineComponentArray<UActorComponent*> Components;
-		InActor->GetComponents(Components);
-
-		for (auto Component : Components)
+		for (UActorComponent* Component : GetComponents())
 		{
-			DeleteInstancesForComponent(Component);
+			if (Component)
+			{
+				DeleteInstancesForComponent(Component);
+			}
 		}
 	}
 }
diff --git a/Engine/Source/Runtime/Landscape/Private/LandscapeGrass.cpp b/Engine/Source/Runtime/Landscape/Private/LandscapeGrass.cpp
--- a/Engine/Source/Runtime/Landscape/Private/LandscapeGrass.cpp
+++ b/Engine/Source/Runtime/Landscape/Private/LandscapeGrass.cpp
@@ -1946,11 +1946,12 @@ void ALandscapeProxy::FlushGrassComponents(const TSet<ULandscapeComponent*>* Onl
 		if (GIsEditor && bFlushGrassMaps && GetWorld() && GetWorld()->Scene->GetFeatureLevel() >= ERHIFeatureLevel::SM4)
 		{
 			// Clear GrassMaps
-			TInlineComponentArray<ULandscapeComponent*> LandComps;
-			GetComponents(LandComps);
-			for (ULandscapeComponent* Component : LandComps)
+			for (UActorComponent* Component : GetComponents())
 			{
-				Component->RemoveGrassMap();
+				if (ULandscapeComponent* LandscapeComp = Cast<ULandscapeComponent>(Component))
+				{
+					LandscapeComp->RemoveGrassMap();
+				}
 			}
 		}
 #endif
diff --git a/Engine/Source/Runtime/LevelSequence/Private/LevelSequenceActorSpawner.cpp b/Engine/Source/Runtime/LevelSequence/Private/LevelSequenceActorSpawner.cpp
--- a/Engine/Source/Runtime/LevelSequence/Private/LevelSequenceActorSpawner.cpp
+++ b/Engine/Source/Runtime/LevelSequence/Private/LevelSequenceActorSpawner.cpp
@@ -134,9 +134,12 @@ UObject* FLevelSequenceActorSpawner::SpawnObject(FMovieSceneSpawnable& Spawnable
 		// Explicitly set RF_Transactional on spawned actors so we can undo/redo properties on them. We don't add this as a spawn flag since we don't want to transact spawn/destroy events.
 		SpawnedActor->SetFlags(RF_Transactional);
 
-		for (UActorComponent* Component : TInlineComponentArray<UActorComponent*>(SpawnedActor))
+		for (UActorComponent* Component : SpawnedActor->GetComponents())
 		{
-			Component->SetFlags(RF_Transactional);
+			if (Component)
+			{
+				Component->SetFlags(RF_Transactional);
+			}
 		}
 	}
 
@@ -162,9 +165,12 @@ void FLevelSequenceActorSpawner::DestroySpawnedObject(UObject& Object)
 	{
 		// Explicitly remove RF_Transactional on spawned actors since we don't want to trasact spawn/destroy events
 		Actor->ClearFlags(RF_Transactional);
-		for (UActorComponent* Component : TInlineComponentArray<UActorComponent*>(Actor))
+		for (UActorComponent* Component : Actor->GetComponents())
 		{
-			Component->ClearFlags(RF_Transactional);
+			if (Component)
+			{
+				Component->ClearFlags(RF_Transactional);
+			}
 		}
 	}
 #endif
diff --git a/Engine/Source/Runtime/MovieSceneTracks/Private/Evaluation/MovieScene3DPathTemplate.cpp b/Engine/Source/Runtime/MovieSceneTracks/Private/Evaluation/MovieScene3DPathTemplate.cpp
--- a/Engine/Source/Runtime/MovieSceneTracks/Private/Evaluation/MovieScene3DPathTemplate.cpp
+++ b/Engine/Source/Runtime/MovieSceneTracks/Private/Evaluation/MovieScene3DPathTemplate.cpp
@@ -146,15 +146,21 @@ struct F3DPathExecutionToken
 		UObject* ParentObject = Objects[0].Get();
 		AActor* Actor = Cast<AActor>(ParentObject);
 
-		TArray<USplineComponent*> SplineComponents;
-		Actor->GetComponents(SplineComponents);
+		USplineComponent* SplineComponent = nullptr;
 
-		if (!SplineComponents.Num())
+		for (UActorComponent* Component : Actor->GetComponents())
 		{
-			return;
+			SplineComponent = Cast<USplineComponent>(Component);
+			if (SplineComponent)
+			{
+				break;
+			}
 		}
 
-		USplineComponent* SplineComponent = SplineComponents[0];
+		if (SplineComponent == nullptr)
+		{
+			return;
+		}
 
 		for (TWeakObjectPtr<> Object : Player.FindBoundObjects(Operand))
 		{
diff --git a/Engine/Source/Runtime/NavigationSystem/Private/NavModifierComponent.cpp b/Engine/Source/Runtime/NavigationSystem/Private/NavModifierComponent.cpp
--- a/Engine/Source/Runtime/NavigationSystem/Private/NavModifierComponent.cpp
+++ b/Engine/Source/Runtime/NavigationSystem/Private/NavModifierComponent.cpp
@@ -19,14 +19,11 @@ void UNavModifierComponent::CalcAndCacheBounds() const
 	const AActor* MyOwner = GetOwner();
 	if (MyOwner)
 	{
-		TInlineComponentArray<UPrimitiveComponent*> PrimComponents;
-		MyOwner->GetComponents(PrimComponents);
-
 		Bounds = FBox(ForceInit);
 		ComponentBounds.Reset();
-		for (int32 Idx = 0; Idx < PrimComponents.Num(); Idx++)
+		for (UActorComponent* Component : MyOwner->GetComponents())
 		{
-			UPrimitiveComponent* PrimComp = PrimComponents[Idx];
+			UPrimitiveComponent* PrimComp = Cast<UPrimitiveComponent>(Component);
 			if (PrimComp && PrimComp->IsRegistered() && PrimComp->IsCollisionEnabled())
 			{
 				UBodySetup* BodySetup = PrimComp->GetBodySetup();
diff --git a/Engine/Source/Runtime/NavigationSystem/Private/NavigationSystem.cpp b/Engine/Source/Runtime/NavigationSystem/Private/NavigationSystem.cpp
--- a/Engine/Source/Runtime/NavigationSystem/Private/NavigationSystem.cpp
+++ b/Engine/Source/Runtime/NavigationSystem/Private/NavigationSystem.cpp
@@ -2645,13 +2645,12 @@ void UNavigationSystemV1::UpdateActorAndComponentsInNavOctree(AActor& Actor, boo
 {
 	UpdateActorInNavOctree(Actor);
 		
-	TInlineComponentArray<UActorComponent*> Components;
-	Actor.GetComponents(Components);
-
-	for (int32 ComponentIndex = 0; ComponentIndex < Components.Num(); ComponentIndex++)
+	for (UActorComponent* Component : Actor.GetComponents())
 	{
-		check(Components[ComponentIndex]);
-		UpdateComponentInNavOctree(*Components[ComponentIndex]);
+		if (Component)
+		{
+			UpdateComponentInNavOctree(*Component);
+		}
 	}
 
 	if (bUpdateAttachedActors)
@@ -2697,12 +2696,9 @@ void UNavigationSystemV1::UpdateAttachedActorsInNavOctree(AActor& RootActor)
 
 void UNavigationSystemV1::UpdateNavOctreeBounds(AActor* Actor)
 {
-	TInlineComponentArray<UActorComponent*> Components;
-	Actor->GetComponents(Components);
-
-	for (int32 Idx = 0; Idx < Components.Num(); Idx++)
+	for (UActorComponent* Component : Actor->GetComponents())
 	{
-		INavRelevantInterface* NavElement = Cast<INavRelevantInterface>(Components[Idx]);
+		INavRelevantInterface* NavElement = Cast<INavRelevantInterface>(Component);
 		if (NavElement)
 		{
 			NavElement->UpdateNavigationBounds();
diff --git a/Engine/Source/Runtime/Renderer/Private/SceneCaptureRendering.cpp b/Engine/Source/Runtime/Renderer/Private/SceneCaptureRendering.cpp
--- a/Engine/Source/Runtime/Renderer/Private/SceneCaptureRendering.cpp
+++ b/Engine/Source/Runtime/Renderer/Private/SceneCaptureRendering.cpp
@@ -540,11 +540,12 @@ void SetupViewVamilyForSceneCapture(
 
 			if (Actor)
 			{
-				TInlineComponentArray<UPrimitiveComponent*> PrimitiveComponents;
-				Actor->GetComponents(PrimitiveComponents);
-				for (int32 ComponentIndex = 0; ComponentIndex < PrimitiveComponents.Num(); ++ComponentIndex)
+				for (UActorComponent* Component : Actor->GetComponents())
 				{
-					View->HiddenPrimitives.Add(PrimitiveComponents[ComponentIndex]->ComponentId);
+					if (UPrimitiveComponent* PrimComp = Cast<UPrimitiveComponent>(Component))
+					{
+						View->HiddenPrimitives.Add(PrimComp->ComponentId);
+					}
 				}
 			}
 		}
@@ -569,11 +570,12 @@ void SetupViewVamilyForSceneCapture(
 
 				if (Actor)
 				{
-					TInlineComponentArray<UPrimitiveComponent*> PrimitiveComponents;
-					Actor->GetComponents(PrimitiveComponents);
-					for (int32 ComponentIndex = 0; ComponentIndex < PrimitiveComponents.Num(); ++ComponentIndex)
+					for (UActorComponent* Component : Actor->GetComponents())
 					{
-						View->ShowOnlyPrimitives->Add(PrimitiveComponents[ComponentIndex]->ComponentId);
+						if (UPrimitiveComponent* PrimComp = Cast<UPrimitiveComponent>(Component))
+						{
+							View->ShowOnlyPrimitives->Add(PrimComp->ComponentId);
+						}
 					}
 				}
 			}
