SceneCapture DrawFrustum and CineCamera visualization properties should be editoronly
For backwards compatibility bIsEditorOnly not editor only ... sigh
#rnx

diff --git a/Engine/Source/Runtime/CinematicCamera/Private/CineCameraComponent.cpp b/Engine/Source/Runtime/CinematicCamera/Private/CineCameraComponent.cpp
--- a/Engine/Source/Runtime/CinematicCamera/Private/CineCameraComponent.cpp
+++ b/Engine/Source/Runtime/CinematicCamera/Private/CineCameraComponent.cpp
@@ -47,13 +47,12 @@ UCineCameraComponent::UCineCameraComponent()
 
 	RecalcDerivedData();
 
+#if WITH_EDITORONLY_DATA
 	if (!IsRunningCommandlet())
 	{
-#if WITH_EDITORONLY_DATA
 		// overrides CameraComponent's camera mesh
 		static ConstructorHelpers::FObjectFinder<UStaticMesh> EditorCameraMesh(TEXT("/Engine/EditorMeshes/Camera/SM_CineCam.SM_CineCam"));
 		CameraMesh = EditorCameraMesh.Object;
-#endif
 	}
 
 	static ConstructorHelpers::FObjectFinder<UStaticMesh> PlaneMesh(TEXT("/Engine/ArtTools/RenderToTexture/Meshes/S_1_Unit_Plane.S_1_Unit_Plane"));
@@ -61,6 +60,7 @@ UCineCameraComponent::UCineCameraComponent()
 
 	static ConstructorHelpers::FObjectFinder<UMaterial> PlaneMat(TEXT("/Engine/EngineDebugMaterials/M_SimpleTranslucent.M_SimpleTranslucent"));
 	FocusPlaneVisualizationMaterial = PlaneMat.Object;
+#endif
 }
 
 void UCineCameraComponent::PostInitProperties()
@@ -90,6 +90,7 @@ static const FColor DebugFocusPointOutlineColor = FColor::Black;
 
 void UCineCameraComponent::TickComponent(float DeltaTime, ELevelTick TickType, FActorComponentTickFunction* ThisTickFunction)
 {
+#if WITH_EDITORONLY_DATA
 	// make sure drawing is set up
 	if (FocusSettings.bDrawDebugFocusPlane)
 	{
@@ -107,6 +108,7 @@ void UCineCameraComponent::TickComponent(float DeltaTime, ELevelTick TickType, F
 			DestroyDebugFocusPlane();
 		}
 	}
+#endif
 
 #if ENABLE_DRAW_DEBUG
 	if (FocusSettings.TrackingFocusSettings.bDrawDebugTrackingFocusPoint)
@@ -337,6 +339,7 @@ void UCineCameraComponent::GetCameraView(float DeltaTime, FMinimalViewInfo& Desi
 	bResetInterpolation = false;
 }
 
+#if WITH_EDITORONLY_DATA
 void UCineCameraComponent::UpdateDebugFocusPlane()
 {
 	if (FocusPlaneVisualizationMesh && DebugFocusPlaneComponent)
@@ -351,6 +354,7 @@ void UCineCameraComponent::UpdateDebugFocusPlane()
 		DebugFocusPlaneComponent->SetWorldLocation(FocusPoint);
 	}
 }
+#endif
 
 void UCineCameraComponent::UpdateCameraLens(float DeltaTime, FMinimalViewInfo& DesiredView)
 {
@@ -411,6 +415,7 @@ void UCineCameraComponent::NotifyCameraCut()
 	bResetInterpolation = true;
 }
 
+#if WITH_EDITORONLY_DATA
 void UCineCameraComponent::CreateDebugFocusPlane()
 {
 	if (AActor* const MyOwner = GetOwner())
@@ -452,6 +457,7 @@ void UCineCameraComponent::DestroyDebugFocusPlane()
 		DebugFocusPlaneMID = nullptr;
 	}
 }
+#endif
 
 void UCineCameraComponent::OnRegister()
 {
diff --git a/Engine/Source/Runtime/CinematicCamera/Public/CineCameraComponent.h b/Engine/Source/Runtime/CinematicCamera/Public/CineCameraComponent.h
--- a/Engine/Source/Runtime/CinematicCamera/Public/CineCameraComponent.h
+++ b/Engine/Source/Runtime/CinematicCamera/Public/CineCameraComponent.h
@@ -275,8 +275,10 @@ public:
 	/** Returns a list of available lens presets. */
 	static TArray<FNamedLensPreset> const& GetLensPresets();
 
+#if WITH_EDITOR
 	/** Update the debug focus plane position and orientation. */
 	void UpdateDebugFocusPlane();
+#endif
 
 protected:
 
@@ -300,6 +302,7 @@ protected:
 	virtual void OnComponentDestroyed(bool bDestroyingHierarchy) override;
 #endif
 
+#if WITH_EDITORONLY_DATA
 	/** Mesh used for debug focus plane visualization */
 	UPROPERTY(transient)
 	UStaticMesh* FocusPlaneVisualizationMesh;
@@ -316,7 +319,6 @@ protected:
 	UPROPERTY(transient)
 	UMaterialInstanceDynamic* DebugFocusPlaneMID;
 
-#if WITH_EDITORONLY_DATA
 	virtual void PostEditChangeProperty(FPropertyChangedEvent& PropertyChangedEvent) override;
 	virtual void ResetProxyMeshTransform() override;
 #endif
@@ -355,6 +357,8 @@ private:
 	float GetDesiredFocusDistance(const FVector& InLocation) const;
 	float GetWorldToMetersScale() const;
 
+#if WITH_EDITORONLY_DATA
 	void CreateDebugFocusPlane();
 	void DestroyDebugFocusPlane();
+#endif
 };
diff --git a/Engine/Source/Runtime/Engine/Classes/Engine/SceneCapture2D.h b/Engine/Source/Runtime/Engine/Classes/Engine/SceneCapture2D.h
--- a/Engine/Source/Runtime/Engine/Classes/Engine/SceneCapture2D.h
+++ b/Engine/Source/Runtime/Engine/Classes/Engine/SceneCapture2D.h
@@ -21,12 +21,15 @@ private:
 	/** Scene capture component. */
 	UPROPERTY(Category = DecalActor, VisibleAnywhere, BlueprintReadOnly, meta=(AllowPrivateAccess = "true"))
 	class USceneCaptureComponent2D* CaptureComponent2D;
+#if WITH_EDITORONLY_DATA
 	/** To allow drawing the camera frustum in the editor. */
 	UPROPERTY()
 	class UDrawFrustumComponent* DrawFrustum;
+#endif
 
 public:
 
+#if WITH_EDITOR
 	//~ Begin AActor Interface
 	ENGINE_API virtual void PostActorCreated() override;
 	//~ End AActor Interface.
@@ -34,13 +37,15 @@ public:
 	/** Used to synchronize the DrawFrustumComponent with the SceneCaptureComponent2D settings. */
 	void UpdateDrawFrustum();
 
+	/** Returns DrawFrustum subobject **/
+	ENGINE_API class UDrawFrustumComponent* GetDrawFrustum() const { return DrawFrustum; }
+#endif
+
 	UFUNCTION(BlueprintCallable, Category="Rendering")
 	void OnInterpToggle(bool bEnable);
 
 	/** Returns CaptureComponent2D subobject **/
 	ENGINE_API class USceneCaptureComponent2D* GetCaptureComponent2D() const { return CaptureComponent2D; }
-	/** Returns DrawFrustum subobject **/
-	ENGINE_API class UDrawFrustumComponent* GetDrawFrustum() const { return DrawFrustum; }
 };
 
 
diff --git a/Engine/Source/Runtime/Engine/Classes/Engine/SceneCaptureCube.h b/Engine/Source/Runtime/Engine/Classes/Engine/SceneCaptureCube.h
--- a/Engine/Source/Runtime/Engine/Classes/Engine/SceneCaptureCube.h
+++ b/Engine/Source/Runtime/Engine/Classes/Engine/SceneCaptureCube.h
@@ -22,29 +22,32 @@ private:
 	UPROPERTY(Category = DecalActor, VisibleAnywhere, BlueprintReadOnly, meta=(AllowPrivateAccess = "true"))
 	class USceneCaptureComponentCube* CaptureComponentCube;
 
+#if WITH_EDITORONLY_DATA
 	/** To allow drawing the camera frustum in the editor. */
 	UPROPERTY()
 	class UDrawFrustumComponent* DrawFrustum;
+#endif
 
 public:
 
+#if WITH_EDITOR
 	//~ Begin AActor Interface
 	virtual void PostActorCreated() override;
-#if WITH_EDITOR
 	virtual void PostEditMove(bool bFinished) override;
-#endif
 	//~ End AActor Interface.
 
 	/** Used to synchronize the DrawFrustumComponent with the SceneCaptureComponentCube settings. */
 	void UpdateDrawFrustum();
 
+	/** Returns DrawFrustum subobject **/
+	class UDrawFrustumComponent* GetDrawFrustum() const { return DrawFrustum; }
+#endif
+
 	UFUNCTION(BlueprintCallable, Category = "Rendering")
 	void OnInterpToggle(bool bEnable);
 
 	/** Returns CaptureComponentCube subobject **/
 	class USceneCaptureComponentCube* GetCaptureComponentCube() const { return CaptureComponentCube; }
-	/** Returns DrawFrustum subobject **/
-	class UDrawFrustumComponent* GetDrawFrustum() const { return DrawFrustum; }
 };
 
 
diff --git a/Engine/Source/Runtime/Engine/Private/Components/SceneCaptureComponent.cpp b/Engine/Source/Runtime/Engine/Private/Components/SceneCaptureComponent.cpp
--- a/Engine/Source/Runtime/Engine/Private/Components/SceneCaptureComponent.cpp
+++ b/Engine/Source/Runtime/Engine/Private/Components/SceneCaptureComponent.cpp
@@ -50,9 +50,11 @@ ASceneCapture::ASceneCapture(const FObjectInitializer& ObjectInitializer)
 ASceneCapture2D::ASceneCapture2D(const FObjectInitializer& ObjectInitializer)
 	: Super(ObjectInitializer)
 {
+#if WITH_EDITOR
 	DrawFrustum = CreateDefaultSubobject<UDrawFrustumComponent>(TEXT("DrawFrust0"));
 	DrawFrustum->SetIsVisualizationComponent(true);
 	DrawFrustum->SetupAttachment(GetMeshComp());
+#endif
 
 	CaptureComponent2D = CreateDefaultSubobject<USceneCaptureComponent2D>(TEXT("NewSceneCaptureComponent2D"));
 	CaptureComponent2D->SetupAttachment(GetMeshComp());
@@ -63,6 +65,7 @@ void ASceneCapture2D::OnInterpToggle(bool bEnable)
 	CaptureComponent2D->SetVisibility(bEnable);
 }
 
+#if WITH_EDITOR
 void ASceneCapture2D::UpdateDrawFrustum()
 {
 	if(DrawFrustum && CaptureComponent2D)
@@ -83,7 +86,6 @@ void ASceneCapture2D::PostActorCreated()
 	Super::PostActorCreated();
 
 	// no need load the editor mesh when there is no editor
-#if WITH_EDITOR
 	if(GetMeshComp())
 	{
 		if (!IsRunningCommandlet())
@@ -95,19 +97,21 @@ void ASceneCapture2D::PostActorCreated()
 			}
 		}
 	}
-#endif
 
 	// Sync component with CameraActor frustum settings.
 	UpdateDrawFrustum();
 }
+#endif
 // -----------------------------------------------
 
 ASceneCaptureCube::ASceneCaptureCube(const FObjectInitializer& ObjectInitializer)
 	: Super(ObjectInitializer)
 {
+#if WITH_EDITOR
 	DrawFrustum = CreateDefaultSubobject<UDrawFrustumComponent>(TEXT("DrawFrust0"));
 	DrawFrustum->SetIsVisualizationComponent(true);
 	DrawFrustum->SetupAttachment(GetMeshComp());
+#endif
 
 	CaptureComponentCube = CreateDefaultSubobject<USceneCaptureComponentCube>(TEXT("NewSceneCaptureComponentCube"));
 	CaptureComponentCube->SetupAttachment(GetMeshComp());
@@ -118,6 +122,7 @@ void ASceneCaptureCube::OnInterpToggle(bool bEnable)
 	CaptureComponentCube->SetVisibility(bEnable);
 }
 
+#if WITH_EDITOR
 void ASceneCaptureCube::UpdateDrawFrustum()
 {
 	if(DrawFrustum && CaptureComponentCube)
@@ -137,7 +142,6 @@ void ASceneCaptureCube::PostActorCreated()
 	Super::PostActorCreated();
 
 	// no need load the editor mesh when there is no editor
-#if WITH_EDITOR
 	if(GetMeshComp())
 	{
 		if (!IsRunningCommandlet())
@@ -149,12 +153,10 @@ void ASceneCaptureCube::PostActorCreated()
 			}
 		}
 	}
-#endif
 
 	// Sync component with CameraActor frustum settings.
 	UpdateDrawFrustum();
 }
-#if WITH_EDITOR
 
 void ASceneCaptureCube::PostEditMove(bool bFinished)
 {
@@ -250,7 +252,7 @@ void USceneCaptureComponent::ShowOnlyActorComponents(AActor* InActor)
 		PrimitiveRenderMode = ESceneCapturePrimitiveRenderMode::PRM_UseShowOnlyList;
 
 		for (UActorComponent* Component : InActor->GetComponents())
-		{
+			{
 			if (UPrimitiveComponent* PrimComp = Cast<UPrimitiveComponent>(Component))
 			{
 				ShowOnlyComponents.Add(PrimComp);
