Engine: Add a BP method DoesDataTableRowExist to determine whether or not a table contains the specified row

Adapted from:
PR #4040: Add Blueprint function to check is data table row exist. (Contributed by [REDACTED_PERSON])

diff --git a/Engine/Source/Runtime/Engine/Classes/Kismet/DataTableFunctionLibrary.h b/Engine/Source/Runtime/Engine/Classes/Kismet/DataTableFunctionLibrary.h
--- a/Engine/Source/Runtime/Engine/Classes/Kismet/DataTableFunctionLibrary.h
+++ b/Engine/Source/Runtime/Engine/Classes/Kismet/DataTableFunctionLibrary.h
@@ -35,6 +35,10 @@ class ENGINE_API UDataTableFunctionLibrary : public UBlueprintFunctionLibrary
 	UFUNCTION(BlueprintCallable, Category = "DataTable", meta = (ExpandEnumAsExecs="OutResult", DataTablePin="CurveTable"))
 	static void EvaluateCurveTableRow(UCurveTable* CurveTable, FName RowName, float InXY, TEnumAsByte<EEvaluateCurveTableResult::Type>& OutResult, float& OutXY,const FString& ContextString);
     
+	// Returns whether or not Table contains a row named RowName
+  	UFUNCTION(BlueprintCallable, Category = "DataTable")
+ 	static bool DoesDataTableRowExist(UDataTable* Table, FName RowName);
+    
 	UFUNCTION(BlueprintCallable, Category = "DataTable")
 	static void GetDataTableRowNames(UDataTable* Table, TArray<FName>& OutRowNames);
 
diff --git a/Engine/Source/Runtime/Engine/Private/DataTableFunctionLibrary.cpp b/Engine/Source/Runtime/Engine/Private/DataTableFunctionLibrary.cpp
--- a/Engine/Source/Runtime/Engine/Private/DataTableFunctionLibrary.cpp
+++ b/Engine/Source/Runtime/Engine/Private/DataTableFunctionLibrary.cpp
@@ -30,6 +30,19 @@ void UDataTableFunctionLibrary::EvaluateCurveTableRow(UCurveTable* CurveTable, F
 	}
 }
 
+bool UDataTableFunctionLibrary::DoesDataTableRowExist(UDataTable* Table, FName RowName)
+{
+	if (!Table)
+	{
+		return false;
+	}
+	else if (Table->RowStruct == nullptr)
+	{
+		return false;
+	}
+	return Table->RowMap.Find(RowName) != nullptr;
+}
+
 TArray<FString> UDataTableFunctionLibrary::GetDataTableColumnAsString(const UDataTable* DataTable, FName PropertyName)
 {
 	if (DataTable && PropertyName != NAME_None)
