PR #4966: [UE-62736] Add MinimalApi to ACullDistanceVolume (Contributed by [REDACTED_PERSON])

diff --git a/Engine/Source/Runtime/Engine/Classes/Engine/CullDistanceVolume.h b/Engine/Source/Runtime/Engine/Classes/Engine/CullDistanceVolume.h
--- a/Engine/Source/Runtime/Engine/Classes/Engine/CullDistanceVolume.h
+++ b/Engine/Source/Runtime/Engine/Classes/Engine/CullDistanceVolume.h
@@ -38,7 +38,7 @@ struct FCullDistanceSizePair
 };
 
 
-UCLASS(hidecategories=(Advanced, Attachment, Collision, Volume))
+UCLASS(hidecategories=(Advanced, Attachment, Collision, Volume), MinimalAPI)
 class ACullDistanceVolume
 	: public AVolume
 {
