Engine: Fixed an issue where spring arm components with bDoRotationLag=true would never pass through pitches of 90 or -90 degrees

Adapted from:
PR #4862: FIX: ticket UE-54222 (Contributed by [REDACTED_PERSON])

diff --git a/Engine/Source/Runtime/Engine/Private/GameFramework/SpringArmComponent.cpp b/Engine/Source/Runtime/Engine/Private/GameFramework/SpringArmComponent.cpp
--- a/Engine/Source/Runtime/Engine/Private/GameFramework/SpringArmComponent.cpp
+++ b/Engine/Source/Runtime/Engine/Private/GameFramework/SpringArmComponent.cpp
@@ -100,13 +100,13 @@ void USpringArmComponent::UpdateDesiredArmLocation(bool bDoTrace, bool bDoLocati
 				LerpTarget += ArmRotStep * (LerpAmount * InverseCameraLagMaxTimeStep);
 				RemainingTime -= LerpAmount;
 
-				DesiredRot = FMath::RInterpTo(PreviousDesiredRot, LerpTarget, LerpAmount, CameraRotationLagSpeed);
+				DesiredRot = FRotator(FMath::QInterpTo(FQuat(PreviousDesiredRot), FQuat(LerpTarget), LerpAmount, CameraRotationLagSpeed));
 				PreviousDesiredRot = DesiredRot;
 			}
 		}
 		else
 		{
-			DesiredRot = FMath::RInterpTo(PreviousDesiredRot, DesiredRot, DeltaTime, CameraRotationLagSpeed);
+			DesiredRot = FRotator(FMath::QInterpTo(FQuat(PreviousDesiredRot), FQuat(DesiredRot), DeltaTime, CameraRotationLagSpeed));
 		}
 	}
 	PreviousDesiredRot = DesiredRot;
