Cut out editor only functions and variables from cooked builds
#rnx

diff --git a/Engine/Source/Runtime/Engine/Classes/Engine/SCS_Node.h b/Engine/Source/Runtime/Engine/Classes/Engine/SCS_Node.h
--- a/Engine/Source/Runtime/Engine/Classes/Engine/SCS_Node.h
+++ b/Engine/Source/Runtime/Engine/Classes/Engine/SCS_Node.h
@@ -33,12 +33,6 @@ class USCS_Node : public UObject
 	UPROPERTY()
 	FBlueprintCookedComponentInstancingData CookedComponentInstancingData;
 
-	/** If non-None, creates a variable in the class and assigns the created blueprint to it */
-	// Note: Keeping this in place for backwards-compatibility with existing projects/tools that might be accessing it.
-	DEPRECATED(4.14, "VariableName is no longer directly accessible; please use GetVariableName()/SetVariableName() instead.")
-	UPROPERTY(transient)
-	FName VariableName;
-
 #if WITH_EDITORONLY_DATA
 	/** If non-None, the assigned category name */
 	UPROPERTY()
@@ -73,6 +67,7 @@ class USCS_Node : public UObject
 	UPROPERTY()
 	FGuid VariableGuid;
 
+#if WITH_EDITORONLY_DATA
 	/** (DEPRECATED) */
 	UPROPERTY()
 	bool bIsFalseRoot_DEPRECATED;
@@ -88,6 +83,7 @@ class USCS_Node : public UObject
 	/** (DEPRECATED) If true, the variable name was a autogenerated and is not presented to the user */
 	UPROPERTY()
 	bool bVariableNameAutoGenerated_DEPRECATED;
+#endif // WITH_EDITORONLY_DATA
 
 	//~ Begin UObject Interface
 	virtual void Serialize(FArchive& Ar) override;
@@ -166,11 +162,13 @@ class USCS_Node : public UObject
 		return CastChecked<USimpleConstructionScript>(GetOuter());
 	}
 
+#if WITH_EDITOR
 	/** Set delegate to call when name is modified externally */
 	ENGINE_API void SetOnNameChanged(const FSCSNodeNameChanged& OnChange);
 
 	/** Signal back to the SCS that variable name was modified by an external part of program */
 	ENGINE_API void NameWasModified();
+#endif
 
 	/** Set a metadata value on the SCS_Node */
 	ENGINE_API void SetMetaData(const FName& Key, const FString& Value);
@@ -208,7 +206,9 @@ private:
 	UPROPERTY()
 	FName InternalVariableName;
 
+#if WITH_EDITOR
 	/** Delegate to trigger when Variable name is modified outside of SCS */
 	FSCSNodeNameChanged		  OnNameChangedExternal;
+#endif
 };
 
diff --git a/Engine/Source/Runtime/Engine/Classes/Engine/SimpleConstructionScript.h b/Engine/Source/Runtime/Engine/Classes/Engine/SimpleConstructionScript.h
--- a/Engine/Source/Runtime/Engine/Classes/Engine/SimpleConstructionScript.h
+++ b/Engine/Source/Runtime/Engine/Classes/Engine/SimpleConstructionScript.h
@@ -113,6 +113,7 @@ private:
 	UPROPERTY()
 	class USCS_Node* DefaultSceneRootNode;
 
+#if WITH_EDITORONLY_DATA
 	/** (DEPRECATED) Root node of the construction script */
 	UPROPERTY()
 	class USCS_Node* RootNode_DEPRECATED;
@@ -120,6 +121,7 @@ private:
 	/** (DEPRECATED) Actor Component based nodes are stored here.  They cannot be in the tree hierarchy */
 	UPROPERTY()
 	TArray<USCS_Node*> ActorComponentNodes_DEPRECATED;
+#endif
 
 	/** Helper function to find an appropriate child node index that can be promoted to the parent's level */
 	int32 FindPromotableChildNodeIndex(USCS_Node* InParentNode) const;
diff --git a/Engine/Source/Runtime/Engine/Private/SCS_Node.cpp b/Engine/Source/Runtime/Engine/Private/SCS_Node.cpp
--- a/Engine/Source/Runtime/Engine/Private/SCS_Node.cpp
+++ b/Engine/Source/Runtime/Engine/Private/SCS_Node.cpp
@@ -13,13 +13,15 @@
 USCS_Node::USCS_Node(const FObjectInitializer& ObjectInitializer)
 	: Super(ObjectInitializer)
 {
+#if WITH_EDITORONLY_DATA
 	bIsFalseRoot_DEPRECATED = false;
 	bIsNative_DEPRECATED = false;
+#endif
 
 	bIsParentComponentNative = false;
 
 #if WITH_EDITOR
-	EditorComponentInstance = NULL;
+	EditorComponentInstance = nullptr;
 #endif
 }
 
@@ -390,14 +392,9 @@ void USCS_Node::SetVariableName(const FName& NewName, bool bRenameTemplate)
 	}
 
 	InternalVariableName = NewName;
-
-	// >>> Backwards Compatibility: Support existing projects/tools that might be reading the variable name directly. This can be removed in a future release.
-PRAGMA_DISABLE_DEPRECATION_WARNINGS
-	VariableName = InternalVariableName;
-PRAGMA_ENABLE_DEPRECATION_WARNINGS
-	// <<< End Backwards Compatibility
 }
 
+#if WITH_EDITOR
 void USCS_Node::NameWasModified()
 {
 	OnNameChangedExternal.ExecuteIfBound(InternalVariableName);
@@ -407,6 +404,7 @@ void USCS_Node::SetOnNameChanged( const FSCSNodeNameChanged& OnChange )
 {
 	OnNameChangedExternal = OnChange;
 }
+#endif
 
 int32 USCS_Node::FindMetaDataEntryIndexForKey(const FName& Key)
 {
@@ -457,12 +455,6 @@ void USCS_Node::Serialize(FArchive& Ar)
 	{
 		if (Ar.IsPersistent() && !Ar.HasAnyPortFlags(PPF_Duplicate | PPF_DuplicateForPIE))
 		{
-			// >>> Backwards Compatibility: Support existing projects/tools that might be reading the variable name directly. This can be removed in a future release.
-PRAGMA_DISABLE_DEPRECATION_WARNINGS
-			VariableName = InternalVariableName;
-PRAGMA_ENABLE_DEPRECATION_WARNINGS
-			// <<< End Backwards Compatibility
-
 			// Fix up the component class property, if it has not already been set.
 			// Note: This is done here, instead of in PostLoad(), because it needs to be set before Blueprint class compilation.
 			if (ComponentClass == nullptr && ComponentTemplate != nullptr)
diff --git a/Engine/Source/Runtime/Engine/Private/SimpleConstructionScript.cpp b/Engine/Source/Runtime/Engine/Private/SimpleConstructionScript.cpp
--- a/Engine/Source/Runtime/Engine/Private/SimpleConstructionScript.cpp
+++ b/Engine/Source/Runtime/Engine/Private/SimpleConstructionScript.cpp
@@ -23,7 +23,9 @@ const FString USimpleConstructionScript::ComponentTemplateNameSuffix(TEXT("_GEN_
 USimpleConstructionScript::USimpleConstructionScript(const FObjectInitializer& ObjectInitializer)
 	: Super(ObjectInitializer)
 {
+#if WITH_EDITORONLY_DATA
 	RootNode_DEPRECATED = nullptr;
+#endif
 	DefaultSceneRootNode = nullptr;
 
 #if WITH_EDITOR
@@ -45,6 +47,7 @@ void USimpleConstructionScript::Serialize(FArchive& Ar)
 
 	if(Ar.IsLoading())
 	{
+#if WITH_EDITORONLY_DATA
 		if(Ar.UE4Ver() < VER_UE4_REMOVE_NATIVE_COMPONENTS_FROM_BLUEPRINT_SCS)
 		{
 			// If we previously had a root node, we need to move it into the new RootNodes array. This is done in Serialize() in order to support SCS preloading (which relies on a valid RootNodes array).
@@ -114,6 +117,7 @@ void USimpleConstructionScript::Serialize(FArchive& Ar)
 			// Clear the deprecated ActorComponent list
 			ActorComponentNodes_DEPRECATED.Empty();
 		}
+#endif // WITH_EDITORONLY_DATA
 	}
 }
 
