PR #4978: Fix connecting during seamless travel without a transition map (Contributed by [REDACTED_PERSON])

diff --git a/Engine/Source/Runtime/Engine/Private/UnrealEngine.cpp b/Engine/Source/Runtime/Engine/Private/UnrealEngine.cpp
--- a/Engine/Source/Runtime/Engine/Private/UnrealEngine.cpp
+++ b/Engine/Source/Runtime/Engine/Private/UnrealEngine.cpp
@@ -11674,7 +11674,7 @@ void UEngine::TickWorldTravel(FWorldContext& Context, float DeltaSeconds)
 			if (!MakeSureMapNameIsValid(Context.PendingNetGame->URL.Map))
 			{
 				BrowseToDefaultMap(Context);
-				BroadcastTravelFailure(Context.World(), ETravelFailure::PackageMissing, Context.PendingNetGame->URL.RedirectURL);
+				BroadcastTravelFailure(Context.World(), ETravelFailure::PackageMissing, Context.PendingNetGame->URL.Map);
 			}
 			else
 			{
diff --git a/Engine/Source/Runtime/Engine/Private/World.cpp b/Engine/Source/Runtime/Engine/Private/World.cpp
--- a/Engine/Source/Runtime/Engine/Private/World.cpp
+++ b/Engine/Source/Runtime/Engine/Private/World.cpp
@@ -4318,9 +4318,21 @@ void UWorld::WelcomePlayer(UNetConnection* Connection)
 {
 	check(CurrentLevel);
 	Connection->SendPackageMap();
-	
-	FString LevelName = CurrentLevel->GetOutermost()->GetName();
-	Connection->SetClientWorldPackageName(CurrentLevel->GetOutermost()->GetFName());
+
+	FString LevelName;
+
+	const FSeamlessTravelHandler& SeamlessTravelHandler = GEngine->SeamlessTravelHandlerForWorld(this);
+	if (SeamlessTravelHandler.IsInTransition())
+	{
+		// Tell the client to go to the destination map
+		LevelName = SeamlessTravelHandler.GetDestinationMapName();
+		Connection->SetClientWorldPackageName(NAME_None);
+	}
+	else
+	{
+		LevelName = CurrentLevel->GetOutermost()->GetName();
+		Connection->SetClientWorldPackageName(CurrentLevel->GetOutermost()->GetFName());
+	}
 
 	FString GameName;
 	FString RedirectURL;
