Preserve default value overrides on an instance of a component that's based on an archetype inherited by a non-nativized child Blueprint class as a default subobject from a nativized parent Blueprint class.

Change summary:
- Modified UBlueprintGeneratedClass::PostLoadDefaultObject() to call CheckAndApplyComponentTemplateOverrides() after serialization of a non-nativized Blueprint class default object. This will update default subobjects inherited from a nativized parent class with any override values stored in the non-nativized child class's ICH instance *prior* to serializing instances of the child class. Previously, the code was applying these values *after* serialization of each instance of these components, rather than to the CDO, which would incorrectly overwrite any instance-specific overrides.
- Modified UBlueprintGeneratedClass::CheckAndApplyComponentTemplateOverrides() to remove the AActor class dependency.
- Changed UBlueprintGeneratedClass::CheckAndApplyComponentTemplateOverrides() to a 'protected' API.
- Modified AActor::PostLoadSubobjects() and AActor::PostSpawnInitialize() to no longer call CheckAndApplyComponentTemplateOverrides(). Since we now apply the overrides to the CDO, Actor instances will pick them up from there at construction time.
- Modified AActor::PostSpawnInitialize() to properly initialize the root component transform when a native root component originates from a converted Blueprint class type. This ensures that instances of child classes that inherit a native root component from a converted parent class will get the same world transform when instanced dynamically in a nativized cooked build as it would otherwise have in a non-nativized cooked build. Also added comments to clarify.
- Updated comment in USCS_Node::ExecuteNodeOnActor() to clarify meaning of 'bIsDefaultTransform' and its relationship to the alternate code path for converted SCS root components in a nativized cooked build.

diff --git a/Engine/Source/Runtime/Engine/Classes/Engine/BlueprintGeneratedClass.h b/Engine/Source/Runtime/Engine/Classes/Engine/BlueprintGeneratedClass.h
--- a/Engine/Source/Runtime/Engine/Classes/Engine/BlueprintGeneratedClass.h
+++ b/Engine/Source/Runtime/Engine/Classes/Engine/BlueprintGeneratedClass.h
@@ -666,9 +666,6 @@ public:
 	static void CreateComponentsForActor(const UClass* ThisClass, AActor* Actor);
 	static void CreateTimelineComponent(AActor* Actor, const UTimelineTemplate* TimelineTemplate);
 
-	/** Check for and handle manual application of default value overrides to instanced component subobjects that were inherited from a nativized parent class */
-	static void CheckAndApplyComponentTemplateOverrides(AActor* Actor);
-
 	// UObject interface
 	virtual void Serialize(FArchive& Ar) override;
 	virtual void PostLoad() override;
@@ -737,6 +734,9 @@ protected:
 	*/
 	static void InitArrayPropertyFromCustomList(const UArrayProperty* ArrayProperty, const FCustomPropertyListNode* InPropertyList, uint8* DataPtr, const uint8* DefaultDataPtr);
 
+	/** Check for and handle manual application of default value overrides to component subobjects that were inherited from a nativized parent class */
+	static void CheckAndApplyComponentTemplateOverrides(UObject* InClassDefaultObject);
+
 public:
 
 	/** Called when the custom list of properties used during post-construct initialization needs to be rebuilt (e.g. after serialization and recompilation). */
diff --git a/Engine/Source/Runtime/Engine/Private/Actor.cpp b/Engine/Source/Runtime/Engine/Private/Actor.cpp
--- a/Engine/Source/Runtime/Engine/Private/Actor.cpp
+++ b/Engine/Source/Runtime/Engine/Private/Actor.cpp
@@ -661,17 +661,6 @@ void AActor::PostLoadSubobjects(FObjectInstancingGraph* OuterInstanceGraph)
 
 	Super::PostLoadSubobjects(OuterInstanceGraph);
 
-	// If this is a Blueprint class, we may need to manually apply default value overrides to some inherited components in a cooked
-	// build scenario. This can occur, for example, if we have a nativized Blueprint class somewhere in the class inheritance hierarchy.
-	if (FPlatformProperties::RequiresCookedData() && !IsTemplate())
-	{
-		const UBlueprintGeneratedClass* BPGC = Cast<UBlueprintGeneratedClass>(GetClass());
-		if (BPGC != nullptr && BPGC->bHasNativizedParent)
-		{
-			UBlueprintGeneratedClass::CheckAndApplyComponentTemplateOverrides(this);
-		}
-	}
-
 	ResetOwnedComponents();
 
 	if (RootComponent && bHadRoot && OldRoot != RootComponent && OldRoot->IsIn(this))
@@ -2825,48 +2814,56 @@ void AActor::PostSpawnInitialize(FTransform const& UserSpawnTransform, AActor* I
 	check(Role == ROLE_Authority);
 	ExchangeNetRoles(bRemoteOwned);
 
+	// Set the actor's world transform if it has a native rootcomponent.
 	USceneComponent* const SceneRootComponent = FixupNativeActorComponents(this);
 	if (SceneRootComponent != nullptr)
 	{
 		check(SceneRootComponent->GetOwner() == this);
-		
-		FVector RootComponentRelativeLocation = SceneRootComponent->RelativeLocation;
-		FRotator RootComponentRelativeRotation = SceneRootComponent->RelativeRotation;
 
-		// For converted BP root components, we must zero out RelativeLocation/RelativeRotation here. The reason is that in the non-nativized
-		// case, we ignore them when we instance a scene component that will also become the root (see USCS_Node::ExecuteNodeOnActor). Once
-		// a Blueprint class is nativized, we no longer run through that path, but we need to keep the same rotation/translation as before.
-		// We used to ignore them at nativization time, but that doesn't work because existing placements of the Blueprint component may rely
-		// on the value that's stored in the CDO, because it won't have been serialized out to the instance as a result of delta serialization.
-		if (Cast<UDynamicClass>(SceneRootComponent->GetArchetype()->GetOuter()->GetClass()))
+		// Determine if the native root component's archetype originates from a converted (nativized) Blueprint class.
+		UObject* RootComponentArchetype = SceneRootComponent->GetArchetype();
+		UClass* ArchetypeOwnerClass = RootComponentArchetype->GetOuter()->GetClass();
+		if (UBlueprintGeneratedClass* ArchetypeOwnerClassAsBPGC = Cast<UBlueprintGeneratedClass>(ArchetypeOwnerClass))
 		{
-			RootComponentRelativeLocation = FVector::ZeroVector;
-			RootComponentRelativeRotation = FRotator::ZeroRotator;
+			// In this case, the Actor CDO is a non-nativized Blueprint class (e.g. a child class) and the component's archetype
+			// is an instanced default subobject within the non-nativized Blueprint's CDO. If the owner class also has a nativized
+			// parent class somewhere in its inheritance hierarchy, we must redirect the query by walking up the archetype chain.
+			if (ArchetypeOwnerClassAsBPGC->bHasNativizedParent)
+			{
+				do 
+				{
+					RootComponentArchetype = RootComponentArchetype->GetArchetype();
+					ArchetypeOwnerClass = RootComponentArchetype->GetOuter()->GetClass();
+				} while (Cast<UBlueprintGeneratedClass>(ArchetypeOwnerClass) != nullptr);
+			}
 		}
 
-		// Set the actor's location and rotation since it has a native rootcomponent
-		// Note that we respect any initial transformation the root component may have from the CDO, so the final transform
-		// might necessarily be exactly the passed-in UserSpawnTransform.
- 		const FTransform RootTransform(RootComponentRelativeRotation, RootComponentRelativeLocation, SceneRootComponent->RelativeScale3D);
- 		const FTransform FinalRootComponentTransform = RootTransform * UserSpawnTransform;
-		SceneRootComponent->SetWorldTransform(FinalRootComponentTransform);
+		if (Cast<UDynamicClass>(ArchetypeOwnerClass) != nullptr)
+		{
+			// For native root components either belonging to or inherited from a converted (nativized) Blueprint class, we currently do not use
+			// the transformation that's set on the root component in the CDO. The reason is that in the non-nativized case, we ignore the default
+			// transform when we instance a Blueprint-owned scene component that will also become the root (see USCS_Node::ExecuteNodeOnActor; in
+			// the case of dynamically-spawned Blueprint instances, 'bIsDefaultTransform' will be false, and the scale from the SCS node's template
+			// will not be applied in that code path in that case). Once a Blueprint class is nativized, we no longer run through that code path
+			// when we spawn new instances of that class dynamically, but for consistency, we need to keep the same transform as in the non-
+			// nativized case. We used to ignore any non-default transform value set on the root component at cook (nativization) time, but that 
+			// doesn't work because existing placements of the Blueprint component in a scene may rely on the value that's stored in the CDO,
+			// and as a result the instance-specific override value doesn't get serialized out to the instance as a result of delta serialization.
+			SceneRootComponent->SetWorldTransform(UserSpawnTransform);
+		}
+		else
+		{
+			// In the "normal" case we do respect any non-default transform value that the root component may have received from the archetype
+			// that's owned by the native CDO, so the final transform might not always necessarily equate to the passed-in UserSpawnTransform.
+			const FTransform RootTransform(SceneRootComponent->RelativeRotation, SceneRootComponent->RelativeLocation, SceneRootComponent->RelativeScale3D);
+			const FTransform FinalRootComponentTransform = RootTransform * UserSpawnTransform;
+			SceneRootComponent->SetWorldTransform(FinalRootComponentTransform);
+		}
 	}
 
 	// Call OnComponentCreated on all default (native) components
 	DispatchOnComponentsCreated(this);
 
-	// If this is a Blueprint class, we may need to manually apply default value overrides to some inherited components in a
-	// cooked build scenario. This can occur, for example, if we have a nativized Blueprint class in the inheritance hierarchy.
-	// Note: This should be done prior to executing the construction script, in case there are any dependencies on default values.
-	if (FPlatformProperties::RequiresCookedData())
-	{
-		const UBlueprintGeneratedClass* BPGC = Cast<UBlueprintGeneratedClass>(GetClass());
-		if (BPGC != nullptr && BPGC->bHasNativizedParent)
-		{
-			UBlueprintGeneratedClass::CheckAndApplyComponentTemplateOverrides(this);
-		}
-	}
-
 	// Register the actor's default (native) components, but only if we have a native scene root. If we don't, it implies that there could be only non-scene components
 	// at the native class level. In that case, if this is a Blueprint instance, we need to defer native registration until after SCS execution can establish a scene root.
 	// Note: This API will also call PostRegisterAllComponents() on the actor instance. If deferred, PostRegisterAllComponents() won't be called until the root is set by SCS.
diff --git a/Engine/Source/Runtime/Engine/Private/BlueprintGeneratedClass.cpp b/Engine/Source/Runtime/Engine/Private/BlueprintGeneratedClass.cpp
--- a/Engine/Source/Runtime/Engine/Private/BlueprintGeneratedClass.cpp
+++ b/Engine/Source/Runtime/Engine/Private/BlueprintGeneratedClass.cpp
@@ -358,11 +358,19 @@ void UBlueprintGeneratedClass::PostLoadDefaultObject(UObject* Object)
 	{
 		// Rebuild the custom property list used in post-construct initialization logic. Note that PostLoad() may have altered some serialized properties.
 		UpdateCustomPropertyListForPostConstruction();
+
 		// Restore any property values from config file
 		if (HasAnyClassFlags(CLASS_Config))
 		{
 			ClassDefaultObject->LoadConfig();
 		}
+
+		// We may need to manually apply default value overrides to some inherited components in a cooked build
+		// scenario. This can occur if we have a nativized Blueprint class somewhere in the parent class ancestry.
+		if (bHasNativizedParent)
+		{
+			CheckAndApplyComponentTemplateOverrides(ClassDefaultObject);
+		}
 	}
 }
 
@@ -1107,13 +1115,11 @@ void UBlueprintGeneratedClass::CreateComponentsForActor(const UClass* ThisClass,
 	}
 }
 
-void UBlueprintGeneratedClass::CheckAndApplyComponentTemplateOverrides(AActor* Actor)
+void UBlueprintGeneratedClass::CheckAndApplyComponentTemplateOverrides(UObject* InClassDefaultObject)
 {
-	check(Actor != nullptr);
-
 	// Get the Blueprint class hierarchy (if valid).
 	TArray<const UBlueprintGeneratedClass*> ParentBPClassStack;
-	GetGeneratedClassesHierarchy(Actor->GetClass(), ParentBPClassStack);
+	GetGeneratedClassesHierarchy(InClassDefaultObject->GetClass(), ParentBPClassStack);
 	if (ParentBPClassStack.Num() > 0)
 	{
 		// If the nearest native antecedent is also a nativized BP class, we may have an override
@@ -1154,8 +1160,8 @@ void UBlueprintGeneratedClass::CheckAndApplyComponentTemplateOverrides(AActor* A
 							const FBlueprintCookedComponentInstancingData* OverrideData = ICH->GetOverridenComponentTemplateData(ComponentKey);
 							if (OverrideData != nullptr && OverrideData->bIsValid)
 							{
-								// This is the instance of the inherited component subobject that's owned by the given Actor instance
-								if (UObject* NativizedComponentSubobjectInstance = Actor->GetDefaultSubobjectByName(NativizedComponentSubobjectName))
+								// This is the instance of the inherited component subobject that's owned by the given class default object
+								if (UObject* NativizedComponentSubobjectInstance = InClassDefaultObject->GetDefaultSubobjectByName(NativizedComponentSubobjectName))
 								{
 									// Nativized component override data loader implementation.
 									class FNativizedComponentOverrideDataLoader : public FObjectReader
@@ -1176,7 +1182,7 @@ void UBlueprintGeneratedClass::CheckAndApplyComponentTemplateOverrides(AActor* A
 									// Ensure that the ICH has gotten a PostLoad() call - we need to ensure that any cooked data will have been fully processed before proceeding.
 									ICH->ConditionalPostLoad();
 
-									// Serialize cached override data to the instanced subobject that's based on the default subobject from the nativized parent class and owned by the Actor instance.
+									// Serialize cached override data to the instanced subobject that's based on the default subobject from the nativized parent class and owned by the non-nativized child class default object.
 									FNativizedComponentOverrideDataLoader OverrideDataLoader(OverrideData->GetCachedPropertyDataForSerialization(), OverrideData->GetCachedPropertyListForSerialization());
 									NativizedComponentSubobjectInstance->Serialize(OverrideDataLoader);
 								}
diff --git a/Engine/Source/Runtime/Engine/Private/SCS_Node.cpp b/Engine/Source/Runtime/Engine/Private/SCS_Node.cpp
--- a/Engine/Source/Runtime/Engine/Private/SCS_Node.cpp
+++ b/Engine/Source/Runtime/Engine/Private/SCS_Node.cpp
@@ -130,7 +130,10 @@ UActorComponent* USCS_Node::ExecuteNodeOnActor(AActor* Actor, USceneComponent* P
 				FTransform WorldTransform = *RootTransform;
 				if(bIsDefaultTransform)
 				{
-					// Note: We use the scale vector from the component template when spawning (to match what happens with a native root)
+					// Note: We use the scale vector from the component template when spawning (to match what happens with a native root). This
+					// does NOT occur when this component is instanced as part of dynamically spawning a Blueprint class in a cooked build (i.e.
+					// 'bIsDefaultTransform' will be 'false' in that situation). In order to maintain the same behavior between a nativized and
+					// non-nativized cooked build, if this ever changes, we would also need to update the code in AActor::PostSpawnInitialize().
 					WorldTransform.SetScale3D(NewSceneComp->RelativeScale3D);
 				}
 
