Merging CL 4277008 and 4277310 to fix Vignette const issue

diff --git a/Engine/Source/Runtime/Engine/Public/ShowFlagsValues.inl b/Engine/Source/Runtime/Engine/Public/ShowFlagsValues.inl
--- a/Engine/Source/Runtime/Engine/Public/ShowFlagsValues.inl
+++ b/Engine/Source/Runtime/Engine/Public/ShowFlagsValues.inl
@@ -37,7 +37,7 @@ SHOWFLAG_FIXED_IN_SHIPPING(1, LensFlares, SFG_PostProcess, NSLOCTEXT("UnrealEd",
 /** show indirect lighting component, for now SHOWFLAG_ALWAYS_ACCESSIBLE because it's needed by r.GBuffer */
 SHOWFLAG_ALWAYS_ACCESSIBLE(GlobalIllumination, SFG_LightingComponents, NSLOCTEXT("UnrealEd", "GlobalIlluminationSF", "Global Illumination"))
 /** Darkens the screen borders (Camera artifact and artistic effect) */
-SHOWFLAG_FIXED_IN_SHIPPING(1, Vignette, SFG_PostProcess, NSLOCTEXT("UnrealEd", "VignetteSF", "Vignette"))
+SHOWFLAG_ALWAYS_ACCESSIBLE(Vignette, SFG_PostProcess, NSLOCTEXT("UnrealEd", "VignetteSF", "Vignette"))
 /** Fine film grain */
 SHOWFLAG_FIXED_IN_SHIPPING(1, Grain, SFG_PostProcess, NSLOCTEXT("UnrealEd", "GrainSF", "Grain"))
 /** Screen Space Ambient Occlusion, for now SHOWFLAG_ALWAYS_ACCESSIBLE because it's exposed in SceneCapture */
