Merging //UE4/Dev-Main to Dev-Framework (//UE4/Dev-Framework) @ 4287497
#rnx

diff --git a/Engine/Plugins/Developer/ActorLayerUtilities/Source/ActorLayerUtilitiesEditor/Private/ActorLayerUtilitiesEditor.cpp b/Engine/Plugins/Developer/ActorLayerUtilities/Source/ActorLayerUtilitiesEditor/Private/ActorLayerUtilitiesEditor.cpp
--- a/Engine/Plugins/Developer/ActorLayerUtilities/Source/ActorLayerUtilitiesEditor/Private/ActorLayerUtilitiesEditor.cpp
+++ b/Engine/Plugins/Developer/ActorLayerUtilities/Source/ActorLayerUtilitiesEditor/Private/ActorLayerUtilitiesEditor.cpp
@@ -36,4 +36,4 @@ class FActorLayerUtilitiesEditorModule : public IModuleInterface
 };
 
 
-IMPLEMENT_MODULE(FActorLayerUtilitiesEditorModule, ActorLayerUtilities)
\ No newline at end of file
+IMPLEMENT_MODULE(FActorLayerUtilitiesEditorModule, ActorLayerUtilitiesEditor)
\ No newline at end of file
diff --git a/Engine/Source/Programs/UnrealBuildTool/Configuration/UEBuildBinary.cs b/Engine/Source/Programs/UnrealBuildTool/Configuration/UEBuildBinary.cs
--- a/Engine/Source/Programs/UnrealBuildTool/Configuration/UEBuildBinary.cs
+++ b/Engine/Source/Programs/UnrealBuildTool/Configuration/UEBuildBinary.cs
@@ -494,7 +494,11 @@ namespace UnrealBuildTool
 						{
 							foreach(FileReference GeneratedCodeFile in DirectoryReference.EnumerateFiles(Module.GeneratedCodeDirectory))
 							{
-								BuildProducts.Add(GeneratedCodeFile, BuildProductType.BuildResource);
+								// Exclude timestamp files, since they're always updated and cause collisions between builds
+								if(!GeneratedCodeFile.GetFileName().Equals("Timestamp", StringComparison.OrdinalIgnoreCase))
+								{
+									BuildProducts.Add(GeneratedCodeFile, BuildProductType.BuildResource);
+								}
 							}
 						}
 						if(Target.LinkType == TargetLinkType.Monolithic)
diff --git a/Engine/Source/Runtime/Engine/Public/ShowFlagsValues.inl b/Engine/Source/Runtime/Engine/Public/ShowFlagsValues.inl
--- a/Engine/Source/Runtime/Engine/Public/ShowFlagsValues.inl
+++ b/Engine/Source/Runtime/Engine/Public/ShowFlagsValues.inl
@@ -37,7 +37,7 @@ SHOWFLAG_FIXED_IN_SHIPPING(1, LensFlares, SFG_PostProcess, NSLOCTEXT("UnrealEd",
 /** show indirect lighting component, for now SHOWFLAG_ALWAYS_ACCESSIBLE because it's needed by r.GBuffer */
 SHOWFLAG_ALWAYS_ACCESSIBLE(GlobalIllumination, SFG_LightingComponents, NSLOCTEXT("UnrealEd", "GlobalIlluminationSF", "Global Illumination"))
 /** Darkens the screen borders (Camera artifact and artistic effect) */
-SHOWFLAG_FIXED_IN_SHIPPING(1, Vignette, SFG_PostProcess, NSLOCTEXT("UnrealEd", "VignetteSF", "Vignette"))
+SHOWFLAG_ALWAYS_ACCESSIBLE(Vignette, SFG_PostProcess, NSLOCTEXT("UnrealEd", "VignetteSF", "Vignette"))
 /** Fine film grain */
 SHOWFLAG_FIXED_IN_SHIPPING(1, Grain, SFG_PostProcess, NSLOCTEXT("UnrealEd", "GrainSF", "Grain"))
 /** Screen Space Ambient Occlusion, for now SHOWFLAG_ALWAYS_ACCESSIBLE because it's exposed in SceneCapture */
