Merging //UE4/Dev-Main to Dev-Anim (//UE4/Dev-Anim) @ 4289353
#rnx

diff --git a/Engine/Source/Runtime/Engine/Classes/Kismet/KismetGuidLibrary.h b/Engine/Source/Runtime/Engine/Classes/Kismet/KismetGuidLibrary.h
--- a/Engine/Source/Runtime/Engine/Classes/Kismet/KismetGuidLibrary.h
+++ b/Engine/Source/Runtime/Engine/Classes/Kismet/KismetGuidLibrary.h
@@ -16,19 +16,19 @@ class ENGINE_API UKismetGuidLibrary
 	GENERATED_UCLASS_BODY()
 
 	/** Returns true if the values are equal (A == B) */
-	UFUNCTION(BlueprintPure, meta=(DisplayName="Equal (Guid)", CompactNodeTitle="=="), Category="Guid")
+	UFUNCTION(BlueprintPure, meta=(DisplayName="Equal (Guid)", CompactNodeTitle="==", BlueprintThreadSafe), Category="Guid")
 	static bool EqualEqual_GuidGuid( const FGuid& A, const FGuid& B );
 	
 	/** Returns true if the values are not equal (A != B) */
-	UFUNCTION(BlueprintPure, meta=(DisplayName="NotEqual (Guid)", CompactNodeTitle="!="), Category="Guid")
+	UFUNCTION(BlueprintPure, meta=(DisplayName="NotEqual (Guid)", CompactNodeTitle="!=", BlueprintThreadSafe), Category="Guid")
 	static bool NotEqual_GuidGuid( const FGuid& A, const FGuid& B );
 
 	/** Checks whether the given GUID is valid */
-	UFUNCTION(BlueprintPure, meta=(DisplayName="IsValid (Guid)", CompactNodeTitle="isValid?"), Category="Guid")
+	UFUNCTION(BlueprintPure, meta=(DisplayName="IsValid (Guid)", CompactNodeTitle="isValid?", BlueprintThreadSafe), Category="Guid")
 	static bool IsValid_Guid( const FGuid& InGuid );
 
 	/** Invalidates the given GUID */
-	UFUNCTION(BlueprintCallable, meta=(DisplayName="Invalidate_Guid (Guid)"), Category="Guid")
+	UFUNCTION(BlueprintCallable, meta=(DisplayName="Invalidate_Guid (Guid)", BlueprintThreadSafe), Category="Guid")
 	static void Invalidate_Guid( UPARAM(ref) FGuid& InGuid );
 
 	/** Returns a new unique GUID */
@@ -36,10 +36,10 @@ class ENGINE_API UKismetGuidLibrary
 	static FGuid NewGuid();
 	
 	/** Converts a GUID value to a string, in the form 'A-B-C-D' */
-	UFUNCTION(BlueprintPure, meta=(DisplayName="ToString (Guid)", CompactNodeTitle="->", ScriptMethod="ToString"), Category="Guid")
+	UFUNCTION(BlueprintPure, meta=(DisplayName="ToString (Guid)", CompactNodeTitle="->", ScriptMethod="ToString", BlueprintThreadSafe), Category="Guid")
 	static FString Conv_GuidToString( const FGuid& InGuid );
 
 	/** Converts a String of format EGuidFormats to a Guid. Returns Guid OutGuid, Returns bool Success */
-	UFUNCTION(BlueprintPure, meta=(DisplayName="Parse String to Guid"), Category="Guid")
+	UFUNCTION(BlueprintPure, meta=(DisplayName="Parse String to Guid", BlueprintThreadSafe), Category="Guid")
 	static void Parse_StringToGuid( const FString& GuidString, FGuid& OutGuid, bool& Success );
 };
