Fix a couple of issues for Visual Studio 2017 update 15.8.

* Warning 4643 did not previously exist and was ignored. Now is an off-by-default warning about defining types in the std namespace (which Boost triggers).
* Attempt to index an array incorrectly using parentheses, which is now caught by /permissive-.

diff --git a/Engine/Source/Runtime/Engine/Public/Interpolation.h b/Engine/Source/Runtime/Engine/Public/Interpolation.h
--- a/Engine/Source/Runtime/Engine/Public/Interpolation.h
+++ b/Engine/Source/Runtime/Engine/Public/Interpolation.h
@@ -417,7 +417,7 @@ namespace MatineeKeyReduction
 			int32 SegmentCount = SegmentQueue.Num();
 			for ( int32 I = 0; I < SegmentCount; ++I )
 			{
-				FIntPoint& Segment = SegmentQueue(I);
+				FIntPoint& Segment = SegmentQueue[I];
 				if ( Segment.X >= InsertionIndex ) ++Segment.X;
 				if ( Segment.Y >= InsertionIndex ) ++Segment.Y;
 			}
