Copying //UE4/Dev-AnimPhys to Dev-Main (//UE4/Dev-Main)

diff --git a/Engine/Plugins/Animation/LiveLink/Source/LiveLink/Private/LiveLinkRetargetAsset.cpp b/Engine/Plugins/Animation/LiveLink/Source/LiveLink/Private/LiveLinkRetargetAsset.cpp
--- a/Engine/Plugins/Animation/LiveLink/Source/LiveLink/Private/LiveLinkRetargetAsset.cpp
+++ b/Engine/Plugins/Animation/LiveLink/Source/LiveLink/Private/LiveLinkRetargetAsset.cpp
@@ -6,6 +6,7 @@
 #include "LiveLinkTypes.h"
 #include "BonePose.h"
 #include "Animation/AnimTypes.h"
+#include "Animation/AnimCurveTypes.h"
 
 ULiveLinkRetargetAsset::ULiveLinkRetargetAsset(const FObjectInitializer& ObjectInitializer)
 	: Super(ObjectInitializer)
diff --git a/Engine/Source/Developer/AudioSettingsEditor/Private/TargetPlatformAudioCustomization.cpp b/Engine/Source/Developer/AudioSettingsEditor/Private/TargetPlatformAudioCustomization.cpp
--- a/Engine/Source/Developer/AudioSettingsEditor/Private/TargetPlatformAudioCustomization.cpp
+++ b/Engine/Source/Developer/AudioSettingsEditor/Private/TargetPlatformAudioCustomization.cpp
@@ -24,8 +24,14 @@ IMPLEMENT_MODULE(FDefaultModuleImpl, AudioSettingsEditor)
 
 #define LOCTEXT_NAMESPACE "PlatformAudio"
 
+// This string is used for the item on the combo box that, when selected, defers to the custom string entry.
+static const TCHAR* ManualEntryItem = TEXT("Other");
+
 FAudioPluginWidgetManager::FAudioPluginWidgetManager()
 {
+	ManualReverbEntry = TSharedPtr<FText>(new FText(FText::FromString(TEXT("Built-in Reverb"))));
+	ManualSpatializationEntry = TSharedPtr<FText>(new FText(FText::FromString(TEXT("Built-in Spatialization"))));
+	ManualOcclusionEntry = TSharedPtr<FText>(new FText(FText::FromString(TEXT("Built-in Occlusion"))));
 }
 
 FAudioPluginWidgetManager::~FAudioPluginWidgetManager()
@@ -121,6 +127,79 @@ TSharedRef<SWidget> FAudioPluginWidgetManager::MakeAudioPluginSelectorWidget(con
 	}
 #endif // #if WITH_ENGINE
 
+	// This pointer is used to store whatever custom string was input by the user or retrieved from the config file.
+	ValidPluginNames->Add(TSharedPtr<FText>(new FText(FText::FromString(ManualEntryItem))));
+
+	// Text box component:
+	TSharedRef<SEditableTextBox> EditableTextBox = SNew(SEditableTextBox)
+		.Text_Lambda([this, AudioPluginType]() { return OnGetPluginText(AudioPluginType); })
+		.OnTextCommitted_Raw(this, &FAudioPluginWidgetManager::OnPluginTextCommitted, AudioPluginType, PropertyHandle)
+		.SelectAllTextWhenFocused(true)
+		.RevertTextOnEscape(true);
+
+	// Combo box component:
+	TSharedRef<SWidget> ComboBox = SNew(SListView<TSharedPtr<FText>>)
+		.ListItemsSource(ValidPluginNames)
+		.ScrollbarVisibility(EVisibility::Collapsed)
+		.OnGenerateRow_Lambda([](TSharedPtr<FText> InItem, const TSharedRef< class STableViewBase >& Owner)
+		{
+			return SNew(STableRow<TSharedPtr<FText>>, Owner)
+				   .Padding(FMargin(16, 4, 16, 4))
+				   [
+					   SNew(STextBlock).Text(*InItem)
+				   ];
+		})
+		.OnSelectionChanged_Lambda([this, AudioPluginType, PropertyHandle](TSharedPtr<FText> InText, ESelectInfo::Type)
+		{
+			const bool bSelectedManualEntry = (InText->ToString() == FString(ManualEntryItem));
+
+			switch (AudioPluginType)
+			{
+				case EAudioPlugin::SPATIALIZATION:
+					if (bSelectedManualEntry)
+					{
+						SelectedSpatialization = ManualSpatializationEntry;
+					}
+					else
+					{
+						SelectedSpatialization = InText;
+					}
+					
+					OnPluginSelected(bSelectedManualEntry ? ManualSpatializationEntry->ToString() : InText->ToString(), PropertyHandle);
+					break;
+
+				case EAudioPlugin::REVERB:
+					if (bSelectedManualEntry)
+					{
+						SelectedReverb = ManualReverbEntry;
+					}
+					else
+					{
+						SelectedReverb = InText;
+					}
+					
+					OnPluginSelected(bSelectedManualEntry ? ManualReverbEntry->ToString() : InText->ToString(), PropertyHandle);
+					break;
+
+				case EAudioPlugin::OCCLUSION:
+					if (bSelectedManualEntry)
+					{
+						SelectedOcclusion = ManualOcclusionEntry;
+					}
+					else
+					{
+						SelectedOcclusion = InText;
+					}
+					
+					OnPluginSelected((bSelectedManualEntry ? ManualOcclusionEntry->ToString() : InText->ToString()), PropertyHandle);
+					break;
+
+				default:
+					break;
+			}
+		});
+
+
 	//Generate widget:
 	const TSharedRef<SWidget> NewWidget = SNew(SComboButton)
 		.ContentPadding(FMargin(0, 0, 5, 0))
@@ -131,49 +210,12 @@ TSharedRef<SWidget> FAudioPluginWidgetManager::MakeAudioPluginSelectorWidget(con
 			.BorderImage(FEditorStyle::GetBrush("NoBorder"))
 			.Padding(FMargin(0, 0, 5, 0))
 			[
-				SNew(SEditableTextBox)
-				.Text_Lambda([this, AudioPluginType]() { return OnGetPluginText(AudioPluginType); })
-				.OnTextCommitted_Raw(this, &FAudioPluginWidgetManager::OnPluginTextCommitted, AudioPluginType, PropertyHandle)
-				.SelectAllTextWhenFocused(true)
-				.RevertTextOnEscape(true)
+				EditableTextBox
 			]
 		]
 		.MenuContent()
 		[
-			SNew(SListView<TSharedPtr<FText>>)
-			.ListItemsSource(ValidPluginNames)
-			.ScrollbarVisibility(EVisibility::Collapsed)
-			.OnGenerateRow_Lambda([](TSharedPtr<FText> InItem, const TSharedRef< class STableViewBase >& Owner)
-			{
-				return
-					SNew(STableRow<TSharedPtr<FText>>, Owner)
-					.Padding(FMargin(16, 4, 16, 4))
-					[
-						SNew(STextBlock).Text(*InItem)
-					];
-			})
-			.OnSelectionChanged_Lambda([this, AudioPluginType, PropertyHandle](TSharedPtr<FText> InText, ESelectInfo::Type)
-			{
-				switch (AudioPluginType)
-				{
-					case EAudioPlugin::SPATIALIZATION:
-						SelectedSpatialization = InText;
-						break;
-
-					case EAudioPlugin::REVERB:
-						SelectedReverb = InText;
-						break;
-					
-					case EAudioPlugin::OCCLUSION:
-						SelectedOcclusion = InText;
-						break;
-					
-					default:
-						break;
-				}
-
-				OnPluginSelected(InText->ToString(), PropertyHandle);
-			})
+			ComboBox
 		];
 
 	return NewWidget;
@@ -239,15 +281,18 @@ void FAudioPluginWidgetManager::OnPluginTextCommitted(const FText& InText, EText
 	switch (AudioPluginType)
 	{
 		case EAudioPlugin::SPATIALIZATION:
-			*SelectedSpatialization = InText;
+			*ManualSpatializationEntry = InText;
+			SelectedSpatialization = ManualSpatializationEntry;
 			break;
 
 		case EAudioPlugin::REVERB:
-			*SelectedReverb = InText;
+			*ManualReverbEntry = InText;
+			SelectedReverb = ManualReverbEntry;
 			break;
 
 		case EAudioPlugin::OCCLUSION:
-			*SelectedOcclusion = InText;
+			*ManualOcclusionEntry = InText;
+			SelectedOcclusion = ManualOcclusionEntry;
 			break;
 
 		default:
diff --git a/Engine/Source/Developer/AudioSettingsEditor/Public/TargetPlatformAudioCustomization.h b/Engine/Source/Developer/AudioSettingsEditor/Public/TargetPlatformAudioCustomization.h
--- a/Engine/Source/Developer/AudioSettingsEditor/Public/TargetPlatformAudioCustomization.h
+++ b/Engine/Source/Developer/AudioSettingsEditor/Public/TargetPlatformAudioCustomization.h
@@ -34,6 +34,10 @@ private:
 	TSharedPtr<FText> SelectedSpatialization;
 	TSharedPtr<FText> SelectedOcclusion;
 
+	TSharedPtr<FText> ManualSpatializationEntry;
+	TSharedPtr<FText> ManualReverbEntry;
+	TSharedPtr<FText> ManualOcclusionEntry;
+
 	TArray<TSharedPtr<FText>> SpatializationPlugins;
 	TArray<TSharedPtr<FText>> ReverbPlugins;
 	TArray<TSharedPtr<FText>> OcclusionPlugins;
diff --git a/Engine/Source/Editor/AnimGraph/Classes/AnimGraphNode_Base.h b/Engine/Source/Editor/AnimGraph/Classes/AnimGraphNode_Base.h
--- a/Engine/Source/Editor/AnimGraph/Classes/AnimGraphNode_Base.h
+++ b/Engine/Source/Editor/AnimGraph/Classes/AnimGraphNode_Base.h
@@ -5,6 +5,7 @@
 #include "CoreMinimal.h"
 #include "UObject/ObjectMacros.h"
 #include "UObject/UnrealType.h"
+#include "Animation/AnimationAsset.h"
 #include "Animation/AnimBlueprint.h"
 #include "Animation/AnimBlueprintGeneratedClass.h"
 #include "Animation/AnimNodeBase.h"
diff --git a/Engine/Source/Editor/AnimGraph/Public/AnimNodeEditModes.h b/Engine/Source/Editor/AnimGraph/Public/AnimNodeEditModes.h
--- a/Engine/Source/Editor/AnimGraph/Public/AnimNodeEditModes.h
+++ b/Engine/Source/Editor/AnimGraph/Public/AnimNodeEditModes.h
@@ -16,6 +16,4 @@ struct ANIMGRAPH_API AnimNodeEditModes
 	const static FEditorModeID PoseDriver;
 	const static FEditorModeID SplineIK;
 	const static FEditorModeID LookAt;
-	const static FEditorModeID ConstraintFabrik;
-	const static FEditorModeID MultiFabrik;
 };
diff --git a/Engine/Source/Editor/UnrealEd/Private/Commandlets/PackageUtilities.cpp b/Engine/Source/Editor/UnrealEd/Private/Commandlets/PackageUtilities.cpp
--- a/Engine/Source/Editor/UnrealEd/Private/Commandlets/PackageUtilities.cpp
+++ b/Engine/Source/Editor/UnrealEd/Private/Commandlets/PackageUtilities.cpp
@@ -1780,7 +1780,7 @@ struct CompressAnimationsFunctor
 			++ActiveAnimationIndex;
 
 			// If animation hasn't been compressed, force it.
-			bool bForceCompression = (AnimSeq->CompressedTrackOffsets.Num() == 0);
+			bool bForceCompression = !AnimSeq->IsCompressedDataValid();
 
 			// If animation has already been compressed with the commandlet and version is the same. then skip.
 			// We're only interested in new animations.
diff --git a/Engine/Source/Runtime/AnimGraphRuntime/Private/AnimNodes/AnimNode_PoseBlendNode.cpp b/Engine/Source/Runtime/AnimGraphRuntime/Private/AnimNodes/AnimNode_PoseBlendNode.cpp
--- a/Engine/Source/Runtime/AnimGraphRuntime/Private/AnimNodes/AnimNode_PoseBlendNode.cpp
+++ b/Engine/Source/Runtime/AnimGraphRuntime/Private/AnimNodes/AnimNode_PoseBlendNode.cpp
@@ -34,6 +34,8 @@ void FAnimNode_PoseBlendNode::UpdateAssetPlayer(const FAnimationUpdateContext& C
 
 void FAnimNode_PoseBlendNode::Evaluate_AnyThread(FPoseContext& Output)
 {
+	ANIM_MT_SCOPE_CYCLE_COUNTER(PoseBlendNodeEvaluate, !IsInGameThread());
+
 	FPoseContext SourceData(Output);
 	SourcePose.Evaluate(SourceData);
 
diff --git a/Engine/Source/Runtime/AnimGraphRuntime/Public/AnimNodes/AnimNode_RandomPlayer.h b/Engine/Source/Runtime/AnimGraphRuntime/Public/AnimNodes/AnimNode_RandomPlayer.h
--- a/Engine/Source/Runtime/AnimGraphRuntime/Public/AnimNodes/AnimNode_RandomPlayer.h
+++ b/Engine/Source/Runtime/AnimGraphRuntime/Public/AnimNodes/AnimNode_RandomPlayer.h
@@ -6,6 +6,7 @@
 #include "UObject/ObjectMacros.h"
 #include "Math/RandomStream.h"
 #include "Animation/AnimationAsset.h"
+#include "Animation/AnimSequenceDecompressionContext.h"
 #include "AlphaBlend.h"
 #include "Animation/AnimNodeBase.h"
 #include "AnimNode_RandomPlayer.generated.h"
diff --git a/Engine/Source/Runtime/AnimGraphRuntime/Public/AnimNodes/AnimNode_SequenceEvaluator.h b/Engine/Source/Runtime/AnimGraphRuntime/Public/AnimNodes/AnimNode_SequenceEvaluator.h
--- a/Engine/Source/Runtime/AnimGraphRuntime/Public/AnimNodes/AnimNode_SequenceEvaluator.h
+++ b/Engine/Source/Runtime/AnimGraphRuntime/Public/AnimNodes/AnimNode_SequenceEvaluator.h
@@ -5,6 +5,8 @@
 #include "CoreMinimal.h"
 #include "UObject/ObjectMacros.h"
 #include "Animation/AnimNode_AssetPlayerBase.h"
+#include "Animation/AnimSequenceBase.h"
+#include "Animation/AnimSequenceDecompressionContext.h"
 #include "AnimNode_SequenceEvaluator.generated.h"
 
 UENUM(BlueprintType)
@@ -58,7 +60,7 @@ public:
 	UPROPERTY(Transient)
 	bool bReinitialized;
 
-public:	
+public:
 	FAnimNode_SequenceEvaluator()
 		: Sequence(nullptr)
 		, ExplicitTime(0.0f)
diff --git a/Engine/Source/Runtime/Engine/Classes/Animation/AnimCompress.h b/Engine/Source/Runtime/Engine/Classes/Animation/AnimCompress.h
--- a/Engine/Source/Runtime/Engine/Classes/Animation/AnimCompress.h
+++ b/Engine/Source/Runtime/Engine/Classes/Animation/AnimCompress.h
@@ -10,14 +10,174 @@
 #include "CoreMinimal.h"
 #include "UObject/ObjectMacros.h"
 #include "UObject/Object.h"
+#include "Templates/Function.h"
 #include "Animation/AnimSequence.h"
 #include "AnimationUtils.h"
 #include "AnimEnums.h"
+#include "AnimationCompression.h"
 #include "AnimCompress.generated.h"
 
 //Helper function for ddc key generation
 uint8 MakeBitForFlag(uint32 Item, uint32 Position);
 
+// Logic for tracking top N error items for later display
+template<typename DataType, typename SortType, int MaxItems>
+struct FMaxErrorStatTracker
+{
+public:
+	FMaxErrorStatTracker()
+		: CurrentLowestError(0.f)
+	{
+		Items.Reserve(MaxItems);
+	}
+
+	bool CanUseErrorStat(SortType NewError)
+	{
+		return Items.Num() < MaxItems || NewError > CurrentLowestError;
+	}
+
+	template <typename... ArgsType>
+	void StoreErrorStat(SortType NewError, ArgsType&&... Args)
+	{
+		bool bModified = false;
+
+		if (Items.Num() < MaxItems)
+		{
+			Items.Emplace(Forward<ArgsType>(Args)...);
+			bModified = true;
+		}
+		else if(NewError > CurrentLowestError)
+		{
+			Items[MaxItems - 1] = DataType(Forward<ArgsType>(Args)...);
+			bModified = true;
+		}
+
+		if (bModified)
+		{
+			Algo::Sort(Items, TGreater<>());
+			CurrentLowestError = Items.Last().GetErrorValue();
+		}
+	}
+
+	void LogErrorStat()
+	{
+		for (int ItemIndex = 0; ItemIndex < Items.Num(); ++ItemIndex)
+		{
+			UE_LOG(LogAnimationCompression, Display, TEXT("%i) %s"), ItemIndex+1, *Items[ItemIndex].ToText().ToString());
+		}
+	}
+
+	const DataType& GetMaxErrorItem() const
+	{
+		return Items[0];
+	}
+
+private:
+	//Storage of tracked items
+	TArray<DataType> Items;
+
+	//For ease cache current lowest error value
+	SortType CurrentLowestError;
+};
+
+struct FErrorTrackerWorstBone
+{
+	FErrorTrackerWorstBone()
+		: BoneError(0)
+		, BoneErrorTime(0)
+		, BoneErrorBone(0)
+		, BoneErrorBoneName(NAME_None)
+		, BoneErrorAnimName(NAME_None)
+	{}
+
+	FErrorTrackerWorstBone(float InBoneError, float InBoneErrorTime, int32 InBoneErrorBone, FName InBoneErrorBoneName, FName InBoneErrorAnimName)
+		: BoneError(InBoneError)
+		, BoneErrorTime(InBoneErrorTime)
+		, BoneErrorBone(InBoneErrorBone)
+		, BoneErrorBoneName(InBoneErrorBoneName)
+		, BoneErrorAnimName(InBoneErrorAnimName)
+	{}
+
+	bool operator<(const FErrorTrackerWorstBone& Rhs) const
+	{
+		return BoneError < Rhs.BoneError;
+	}
+
+	float GetErrorValue() const { return BoneError; }
+
+	FText ToText() const
+	{
+		FNumberFormattingOptions Options;
+		Options.MinimumIntegralDigits = 1;
+		Options.MinimumFractionalDigits = 3;
+
+		FFormatNamedArguments Args;
+		Args.Add(TEXT("BoneError"), FText::AsNumber(BoneError, &Options));
+		Args.Add(TEXT("BoneErrorAnimName"), FText::FromName(BoneErrorAnimName));
+		Args.Add(TEXT("BoneErrorBoneName"), FText::FromName(BoneErrorBoneName));
+		Args.Add(TEXT("BoneErrorBone"), BoneErrorBone);
+		Args.Add(TEXT("BoneErrorTime"), FText::AsNumber(BoneErrorTime, &Options));
+
+		return FText::Format(NSLOCTEXT("Engine", "CompressionWorstBoneSummary", "{BoneError} in Animation {BoneErrorAnimName}, Bone : {BoneErrorBoneName}(#{BoneErrorBone}), at Time {BoneErrorTime}"), Args);
+	}
+
+	// Error of this bone
+	float BoneError;
+
+	// Time in the sequence that the error occurred at
+	float BoneErrorTime;
+
+	// Bone index the error occurred on
+	int32 BoneErrorBone;
+
+	// Bone name the error occurred on 
+	FName BoneErrorBoneName;
+
+	// Animation the error occurred on
+	FName BoneErrorAnimName;
+};
+
+struct FErrorTrackerWorstAnimation
+{
+	FErrorTrackerWorstAnimation()
+		: AvgError(0)
+		, AnimName(NAME_None)
+	{}
+
+	FErrorTrackerWorstAnimation(float InAvgError, FName InMaxErrorAnimName)
+		: AvgError(InAvgError)
+		, AnimName(InMaxErrorAnimName)
+	{}
+
+	bool operator<(const FErrorTrackerWorstAnimation& Rhs) const
+	{
+		return AvgError < Rhs.AvgError;
+	}
+
+	float GetErrorValue() const { return AvgError; }
+
+	FText ToText() const
+	{
+		FNumberFormattingOptions Options;
+		Options.MinimumIntegralDigits = 1;
+		Options.MinimumFractionalDigits = 3;
+
+		FFormatNamedArguments Args;
+		Args.Add(TEXT("AvgError"), FText::AsNumber(AvgError, &Options));
+		Args.Add(TEXT("AnimName"), FText::FromName(AnimName));
+
+		return FText::Format(NSLOCTEXT("Engine", "CompressionWorstAnimationSummary", "{AvgError} in Animation {AnimName}"), Args);
+	}
+
+private:
+
+	// Average error of this animation
+	float AvgError;
+
+	// Animation being tracked
+	FName AnimName;
+};
+
 class ENGINE_API FCompressionMemorySummary
 {
 public:
@@ -25,7 +185,7 @@ public:
 
 	void GatherPreCompressionStats(UAnimSequence* Seq, int32 ProgressNumerator, int32 ProgressDenominator);
 
-	void GatherPostCompressionStats(UAnimSequence* Seq, TArray<FBoneData>& BoneData);
+	void GatherPostCompressionStats(UAnimSequence* Seq, const TArray<FBoneData>& BoneData, double CompressionTime);
 
 	~FCompressionMemorySummary();
 
@@ -35,15 +195,21 @@ private:
 	int32 TotalRaw;
 	int32 TotalBeforeCompressed;
 	int32 TotalAfterCompressed;
+	int32 NumberOfAnimations;
+
+	// Total time spent compressing animations
+	double TotalCompressionExecutionTime;
 
+	// Stats across all animations
 	float ErrorTotal;
 	float ErrorCount;
 	float AverageError;
-	float MaxError;
-	float MaxErrorTime;
-	int32 MaxErrorBone;
-	FName MaxErrorBoneName;
-	FName MaxErrorAnimName;
+
+	// Track the largest errors on a single bone
+	FMaxErrorStatTracker<FErrorTrackerWorstBone, float, 10> WorstBoneError;
+
+	// Track the animations with the largest average error
+	FMaxErrorStatTracker<FErrorTrackerWorstAnimation, float, 10> WorstAnimationError;
 };
 
 //////////////////////////////////////////////////////////////////////////
@@ -51,7 +217,15 @@ private:
 // animation compression
 struct ENGINE_API FAnimCompressContext
 {
+private:
 	FCompressionMemorySummary	CompressionSummary;
+
+	void GatherPreCompressionStats(UAnimSequence* Seq);
+
+	void GatherPostCompressionStats(UAnimSequence* Seq, const TArray<FBoneData>& BoneData, double CompressionTime);
+
+
+public:
 	uint32						AnimIndex;
 	uint32						MaxAnimations;
 	bool						bAllowAlternateCompressor;
@@ -59,9 +233,32 @@ struct ENGINE_API FAnimCompressContext
 
 	FAnimCompressContext(bool bInAllowAlternateCompressor, bool bInOutput, uint32 InMaxAnimations = 1) : CompressionSummary(bInOutput), AnimIndex(0), MaxAnimations(InMaxAnimations), bAllowAlternateCompressor(bInAllowAlternateCompressor), bOutput(bInOutput) {}
 
-	void GatherPreCompressionStats(UAnimSequence* Seq);
+	// If we are duping a compression context we don't want the CompressionSummary to output
+	FAnimCompressContext(const FAnimCompressContext& Rhs) : CompressionSummary(false), AnimIndex(Rhs.AnimIndex), MaxAnimations(Rhs.MaxAnimations), bAllowAlternateCompressor(Rhs.bAllowAlternateCompressor), bOutput(Rhs.bOutput) {}
+
+	friend class FAnimationUtils;
+};
+
+//////////////////////////////////////////////////////////////////////////
+// FAnimSegmentContext - This holds the relevant intermediate information
+// when compressing animation sequence segments.
+struct FAnimSegmentContext
+{
+	int32 StartFrame;
+	int32 NumFrames;
+
+	TArray<struct FTranslationTrack> TranslationData;
+	TArray<struct FRotationTrack> RotationData;
+	TArray<struct FScaleTrack> ScaleData;
+
+	AnimationCompressionFormat TranslationCompressionFormat;
+	AnimationCompressionFormat RotationCompressionFormat;
+	AnimationCompressionFormat ScaleCompressionFormat;
 
-	void GatherPostCompressionStats(UAnimSequence* Seq, TArray<FBoneData>& BoneData);
+	TArray<int32> CompressedTrackOffsets;
+	FCompressedOffsetData CompressedScaleOffsets;
+	TArray<uint8> CompressedByteStream;
+	TArray<uint8> CompressedTrivialTracksByteStream;
 };
 
 UCLASS(abstract, hidecategories=Object, MinimalAPI, EditInlineNew)
@@ -77,6 +274,19 @@ class UAnimCompress : public UObject
 	UPROPERTY()
 	uint32 bNeedsSkeleton:1;
 
+	/** Whether to enable segmenting or not. Needed for USE_SEGMENTING_CONTEXT (currently turned off) */
+	//UPROPERTY(Category = Compression, EditAnywhere)
+	UPROPERTY()
+	uint32 bEnableSegmenting:1;
+
+	/** When splitting the sequence into segments, we will try to approach this value as much as possible. */
+	UPROPERTY(Category = Compression, EditAnywhere, meta = (ClampMin = "16"))
+	uint32 IdealNumFramesPerSegment;
+
+	/** When splitting the sequence into segments, we will never allow a segment with more than the maximum number of frames. */
+	UPROPERTY(Category = Compression, EditAnywhere)
+	uint32 MaxNumFramesPerSegment;
+
 	/** Format for bitwise compression of translation data. */
 	UPROPERTY()
 	TEnumAsByte<AnimationCompressionFormat> TranslationCompressionFormat;
@@ -102,7 +312,7 @@ public:
 	 * @param	bOutput		If false don't generate output or compute memory savings.
 	 * @return				false if a skeleton was needed by the algorithm but not provided.
 	 */
-	ENGINE_API bool Reduce(class UAnimSequence* AnimSeq, bool bOutput);
+	ENGINE_API bool Reduce(class UAnimSequence* AnimSeq, bool bOutput, const TArray<FBoneData>& BoneData);
 
 	/**
 	 * Reduce the number of keyframes and bitwise compress all sequences in the specified array.
@@ -111,7 +321,7 @@ public:
 	 * @param	bOutput			If false don't generate output or compute memory savings.
 	 * @return					false if a skeleton was needed by the algorithm but not provided.
 	 */
-	ENGINE_API bool Reduce(class UAnimSequence* AnimSeq, FAnimCompressContext& Context);
+	ENGINE_API bool Reduce(class UAnimSequence* AnimSeq, FAnimCompressContext& Context, const TArray<FBoneData>& BoneData);
 #endif // WITH_EDITOR
 protected:
 #if WITH_EDITOR
@@ -189,10 +399,10 @@ protected:
 	 *
 	 * @param	PositionTracks	Array of position track elements to reduce
 	 * @param	RotationTracks	Array of rotation track elements to reduce
-	 * @param	ScaleTracks		Array of scale track elements to reduce	 
+	 * @param	ScaleTracks		Array of scale track elements to reduce
 	 * @param	MaxPosDelta		Maximum local-space threshold for stationary motion
 	 * @param	MaxRotDelta		Maximum angle threshold to consider stationary motion
-	 * @param	MaxScaleDelta	Maximum scale threshold to consider stationary motion	 
+	 * @param	MaxScaleDelta	Maximum scale threshold to consider stationary motion
 	 */
 	static void FilterTrivialKeys(
 		TArray<struct FTranslationTrack>& PositionTracks,
@@ -202,13 +412,19 @@ protected:
 		float MaxRotDelta, 
 		float MaxScaleDelta);
 
-	/** 
-	 * Remove translation keys from tracks marked bAnimRotationOnly.
+	/**
+	 * Same as above function but it will execute for every segment.
 	 *
-	 * @param PositionTracks	Array of position track elements to reduce
-	 * @param AnimSeq			AnimSequence the track is from.
+	 * @param	RawSegments		Array of segments being compressed
+	 * @param	MaxPosDelta		Maximum local-space threshold for stationary motion
+	 * @param	MaxRotDelta		Maximum angle threshold to consider stationary motion
+	 * @param	MaxScaleDelta	Maximum scale threshold to consider stationary motion
 	 */
-	static void FilterAnimRotationOnlyKeys(TArray<FTranslationTrack>& PositionTracks, UAnimSequence* AnimSeq);
+	static void FilterTrivialKeys(
+		TArray<FAnimSegmentContext>& RawSegments,
+		float MaxPosDelta,
+		float MaxRotDelta,
+		float MaxScaleDelta);
 
 	/**
 	 * Common compression utility to retain only intermittent position keys. For example,
@@ -294,6 +510,26 @@ protected:
 		TArray<struct FRotationTrack>& OutRotationData, 
 		TArray<struct FScaleTrack>& OutScaleData);
 
+	/**
+	 * This will populate segments from the raw animation data.
+	 *
+	 * @param	AnimSeq						AnimSequence being compressed
+	 * @param	TranslationData				Translation tracks
+	 * @param	RotationData				Rotation tracks
+	 * @param	ScaleData					Scale tracks
+	 * @param	IdealNumFramesPerSegment	The ideal number of frames within a segment
+	 * @param	MaxNumFramesPerSegment		The maximum number of frames within a segment
+	 * @param	OutRawSegments				The array of segments to populate
+	 */
+	static void SeparateRawDataIntoTracks(
+		const UAnimSequence& AnimSeq,
+		const TArray<struct FTranslationTrack>& TranslationData,
+		const TArray<struct FRotationTrack>& RotationData,
+		const TArray<struct FScaleTrack>& ScaleData,
+		int32 IdealNumFramesPerSegment,
+		int32 MaxNumFramesPerSegment,
+		TArray<FAnimSegmentContext>& OutRawSegments);
+
 	/**
 	 * Common compression utility to walk an array of rotation tracks and enforce
 	 * that all adjacent rotation keys are represented by shortest-arc quaternion pairs.
@@ -317,7 +553,7 @@ public:
 	 * @param	IncludeKeyTable				true if the compressed data should also contain a table of frame indices for each key. (required by some codecs)
 	 */
 	static void BitwiseCompressAnimationTracks(
-		class UAnimSequence* Seq, 
+		UAnimSequence* Seq, 
 		AnimationCompressionFormat TargetTranslationFormat, 
 		AnimationCompressionFormat TargetRotationFormat,
 		AnimationCompressionFormat TargetScaleFormat,
@@ -326,12 +562,339 @@ public:
 		const TArray<FScaleTrack>& ScaleData,
 		bool IncludeKeyTable = false);
 
+	/**
+	 * Encodes individual key arrays into an AnimSequence using the desired bit packing formats.
+	 * This action is performed for every segment supplied.
+	 *
+	 * @param	AnimSeq						Animation Sequence which will contain the bit-packed data.
+	 * @param	TargetTranslationFormat		The format to use when encoding translation keys.
+	 * @param	TargetRotationFormat		The format to use when encoding rotation keys.
+	 * @param	TargetScaleFormat			The format to use when encoding scale keys.
+	 * @param	RawSegments					The array of segments to compress
+	 * @param	bIsSorted					For variable interpolation, is the compressed data sorted or not?
+	 */
+	static void BitwiseCompressAnimationTracks(
+		UAnimSequence& AnimSeq,
+		AnimationCompressionFormat TargetTranslationFormat,
+		AnimationCompressionFormat TargetRotationFormat,
+		AnimationCompressionFormat TargetScaleFormat,
+		TArray<FAnimSegmentContext>& RawSegments,
+		bool bIsSorted = false);
+
+	/**
+	 * Encodes individual key arrays into an AnimSequence using the desired bit packing formats.
+	 * This action is performed for a single segment.
+	 *
+	 * @param	AnimSeq						Animation Sequence which will contain the bit-packed data.
+	 * @param	TargetTranslationFormat		The format to use when encoding translation keys.
+	 * @param	TargetRotationFormat		The format to use when encoding rotation keys.
+	 * @param	TargetScaleFormat			The format to use when encoding scale keys.
+	 * @param	RawSegment					The segment to compress
+	 * @param	bIsSorted					For variable interpolation, is the compressed data sorted or not?
+	 */
+	static void BitwiseCompressAnimationTracks(
+		const UAnimSequence& AnimSeq,
+		AnimationCompressionFormat TargetTranslationFormat,
+		AnimationCompressionFormat TargetRotationFormat,
+		AnimationCompressionFormat TargetScaleFormat,
+		FAnimSegmentContext& RawSegment,
+		bool bIsSorted = false);
+
+	/**
+	 * Encodes the trivial tracks within a segment.
+	 *
+	 * @param	AnimSeq						Animation Sequence which will contain the bit-packed data.
+	 * @param	RawSegment					The segment to compress
+	 */
+	static void BitwiseCompressTrivialAnimationTracks(const UAnimSequence& AnimSeq, FAnimSegmentContext& RawSegment);
+
+	/**
+	 * Coalesces the compressed data from every segment into a single contiguous array stored within the sequence.
+	 *
+	 * @param	AnimSeq						Animation Sequence which will contain the bit-packed data.
+	 * @param	RawSegments					The array of segments to coalesce
+	 * @param	bIsSorted					For variable interpolation, is the compressed data sorted or not?
+	 */
+	static void CoalesceCompressedSegments(UAnimSequence& AnimSeq, const TArray<FAnimSegmentContext>& RawSegments, bool bIsSorted = false);
+
 #if WITH_EDITOR
 	FString MakeDDCKey();
 
 protected:
 	virtual void PopulateDDCKey(FArchive& Ar);
 #endif // WITH_EDITOR
+
+	/**
+	 * Structure that holds the range min/extent for a track.
+	 */
+	struct FAnimTrackRange
+	{
+		FVector RotMin;
+		FVector RotExtent;
+		FVector TransMin;
+		FVector TransExtent;
+		FVector ScaleMin;
+		FVector ScaleExtent;
+	};
+
+	/**
+	 * Utility function to append data to a byte stream.
+	 *
+	 * @param	ByteStream					Byte stream to append to
+	 * @param	Src							Pointer to the source data to append
+	 * @param	Len							Length in bytes of the source data to append
+	 */
+	static void UnalignedWriteToStream(TArray<uint8>& ByteStream, const void* Src, SIZE_T Len);
+
+	/**
+	* Utility function to write data to a byte stream.
+	*
+	* @param	ByteStream					Byte stream to write to
+	* @param	StreamOffset				Offset in stream to start writing to
+	* @param	Src							Pointer to the source data to write
+	* @param	Len							Length in bytes of the source data to write
+	*/
+	static void UnalignedWriteToStream(TArray<uint8>& ByteStream, int32& StreamOffset, const void* Src, SIZE_T Len);
+
+	/**
+	 * Utility function to append a packed FVector to a byte stream.
+	 *
+	 * @param	ByteStream					Byte stream to append to
+	 * @param	Format						Compression format to pack with
+	 * @param	Vec							The FVector to pack
+	 * @param	Mins						The range minimum of the input value to pack for range normalization
+	 * @param	Ranges						The range extent of the input value to pack for range normalization
+	 */
+	static void PackVectorToStream(
+		TArray<uint8>& ByteStream,
+		AnimationCompressionFormat Format,
+		const FVector& Vec,
+		const float* Mins,
+		const float* Ranges);
+
+	/**
+	* Utility function to append a packed FQuat to a byte stream.
+	*
+	* @param	ByteStream					Byte stream to append to
+	* @param	Format						Compression format to pack with
+	* @param	Vec							The FQuat to pack
+	* @param	Mins						The range minimum of the input value to pack for range normalization
+	* @param	Ranges						The range extent of the input value to pack for range normalization
+	*/
+	static void PackQuaternionToStream(
+		TArray<uint8>& ByteStream,
+		AnimationCompressionFormat Format,
+		const FQuat& Quat,
+		const float* Mins,
+		const float* Ranges);
+
+	/**
+	 * Utility function that performs minimal sanity checks.
+	 *
+	 * @param	AnimSeq						The anim sequence to check
+	 * @param	Segment						The segment to check
+	 */
+	static void SanityCheckTrackData(const UAnimSequence& AnimSeq, const FAnimSegmentContext& Segment);
+
+	/**
+	 * Calculates the translation track range.
+	 *
+	 * @param	TranslationData				The translation track data
+	 * @param	Format						Compression format
+	 * @param	OutMin						The output range minimum
+	 * @param	OutExtent					The output range extent
+	 */
+	static void CalculateTrackRange(const FTranslationTrack& TranslationData, AnimationCompressionFormat Format, FVector& OutMin, FVector& OutExtent);
+
+	/**
+	 * Calculates the rotation track range.
+	 *
+	 * @param	TranslationData				The rotation track data
+	 * @param	Format						Compression format
+	 * @param	OutMin						The output range minimum
+	 * @param	OutExtent					The output range extent
+	 */
+	static void CalculateTrackRange(const FRotationTrack& RotationData, AnimationCompressionFormat Format, FVector& OutMin, FVector& OutExtent);
+
+	/**
+	 * Calculates the sacle track range.
+	 *
+	 * @param	TranslationData				The scale track data
+	 * @param	Format						Compression format
+	 * @param	OutMin						The output range minimum
+	 * @param	OutExtent					The output range extent
+	 */
+	static void CalculateTrackRange(const FScaleTrack& ScaleData, AnimationCompressionFormat Format, FVector& OutMin, FVector& OutExtent);
+
+	/**
+	 * Calculates the track ranges within a segment.
+	 *
+	 * @param	TargetTranslationFormat		Compression format for translations
+	 * @param	TargetRotationFormat		Compression format for rotations
+	 * @param	TargetScaleFormat			Compression format for scales
+	 * @param	Segment						The segment that contains our tracks to process
+	 * @param	TrackRanges					The calculated track ranges
+	 */
+	static void CalculateTrackRanges(
+		AnimationCompressionFormat TargetTranslationFormat,
+		AnimationCompressionFormat TargetRotationFormat,
+		AnimationCompressionFormat TargetScaleFormat,
+		const FAnimSegmentContext& Segment,
+		TArray<FAnimTrackRange>& TrackRanges);
+
+	/**
+	 * Structure to wrap and represent track key flags.
+	 */
+	struct FTrackKeyFlags
+	{
+		constexpr FTrackKeyFlags() : Flags(0) {}
+		constexpr explicit FTrackKeyFlags(uint8 InFlags) : Flags(InFlags) {}
+
+		constexpr bool IsComponentNeededX() const { return (Flags & 0x1) != 0; }
+		constexpr bool IsComponentNeededY() const { return (Flags & 0x2) != 0; }
+		constexpr bool IsComponentNeededZ() const { return (Flags & 0x4) != 0; }
+
+		constexpr bool IsValid() const { return (Flags & ~0x7) == 0; }
+
+		uint8 Flags;
+	};
+
+	/**
+	 * Writes the necessary track ranges to a byte stream.
+	 *
+	 * @param	ByteStream					Byte stream to append to
+	 * @param	GetTranslationFormatFun		Function that returns the translation format for a track index
+	 * @param	GetRotationFormatFun		Function that returns the rotation format for a track index
+	 * @param	GetScaleFormatFun			Function that returns the scale format for a track index
+	 * @param	GetTranslationFlagsFun		Function that returns the translation flags for a track index
+	 * @param	GetRotationFlagsFun			Function that returns the rotation flags for a track index
+	 * @param	GetScaleFlagsFun			Function that returns the scale flags for a track index
+	 * @param	Segment						Segment to process
+	 * @param	TrackRanges					Track ranges to pack
+	 * @param	bInterleaveValues			Whether to interleave range values or not
+	 */
+	static void WriteTrackRanges(
+		TArray<uint8>& ByteStream,
+		TFunction<AnimationCompressionFormat(int32 TrackIndex)> GetTranslationFormatFun,
+		TFunction<AnimationCompressionFormat(int32 TrackIndex)> GetRotationFormatFun,
+		TFunction<AnimationCompressionFormat(int32 TrackIndex)> GetScaleFormatFun,
+		TFunction<FTrackKeyFlags(int32 TrackIndex)> GetTranslationFlagsFun,
+		TFunction<FTrackKeyFlags(int32 TrackIndex)> GetRotationFlagsFun,
+		TFunction<FTrackKeyFlags(int32 TrackIndex)> GetScaleFlagsFun,
+		const FAnimSegmentContext& Segment,
+		const TArray<FAnimTrackRange>& TrackRanges,
+		bool bInterleaveValues);
+
+	/**
+	 * Writes a segment's uniform track data to a byte stream.
+	 * A track's data is uniform if no keys are removed.
+	 *
+	 * @param	ByteStream					Byte stream to append to
+	 * @param	GetTranslationFormatFun		Function that returns the translation format for a track index
+	 * @param	GetRotationFormatFun		Function that returns the rotation format for a track index
+	 * @param	GetScaleFormatFun			Function that returns the scale format for a track index
+	 * @param	IsTranslationUniformFun		Function that returns if a track translation's data is uniform or not for a track index
+	 * @param	IsRotationUniformFun		Function that returns if a track rotation's data is uniform or not for a track index
+	 * @param	IsScaleUniformFun			Function that returns if a track scale's data is uniform or not for a track index
+	 * @param	PackTranslationKeyFun		Function that packs a translation key for a track index
+	 * @param	PackRotationKeyFun			Function that packs a rotation key for a track index
+	 * @param	PackScaleKeyFun				Function that packs a scale key for a track index
+	 * @param	Segment						Segment to process
+	 * @param	TrackRanges					Track ranges for the segment to process
+	 */
+	static void WriteUniformTrackData(
+		TArray<uint8>& ByteStream,
+		TFunction<AnimationCompressionFormat(int32 TrackIndex)> GetTranslationFormatFun,
+		TFunction<AnimationCompressionFormat(int32 TrackIndex)> GetRotationFormatFun,
+		TFunction<AnimationCompressionFormat(int32 TrackIndex)> GetScaleFormatFun,
+		TFunction<bool(int32 TrackIndex)> IsTranslationUniformFun,
+		TFunction<bool(int32 TrackIndex)> IsRotationUniformFun,
+		TFunction<bool(int32 TrackIndex)> IsScaleUniformFun,
+		TFunction<void(TArray<uint8>& ByteStream, AnimationCompressionFormat Format, const FVector& Key, const float* Mins, const float* Ranges, int32 TrackIndex)> PackTranslationKeyFun,
+		TFunction<void(TArray<uint8>& ByteStream, AnimationCompressionFormat Format, const FQuat& Key, const float* Mins, const float* Ranges, int32 TrackIndex)> PackRotationKeyFun,
+		TFunction<void(TArray<uint8>& ByteStream, AnimationCompressionFormat Format, const FVector& Key, const float* Mins, const float* Ranges, int32 TrackIndex)> PackScaleKeyFun,
+		const FAnimSegmentContext& Segment,
+		const TArray<FAnimTrackRange>& TrackRanges);
+
+	/**
+	 * Writes a segment's variable track data to a byte stream with a sorted ordering.
+	 * A track's data is variable if some keys are removed.
+	 *
+	 * @param	ByteStream					Byte stream to append to
+	 * @param	AnimSeq						The anim sequence to process
+	 * @param	GetTranslationFormatFun		Function that returns the translation format for a track index
+	 * @param	GetRotationFormatFun		Function that returns the rotation format for a track index
+	 * @param	GetScaleFormatFun			Function that returns the scale format for a track index
+	 * @param	IsTranslationUniformFun		Function that returns if a track translation's data is uniform or not for a track index
+	 * @param	IsRotationUniformFun		Function that returns if a track rotation's data is uniform or not for a track index
+	 * @param	IsScaleUniformFun			Function that returns if a track scale's data is uniform or not for a track index
+	 * @param	PackTranslationKeyFun		Function that packs a translation key for a track index
+	 * @param	PackRotationKeyFun			Function that packs a rotation key for a track index
+	 * @param	PackScaleKeyFun				Function that packs a scale key for a track index
+	 * @param	Segment						Segment to process
+	 * @param	TrackRanges					Track ranges for the segment to process
+	 */
+	static void WriteSortedVariableTrackData(
+		TArray<uint8>& ByteStream,
+		const UAnimSequence& AnimSeq,
+		TFunction<AnimationCompressionFormat(int32 TrackIndex)> GetTranslationFormatFun,
+		TFunction<AnimationCompressionFormat(int32 TrackIndex)> GetRotationFormatFun,
+		TFunction<AnimationCompressionFormat(int32 TrackIndex)> GetScaleFormatFun,
+		TFunction<bool(int32 TrackIndex)> IsTranslationVariableFun,
+		TFunction<bool(int32 TrackIndex)> IsRotationVariableFun,
+		TFunction<bool(int32 TrackIndex)> IsScaleVariableFun,
+		TFunction<void(TArray<uint8>& ByteStream, AnimationCompressionFormat Format, const FVector& Key, const float* Mins, const float* Ranges, int32 TrackIndex)> PackTranslationKeyFun,
+		TFunction<void(TArray<uint8>& ByteStream, AnimationCompressionFormat Format, const FQuat& Key, const float* Mins, const float* Ranges, int32 TrackIndex)> PackRotationKeyFun,
+		TFunction<void(TArray<uint8>& ByteStream, AnimationCompressionFormat Format, const FVector& Key, const float* Mins, const float* Ranges, int32 TrackIndex)> PackScaleKeyFun,
+		FAnimSegmentContext& Segment,
+		const TArray<FAnimTrackRange>& TrackRanges);
+
+	/**
+	 * Writes a segment's variable track data to a byte stream with a linear ordering.
+	 * A track's data is variable if some keys are removed.
+	 *
+	 * @param	ByteStream					Byte stream to append to
+	 * @param	AnimSeq						The anim sequence to process
+	 * @param	GetTranslationFormatFun		Function that returns the translation format for a track index
+	 * @param	GetRotationFormatFun		Function that returns the rotation format for a track index
+	 * @param	GetScaleFormatFun			Function that returns the scale format for a track index
+	 * @param	IsTranslationUniformFun		Function that returns if a track translation's data is uniform or not for a track index
+	 * @param	IsRotationUniformFun		Function that returns if a track rotation's data is uniform or not for a track index
+	 * @param	IsScaleUniformFun			Function that returns if a track scale's data is uniform or not for a track index
+	 * @param	PackTranslationKeyFun		Function that packs a translation key for a track index
+	 * @param	PackRotationKeyFun			Function that packs a rotation key for a track index
+	 * @param	PackScaleKeyFun				Function that packs a scale key for a track index
+	 * @param	Segment						Segment to process
+	 * @param	TrackRanges					Track ranges for the segment to process
+	 */
+	static void WriteLinearVariableTrackData(
+		TArray<uint8>& ByteStream,
+		const UAnimSequence& AnimSeq,
+		TFunction<AnimationCompressionFormat(int32 TrackIndex)> GetTranslationFormatFun,
+		TFunction<AnimationCompressionFormat(int32 TrackIndex)> GetRotationFormatFun,
+		TFunction<AnimationCompressionFormat(int32 TrackIndex)> GetScaleFormatFun,
+		TFunction<bool(int32 TrackIndex)> IsTranslationVariableFun,
+		TFunction<bool(int32 TrackIndex)> IsRotationVariableFun,
+		TFunction<bool(int32 TrackIndex)> IsScaleVariableFun,
+		TFunction<void(TArray<uint8>& ByteStream, AnimationCompressionFormat Format, const FVector& Key, const float* Mins, const float* Ranges, int32 TrackIndex)> PackTranslationKeyFun,
+		TFunction<void(TArray<uint8>& ByteStream, AnimationCompressionFormat Format, const FQuat& Key, const float* Mins, const float* Ranges, int32 TrackIndex)> PackRotationKeyFun,
+		TFunction<void(TArray<uint8>& ByteStream, AnimationCompressionFormat Format, const FVector& Key, const float* Mins, const float* Ranges, int32 TrackIndex)> PackScaleKeyFun,
+		const FAnimSegmentContext& Segment,
+		const TArray<FAnimTrackRange>& TrackRanges);
+
+	/**
+	 * Pads a byte stream to force a particular alignment for the data to follow.
+	 *
+	 * @param	ByteStream					Byte stream to append to
+	 * @param	Alignment					Required alignment
+	 * @param	Sentinel					If we need to add padding to meet the requested alignment, this is the padding value used
+	 */
+	static void PadByteStream(TArray<uint8>& ByteStream, const int32 Alignment, uint8 Sentinel);
+
+	/**
+	 * Default animation padding value.
+	 */
+	static constexpr uint8 AnimationPadSentinel = 85; //(1<<1)+(1<<3)+(1<<5)+(1<<7)
 };
 
 
diff --git a/Engine/Source/Runtime/Engine/Classes/Animation/AnimCompress_Automatic.h b/Engine/Source/Runtime/Engine/Classes/Animation/AnimCompress_Automatic.h
--- a/Engine/Source/Runtime/Engine/Classes/Animation/AnimCompress_Automatic.h
+++ b/Engine/Source/Runtime/Engine/Classes/Animation/AnimCompress_Automatic.h
@@ -27,22 +27,6 @@ class UAnimCompress_Automatic : public UAnimCompress
 	UPROPERTY(EditAnywhere, Category=AnimationCompressionAlgorithm_Automatic)
 	float MaxEndEffectorError;
 
-	/** If true, the uniform bitwise techniques will be tried */
-	UPROPERTY(EditAnywhere, Category=AnimationCompressionAlgorithm_Automatic)
-	uint32 bTryFixedBitwiseCompression:1;
-
-	/** If true, the per-track compressor techniques will be tried */
-	UPROPERTY(EditAnywhere, Category=AnimationCompressionAlgorithm_Automatic)
-	uint32 bTryPerTrackBitwiseCompression:1;
-
-	/** If true, the linear key removal techniques will be tried */
-	UPROPERTY(EditAnywhere, Category=AnimationCompressionAlgorithm_Automatic)
-	uint32 bTryLinearKeyRemovalCompression:1;
-
-	/** If true, the resampling techniques will be tried */
-	UPROPERTY(EditAnywhere, Category=AnimationCompressionAlgorithm_Automatic)
-	uint32 bTryIntervalKeyRemoval:1;
-
 	/** If true, then the animation will be first recompressed with it's current compressor if non-NULL, or with the global default compressor (specified in the engine ini)
 	* Also known as "First Recompress Using Current Or Default"
 	*/
@@ -59,6 +43,10 @@ class UAnimCompress_Automatic : public UAnimCompress
 	UPROPERTY(EditAnywhere, Category=AnimationCompressionAlgorithm_Automatic)
 	uint32 bRaiseMaxErrorToExisting:1;
 
+	/** If true, then an exhaustive search is used otherwise only a short list of the best methods is tried */
+	UPROPERTY(EditAnywhere, Category=AnimationCompressionAlgorithm_Automatic)
+	uint32 bTryExhaustiveSearch:1;
+
 
 protected:
 	//~ Begin UAnimCompress Interface
diff --git a/Engine/Source/Runtime/Engine/Classes/Animation/AnimCompress_PerTrackCompression.h b/Engine/Source/Runtime/Engine/Classes/Animation/AnimCompress_PerTrackCompression.h
--- a/Engine/Source/Runtime/Engine/Classes/Animation/AnimCompress_PerTrackCompression.h
+++ b/Engine/Source/Runtime/Engine/Classes/Animation/AnimCompress_PerTrackCompression.h
@@ -153,6 +153,19 @@ protected:
 		const TArray<FRotationTrack>& RotationData,
 		const TArray<FScaleTrack>& ScaleData,
 		const bool bFinalPass) override;
+
+	virtual void CompressUsingUnderlyingCompressor(
+		UAnimSequence& AnimSeq,
+		const TArray<FBoneData>& BoneData,
+		TArray<FAnimSegmentContext>& RawSegments,
+		const bool bFinalPass) override;
+
+	/**
+	 * Performs the per track compression optimization for a single segment.
+	 * This can be called from multiple threads concurrently.
+	 */
+	void OptimizeSegmentTracks(struct FOptimizeSegmentTracksContext& Context) const;
+
 	virtual void FilterBeforeMainKeyRemoval(
 		UAnimSequence* AnimSeq, 
 		const TArray<FBoneData>& BoneData, 
@@ -160,6 +173,30 @@ protected:
 		TArray<FRotationTrack>& RotationData,
 		TArray<FScaleTrack>& ScaleData) override;
 	//~ End UAnimCompress_RemoveLinearKeys Interface
+
+	/**
+	* Structure to hold the track format information that we calculate is most optimal.
+	*/
+	struct FPerTrackFormat
+	{
+		AnimationCompressionFormat RotationFormat;
+		AnimationCompressionFormat TranslationFormat;
+		AnimationCompressionFormat ScaleFormat;
+
+		bool bHasRotationTimeMarkers;
+		bool bHasTranslationTimeMarkers;
+		bool bHasScaleTimeMarkers;
+
+		FTrackKeyFlags RotationKeyFlags;
+		FTrackKeyFlags TranslationKeyFlags;
+		FTrackKeyFlags ScaleKeyFlags;
+	};
+
+	static void PackTranslationKey(TArray<uint8>& ByteStream, AnimationCompressionFormat Format, const FVector& Key, const float* Mins, const float* Ranges, const struct FPerTrackFormat& TrackFormat);
+	static void PackRotationKey(TArray<uint8>& ByteStream, AnimationCompressionFormat Format, const FQuat& Key, const float* Mins, const float* Ranges, const struct FPerTrackFormat& TrackFormat);
+	static void PackScaleKey(TArray<uint8>& ByteStream, AnimationCompressionFormat Format, const FVector& Key, const float* Mins, const float* Ranges, const struct FPerTrackFormat& TrackFormat);
+
+	friend struct FAsyncOptimizeSegmentTracksTaskGroupContext;
 #endif // WITH_EDITOR
 };
 
diff --git a/Engine/Source/Runtime/Engine/Classes/Animation/AnimCompress_RemoveLinearKeys.h b/Engine/Source/Runtime/Engine/Classes/Animation/AnimCompress_RemoveLinearKeys.h
--- a/Engine/Source/Runtime/Engine/Classes/Animation/AnimCompress_RemoveLinearKeys.h
+++ b/Engine/Source/Runtime/Engine/Classes/Animation/AnimCompress_RemoveLinearKeys.h
@@ -14,6 +14,8 @@
 #include "Animation/AnimCompress.h"
 #include "AnimCompress_RemoveLinearKeys.generated.h"
 
+struct FProcessAnimationTracksContext;
+
 UCLASS(MinimalAPI)
 class UAnimCompress_RemoveLinearKeys : public UAnimCompress
 {
@@ -75,6 +77,31 @@ class UAnimCompress_RemoveLinearKeys : public UAnimCompress
 	UPROPERTY(EditAnywhere, Category=LinearKeyRemoval)
 	uint32 bActuallyFilterLinearKeys:1;
 
+	// USE_SEGMENTING_CONTEXT - Hide property 
+	/**
+	 * Whether or not to optimize the memory layout for forward playback. Note that this will degrade the performance
+	 * if the sequence is played backwards quite dramatically depending on the circumstances.
+	 */
+	//UPROPERTY(EditAnywhere, Category = LinearKeyRemoval, meta = (EditCondition="bEnableSegmenting"))
+	UPROPERTY()
+	uint32 bOptimizeForForwardPlayback : 1;
+
+	/**
+	 * Whether to use the legacy behavior that compressed and decompressed the whole sequence whenever a key was removed.
+	 * This is left in for now in order to debug the code if issues arise with the new stateless version.
+	 * The legacy behavior to compress the whole sequence constantly makes it impossible to compress segments
+	 * in parallel safely. The new code samples the sequence without the compression step in order to make it
+	 * thread safe.
+	 * Once the stateless code is robust, this can be safely removed. It should never be used.
+	 */
+	UPROPERTY(EditAnywhere, Category = LinearKeyRemoval)
+	uint32 bUseDecompression : 1;
+
+	/**
+	 * Whether or not to use multiple threads during compression.
+	 */
+	UPROPERTY(EditAnywhere, Category = LinearKeyRemoval)
+	uint32 bUseMultithreading : 1;
 
 protected:
 	//~ Begin UAnimCompress Interface
@@ -106,6 +133,15 @@ protected:
 		const TArray<FScaleTrack>& ScaleData,
 		const bool bFinalPass);
 
+	/**
+	 * Compresses the tracks passed in using the underlying compressor for this key removal codec for a single segment
+	 */
+	virtual void CompressUsingUnderlyingCompressor(
+		UAnimSequence& AnimSeq,
+		const TArray<FBoneData>& BoneData,
+		TArray<FAnimSegmentContext>& RawSegments,
+		const bool bFinalPass);
+
 	/**
 	 * Updates the world bone transforms for a range of bone indices
 	 */
@@ -121,19 +157,38 @@ protected:
 		bool UseRaw,
 		TArray<FTransform>& OutputWorldBones);
 
+	/**
+	 * Updates the world bone transforms for a range of bone indices for a single segment
+	 */
+	void UpdateWorldBoneTransformRange(
+		FProcessAnimationTracksContext& Context,
+		int32 StartingBoneIndex,
+		int32 EndingBoneIndex);
+
 	/**
 	 * To guide the key removal process, we need to maintain a table of world transforms
 	 * for the bones we are investigating. This helper function fills a row of the 
 	 * table for a specified bone.
 	 */
 	void UpdateWorldBoneTransformTable(
-		UAnimSequence* AnimSeq, 
-		const TArray<FBoneData>& BoneData, 
+		UAnimSequence* AnimSeq,
+		const TArray<FBoneData>& BoneData,
 		const TArray<FTransform>& RefPose,
 		int32 BoneIndex,
 		bool UseRaw,
 		TArray<FTransform>& OutputWorldBones);
 
+	/**
+	 * To guide the key removal process, we need to maintain a table of world transforms
+	 * for the bones we are investigating. This helper function fills a row of the
+	 * table for a specified bone. For a single segment.
+	 */
+	void UpdateWorldBoneTransformTable(
+		const FProcessAnimationTracksContext& Context,
+		int32 BoneIndex,
+		bool UseRaw,
+		TArray<FTransform>& OutputWorldBones);
+
 	/**
 	 * Creates a list of the bone atom result for every frame of a given track
 	 */
@@ -145,6 +200,14 @@ protected:
 		float TimePerFrame,
 		TArray<FTransform>& BoneAtoms);
 
+	/**
+	 * Creates a list of the bone atom result for every frame of a given track for a single segment
+	 */
+	void UpdateBoneAtomList(
+		const FProcessAnimationTracksContext& Context,
+		int32 TrackIndex,
+		TArray<FTransform>& BoneAtoms) const;
+
 	/**
 	 * If the passed in animation sequence is additive, converts it to absolute (using the frame 0 pose) and returns true
 	 * (indicating it should be converted back to relative later with ConvertToRelativeSpace)
@@ -156,15 +219,35 @@ protected:
 	bool ConvertFromRelativeSpace(UAnimSequence* AnimSeq);
 
 	/**
-	 * Converts an absolute animation sequence to a relative (additive) one.
+	 * Converts an absolute animation sequence and matching track data to a relative (additive) one.
 	 *
 	 * @param AnimSeq			The animation sequence being compressed
 	 * @param TranslationData	Translation Tracks to convert to relative space
-	 * @param RotationData		Rotation Tracks  to convert to relative space
+	 * @param RotationData		Rotation Tracks to convert to relative space
+	 * @param ScaleData			Scale Tracks to convert to relative space
 	 *
 	 */
 	void ConvertToRelativeSpace(UAnimSequence* AnimSeq, TArray<FTranslationTrack>& TranslationData, TArray<FRotationTrack>& RotationData, TArray<FScaleTrack>& ScaleData);
 
+	/**
+	 * Converts an absolute animation sequence to a relative (additive) one.
+	 *
+	 * @param AnimSeq			The animation sequence being compressed and converted
+	 *
+	 */
+	void ConvertToRelativeSpace(UAnimSequence& AnimSeq) const;
+
+	/**
+	 * Converts track data to relative (additive) space.
+	 *
+	 * @param AnimSeq			The animation sequence being compressed
+	 * @param TranslationData	Translation Tracks to convert to relative space
+	 * @param RotationData		Rotation Tracks to convert to relative space
+	 * @param ScaleData			Scale Tracks to convert to relative space
+	 *
+	 */
+	void ConvertToRelativeSpace(const UAnimSequence& AnimSeq, TArray<FTranslationTrack>& TranslationData, TArray<FRotationTrack>& RotationData, TArray<FScaleTrack>& ScaleData) const;
+
 
 	/**
 	 * Locates spans of keys within the position and rotation tracks provided which can be estimated
@@ -185,6 +268,43 @@ protected:
 		TArray<FRotationTrack>& RotationTracks, 
 		TArray<FScaleTrack>& ScaleTracks);
 
+	/**
+	 * Locates spans of keys within the position and rotation tracks provided which can be estimated
+	 * through linear interpolation of the surrounding keys. The remaining key values are bit packed into
+	 * the animation sequence provided
+	 *
+	 * @param	AnimSeq				The animation sequence being compressed
+	 * @param	BoneData			BoneData array describing the hierarchy of the animated skeleton
+	 * @param	RawSegments			Segment list to process
+	 */
+	void ProcessAnimationTracks(
+		UAnimSequence& AnimSeq,
+		const TArray<FBoneData>& BoneData,
+		TArray<FAnimSegmentContext>& RawSegments);
+
+	/**
+	 * Performs retargetting for a single segment.
+	 */
+	void PerformRetargeting(
+		FProcessAnimationTracksContext& Context,
+		int32 BoneIndex,
+		int32 HighestTargetBoneIndex,
+		int32 FurthestTargetBoneIndex);
+
+	/**
+	 * Samples a track within a segment at a specified time with the current track format.
+	 */
+	static FTransform SampleSegment(const FProcessAnimationTracksContext& Context, int32 TrackIndex, float Time);
+
+private:
+	/**
+	 * Performs the linear key reduction and retargetting for a single segment.
+	 * This can be called from multiple threads concurrently.
+	 */
+	void ProcessAnimationTracks(FProcessAnimationTracksContext& Context);
+
+	friend struct FAsyncProcessAnimationTracksTaskGroupContext;
+
 #endif // WITH_EDITOR
 };
 
diff --git a/Engine/Source/Runtime/Engine/Classes/Animation/AnimInstance.h b/Engine/Source/Runtime/Engine/Classes/Animation/AnimInstance.h
--- a/Engine/Source/Runtime/Engine/Classes/Animation/AnimInstance.h
+++ b/Engine/Source/Runtime/Engine/Classes/Animation/AnimInstance.h
@@ -30,6 +30,9 @@ struct FAnimNode_AssetPlayerBase;
 struct FAnimNode_StateMachine;
 struct FAnimNode_SubInput;
 struct FBakedAnimationStateMachine;
+struct FBoneContainer;
+
+typedef TArray<FTransform> FTransformArrayA2;
 
 UENUM()
 enum class EAnimCurveType : uint8 
diff --git a/Engine/Source/Runtime/Engine/Classes/Animation/AnimNodeBase.h b/Engine/Source/Runtime/Engine/Classes/Animation/AnimNodeBase.h
--- a/Engine/Source/Runtime/Engine/Classes/Animation/AnimNodeBase.h
+++ b/Engine/Source/Runtime/Engine/Classes/Animation/AnimNodeBase.h
@@ -5,6 +5,7 @@
 #include "CoreMinimal.h"
 #include "UObject/ObjectMacros.h"
 #include "UObject/Class.h"
+#include "Engine/EngineTypes.h"
 #include "Animation/AnimTypes.h"
 #include "Animation/AnimCurveTypes.h"
 #include "BonePose.h"
@@ -653,7 +654,7 @@ struct ENGINE_API FAnimNode_Base
 	 * This is called during anim blueprint compilation to handle child anim blueprints.
 	 * @param	NewAsset	The new asset that is being set
 	 */
-	virtual void OverrideAsset(UAnimationAsset* NewAsset) {}
+	virtual void OverrideAsset(class UAnimationAsset* NewAsset) {}
 
 	/**
 	 * Called to gather on-screen debug data. 
diff --git a/Engine/Source/Runtime/Engine/Classes/Animation/AnimNode_SequencePlayer.h b/Engine/Source/Runtime/Engine/Classes/Animation/AnimNode_SequencePlayer.h
--- a/Engine/Source/Runtime/Engine/Classes/Animation/AnimNode_SequencePlayer.h
+++ b/Engine/Source/Runtime/Engine/Classes/Animation/AnimNode_SequencePlayer.h
@@ -6,9 +6,10 @@
 #include "UObject/ObjectMacros.h"
 #include "Animation/AnimNode_AssetPlayerBase.h"
 #include "Animation/InputScaleBias.h"
+#include "Animation/AnimSequenceDecompressionContext.h"
+#include "Animation/AnimSequenceBase.h"
 #include "AnimNode_SequencePlayer.generated.h"
 
-#pragma once
 
 // Sequence player node
 USTRUCT(BlueprintInternalUseOnly)
@@ -54,9 +55,10 @@ public:
 	}
 
 	// FAnimNode_AssetPlayerBase interface
-	virtual float GetCurrentAssetTime();
-	virtual float GetCurrentAssetTimePlayRateAdjusted();
-	virtual float GetCurrentAssetLength();
+	virtual float GetCurrentAssetTime() override;
+	virtual float GetCurrentAssetTimePlayRateAdjusted() override;
+	virtual float GetCurrentAssetLength() override;
+	virtual UAnimationAsset* GetAnimAsset() override { return Sequence; }
 	// End of FAnimNode_AssetPlayerBase interface
 
 	// FAnimNode_Base interface
@@ -68,7 +70,5 @@ public:
 	virtual void GatherDebugData(FNodeDebugData& DebugData) override;
 	// End of FAnimNode_Base interface
 
-	UAnimationAsset* GetAnimAsset() { return Sequence; }
-
 	float GetTimeFromEnd(float CurrentNodeTime);
 };
diff --git a/Engine/Source/Runtime/Engine/Classes/Animation/AnimSequence.h b/Engine/Source/Runtime/Engine/Classes/Animation/AnimSequence.h
--- a/Engine/Source/Runtime/Engine/Classes/Animation/AnimSequence.h
+++ b/Engine/Source/Runtime/Engine/Classes/Animation/AnimSequence.h
@@ -17,13 +17,11 @@
 #include "Animation/AnimSequenceBase.h"
 #include "AnimSequence.generated.h"
 
-struct FAnimCompressContext;
-struct FCompactPose;
-
 typedef TArray<FTransform> FTransformArrayA2;
 
 class USkeletalMesh;
 struct FAnimCompressContext;
+struct FAnimSequenceDecompressionContext;
 struct FCompactPose;
 
 /**
@@ -347,6 +345,49 @@ private:
 
 FArchive& operator<<(FArchive& Ar, FCompressedOffsetData& D);
 
+/**
+ * Represents a segment of the anim sequence that is compressed.
+ */
+USTRUCT()
+struct ENGINE_API FCompressedSegment
+{
+	GENERATED_USTRUCT_BODY()
+
+	// Frame where the segment begins in the anim sequence
+	int32 StartFrame;
+
+	// Num of frames contained in the segment
+	int32 NumFrames;
+
+	// Segment data offset in CompressedByteStream
+	int32 ByteStreamOffset;
+
+	/** The compression format that was used to compress translation tracks. */
+	TEnumAsByte<enum AnimationCompressionFormat> TranslationCompressionFormat;
+
+	/** The compression format that was used to compress rotation tracks. */
+	TEnumAsByte<enum AnimationCompressionFormat> RotationCompressionFormat;
+
+	/** The compression format that was used to compress rotation tracks. */
+	TEnumAsByte<enum AnimationCompressionFormat> ScaleCompressionFormat;
+
+	FCompressedSegment()
+		: StartFrame(0)
+		, NumFrames(0)
+		, ByteStreamOffset(0)
+		, TranslationCompressionFormat(ACF_None)
+		, RotationCompressionFormat(ACF_None)
+		, ScaleCompressionFormat(ACF_None)
+	{
+	}
+
+	friend FArchive& operator<<(FArchive& Ar, FCompressedSegment &Segment)
+	{
+		return Ar << Segment.StartFrame << Segment.NumFrames << Segment.ByteStreamOffset
+			<< Segment.TranslationCompressionFormat << Segment.RotationCompressionFormat << Segment.ScaleCompressionFormat;
+	}
+};
+
 UCLASS(config=Engine, hidecategories=(UObject, Length), BlueprintType)
 class ENGINE_API UAnimSequence : public UAnimSequenceBase
 {
@@ -445,16 +486,15 @@ public:
 
 	/**
 	 * ByteStream for compressed animation data.
-	 * All keys are currently stored at evenly-spaced intervals (ie no explicit key times).
-	 *
-	 * For a translation track of n keys, data is packed as n uncompressed float[3]:
-	 *
-	 * For a rotation track of n>1 keys, the first 24 bytes are reserved for compression info
-	 * (eg Fixed32 stores float Mins[3]; float Ranges[3]), followed by n elements of the compressed type.
-	 * For a rotation track of n=1 keys, the single key is packed as an FQuatFloat96NoW.
+	 * The memory layout is dependent on the algorithm used to compress the anim sequence.
 	 */
 	TArray<uint8> CompressedByteStream;
 
+	/**
+	 * Array of segment descriptors for this compressed anim sequence.
+	 */
+	TArray<FCompressedSegment> CompressedSegments;
+
 	TEnumAsByte<enum AnimationKeyFormat> KeyEncodingFormat;
 
 	/**
@@ -712,6 +752,16 @@ public:
 	 */
 	void GetBoneTransform(FTransform& OutAtom, int32 TrackIndex, float Time, bool bUseRawData) const;
 
+	/**
+	 * Get Bone Transform of the Time given, relative to Parent for the Track Given
+	 *
+	 * @param	OutAtom			[out] Output bone transform.
+	 * @param	TrackIndex		Index of track to interpolate.
+	 * @param	DecompContext	Decompression context to use.
+	 * @param	bUseRawData		If true, use raw animation data instead of compressed data.
+	 */
+	void GetBoneTransform(FTransform& OutAtom, int32 TrackIndex, FAnimSequenceDecompressionContext& DecompContext, bool bUseRawData) const;
+
 	/**
 	 * Extract Bone Transform of the Time given, from InRawAnimationData
 	 *
diff --git a/Engine/Source/Runtime/Engine/Classes/Animation/AnimSequenceBase.h b/Engine/Source/Runtime/Engine/Classes/Animation/AnimSequenceBase.h
--- a/Engine/Source/Runtime/Engine/Classes/Animation/AnimSequenceBase.h
+++ b/Engine/Source/Runtime/Engine/Classes/Animation/AnimSequenceBase.h
@@ -144,6 +144,9 @@ public:
 #endif // WITH_EDITOR
 
 	//~ Begin UAnimationAsset Interface
+#if WITH_EDITOR
+	virtual void RemapTracksToNewSkeleton(USkeleton* NewSkeleton, bool bConvertSpaces) override;
+#endif
 	virtual void TickAssetPlayer(FAnimTickRecord& Instance, struct FAnimNotifyQueue& NotifyQueue, FAnimAssetTickContext& Context) const override;
 
 	void TickByMarkerAsFollower(FMarkerTickRecord &Instance, FMarkerTickContext &MarkerContext, float& CurrentTime, float& OutPreviousTime, const float MoveDelta, const bool bLooping) const;
diff --git a/Engine/Source/Runtime/Engine/Classes/Animation/AnimationSettings.h b/Engine/Source/Runtime/Engine/Classes/Animation/AnimationSettings.h
--- a/Engine/Source/Runtime/Engine/Classes/Animation/AnimationSettings.h
+++ b/Engine/Source/Runtime/Engine/Classes/Animation/AnimationSettings.h
@@ -71,21 +71,13 @@ class ENGINE_API UAnimationSettings : public UDeveloperSettings
 	UPROPERTY(config, EditAnywhere, Category = Compression)
 	bool bRaiseMaxErrorToExisting;
 
-	/** If true, the uniform bitwise techniques will be tried */
+	/** If true, then an exhaustive search is used otherwise only a short list of the best methods is tried */
 	UPROPERTY(config, EditAnywhere, Category = Compression)
-	bool bTryFixedBitwiseCompression;
+	bool bTryExhaustiveSearch;
 
-	/** If true, the per-track compressor techniques will be tried */
+	/** If true, anim sequences are compressed in segments. This allows for parallel compression of individual sequences. */
 	UPROPERTY(config, EditAnywhere, Category = Compression)
-	bool bTryPerTrackBitwiseCompression;
-
-	/** If true, the linear key removal techniques will be tried */
-	UPROPERTY(config, EditAnywhere, Category = Compression)
-	bool bTryLinearKeyRemovalCompression;
-
-	/** If true, the resampling techniques will be tried */
-	UPROPERTY(config, EditAnywhere, Category = Compression)
-	bool bTryIntervalKeyRemoval;
+	bool bEnableSegmenting;
 
 	UPROPERTY(config, EditAnywhere, Category = Performance)
 	bool bEnablePerformanceLog;
diff --git a/Engine/Source/Runtime/Engine/Classes/Animation/BlendSpaceBase.h b/Engine/Source/Runtime/Engine/Classes/Animation/BlendSpaceBase.h
--- a/Engine/Source/Runtime/Engine/Classes/Animation/BlendSpaceBase.h
+++ b/Engine/Source/Runtime/Engine/Classes/Animation/BlendSpaceBase.h
@@ -239,7 +239,7 @@ class UBlendSpaceBase : public UAnimationAsset, public IInterpolationIndexProvid
 	// this doesn't mean max time. In Sequence, this is SequenceLength,
 	// but for BlendSpace CurrentTime is normalized [0,1], so this is 1
 	virtual float GetMaxCurrentTime() override { return 1.f; }	
-	virtual TArray<FName>* GetUniqueMarkerNames() override { return (SampleIndexWithMarkers != INDEX_NONE && SampleData.Num() > 0) ? SampleData[SampleIndexWithMarkers].Animation->GetUniqueMarkerNames() : nullptr; }
+	virtual TArray<FName>* GetUniqueMarkerNames() override;
 	virtual bool IsValidAdditive() const override;
 #if WITH_EDITOR
 	virtual bool GetAllAnimationSequencesReferred(TArray<UAnimationAsset*>& AnimationAssets, bool bRecursive = true) override;
