Copying //UE4/Dev-AnimPhys to Dev-Main (//UE4/Dev-Main)

diff --git a/Engine/Source/Runtime/Engine/Private/Animation/AnimCompress.cpp b/Engine/Source/Runtime/Engine/Private/Animation/AnimCompress.cpp
--- a/Engine/Source/Runtime/Engine/Private/Animation/AnimCompress.cpp
+++ b/Engine/Source/Runtime/Engine/Private/Animation/AnimCompress.cpp
@@ -15,74 +15,77 @@
 DEFINE_LOG_CATEGORY(LogAnimationCompression);
 
 
-// Writes the specified data to Seq->CompresedByteStream with four-byte alignment.
-#define AC_UnalignedWriteToStream( Src, Len )										\
-	{																				\
-		const int32 Ofs = Seq->CompressedByteStream.AddUninitialized( Len );						\
-		FMemory::Memcpy( Seq->CompressedByteStream.GetData()+Ofs, (Src), (Len) );	\
-	}
+void UAnimCompress::UnalignedWriteToStream(TArray<uint8>& ByteStream, const void* Src, SIZE_T Len)
+{
+	const int32 Offset = ByteStream.AddUninitialized(Len);
+	FMemory::Memcpy(ByteStream.GetData() + Offset, Src, Len);
+}
 
-static const uint8 AnimationPadSentinel = 85; //(1<<1)+(1<<3)+(1<<5)+(1<<7)
+void UAnimCompress::UnalignedWriteToStream(TArray<uint8>& ByteStream, int32& StreamOffset, const void* Src, SIZE_T Len)
+{
+	FMemory::Memcpy(ByteStream.GetData() + StreamOffset, Src, Len);
+	StreamOffset += static_cast<int32>(Len);
+}
 
 
-static void PackVectorToStream(
-	UAnimSequence* Seq, 
-	AnimationCompressionFormat TargetTranslationFormat, 
+void UAnimCompress::PackVectorToStream(
+	TArray<uint8>& ByteStream,
+	AnimationCompressionFormat Format,
 	const FVector& Vec,
 	const float* Mins,
 	const float* Ranges)
 {
-	if ( TargetTranslationFormat == ACF_None )
+	if ( Format == ACF_None )
 	{
-		AC_UnalignedWriteToStream( &Vec, sizeof(Vec) );
+		UnalignedWriteToStream( ByteStream, &Vec, sizeof(Vec) );
 	}
-	else if ( TargetTranslationFormat == ACF_Float96NoW )
+	else if ( Format == ACF_Float96NoW )
 	{
-		AC_UnalignedWriteToStream( &Vec, sizeof(Vec) );
+		UnalignedWriteToStream( ByteStream, &Vec, sizeof(Vec) );
 	}
-	else if ( TargetTranslationFormat == ACF_IntervalFixed32NoW )
+	else if ( Format == ACF_IntervalFixed32NoW )
 	{
 		const FVectorIntervalFixed32NoW CompressedVec( Vec, Mins, Ranges );
 
-		AC_UnalignedWriteToStream( &CompressedVec, sizeof(CompressedVec) );
+		UnalignedWriteToStream( ByteStream, &CompressedVec, sizeof(CompressedVec) );
 	}
 }
 
-static void PackQuaternionToStream(
-	UAnimSequence* Seq, 
-	AnimationCompressionFormat TargetRotationFormat, 
+void UAnimCompress::PackQuaternionToStream(
+	TArray<uint8>& ByteStream,
+	AnimationCompressionFormat Format,
 	const FQuat& Quat,
 	const float* Mins,
 	const float* Ranges)
 {
-	if ( TargetRotationFormat == ACF_None )
+	if ( Format == ACF_None )
 	{
-		AC_UnalignedWriteToStream( &Quat, sizeof(FQuat) );
+		UnalignedWriteToStream( ByteStream, &Quat, sizeof(FQuat) );
 	}
-	else if ( TargetRotationFormat == ACF_Float96NoW )
+	else if ( Format == ACF_Float96NoW )
 	{
 		const FQuatFloat96NoW QuatFloat96NoW( Quat );
-		AC_UnalignedWriteToStream( &QuatFloat96NoW, sizeof(FQuatFloat96NoW) );
+		UnalignedWriteToStream( ByteStream, &QuatFloat96NoW, sizeof(FQuatFloat96NoW) );
 	}
-	else if ( TargetRotationFormat == ACF_Fixed32NoW )
+	else if ( Format == ACF_Fixed32NoW )
 	{
 		const FQuatFixed32NoW QuatFixed32NoW( Quat );
-		AC_UnalignedWriteToStream( &QuatFixed32NoW, sizeof(FQuatFixed32NoW) );
+		UnalignedWriteToStream( ByteStream, &QuatFixed32NoW, sizeof(FQuatFixed32NoW) );
 	}
-	else if ( TargetRotationFormat == ACF_Fixed48NoW )
+	else if ( Format == ACF_Fixed48NoW )
 	{
 		const FQuatFixed48NoW QuatFixed48NoW( Quat );
-		AC_UnalignedWriteToStream( &QuatFixed48NoW, sizeof(FQuatFixed48NoW) );
+		UnalignedWriteToStream( ByteStream, &QuatFixed48NoW, sizeof(FQuatFixed48NoW) );
 	}
-	else if ( TargetRotationFormat == ACF_IntervalFixed32NoW )
+	else if ( Format == ACF_IntervalFixed32NoW )
 	{
 		const FQuatIntervalFixed32NoW QuatIntervalFixed32NoW( Quat, Mins, Ranges );
-		AC_UnalignedWriteToStream( &QuatIntervalFixed32NoW, sizeof(FQuatIntervalFixed32NoW) );
+		UnalignedWriteToStream( ByteStream, &QuatIntervalFixed32NoW, sizeof(FQuatIntervalFixed32NoW) );
 	}
-	else if ( TargetRotationFormat == ACF_Float32NoW )
+	else if ( Format == ACF_Float32NoW )
 	{
 		const FQuatFloat32NoW QuatFloat32NoW( Quat );
-		AC_UnalignedWriteToStream( &QuatFloat32NoW, sizeof(FQuatFloat32NoW) );
+		UnalignedWriteToStream( ByteStream, &QuatFloat32NoW, sizeof(FQuatFloat32NoW) );
 	}
 }
 
@@ -101,27 +104,18 @@ FCompressionMemorySummary::FCompressionMemorySummary(bool bInEnabled)
 	, TotalRaw(0)
 	, TotalBeforeCompressed(0)
 	, TotalAfterCompressed(0)
+	, NumberOfAnimations(0)
+	, TotalCompressionExecutionTime(0.0)
 	, ErrorTotal(0)
 	, ErrorCount(0)
 	, AverageError(0)
-	, MaxError(0)
-	, MaxErrorTime(0)
-	, MaxErrorBone(0)
-	, MaxErrorBoneName(NAME_None)
-	, MaxErrorAnimName(NAME_None)
 {
-	if (bEnabled)
-	{
-		GWarn->BeginSlowTask(NSLOCTEXT("CompressionMemorySummary", "BeginCompressingTaskMessage", "Compressing animations"), true);
-	}
 }
 
 FCompressionMemorySummary::~FCompressionMemorySummary()
 {
 	if (bEnabled)
 	{
-		GWarn->EndSlowTask();
-
 		if (bUsed)
 		{
 			const int32 TotalBeforeSaving = TotalRaw - TotalBeforeCompressed;
@@ -137,22 +131,29 @@ FCompressionMemorySummary::~FCompressionMemorySummary()
 			Args.Add(TEXT("TotalRaw"), FText::AsMemory(TotalRaw, &Options));
 			Args.Add(TEXT("TotalBeforeCompressed"), FText::AsMemory(TotalBeforeCompressed, &Options));
 			Args.Add(TEXT("TotalBeforeSaving"), FText::AsMemory(TotalBeforeSaving, &Options));
+			Args.Add(TEXT("NumberOfAnimations"), FText::AsNumber(NumberOfAnimations));
 			Args.Add(TEXT("OldCompressionRatio"), OldCompressionRatio);
 
 			Args.Add(TEXT("TotalAfterCompressed"), FText::AsMemory(TotalAfterCompressed, &Options));
 			Args.Add(TEXT("TotalAfterSaving"), FText::AsMemory(TotalAfterSaving, &Options));
 			Args.Add(TEXT("NewCompressionRatio"), NewCompressionRatio);
+			Args.Add(TEXT("TotalTimeSpentCompressingPretty"), FText::FromString(FPlatformTime::PrettyTime(TotalCompressionExecutionTime)));
+			Args.Add(TEXT("TotalTimeSpentCompressingRawSeconds"), FText::AsNumber((float)TotalCompressionExecutionTime, &Options));
+
+			const FErrorTrackerWorstBone WorstBone = WorstBoneError.GetMaxErrorItem();
+			const FErrorTrackerWorstAnimation WorstAnimation = WorstAnimationError.GetMaxErrorItem();
 
 			Args.Add(TEXT("AverageError"), FText::AsNumber(AverageError, &Options));
-			Args.Add(TEXT("MaxError"), FText::AsNumber(MaxError, &Options));
 
-			Args.Add(TEXT("MaxErrorAnimName"), FText::FromName(MaxErrorAnimName));
-			Args.Add(TEXT("MaxErrorBoneName"), FText::FromName(MaxErrorBoneName));
-			Args.Add(TEXT("MaxErrorBone"), MaxErrorBone);
-			Args.Add(TEXT("MaxErrorTime"), FText::AsNumber(MaxErrorTime, &Options));
+			Args.Add(TEXT("WorstBoneError"), WorstBone.ToText());
+			Args.Add(TEXT("WorstAnimationError"), WorstAnimation.ToText());
 
-			const FText Message = FText::Format(NSLOCTEXT("Engine", "CompressionMemorySummary", "Raw: {TotalRaw} - Compressed: {TotalBeforeCompressed}\nSaving: {TotalBeforeSaving} ({OldCompressionRatio})\nRaw: {TotalRaw} - Compressed: {TotalAfterCompressed}\nSaving: {TotalAfterSaving} ({NewCompressionRatio})\n\nEnd Effector Translation Added By Compression:\n{AverageError} avg, {MaxError} max\nMax occurred in {MaxErrorAnimName}, Bone {MaxErrorBoneName}(#{MaxErrorBone}), at Time {MaxErrorTime}\n"), Args);
+			const FText Message = FText::Format(NSLOCTEXT("Engine", "CompressionMemorySummary", "Compressed {NumberOfAnimations} Animation(s)\n\nPre Compression:\n\nRaw: {TotalRaw} - Compressed: {TotalBeforeCompressed}\nSaving: {TotalBeforeSaving} ({OldCompressionRatio})\n\nPost Compression:\n\nRaw: {TotalRaw} - Compressed: {TotalAfterCompressed}\nSaving: {TotalAfterSaving} ({NewCompressionRatio})\n\nTotal Compression Time: {TotalTimeSpentCompressingPretty} (Seconds: {TotalTimeSpentCompressingRawSeconds})\n\nEnd Effector Translation Added By Compression:\n Average: {AverageError} Max:\n{WorstBoneError}\n\nMax Average Animation Error:\n{WorstAnimationError}"), Args);
 
+			UE_LOG(LogAnimationCompression, Display, TEXT("Top 10 Worst Bone Errors:"));
+			WorstBoneError.LogErrorStat();
+			UE_LOG(LogAnimationCompression, Display, TEXT("Top 10 Worst Average Animation Errors:"));
+			WorstAnimationError.LogErrorStat();
 			FMessageDialog::Open(EAppMsgType::Ok, Message);
 		}
 	}
@@ -174,14 +175,16 @@ void FCompressionMemorySummary::GatherPreCompressionStats(UAnimSequence* Seq, in
 
 		TotalRaw += Seq->GetApproxRawSize();
 		TotalBeforeCompressed += Seq->GetApproxCompressedSize();
+		++NumberOfAnimations;
 	}
 }
 
-void FCompressionMemorySummary::GatherPostCompressionStats(UAnimSequence* Seq, TArray<FBoneData>& BoneData)
+void FCompressionMemorySummary::GatherPostCompressionStats(UAnimSequence* Seq, const TArray<FBoneData>& BoneData, double CompressionTime)
 {
 	if (bEnabled)
 	{
 		TotalAfterCompressed += Seq->GetApproxCompressedSize();
+		TotalCompressionExecutionTime += CompressionTime;
 
 		if (Seq->GetSkeleton() != NULL)
 		{
@@ -193,14 +196,9 @@ void FCompressionMemorySummary::GatherPostCompressionStats(UAnimSequence* Seq, T
 			ErrorCount += 1.0f;
 			AverageError = ErrorTotal / ErrorCount;
 
-			if (ErrorStats.MaxError > MaxError)
-			{
-				MaxError = ErrorStats.MaxError;
-				MaxErrorTime = ErrorStats.MaxErrorTime;
-				MaxErrorBone = ErrorStats.MaxErrorBone;
-				MaxErrorAnimName = Seq->GetFName();
-				MaxErrorBoneName = BoneData[ErrorStats.MaxErrorBone].Name;
-			}
+			WorstBoneError.StoreErrorStat(ErrorStats.MaxError, ErrorStats.MaxError, ErrorStats.MaxErrorTime, ErrorStats.MaxErrorBone, BoneData[ErrorStats.MaxErrorBone].Name, Seq->GetFName());
+
+			WorstAnimationError.StoreErrorStat(ErrorStats.AverageError, ErrorStats.AverageError, Seq->GetFName());
 		}
 	}
 }
@@ -211,9 +209,9 @@ void FAnimCompressContext::GatherPreCompressionStats(UAnimSequence* Seq)
 	CompressionSummary.GatherPreCompressionStats(Seq, AnimIndex, MaxAnimations);
 }
 
-void FAnimCompressContext::GatherPostCompressionStats(UAnimSequence* Seq, TArray<FBoneData>& BoneData)
+void FAnimCompressContext::GatherPostCompressionStats(UAnimSequence* Seq, const TArray<FBoneData>& BoneData, double CompressionTime)
 {
-	CompressionSummary.GatherPostCompressionStats(Seq, BoneData);
+	CompressionSummary.GatherPostCompressionStats(Seq, BoneData, CompressionTime);
 }
 
 //////////////////////////////////////////////////////////////////////////////////////
@@ -228,6 +226,9 @@ UAnimCompress::UAnimCompress(const FObjectInitializer& ObjectInitializer)
 
 	UAnimationSettings* AnimationSettings = UAnimationSettings::Get();
 	MaxCurveError = AnimationSettings->MaxCurveError;
+	bEnableSegmenting = AnimationSettings->bEnableSegmenting;
+	IdealNumFramesPerSegment = 64;
+	MaxNumFramesPerSegment = (IdealNumFramesPerSegment * 2) - 1;
 }
 
 
@@ -253,19 +254,19 @@ void UAnimCompress::PrecalculateShortestQuaternionRoutes(
 	}
 }
 
-void PadByteStream(TArray<uint8>& CompressedByteStream, const int32 Alignment, uint8 sentinel)
+void UAnimCompress::PadByteStream(TArray<uint8>& ByteStream, const int32 Alignment, uint8 Sentinel)
 {
-	int32 Pad = Align( CompressedByteStream.Num(), 4 ) - CompressedByteStream.Num();
+	int32 Pad = Align( ByteStream.Num(), Alignment ) - ByteStream.Num();
 	for ( int32 i = 0 ; i < Pad ; ++i )
 	{
-		CompressedByteStream.Add(sentinel);
+		ByteStream.Add(Sentinel);
 	}
 }
 
 
 void UAnimCompress::BitwiseCompressAnimationTracks(
-	UAnimSequence* Seq, 
-	AnimationCompressionFormat TargetTranslationFormat, 
+	class UAnimSequence* Seq,
+	AnimationCompressionFormat TargetTranslationFormat,
 	AnimationCompressionFormat TargetRotationFormat,
 	AnimationCompressionFormat TargetScaleFormat,
 	const TArray<FTranslationTrack>& TranslationData,
@@ -275,156 +276,158 @@ void UAnimCompress::BitwiseCompressAnimationTracks(
 {
 	// Ensure supported compression formats.
 	bool bInvalidCompressionFormat = false;
-	if( !(TargetTranslationFormat == ACF_None) && !(TargetTranslationFormat == ACF_IntervalFixed32NoW) && !(TargetTranslationFormat == ACF_Float96NoW) )
+	if (!(TargetTranslationFormat == ACF_None) && !(TargetTranslationFormat == ACF_IntervalFixed32NoW) && !(TargetTranslationFormat == ACF_Float96NoW))
 	{
-		FMessageDialog::Open( EAppMsgType::Ok, FText::Format( NSLOCTEXT("Engine", "UnknownTranslationCompressionFormat", "Unknown or unsupported translation compression format ({0})"), FText::AsNumber( (int32)TargetTranslationFormat ) ) );
+		FMessageDialog::Open(EAppMsgType::Ok, FText::Format(NSLOCTEXT("Engine", "UnknownTranslationCompressionFormat", "Unknown or unsupported translation compression format ({0})"), FText::AsNumber((int32)TargetTranslationFormat)));
 		bInvalidCompressionFormat = true;
 	}
-	if ( !(TargetRotationFormat >= ACF_None && TargetRotationFormat < ACF_MAX) )
+	if (!(TargetRotationFormat >= ACF_None && TargetRotationFormat < ACF_MAX))
 	{
-		FMessageDialog::Open( EAppMsgType::Ok, FText::Format( NSLOCTEXT("Engine", "UnknownRotationCompressionFormat", "Unknown or unsupported rotation compression format ({0})"), FText::AsNumber( (int32)TargetRotationFormat ) ) );
+		FMessageDialog::Open(EAppMsgType::Ok, FText::Format(NSLOCTEXT("Engine", "UnknownRotationCompressionFormat", "Unknown or unsupported rotation compression format ({0})"), FText::AsNumber((int32)TargetRotationFormat)));
 		bInvalidCompressionFormat = true;
 	}
-	if( !(TargetScaleFormat == ACF_None) && !(TargetScaleFormat == ACF_IntervalFixed32NoW) && !(TargetScaleFormat == ACF_Float96NoW) )
+	if (!(TargetScaleFormat == ACF_None) && !(TargetScaleFormat == ACF_IntervalFixed32NoW) && !(TargetScaleFormat == ACF_Float96NoW))
 	{
-		FMessageDialog::Open( EAppMsgType::Ok, FText::Format( NSLOCTEXT("Engine", "UnknownScaleCompressionFormat", "Unknown or unsupported Scale compression format ({0})"), FText::AsNumber( (int32)TargetScaleFormat ) ) );
+		FMessageDialog::Open(EAppMsgType::Ok, FText::Format(NSLOCTEXT("Engine", "UnknownScaleCompressionFormat", "Unknown or unsupported Scale compression format ({0})"), FText::AsNumber((int32)TargetScaleFormat)));
 		bInvalidCompressionFormat = true;
 	}
-	if ( bInvalidCompressionFormat )
+	if (bInvalidCompressionFormat)
 	{
-		Seq->TranslationCompressionFormat	= ACF_None;
-		Seq->RotationCompressionFormat		= ACF_None;
-		Seq->ScaleCompressionFormat			= ACF_None;
+		Seq->TranslationCompressionFormat = ACF_None;
+		Seq->RotationCompressionFormat = ACF_None;
+		Seq->ScaleCompressionFormat = ACF_None;
 		Seq->CompressedTrackOffsets.Empty();
 		Seq->CompressedScaleOffsets.Empty();
 		Seq->CompressedByteStream.Empty();
+		Seq->CompressedSegments.Empty();
 	}
 	else
 	{
-		Seq->RotationCompressionFormat		= TargetRotationFormat;
-		Seq->TranslationCompressionFormat	= TargetTranslationFormat;
-		Seq->ScaleCompressionFormat	= TargetScaleFormat;
+		Seq->RotationCompressionFormat = TargetRotationFormat;
+		Seq->TranslationCompressionFormat = TargetTranslationFormat;
+		Seq->ScaleCompressionFormat = TargetScaleFormat;
 
-		check( TranslationData.Num() == RotationData.Num() );
+		check(TranslationData.Num() == RotationData.Num());
 		const int32 NumTracks = RotationData.Num();
 		const bool bHasScale = ScaleData.Num() > 0;
 
-		if ( NumTracks == 0 )
+		if (NumTracks == 0)
 		{
-			UE_LOG(LogAnimationCompression, Warning, TEXT("When compressing %s: no key-reduced data"), *Seq->GetName() );
+			UE_LOG(LogAnimationCompression, Warning, TEXT("When compressing %s: no key-reduced data"), *Seq->GetName());
 		}
 
-		Seq->CompressedTrackOffsets.Empty( NumTracks*4 );
-		Seq->CompressedTrackOffsets.AddUninitialized( NumTracks*4 );
+		Seq->CompressedTrackOffsets.Empty(NumTracks * 4);
+		Seq->CompressedTrackOffsets.AddUninitialized(NumTracks * 4);
 
 		// just empty it since there is chance this can be 0
 		Seq->CompressedScaleOffsets.Empty();
 		// only do this if Scale exists;
-		if ( bHasScale )
+		if (bHasScale)
 		{
 			Seq->CompressedScaleOffsets.SetStripSize(2);
-			Seq->CompressedScaleOffsets.AddUninitialized( NumTracks );
+			Seq->CompressedScaleOffsets.AddUninitialized(NumTracks);
 		}
 
 		Seq->CompressedByteStream.Empty();
+		Seq->CompressedSegments.Empty();
 
-		for ( int32 TrackIndex = 0 ; TrackIndex < NumTracks ; ++TrackIndex )
+		for (int32 TrackIndex = 0; TrackIndex < NumTracks; ++TrackIndex)
 		{
 			// Translation data.
-			const FTranslationTrack& SrcTrans	= TranslationData[TrackIndex];
+			const FTranslationTrack& SrcTrans = TranslationData[TrackIndex];
 
-			const int32 OffsetTrans				= Seq->CompressedByteStream.Num();
-			const int32 NumKeysTrans				= SrcTrans.PosKeys.Num();
+			const int32 OffsetTrans = Seq->CompressedByteStream.Num();
+			const int32 NumKeysTrans = SrcTrans.PosKeys.Num();
 
 			// Warn on empty data.
-			if ( NumKeysTrans == 0 )
+			if (NumKeysTrans == 0)
 			{
-				UE_LOG(LogAnimationCompression, Warning, TEXT("When compressing %s track %i: no translation keys"), *Seq->GetName(), TrackIndex );
+				UE_LOG(LogAnimationCompression, Warning, TEXT("When compressing %s track %i: no translation keys"), *Seq->GetName(), TrackIndex);
 			}
 
-			checkf( (OffsetTrans % 4) == 0, TEXT("CompressedByteStream not aligned to four bytes" ) );
-			Seq->CompressedTrackOffsets[TrackIndex*4] = OffsetTrans;
-			Seq->CompressedTrackOffsets[TrackIndex*4+1] = NumKeysTrans;
+			checkf((OffsetTrans % 4) == 0, TEXT("CompressedByteStream not aligned to four bytes"));
+			Seq->CompressedTrackOffsets[TrackIndex * 4] = OffsetTrans;
+			Seq->CompressedTrackOffsets[TrackIndex * 4 + 1] = NumKeysTrans;
 
 			// Calculate the bounding box of the translation keys
 			FBox PositionBounds(SrcTrans.PosKeys);
 
 			float TransMins[3] = { PositionBounds.Min.X, PositionBounds.Min.Y, PositionBounds.Min.Z };
 			float TransRanges[3] = { PositionBounds.Max.X - PositionBounds.Min.X, PositionBounds.Max.Y - PositionBounds.Min.Y, PositionBounds.Max.Z - PositionBounds.Min.Z };
-			if ( TransRanges[0] == 0.f ) { TransRanges[0] = 1.f; }
-			if ( TransRanges[1] == 0.f ) { TransRanges[1] = 1.f; }
-			if ( TransRanges[2] == 0.f ) { TransRanges[2] = 1.f; }
+			if (TransRanges[0] == 0.f) { TransRanges[0] = 1.f; }
+			if (TransRanges[1] == 0.f) { TransRanges[1] = 1.f; }
+			if (TransRanges[2] == 0.f) { TransRanges[2] = 1.f; }
 
 			if (NumKeysTrans > 1)
 			{
 				// Write the mins and ranges if they'll be used on the other side
 				if (TargetTranslationFormat == ACF_IntervalFixed32NoW)
 				{
-					AC_UnalignedWriteToStream( TransMins, sizeof(float)*3 );
-					AC_UnalignedWriteToStream( TransRanges, sizeof(float)*3 );
+					UnalignedWriteToStream(Seq->CompressedByteStream, TransMins, sizeof(float) * 3);
+					UnalignedWriteToStream(Seq->CompressedByteStream, TransRanges, sizeof(float) * 3);
 				}
 
 				// Pack the positions into the stream
-				for ( int32 KeyIndex = 0 ; KeyIndex < NumKeysTrans ; ++KeyIndex )
+				for (int32 KeyIndex = 0; KeyIndex < NumKeysTrans; ++KeyIndex)
 				{
 					const FVector& Vec = SrcTrans.PosKeys[KeyIndex];
-					PackVectorToStream(Seq, TargetTranslationFormat, Vec, TransMins, TransRanges);
+					PackVectorToStream(Seq->CompressedByteStream, TargetTranslationFormat, Vec, TransMins, TransRanges);
 				}
 
 				if (IncludeKeyTable)
 				{
 					// Align to four bytes.
-					PadByteStream(Seq->CompressedByteStream, 4, AnimationPadSentinel );
+					PadByteStream(Seq->CompressedByteStream, 4, AnimationPadSentinel);
 
 					// write the key table
 					const int32 NumFrames = Seq->NumFrames;
-					const int32 LastFrame = Seq->NumFrames-1;
+					const int32 LastFrame = Seq->NumFrames - 1;
 					const size_t FrameSize = Seq->NumFrames > 0xff ? sizeof(uint16) : sizeof(uint8);
-					const float FrameRate = NumFrames / Seq->SequenceLength;
+					const float FrameRate = LastFrame / Seq->SequenceLength;
 
-					const int32 TableSize= NumKeysTrans*FrameSize;
-					const int32 TableDwords= (TableSize+3)>>2;
+					const int32 TableSize = NumKeysTrans*FrameSize;
+					const int32 TableDwords = (TableSize + 3) >> 2;
 					const int32 StartingOffset = Seq->CompressedByteStream.Num();
 
-					for ( int32 KeyIndex = 0 ; KeyIndex < NumKeysTrans ; ++KeyIndex )
+					for (int32 KeyIndex = 0; KeyIndex < NumKeysTrans; ++KeyIndex)
 					{
 						// write the frame values for each key
-						float KeyTime= SrcTrans.Times[KeyIndex];
+						float KeyTime = SrcTrans.Times[KeyIndex];
 						float FrameTime = KeyTime * FrameRate;
-						int32 FrameIndex= FMath::Clamp(FMath::TruncToInt(FrameTime), 0, LastFrame);
-						AC_UnalignedWriteToStream( &FrameIndex, FrameSize );
+						int32 FrameIndex = FMath::Clamp(FMath::TruncToInt(FrameTime + 0.5f), 0, LastFrame);
+						UnalignedWriteToStream(Seq->CompressedByteStream, &FrameIndex, FrameSize);
 					}
 
 					// Align to four bytes. Padding with 0's to round out the key table
-					PadByteStream(Seq->CompressedByteStream, 4, 0 );
+					PadByteStream(Seq->CompressedByteStream, 4, 0);
 
 					const int32 EndingOffset = Seq->CompressedByteStream.Num();
-					check((EndingOffset - StartingOffset) == (TableDwords*4));
+					check((EndingOffset - StartingOffset) == (TableDwords * 4));
 				}
 			}
 			else if (NumKeysTrans == 1)
 			{
 				// A single translation key gets written out a single uncompressed float[3].
-				AC_UnalignedWriteToStream( &(SrcTrans.PosKeys[0]), sizeof(FVector) );
+				UnalignedWriteToStream(Seq->CompressedByteStream, &(SrcTrans.PosKeys[0]), sizeof(FVector));
 			}
 			else
 			{
-				UE_LOG(LogAnimationCompression, Warning, TEXT("When compressing %s track %i: no translation keys"), *Seq->GetName(), TrackIndex );
+				UE_LOG(LogAnimationCompression, Warning, TEXT("When compressing %s track %i: no translation keys"), *Seq->GetName(), TrackIndex);
 			}
 
 			// Align to four bytes.
-			PadByteStream(Seq->CompressedByteStream, 4, AnimationPadSentinel );
+			PadByteStream(Seq->CompressedByteStream, 4, AnimationPadSentinel);
 
 			// Compress rotation data.
-			const FRotationTrack& SrcRot	= RotationData[TrackIndex];
-			const int32 OffsetRot				= Seq->CompressedByteStream.Num();
-			const int32 NumKeysRot			= SrcRot.RotKeys.Num();
+			const FRotationTrack& SrcRot = RotationData[TrackIndex];
+			const int32 OffsetRot = Seq->CompressedByteStream.Num();
+			const int32 NumKeysRot = SrcRot.RotKeys.Num();
 
-			checkf( (OffsetRot % 4) == 0, TEXT("CompressedByteStream not aligned to four bytes" ) );
-			Seq->CompressedTrackOffsets[TrackIndex*4+2] = OffsetRot;
-			Seq->CompressedTrackOffsets[TrackIndex*4+3] = NumKeysRot;
+			checkf((OffsetRot % 4) == 0, TEXT("CompressedByteStream not aligned to four bytes"));
+			Seq->CompressedTrackOffsets[TrackIndex * 4 + 2] = OffsetRot;
+			Seq->CompressedTrackOffsets[TrackIndex * 4 + 3] = NumKeysRot;
 
-			if ( NumKeysRot > 1 )
+			if (NumKeysRot > 1)
 			{
 				// Calculate the min/max of the XYZ components of the quaternion
 				float MinX = 1.f;
@@ -433,10 +436,10 @@ void UAnimCompress::BitwiseCompressAnimationTracks(
 				float MaxX = -1.f;
 				float MaxY = -1.f;
 				float MaxZ = -1.f;
-				for ( int32 KeyIndex = 0 ; KeyIndex < SrcRot.RotKeys.Num() ; ++KeyIndex )
+				for (int32 KeyIndex = 0; KeyIndex < SrcRot.RotKeys.Num(); ++KeyIndex)
 				{
-					FQuat Quat( SrcRot.RotKeys[KeyIndex] );
-					if ( Quat.W < 0.f )
+					FQuat Quat(SrcRot.RotKeys[KeyIndex]);
+					if (Quat.W < 0.f)
 					{
 						Quat.X = -Quat.X;
 						Quat.Y = -Quat.Y;
@@ -445,93 +448,93 @@ void UAnimCompress::BitwiseCompressAnimationTracks(
 					}
 					Quat.Normalize();
 
-					MinX = FMath::Min( MinX, Quat.X );
-					MaxX = FMath::Max( MaxX, Quat.X );
-					MinY = FMath::Min( MinY, Quat.Y );
-					MaxY = FMath::Max( MaxY, Quat.Y );
-					MinZ = FMath::Min( MinZ, Quat.Z );
-					MaxZ = FMath::Max( MaxZ, Quat.Z );
+					MinX = FMath::Min(MinX, Quat.X);
+					MaxX = FMath::Max(MaxX, Quat.X);
+					MinY = FMath::Min(MinY, Quat.Y);
+					MaxY = FMath::Max(MaxY, Quat.Y);
+					MinZ = FMath::Min(MinZ, Quat.Z);
+					MaxZ = FMath::Max(MaxZ, Quat.Z);
 				}
-				const float Mins[3]	= { MinX,		MinY,		MinZ };
-				float Ranges[3]		= { MaxX-MinX,	MaxY-MinY,	MaxZ-MinZ };
-				if ( Ranges[0] == 0.f ) { Ranges[0] = 1.f; }
-				if ( Ranges[1] == 0.f ) { Ranges[1] = 1.f; }
-				if ( Ranges[2] == 0.f ) { Ranges[2] = 1.f; }
+				const float Mins[3] = { MinX,		MinY,		MinZ };
+				float Ranges[3] = { MaxX - MinX,	MaxY - MinY,	MaxZ - MinZ };
+				if (Ranges[0] == 0.f) { Ranges[0] = 1.f; }
+				if (Ranges[1] == 0.f) { Ranges[1] = 1.f; }
+				if (Ranges[2] == 0.f) { Ranges[2] = 1.f; }
 
 				// Write the mins and ranges if they'll be used on the other side
 				if (TargetRotationFormat == ACF_IntervalFixed32NoW)
 				{
-					AC_UnalignedWriteToStream( Mins, sizeof(float)*3 );
-					AC_UnalignedWriteToStream( Ranges, sizeof(float)*3 );
+					UnalignedWriteToStream(Seq->CompressedByteStream, Mins, sizeof(float) * 3);
+					UnalignedWriteToStream(Seq->CompressedByteStream, Ranges, sizeof(float) * 3);
 				}
 
 				// n elements of the compressed type.
-				for ( int32 KeyIndex = 0 ; KeyIndex < SrcRot.RotKeys.Num() ; ++KeyIndex )
+				for (int32 KeyIndex = 0; KeyIndex < SrcRot.RotKeys.Num(); ++KeyIndex)
 				{
 					const FQuat& Quat = SrcRot.RotKeys[KeyIndex];
-					PackQuaternionToStream(Seq, TargetRotationFormat, Quat, Mins, Ranges);
+					PackQuaternionToStream(Seq->CompressedByteStream, TargetRotationFormat, Quat, Mins, Ranges);
 				}
 
 				// n elements of frame indices
 				if (IncludeKeyTable)
 				{
 					// Align to four bytes.
-					PadByteStream(Seq->CompressedByteStream, 4, AnimationPadSentinel );
+					PadByteStream(Seq->CompressedByteStream, 4, AnimationPadSentinel);
 
 					// write the key table
 					const int32 NumFrames = Seq->NumFrames;
-					const int32 LastFrame= Seq->NumFrames-1;
-					const size_t FrameSize= Seq->NumFrames > 0xff ? sizeof(uint16) : sizeof(uint8);
-					const float FrameRate = NumFrames / Seq->SequenceLength;
+					const int32 LastFrame = Seq->NumFrames - 1;
+					const size_t FrameSize = Seq->NumFrames > 0xff ? sizeof(uint16) : sizeof(uint8);
+					const float FrameRate = LastFrame / Seq->SequenceLength;
 
-					const int32 TableSize= NumKeysRot*FrameSize;
-					const int32 TableDwords= (TableSize+3)>>2;
+					const int32 TableSize = NumKeysRot*FrameSize;
+					const int32 TableDwords = (TableSize + 3) >> 2;
 					const int32 StartingOffset = Seq->CompressedByteStream.Num();
 
-					for ( int32 KeyIndex = 0 ; KeyIndex < NumKeysRot ; ++KeyIndex )
+					for (int32 KeyIndex = 0; KeyIndex < NumKeysRot; ++KeyIndex)
 					{
 						// write the frame values for each key
 						float KeyTime = SrcRot.Times[KeyIndex];
 						float FrameTime = KeyTime * FrameRate;
-						int32 FrameIndex = FMath::Clamp(FMath::TruncToInt(FrameTime), 0, LastFrame);
-						AC_UnalignedWriteToStream( &FrameIndex, FrameSize );
+						int32 FrameIndex = FMath::Clamp(FMath::TruncToInt(FrameTime + 0.5f), 0, LastFrame);
+						UnalignedWriteToStream(Seq->CompressedByteStream, &FrameIndex, FrameSize);
 					}
 
 					// Align to four bytes. Padding with 0's to round out the key table
-					PadByteStream(Seq->CompressedByteStream, 4, 0 );
+					PadByteStream(Seq->CompressedByteStream, 4, 0);
 
 					const int32 EndingOffset = Seq->CompressedByteStream.Num();
-					check((EndingOffset - StartingOffset) == (TableDwords*4));
+					check((EndingOffset - StartingOffset) == (TableDwords * 4));
 
 				}
 			}
-			else if ( NumKeysRot == 1 )
+			else if (NumKeysRot == 1)
 			{
 				// For a rotation track of n=1 keys, the single key is packed as an FQuatFloat96NoW.
 				const FQuat& Quat = SrcRot.RotKeys[0];
-				const FQuatFloat96NoW QuatFloat96NoW( Quat );
-				AC_UnalignedWriteToStream( &QuatFloat96NoW, sizeof(FQuatFloat96NoW) );
+				const FQuatFloat96NoW QuatFloat96NoW(Quat);
+				UnalignedWriteToStream(Seq->CompressedByteStream, &QuatFloat96NoW, sizeof(FQuatFloat96NoW));
 			}
 			else
 			{
-				UE_LOG(LogAnimationCompression, Warning, TEXT("When compressing %s track %i: no rotation keys"), *Seq->GetName(), TrackIndex );
+				UE_LOG(LogAnimationCompression, Warning, TEXT("When compressing %s track %i: no rotation keys"), *Seq->GetName(), TrackIndex);
 			}
 
 
 			// Align to four bytes.
-			PadByteStream(Seq->CompressedByteStream, 4, AnimationPadSentinel );
+			PadByteStream(Seq->CompressedByteStream, 4, AnimationPadSentinel);
 
 			// we also should do this only when scale exists. 
 			if (bHasScale)
 			{
-				const FScaleTrack& SrcScale			= ScaleData[TrackIndex];
+				const FScaleTrack& SrcScale = ScaleData[TrackIndex];
 
-				const int32 OffsetScale				= Seq->CompressedByteStream.Num();
-				const int32 NumKeysScale			= SrcScale.ScaleKeys.Num();
+				const int32 OffsetScale = Seq->CompressedByteStream.Num();
+				const int32 NumKeysScale = SrcScale.ScaleKeys.Num();
 
-				checkf( (OffsetScale % 4) == 0, TEXT("CompressedByteStream not aligned to four bytes" ) );
+				checkf((OffsetScale % 4) == 0, TEXT("CompressedByteStream not aligned to four bytes"));
 				Seq->CompressedScaleOffsets.SetOffsetData(TrackIndex, 0, OffsetScale);
-				Seq->CompressedScaleOffsets.SetOffsetData(TrackIndex, 1,  NumKeysScale);
+				Seq->CompressedScaleOffsets.SetOffsetData(TrackIndex, 1, NumKeysScale);
 
 				// Calculate the bounding box of the Scalelation keys
 				FBox ScaleBoundsBounds(SrcScale.ScaleKeys);
@@ -548,60 +551,60 @@ void UAnimCompress::BitwiseCompressAnimationTracks(
 					// Write the mins and ranges if they'll be used on the other side
 					if (TargetScaleFormat == ACF_IntervalFixed32NoW)
 					{
-						AC_UnalignedWriteToStream( ScaleMins, sizeof(float)*3 );
-						AC_UnalignedWriteToStream( ScaleRanges, sizeof(float)*3 );
+						UnalignedWriteToStream(Seq->CompressedByteStream, ScaleMins, sizeof(float) * 3);
+						UnalignedWriteToStream(Seq->CompressedByteStream, ScaleRanges, sizeof(float) * 3);
 					}
 
 					// Pack the positions into the stream
-					for ( int32 KeyIndex = 0 ; KeyIndex < NumKeysScale ; ++KeyIndex )
+					for (int32 KeyIndex = 0; KeyIndex < NumKeysScale; ++KeyIndex)
 					{
 						const FVector& Vec = SrcScale.ScaleKeys[KeyIndex];
-						PackVectorToStream(Seq, TargetScaleFormat, Vec, ScaleMins, ScaleRanges);
+						PackVectorToStream(Seq->CompressedByteStream, TargetScaleFormat, Vec, ScaleMins, ScaleRanges);
 					}
 
 					if (IncludeKeyTable)
 					{
 						// Align to four bytes.
-						PadByteStream(Seq->CompressedByteStream, 4, AnimationPadSentinel );
+						PadByteStream(Seq->CompressedByteStream, 4, AnimationPadSentinel);
 
 						// write the key table
 						const int32 NumFrames = Seq->NumFrames;
-						const int32 LastFrame = Seq->NumFrames-1;
+						const int32 LastFrame = Seq->NumFrames - 1;
 						const size_t FrameSize = Seq->NumFrames > 0xff ? sizeof(uint16) : sizeof(uint8);
-						const float FrameRate = NumFrames / Seq->SequenceLength;
+						const float FrameRate = LastFrame / Seq->SequenceLength;
 
-						const int32 TableSize= NumKeysScale*FrameSize;
-						const int32 TableDwords= (TableSize+3)>>2;
+						const int32 TableSize = NumKeysScale*FrameSize;
+						const int32 TableDwords = (TableSize + 3) >> 2;
 						const int32 StartingOffset = Seq->CompressedByteStream.Num();
 
-						for ( int32 KeyIndex = 0 ; KeyIndex < NumKeysScale ; ++KeyIndex )
+						for (int32 KeyIndex = 0; KeyIndex < NumKeysScale; ++KeyIndex)
 						{
 							// write the frame values for each key
-							float KeyTime= SrcScale.Times[KeyIndex];
+							float KeyTime = SrcScale.Times[KeyIndex];
 							float FrameTime = KeyTime * FrameRate;
-							int32 FrameIndex= FMath::Clamp(FMath::TruncToInt(FrameTime), 0, LastFrame);
-							AC_UnalignedWriteToStream( &FrameIndex, FrameSize );
+							int32 FrameIndex = FMath::Clamp(FMath::TruncToInt(FrameTime + 0.5f), 0, LastFrame);
+							UnalignedWriteToStream(Seq->CompressedByteStream, &FrameIndex, FrameSize);
 						}
 
 						// Align to four bytes. Padding with 0's to round out the key table
-						PadByteStream(Seq->CompressedByteStream, 4, 0 );
+						PadByteStream(Seq->CompressedByteStream, 4, 0);
 
 						const int32 EndingOffset = Seq->CompressedByteStream.Num();
-						check((EndingOffset - StartingOffset) == (TableDwords*4));
+						check((EndingOffset - StartingOffset) == (TableDwords * 4));
 					}
 				}
 				else if (NumKeysScale == 1)
 				{
 					// A single Scalelation key gets written out a single uncompressed float[3].
-					AC_UnalignedWriteToStream( &(SrcScale.ScaleKeys[0]), sizeof(FVector) );
+					UnalignedWriteToStream(Seq->CompressedByteStream, &(SrcScale.ScaleKeys[0]), sizeof(FVector));
 				}
 				else
 				{
-					UE_LOG(LogAnimationCompression, Warning, TEXT("When compressing %s track %i: no Scale keys"), *Seq->GetName(), TrackIndex );
+					UE_LOG(LogAnimationCompression, Warning, TEXT("When compressing %s track %i: no Scale keys"), *Seq->GetName(), TrackIndex);
 				}
 
 				// Align to four bytes.
-				PadByteStream(Seq->CompressedByteStream, 4, AnimationPadSentinel );
+				PadByteStream(Seq->CompressedByteStream, 4, AnimationPadSentinel);
 			}
 		}
 
@@ -610,6 +613,965 @@ void UAnimCompress::BitwiseCompressAnimationTracks(
 	}
 }
 
+void UAnimCompress::BitwiseCompressAnimationTracks(
+	class UAnimSequence& AnimSeq,
+	AnimationCompressionFormat TargetTranslationFormat,
+	AnimationCompressionFormat TargetRotationFormat,
+	AnimationCompressionFormat TargetScaleFormat,
+	TArray<FAnimSegmentContext>& RawSegments,
+	bool bIsSorted)
+{
+	// Ensure supported compression formats.
+	bool bInvalidCompressionFormat = false;
+	if (!(TargetTranslationFormat == ACF_None) && !(TargetTranslationFormat == ACF_IntervalFixed32NoW) && !(TargetTranslationFormat == ACF_Float96NoW))
+	{
+		FMessageDialog::Open(EAppMsgType::Ok, FText::Format(NSLOCTEXT("Engine", "UnknownTranslationCompressionFormat", "Unknown or unsupported translation compression format ({0})"), FText::AsNumber((int32)TargetTranslationFormat)));
+		bInvalidCompressionFormat = true;
+	}
+	if (!(TargetRotationFormat >= ACF_None && TargetRotationFormat < ACF_MAX))
+	{
+		FMessageDialog::Open(EAppMsgType::Ok, FText::Format(NSLOCTEXT("Engine", "UnknownRotationCompressionFormat", "Unknown or unsupported rotation compression format ({0})"), FText::AsNumber((int32)TargetRotationFormat)));
+		bInvalidCompressionFormat = true;
+	}
+	if (!(TargetScaleFormat == ACF_None) && !(TargetScaleFormat == ACF_IntervalFixed32NoW) && !(TargetScaleFormat == ACF_Float96NoW))
+	{
+		FMessageDialog::Open(EAppMsgType::Ok, FText::Format(NSLOCTEXT("Engine", "UnknownScaleCompressionFormat", "Unknown or unsupported Scale compression format ({0})"), FText::AsNumber((int32)TargetScaleFormat)));
+		bInvalidCompressionFormat = true;
+	}
+
+	if (!bInvalidCompressionFormat && RawSegments.Num() > 0)
+	{
+		// First segment holds the compressed trivial tracks
+		BitwiseCompressTrivialAnimationTracks(AnimSeq, RawSegments[0]);
+	}
+
+	for (FAnimSegmentContext& Segment : RawSegments)
+	{
+		if (bInvalidCompressionFormat)
+		{
+			Segment.TranslationCompressionFormat = ACF_None;
+			Segment.RotationCompressionFormat = ACF_None;
+			Segment.ScaleCompressionFormat = ACF_None;
+			Segment.CompressedTrackOffsets.Empty();
+			Segment.CompressedScaleOffsets.Empty();
+			Segment.CompressedByteStream.Empty();
+			Segment.CompressedTrivialTracksByteStream.Empty();
+		}
+		else
+		{
+			BitwiseCompressAnimationTracks(AnimSeq, TargetTranslationFormat, TargetRotationFormat, TargetScaleFormat, Segment, bIsSorted);
+		}
+	}
+}
+
+void UAnimCompress::SanityCheckTrackData(const UAnimSequence& AnimSeq, const FAnimSegmentContext& Segment)
+{
+	check(Segment.TranslationData.Num() == Segment.RotationData.Num());
+
+	const int32 NumTracks = Segment.RotationData.Num();
+	const bool bHasScale = Segment.ScaleData.Num() > 0;
+
+	if (NumTracks == 0)
+	{
+		UE_LOG(LogAnimationCompression, Warning, TEXT("When compressing %s: no key-reduced data"), *AnimSeq.GetFullName());
+	}
+
+	for (int32 TrackIndex = 0; TrackIndex < NumTracks; ++TrackIndex)
+	{
+		const FTranslationTrack& SrcTrans = Segment.TranslationData[TrackIndex];
+		if (SrcTrans.PosKeys.Num() == 0)
+		{
+			UE_LOG(LogAnimationCompression, Warning, TEXT("When compressing %s track %i: no translation keys"), *AnimSeq.GetFullName(), TrackIndex);
+		}
+
+		const FRotationTrack& SrcRot = Segment.RotationData[TrackIndex];
+		if (SrcRot.RotKeys.Num() == 0)
+		{
+			UE_LOG(LogAnimationCompression, Warning, TEXT("When compressing %s track %i: no rotation keys"), *AnimSeq.GetFullName(), TrackIndex);
+		}
+
+		if (bHasScale)
+		{
+			const FScaleTrack& SrcScale = Segment.ScaleData[TrackIndex];
+			if (SrcScale.ScaleKeys.Num() == 0)
+			{
+				UE_LOG(LogAnimationCompression, Warning, TEXT("When compressing %s track %i: no Scale keys"), *AnimSeq.GetFullName(), TrackIndex);
+			}
+		}
+	}
+}
+
+void UAnimCompress::CalculateTrackRange(const FTranslationTrack& TranslationData, AnimationCompressionFormat Format, FVector& OutMin, FVector& OutExtent)
+{
+	FVector Mins;
+	FVector Ranges;
+
+	const int32 NumKeysTrans = TranslationData.PosKeys.Num();
+	if (NumKeysTrans > 1 && Format == ACF_IntervalFixed32NoW)
+	{
+		// Calculate the range of the XYZ components of the translation
+		FBox PositionBounds(TranslationData.PosKeys);
+
+		Mins = PositionBounds.Min;
+		Ranges = PositionBounds.Max - PositionBounds.Min;
+
+		if (Ranges.X == 0.f) { Ranges.X = 1.f; }
+		if (Ranges.Y == 0.f) { Ranges.Y = 1.f; }
+		if (Ranges.Z == 0.f) { Ranges.Z = 1.f; }
+	}
+	else
+	{
+		Mins = FVector::ZeroVector;
+		Ranges = FVector::ZeroVector;
+	}
+
+	OutMin = Mins;
+	OutExtent = Ranges;
+}
+
+void UAnimCompress::CalculateTrackRange(const FRotationTrack& RotationData, AnimationCompressionFormat Format, FVector& OutMin, FVector& OutExtent)
+{
+	FVector Mins;
+	FVector Ranges;
+
+	const int32 NumKeysRot = RotationData.RotKeys.Num();
+	if (NumKeysRot > 1 && Format == ACF_IntervalFixed32NoW)
+	{
+		// Calculate the range of the XYZ components of the quaternion
+		Mins = FVector(1.f, 1.f, 1.f);
+		FVector Maxs(-1.f, -1.f, -1.f);
+
+		for (int32 KeyIndex = 0; KeyIndex < NumKeysRot; ++KeyIndex)
+		{
+			FQuat Quat(RotationData.RotKeys[KeyIndex]);
+			if (Quat.W < 0.f)
+			{
+				Quat.X = -Quat.X;
+				Quat.Y = -Quat.Y;
+				Quat.Z = -Quat.Z;
+				Quat.W = -Quat.W;
+			}
+			Quat.Normalize();
+
+			FVector QuatV(Quat.X, Quat.Y, Quat.Z);
+			Mins = Mins.ComponentMin(QuatV);
+			Maxs = Maxs.ComponentMax(QuatV);
+		}
+
+		Ranges = Maxs - Mins;
+		if (Ranges.X == 0.f) { Ranges.X = 1.f; }
+		if (Ranges.Y == 0.f) { Ranges.Y = 1.f; }
+		if (Ranges.Z == 0.f) { Ranges.Z = 1.f; }
+	}
+	else
+	{
+		Mins = FVector::ZeroVector;
+		Ranges = FVector::ZeroVector;
+	}
+
+	OutMin = Mins;
+	OutExtent = Ranges;
+}
+
+void UAnimCompress::CalculateTrackRange(const FScaleTrack& ScaleData, AnimationCompressionFormat Format, FVector& OutMin, FVector& OutExtent)
+{
+	FVector Mins;
+	FVector Ranges;
+
+	const int32 NumKeysScale = ScaleData.ScaleKeys.Num();
+	if (NumKeysScale > 1 && Format == ACF_IntervalFixed32NoW)
+	{
+		// Calculate the range of the XYZ components of the translation
+		FBox Bounds(ScaleData.ScaleKeys);
+
+		Mins = Bounds.Min;
+		Ranges = Bounds.Max - Bounds.Min;
+
+		if (Ranges.X == 0.f) { Ranges.X = 1.f; }
+		if (Ranges.Y == 0.f) { Ranges.Y = 1.f; }
+		if (Ranges.Z == 0.f) { Ranges.Z = 1.f; }
+	}
+	else
+	{
+		Mins = FVector::ZeroVector;
+		Ranges = FVector::ZeroVector;
+	}
+
+	OutMin = Mins;
+	OutExtent = Ranges;
+}
+
+void UAnimCompress::CalculateTrackRanges(
+	AnimationCompressionFormat TargetTranslationFormat,
+	AnimationCompressionFormat TargetRotationFormat,
+	AnimationCompressionFormat TargetScaleFormat,
+	const FAnimSegmentContext& Segment,
+	TArray<FAnimTrackRange>& TrackRanges)
+{
+	const int32 NumTracks = Segment.RotationData.Num();
+	const bool bHasScale = Segment.ScaleData.Num() > 0;
+
+	TrackRanges.Empty(NumTracks);
+	TrackRanges.AddZeroed(NumTracks);
+
+	for (int32 TrackIndex = 0; TrackIndex < NumTracks; ++TrackIndex)
+	{
+		FAnimTrackRange& TrackRange = TrackRanges[TrackIndex];
+
+		CalculateTrackRange(Segment.TranslationData[TrackIndex], TargetTranslationFormat, TrackRange.TransMin, TrackRange.TransExtent);
+		CalculateTrackRange(Segment.RotationData[TrackIndex], TargetRotationFormat, TrackRange.RotMin, TrackRange.RotExtent);
+
+		if (bHasScale)
+		{
+			CalculateTrackRange(Segment.ScaleData[TrackIndex], TargetScaleFormat, TrackRange.ScaleMin, TrackRange.ScaleExtent);
+		}
+	}
+}
+
+void UAnimCompress::WriteTrackRanges(
+	TArray<uint8>& ByteStream,
+	TFunction<AnimationCompressionFormat(int32 TrackIndex)> GetTranslationFormatFun,
+	TFunction<AnimationCompressionFormat(int32 TrackIndex)> GetRotationFormatFun,
+	TFunction<AnimationCompressionFormat(int32 TrackIndex)> GetScaleFormatFun,
+	TFunction<FTrackKeyFlags(int32 TrackIndex)> GetTranslationFlagsFun,
+	TFunction<FTrackKeyFlags(int32 TrackIndex)> GetRotationFlagsFun,
+	TFunction<FTrackKeyFlags(int32 TrackIndex)> GetScaleFlagsFun,
+	const FAnimSegmentContext& Segment,
+	const TArray<FAnimTrackRange>& TrackRanges,
+	bool bInterleaveValues)
+{
+	const int32 NumTracks = Segment.RotationData.Num();
+	const bool bHasScale = Segment.ScaleData.Num() > 0;
+
+	for (int32 TrackIndex = 0; TrackIndex < NumTracks; ++TrackIndex)
+	{
+		const FAnimTrackRange& TrackRange = TrackRanges[TrackIndex];
+
+		const FTranslationTrack& SrcTrans = Segment.TranslationData[TrackIndex];
+		if (SrcTrans.PosKeys.Num() > 1 && GetTranslationFormatFun(TrackIndex) == ACF_IntervalFixed32NoW)
+		{
+			if (bInterleaveValues)
+			{
+				const FTrackKeyFlags Flags = GetTranslationFlagsFun(TrackIndex);
+				if (Flags.IsComponentNeededX())
+				{
+					UnalignedWriteToStream(ByteStream, &TrackRange.TransMin.X, sizeof(float));
+					UnalignedWriteToStream(ByteStream, &TrackRange.TransExtent.X, sizeof(float));
+				}
+
+				if (Flags.IsComponentNeededY())
+				{
+					UnalignedWriteToStream(ByteStream, &TrackRange.TransMin.Y, sizeof(float));
+					UnalignedWriteToStream(ByteStream, &TrackRange.TransExtent.Y, sizeof(float));
+				}
+
+				if (Flags.IsComponentNeededZ())
+				{
+					UnalignedWriteToStream(ByteStream, &TrackRange.TransMin.Z, sizeof(float));
+					UnalignedWriteToStream(ByteStream, &TrackRange.TransExtent.Z, sizeof(float));
+				}
+			}
+			else
+			{
+				UnalignedWriteToStream(ByteStream, &TrackRange.TransMin, sizeof(FVector));
+				UnalignedWriteToStream(ByteStream, &TrackRange.TransExtent, sizeof(FVector));
+			}
+		}
+
+		const FRotationTrack& SrcRot = Segment.RotationData[TrackIndex];
+		if (SrcRot.RotKeys.Num() > 1 && GetRotationFormatFun(TrackIndex) == ACF_IntervalFixed32NoW)
+		{
+			if (bInterleaveValues)
+			{
+				const FTrackKeyFlags Flags = GetRotationFlagsFun(TrackIndex);
+				if (Flags.IsComponentNeededX())
+				{
+					UnalignedWriteToStream(ByteStream, &TrackRange.RotMin.X, sizeof(float));
+					UnalignedWriteToStream(ByteStream, &TrackRange.RotExtent.X, sizeof(float));
+				}
+
+				if (Flags.IsComponentNeededY())
+				{
+					UnalignedWriteToStream(ByteStream, &TrackRange.RotMin.Y, sizeof(float));
+					UnalignedWriteToStream(ByteStream, &TrackRange.RotExtent.Y, sizeof(float));
+				}
+
+				if (Flags.IsComponentNeededZ())
+				{
+					UnalignedWriteToStream(ByteStream, &TrackRange.RotMin.Z, sizeof(float));
+					UnalignedWriteToStream(ByteStream, &TrackRange.RotExtent.Z, sizeof(float));
+				}
+			}
+			else
+			{
+				UnalignedWriteToStream(ByteStream, &TrackRange.RotMin, sizeof(FVector));
+				UnalignedWriteToStream(ByteStream, &TrackRange.RotExtent, sizeof(FVector));
+			}
+		}
+
+		if (bHasScale)
+		{
+			const FScaleTrack& SrcScale = Segment.ScaleData[TrackIndex];
+			if (SrcScale.ScaleKeys.Num() > 1 && GetScaleFormatFun(TrackIndex) == ACF_IntervalFixed32NoW)
+			{
+				if (bInterleaveValues)
+				{
+					const FTrackKeyFlags Flags = GetScaleFlagsFun(TrackIndex);
+					if (Flags.IsComponentNeededX())
+					{
+						UnalignedWriteToStream(ByteStream, &TrackRange.ScaleMin.X, sizeof(float));
+						UnalignedWriteToStream(ByteStream, &TrackRange.ScaleExtent.X, sizeof(float));
+					}
+
+					if (Flags.IsComponentNeededY())
+					{
+						UnalignedWriteToStream(ByteStream, &TrackRange.ScaleMin.Y, sizeof(float));
+						UnalignedWriteToStream(ByteStream, &TrackRange.ScaleExtent.Y, sizeof(float));
+					}
+
+					if (Flags.IsComponentNeededZ())
+					{
+						UnalignedWriteToStream(ByteStream, &TrackRange.ScaleMin.Z, sizeof(float));
+						UnalignedWriteToStream(ByteStream, &TrackRange.ScaleExtent.Z, sizeof(float));
+					}
+				}
+				else
+				{
+					UnalignedWriteToStream(ByteStream, &TrackRange.ScaleMin, sizeof(FVector));
+					UnalignedWriteToStream(ByteStream, &TrackRange.ScaleExtent, sizeof(FVector));
+				}
+			}
+		}
+	}
+}
+
+void UAnimCompress::WriteUniformTrackData(
+	TArray<uint8>& ByteStream,
+	TFunction<AnimationCompressionFormat(int32 TrackIndex)> GetTranslationFormatFun,
+	TFunction<AnimationCompressionFormat(int32 TrackIndex)> GetRotationFormatFun,
+	TFunction<AnimationCompressionFormat(int32 TrackIndex)> GetScaleFormatFun,
+	TFunction<bool(int32 TrackIndex)> IsTranslationUniformFun,
+	TFunction<bool(int32 TrackIndex)> IsRotationUniformFun,
+	TFunction<bool(int32 TrackIndex)> IsScaleUniformFun,
+	TFunction<void(TArray<uint8>& ByteStream, AnimationCompressionFormat Format, const FVector& Key, const float* Mins, const float* Ranges, int32 TrackIndex)> PackTranslationKeyFun,
+	TFunction<void(TArray<uint8>& ByteStream, AnimationCompressionFormat Format, const FQuat& Key, const float* Mins, const float* Ranges, int32 TrackIndex)> PackRotationKeyFun,
+	TFunction<void(TArray<uint8>& ByteStream, AnimationCompressionFormat Format, const FVector& Key, const float* Mins, const float* Ranges, int32 TrackIndex)> PackScaleKeyFun,
+	const FAnimSegmentContext& Segment,
+	const TArray<FAnimTrackRange>& TrackRanges)
+{
+	const int32 NumFrames = Segment.NumFrames;
+	const int32 NumTracks = Segment.RotationData.Num();
+	const bool bHasScale = Segment.ScaleData.Num() > 0;
+
+	// Samples are sorted by time first, track second to ensure that when we sample a time T, data for all tracks at
+	// that time is contiguous in memory
+	for (int32 KeyIndex = 0; KeyIndex < NumFrames; ++KeyIndex)
+	{
+		for (int32 TrackIndex = 0; TrackIndex < NumTracks; ++TrackIndex)
+		{
+			const FAnimTrackRange& TrackRange = TrackRanges[TrackIndex];
+
+			const FTranslationTrack& SrcTrans = Segment.TranslationData[TrackIndex];
+			if (SrcTrans.PosKeys.Num() > 1 && IsTranslationUniformFun(TrackIndex))
+			{
+				const AnimationCompressionFormat Format = GetTranslationFormatFun(TrackIndex);
+				const FVector& Vec = SrcTrans.PosKeys[KeyIndex];
+				PackTranslationKeyFun(ByteStream, Format, Vec, &TrackRange.TransMin.X, &TrackRange.TransExtent.X, TrackIndex);
+			}
+
+			const FRotationTrack& SrcRot = Segment.RotationData[TrackIndex];
+			if (SrcRot.RotKeys.Num() > 1 && IsRotationUniformFun(TrackIndex))
+			{
+				const AnimationCompressionFormat Format = GetRotationFormatFun(TrackIndex);
+				const FQuat& Quat = SrcRot.RotKeys[KeyIndex];
+				PackRotationKeyFun(ByteStream, Format, Quat, &TrackRange.RotMin.X, &TrackRange.RotExtent.X, TrackIndex);
+			}
+
+			if (bHasScale)
+			{
+				const FScaleTrack& SrcScale = Segment.ScaleData[TrackIndex];
+				if (SrcScale.ScaleKeys.Num() > 1 && IsScaleUniformFun(TrackIndex))
+				{
+					const AnimationCompressionFormat Format = GetScaleFormatFun(TrackIndex);
+					const FVector& Vec = SrcScale.ScaleKeys[KeyIndex];
+					PackScaleKeyFun(ByteStream, Format, Vec, &TrackRange.ScaleMin.X, &TrackRange.ScaleExtent.X, TrackIndex);
+				}
+			}
+		}
+	}
+}
+
+void UAnimCompress::WriteSortedVariableTrackData(
+	TArray<uint8>& ByteStream,
+	const UAnimSequence& AnimSeq,
+	TFunction<AnimationCompressionFormat(int32 TrackIndex)> GetTranslationFormatFun,
+	TFunction<AnimationCompressionFormat(int32 TrackIndex)> GetRotationFormatFun,
+	TFunction<AnimationCompressionFormat(int32 TrackIndex)> GetScaleFormatFun,
+	TFunction<bool(int32 TrackIndex)> IsTranslationVariableFun,
+	TFunction<bool(int32 TrackIndex)> IsRotationVariableFun,
+	TFunction<bool(int32 TrackIndex)> IsScaleVariableFun,
+	TFunction<void(TArray<uint8>& ByteStream, AnimationCompressionFormat Format, const FVector& Key, const float* Mins, const float* Ranges, int32 TrackIndex)> PackTranslationKeyFun,
+	TFunction<void(TArray<uint8>& ByteStream, AnimationCompressionFormat Format, const FQuat& Key, const float* Mins, const float* Ranges, int32 TrackIndex)> PackRotationKeyFun,
+	TFunction<void(TArray<uint8>& ByteStream, AnimationCompressionFormat Format, const FVector& Key, const float* Mins, const float* Ranges, int32 TrackIndex)> PackScaleKeyFun,
+	FAnimSegmentContext& Segment,
+	const TArray<FAnimTrackRange>& TrackRanges)
+{
+	const int32 NumFrames = Segment.NumFrames;
+	const int32 NumTracks = Segment.RotationData.Num();
+	const bool bHasScale = Segment.ScaleData.Num() > 0;
+
+	const int32 LastFrame = AnimSeq.NumFrames - 1;
+	const float FrameRate = float(AnimSeq.NumFrames - 1) / AnimSeq.SequenceLength;
+
+	// Samples are sorted by time needed first, track second to ensure that when we sample a time T, data for all tracks at
+	// that time is contiguous in memory
+	// Imagine the following rotation keys:
+	// rot0 @ 0.0										rot0 @ 0.75
+	// rot1 @ 0.0					rot1 @ 0.5			rot1 @ 0.75
+	// rot2 @ 0.0	rot2 @ 0.25		rot2 @ 0.5			rot2 @ 0.75
+	//
+	// After being sorted, we have:
+	// rot0 @ 0.0, rot1 @ 0.0, rot2 @ 0.0, rot2 @ 0.25, rot1 @ 0.5, rot0 @ 0.75, rot2 @ 0.5, rot1 @ 0.75, rot2 @ 0.75
+	//
+	// The sorting is done by the time of the previous key, which is when the current key is expected in the stream.
+	// The first key at T = 0.0 is special and is needed first, as such for the purpose of the sort, it has a negative time T.
+	// rot1 @ 0.5 is needed at T = 0.0, rot1 @ 0.75 is needed at T = 0.5, etc.
+	//
+	// When we will read from the stream, the time delta can be positive or negative.
+	// When we wish to sample a time T, we must read from the stream until:
+	//    - The time value of the packed sample is greater than T
+	//    - AND the previously read time value for the packed track sample is greater or equal to T
+	//    - OR we reach the end of the stream (we have an EOS marker)
+	//
+	// Sampling at 0.4 with the above example, we'll have:
+	// [rot0 @ 0.0, rot1 @ 0.0, rot2 @ 0.0, rot2 @ 0.25]: The first 4 values have a sample time lower or equal to T,
+	// we consume them and cache the last 2 samples for every track type.
+	// rot1 @ 0.5 is greater than T, but the last sample read for rot1 is 0.0, we consume it and cache it
+	// rot0 @ 0.75 is greater than T, but the last sample read for rot0 is 0.0, we consume it and cache it
+	// rot2 @ 0.5 is greater than T, but the last sample read for rot2 is 0.25, we consume it and cache it
+	// rot1 @ 0.75 is greater than T, but the last sample read for rot1 is 0.5, we don't need this sample and we can stop reading
+	//
+	// Our cached samples are now:
+	// [rot0 @ 0.0, rot0 @ 0.75], [rot1 @ 0.0, rot1 @ 0.5], [rot2 @ 0.25, rot2 @ 0.5]
+	//
+	// Each sample packed has a small header.
+	// 16 bits are used for the track index
+	// The MSB of the next byte tells us if the rest of the header is 8 bits (small) or 16 bits (large) wide.
+	// The following 2 bits tell us if the sample is a rotation (0), translation (1), or scale (2).
+	// If the header is small, the remaining 5 bits represent the time delta from the last packed sample (previous packed time + packed time delta).
+	// If the header is large, the remaining 13 bits represent the time delta from the last packed sample.
+	// Time does not always increment as such the time delta is can be positive or negative.
+	// As such each sample has a header of 24 bits or 32 bits. In practice, we will have a sample at nearly every
+	// key frame as such the time delta should be very small and 24 bits should end up being used for virtually every sample.
+	//
+	// [ track index : 16 ] [ is large header? : 1 ] [ sample type : 2 ] [ time delta : 5 ]
+	// [ track index : 16 ] [ is large header? : 1 ] [ sample type : 2 ] [ time delta : 13 ]
+
+	struct SampleRef
+	{
+		int32 TrackIndex;
+		int32 SampleType;
+		int32 FrameIndex;
+		int32 KeyIndex;
+		float NeededAtTime;
+	};
+
+	TArray<SampleRef> SamplesToPack;
+	SamplesToPack.Empty(NumTracks * NumFrames);
+	for (int32 TrackIndex = 0; TrackIndex < NumTracks; ++TrackIndex)
+	{
+		const FAnimTrackRange& TrackRange = TrackRanges[TrackIndex];
+
+		const FTranslationTrack& SrcTrans = Segment.TranslationData[TrackIndex];
+		const int32 NumKeysTrans = SrcTrans.PosKeys.Num();
+		if (NumKeysTrans > 1 && IsTranslationVariableFun(TrackIndex))
+		{
+			for (int32 KeyIndex = 0; KeyIndex < NumKeysTrans; ++KeyIndex)
+			{
+				float KeyTime = SrcTrans.Times[KeyIndex];
+				float FrameTime = KeyTime * FrameRate;
+				int32 FrameIndex = FMath::Clamp(FMath::TruncToInt(FrameTime + 0.5f), 0, LastFrame);
+				int32 SegmentFrameIndex = FrameIndex - Segment.StartFrame;
+				float NeededAtTime = KeyIndex == 0 ? -1.0f : SrcTrans.Times[KeyIndex - 1];
+				SamplesToPack.Add(SampleRef{ TrackIndex, 0, SegmentFrameIndex, KeyIndex, NeededAtTime });
+			}
+		}
+
+		const FRotationTrack& SrcRot = Segment.RotationData[TrackIndex];
+		const int32 NumKeysRot = SrcRot.RotKeys.Num();
+		if (NumKeysRot > 1 && IsRotationVariableFun(TrackIndex))
+		{
+			for (int32 KeyIndex = 0; KeyIndex < NumKeysRot; ++KeyIndex)
+			{
+				float KeyTime = SrcRot.Times[KeyIndex];
+				float FrameTime = KeyTime * FrameRate;
+				int32 FrameIndex = FMath::Clamp(FMath::TruncToInt(FrameTime + 0.5f), 0, LastFrame);
+				int32 SegmentFrameIndex = FrameIndex - Segment.StartFrame;
+				float NeededAtTime = KeyIndex == 0 ? -1.0f : SrcRot.Times[KeyIndex - 1];
+				SamplesToPack.Add(SampleRef{ TrackIndex, 1, SegmentFrameIndex, KeyIndex, NeededAtTime });
+			}
+		}
+
+		if (bHasScale)
+		{
+			const FScaleTrack& SrcScale = Segment.ScaleData[TrackIndex];
+			const int32 NumKeysScale = SrcScale.ScaleKeys.Num();
+			if (NumKeysScale > 1 && IsScaleVariableFun(TrackIndex))
+			{
+				for (int32 KeyIndex = 0; KeyIndex < NumKeysScale; ++KeyIndex)
+				{
+					// write the frame values for each key
+					float KeyTime = SrcScale.Times[KeyIndex];
+					float FrameTime = KeyTime * FrameRate;
+					int32 FrameIndex = FMath::Clamp(FMath::TruncToInt(FrameTime + 0.5f), 0, LastFrame);
+					int32 SegmentFrameIndex = FrameIndex - Segment.StartFrame;
+					float NeededAtTime = KeyIndex == 0 ? -1.0f : SrcScale.Times[KeyIndex - 1];
+					SamplesToPack.Add(SampleRef{ TrackIndex, 2, SegmentFrameIndex, KeyIndex, NeededAtTime });
+				}
+			}
+		}
+	}
+
+	// Sort by key index, followed by time as explained above
+	SamplesToPack.Sort([](const SampleRef& lhs, const SampleRef& rhs)
+	{
+		// Lowest needed at frame index first
+		if (lhs.NeededAtTime < rhs.NeededAtTime)
+		{
+			return true;
+		}
+
+		if (lhs.NeededAtTime > rhs.NeededAtTime)
+		{
+			return false;
+		}
+
+		// Lowest sample type first
+		return lhs.SampleType < rhs.SampleType;
+	});
+
+	int32 PreviousPackedFrameIndex = 0;
+	int32 SampleIndex = 0;
+	for (const SampleRef& Sample : SamplesToPack)
+	{
+		const FAnimTrackRange& TrackRange = TrackRanges[Sample.TrackIndex];
+
+		const int32 TimeDelta = Sample.FrameIndex - PreviousPackedFrameIndex;
+		const FSortedKeyHeader KeyHeader(Sample.TrackIndex, Sample.SampleType, TimeDelta);
+		UnalignedWriteToStream(ByteStream, &KeyHeader, KeyHeader.GetSize());
+
+		PreviousPackedFrameIndex = Sample.FrameIndex;
+
+		if (Sample.SampleType == 0)
+		{
+			const FTranslationTrack& SrcTrans = Segment.TranslationData[Sample.TrackIndex];
+			const FVector& Vec = SrcTrans.PosKeys[Sample.KeyIndex];
+			const AnimationCompressionFormat Format = GetTranslationFormatFun(Sample.TrackIndex);
+			PackTranslationKeyFun(ByteStream, Format, Vec, &TrackRange.TransMin.X, &TrackRange.TransExtent.X, Sample.TrackIndex);
+		}
+		else if (Sample.SampleType == 1)
+		{
+			const FRotationTrack& SrcRot = Segment.RotationData[Sample.TrackIndex];
+			const FQuat& Quat = SrcRot.RotKeys[Sample.KeyIndex];
+			const AnimationCompressionFormat Format = GetRotationFormatFun(Sample.TrackIndex);
+			PackRotationKeyFun(ByteStream, Format, Quat, &TrackRange.RotMin.X, &TrackRange.RotExtent.X, Sample.TrackIndex);
+		}
+		else
+		{
+			checkSlow(Sample.SampleType == 2);
+			const FScaleTrack& SrcScale = Segment.ScaleData[Sample.TrackIndex];
+			const FVector& Vec = SrcScale.ScaleKeys[Sample.KeyIndex];
+			const AnimationCompressionFormat Format = GetScaleFormatFun(Sample.TrackIndex);
+			PackScaleKeyFun(ByteStream, Format, Vec, &TrackRange.ScaleMin.X, &TrackRange.ScaleExtent.X, Sample.TrackIndex);
+		}
+
+		SampleIndex++;
+	}
+
+	// End the stream with a terminator
+	const FSortedKeyHeader EndOfStreamKeyHeader;
+	UnalignedWriteToStream(ByteStream, &EndOfStreamKeyHeader, EndOfStreamKeyHeader.GetSize());
+}
+
+static int32 GetNumAnimatedTrackStreams(
+	TFunction<bool(int32 TrackIndex)> IsTranslationVariableFun,
+	TFunction<bool(int32 TrackIndex)> IsRotationVariableFun,
+	TFunction<bool(int32 TrackIndex)> IsScaleVariableFun,
+	const FAnimSegmentContext& RawSegment)
+{
+	const int32 NumTracks = RawSegment.RotationData.Num();
+	const bool bHasScale = RawSegment.ScaleData.Num() > 0;
+
+	int32 NumAnimatedTrackTypes = 0;
+	for (int32 TrackIndex = 0; TrackIndex < NumTracks; ++TrackIndex)
+	{
+		NumAnimatedTrackTypes += RawSegment.TranslationData[TrackIndex].PosKeys.Num() > 1 && IsTranslationVariableFun(TrackIndex) ? 1 : 0;
+		NumAnimatedTrackTypes += RawSegment.RotationData[TrackIndex].RotKeys.Num() > 1 && IsRotationVariableFun(TrackIndex) ? 1 : 0;
+		NumAnimatedTrackTypes += bHasScale && RawSegment.ScaleData[TrackIndex].ScaleKeys.Num() > 1 && IsScaleVariableFun(TrackIndex) ? 1 : 0;
+	}
+	return NumAnimatedTrackTypes;
+}
+
+void UAnimCompress::WriteLinearVariableTrackData(
+	TArray<uint8>& ByteStream,
+	const UAnimSequence& AnimSeq,
+	TFunction<AnimationCompressionFormat(int32 TrackIndex)> GetTranslationFormatFun,
+	TFunction<AnimationCompressionFormat(int32 TrackIndex)> GetRotationFormatFun,
+	TFunction<AnimationCompressionFormat(int32 TrackIndex)> GetScaleFormatFun,
+	TFunction<bool(int32 TrackIndex)> IsTranslationVariableFun,
+	TFunction<bool(int32 TrackIndex)> IsRotationVariableFun,
+	TFunction<bool(int32 TrackIndex)> IsScaleVariableFun,
+	TFunction<void(TArray<uint8>& ByteStream, AnimationCompressionFormat Format, const FVector& Key, const float* Mins, const float* Ranges, int32 TrackIndex)> PackTranslationKeyFun,
+	TFunction<void(TArray<uint8>& ByteStream, AnimationCompressionFormat Format, const FQuat& Key, const float* Mins, const float* Ranges, int32 TrackIndex)> PackRotationKeyFun,
+	TFunction<void(TArray<uint8>& ByteStream, AnimationCompressionFormat Format, const FVector& Key, const float* Mins, const float* Ranges, int32 TrackIndex)> PackScaleKeyFun,
+	const FAnimSegmentContext& Segment,
+	const TArray<FAnimTrackRange>& TrackRanges)
+{
+	const int32 NumFrames = Segment.NumFrames;
+	const int32 NumTracks = Segment.RotationData.Num();
+	const bool bHasScale = Segment.ScaleData.Num() > 0;
+
+	const int32 LastFrame = AnimSeq.NumFrames - 1;
+	const float FrameRate = float(AnimSeq.NumFrames - 1) / AnimSeq.SequenceLength;
+
+	// The linear packing format is more or less the same as the legacy format.
+	// First we have a list of pairs for each track type (rot, trans, scale): offset in stream, number of keys in stream
+	// Offsets are 32 bits, number of keys are 16 bits
+	// This is followed by the packed track data and their time markers.
+	//
+	// [rot0 offset, num rot0 keys] [rot1 offset, num rot1 keys] [rot0 t0, rot0 t1, rot0, t2, rot0 key0, rot0 key1, rot0 key2] [rot1 t0, rot1 t1, rot1 key0, rot1 key1]
+
+	const int32 OffsetNumKeysPairSize = sizeof(uint32) + sizeof(uint16);
+	const int32 NumAnimatedTrackStreams = GetNumAnimatedTrackStreams(IsTranslationVariableFun, IsRotationVariableFun, IsScaleVariableFun, Segment);
+
+	int32 OffsetNumKeysPairStreamOffset = ByteStream.Num();
+	ByteStream.AddUninitialized(OffsetNumKeysPairSize * NumAnimatedTrackStreams);
+
+	// If we don't have too many frames in our segment, use uint8 instead of uint16 for the time markers
+	const int32 TimeMarkerSize = NumFrames < 256 ? sizeof(uint8) : sizeof(uint16);
+
+	// TODO: There is 4 byte alignment for the packed keys due to the decompression code not handling unaligned reads properly
+	// Either properly handle unaligned reads of move the time markers after the packed keys
+
+	for (int32 TrackIndex = 0; TrackIndex < NumTracks; ++TrackIndex)
+	{
+		const FAnimTrackRange& TrackRange = TrackRanges[TrackIndex];
+
+		const FTranslationTrack& SrcTrans = Segment.TranslationData[TrackIndex];
+		const int32 NumKeysTrans = SrcTrans.PosKeys.Num();
+		if (NumKeysTrans > 1 && IsTranslationVariableFun(TrackIndex))
+		{
+			if (TimeMarkerSize == sizeof(uint16))
+				PadByteStream(ByteStream, 2, AnimationPadSentinel);
+
+			const uint32 TrackDataOffset = ByteStream.Num();
+			const uint16 NumTrackKeys = static_cast<uint16>(NumKeysTrans);
+			UnalignedWriteToStream(ByteStream, OffsetNumKeysPairStreamOffset, &TrackDataOffset, sizeof(uint32));
+			UnalignedWriteToStream(ByteStream, OffsetNumKeysPairStreamOffset, &NumTrackKeys, sizeof(uint16));
+
+			for (int32 KeyIndex = 0; KeyIndex < NumKeysTrans; ++KeyIndex)
+			{
+				const float KeyTime = SrcTrans.Times[KeyIndex];
+				const float FrameTime = KeyTime * FrameRate;
+				const int32 FrameIndex = FMath::Clamp(FMath::TruncToInt(FrameTime + 0.5f), 0, LastFrame);
+				const int32 SegmentFrameIndex = FrameIndex - Segment.StartFrame;
+				UnalignedWriteToStream(ByteStream, &SegmentFrameIndex, TimeMarkerSize);
+			}
+
+			PadByteStream(ByteStream, 4, AnimationPadSentinel);
+
+			const AnimationCompressionFormat Format = GetTranslationFormatFun(TrackIndex);
+			for (int32 KeyIndex = 0; KeyIndex < NumKeysTrans; ++KeyIndex)
+			{
+				const FVector& Key = SrcTrans.PosKeys[KeyIndex];
+				PackTranslationKeyFun(ByteStream, Format, Key, &TrackRange.TransMin.X, &TrackRange.TransExtent.X, TrackIndex);
+			}
+		}
+
+		const FRotationTrack& SrcRot = Segment.RotationData[TrackIndex];
+		const int32 NumKeysRot = SrcRot.RotKeys.Num();
+		if (NumKeysRot > 1 && IsRotationVariableFun(TrackIndex))
+		{
+			if (TimeMarkerSize == sizeof(uint16))
+				PadByteStream(ByteStream, 2, AnimationPadSentinel);
+
+			const uint32 TrackDataOffset = ByteStream.Num();
+			const uint16 NumTrackKeys = static_cast<uint16>(NumKeysRot);
+			UnalignedWriteToStream(ByteStream, OffsetNumKeysPairStreamOffset, &TrackDataOffset, sizeof(uint32));
+			UnalignedWriteToStream(ByteStream, OffsetNumKeysPairStreamOffset, &NumTrackKeys, sizeof(uint16));
+
+			for (int32 KeyIndex = 0; KeyIndex < NumKeysRot; ++KeyIndex)
+			{
+				const float KeyTime = SrcRot.Times[KeyIndex];
+				const float FrameTime = KeyTime * FrameRate;
+				const int32 FrameIndex = FMath::Clamp(FMath::TruncToInt(FrameTime + 0.5f), 0, LastFrame);
+				const int32 SegmentFrameIndex = FrameIndex - Segment.StartFrame;
+				UnalignedWriteToStream(ByteStream, &SegmentFrameIndex, TimeMarkerSize);
+			}
+
+			PadByteStream(ByteStream, 4, AnimationPadSentinel);
+
+			const AnimationCompressionFormat Format = GetRotationFormatFun(TrackIndex);
+			for (int32 KeyIndex = 0; KeyIndex < NumKeysRot; ++KeyIndex)
+			{
+				const FQuat& Key = SrcRot.RotKeys[KeyIndex];
+				PackRotationKeyFun(ByteStream, Format, Key, &TrackRange.RotMin.X, &TrackRange.RotExtent.X, TrackIndex);
+			}
+		}
+
+		if (bHasScale)
+		{
+			const FScaleTrack& SrcScale = Segment.ScaleData[TrackIndex];
+			const int32 NumKeysScale = SrcScale.ScaleKeys.Num();
+			if (NumKeysScale > 1 && IsScaleVariableFun(TrackIndex))
+			{
+				if (TimeMarkerSize == sizeof(uint16))
+					PadByteStream(ByteStream, 2, AnimationPadSentinel);
+
+				const uint32 TrackDataOffset = ByteStream.Num();
+				const uint16 NumTrackKeys = static_cast<uint16>(NumKeysScale);
+				UnalignedWriteToStream(ByteStream, OffsetNumKeysPairStreamOffset, &TrackDataOffset, sizeof(uint32));
+				UnalignedWriteToStream(ByteStream, OffsetNumKeysPairStreamOffset, &NumTrackKeys, sizeof(uint16));
+
+				for (int32 KeyIndex = 0; KeyIndex < NumKeysScale; ++KeyIndex)
+				{
+					const float KeyTime = SrcScale.Times[KeyIndex];
+					const float FrameTime = KeyTime * FrameRate;
+					const int32 FrameIndex = FMath::Clamp(FMath::TruncToInt(FrameTime + 0.5f), 0, LastFrame);
+					const int32 SegmentFrameIndex = FrameIndex - Segment.StartFrame;
+					UnalignedWriteToStream(ByteStream, &SegmentFrameIndex, TimeMarkerSize);
+				}
+
+				PadByteStream(ByteStream, 4, AnimationPadSentinel);
+
+				const AnimationCompressionFormat Format = GetScaleFormatFun(TrackIndex);
+				for (int32 KeyIndex = 0; KeyIndex < NumKeysScale; ++KeyIndex)
+				{
+					const FVector& Key = SrcScale.ScaleKeys[KeyIndex];
+					PackScaleKeyFun(ByteStream, Format, Key, &TrackRange.ScaleMin.X, &TrackRange.ScaleExtent.X, TrackIndex);
+				}
+			}
+		}
+	}
+}
+
+void UAnimCompress::BitwiseCompressAnimationTracks(
+	const UAnimSequence& AnimSeq,
+	AnimationCompressionFormat TargetTranslationFormat,
+	AnimationCompressionFormat TargetRotationFormat,
+	AnimationCompressionFormat TargetScaleFormat,
+	FAnimSegmentContext& RawSegment,
+	bool bIsSorted)
+{
+	RawSegment.RotationCompressionFormat = TargetRotationFormat;
+	RawSegment.TranslationCompressionFormat = TargetTranslationFormat;
+	RawSegment.ScaleCompressionFormat = TargetScaleFormat;
+
+	SanityCheckTrackData(AnimSeq, RawSegment);
+
+	RawSegment.CompressedByteStream.Empty(64 * 1024);
+
+	TArray<FAnimTrackRange> TrackRanges;
+	CalculateTrackRanges(TargetTranslationFormat, TargetRotationFormat, TargetScaleFormat, RawSegment, TrackRanges);
+
+	// Write track ranges
+	WriteTrackRanges(
+		RawSegment.CompressedByteStream,
+		[&](int32 TrackIndex) { return TargetTranslationFormat; },
+		[&](int32 TrackIndex) { return TargetRotationFormat; },
+		[&](int32 TrackIndex) { return TargetScaleFormat; },
+		[](int32 TrackIndex) { return FTrackKeyFlags(); },
+		[](int32 TrackIndex) { return FTrackKeyFlags(); },
+		[](int32 TrackIndex) { return FTrackKeyFlags(); },
+		RawSegment, TrackRanges, false);
+
+	checkf((RawSegment.CompressedByteStream.Num() % 4) == 0, TEXT("CompressedByteStream not aligned to four bytes"));
+
+	if (AnimSeq.KeyEncodingFormat == AKF_ConstantKeyLerp)
+	{
+		WriteUniformTrackData(
+			RawSegment.CompressedByteStream,
+			[&](int32 TrackIndex) { return TargetTranslationFormat; },
+			[&](int32 TrackIndex) { return TargetRotationFormat; },
+			[&](int32 TrackIndex) { return TargetScaleFormat; },
+			[](int32 TrackIndex) { return true; },
+			[](int32 TrackIndex) { return true; },
+			[](int32 TrackIndex) { return true; },
+			[](TArray<uint8>& ByteStream, AnimationCompressionFormat Format, const FVector& Key, const float* Mins, const float* Ranges, int32 TrackIndex) { PackVectorToStream(ByteStream, Format, Key, Mins, Ranges); },
+			[](TArray<uint8>& ByteStream, AnimationCompressionFormat Format, const FQuat& Key, const float* Mins, const float* Ranges, int32 TrackIndex) { PackQuaternionToStream(ByteStream, Format, Key, Mins, Ranges); },
+			[](TArray<uint8>& ByteStream, AnimationCompressionFormat Format, const FVector& Key, const float* Mins, const float* Ranges, int32 TrackIndex) { PackVectorToStream(ByteStream, Format, Key, Mins, Ranges); },
+			RawSegment, TrackRanges);
+	}
+	else if (AnimSeq.KeyEncodingFormat == AKF_VariableKeyLerp)
+	{
+		if (bIsSorted)
+		{
+			WriteSortedVariableTrackData(
+				RawSegment.CompressedByteStream,
+				AnimSeq,
+				[&](int32 TrackIndex) { return TargetTranslationFormat; },
+				[&](int32 TrackIndex) { return TargetRotationFormat; },
+				[&](int32 TrackIndex) { return TargetScaleFormat; },
+				[](int32 TrackIndex) { return true; },
+				[](int32 TrackIndex) { return true; },
+				[](int32 TrackIndex) { return true; },
+				[](TArray<uint8>& ByteStream, AnimationCompressionFormat Format, const FVector& Key, const float* Mins, const float* Ranges, int32 TrackIndex) { PackVectorToStream(ByteStream, Format, Key, Mins, Ranges); },
+				[](TArray<uint8>& ByteStream, AnimationCompressionFormat Format, const FQuat& Key, const float* Mins, const float* Ranges, int32 TrackIndex) { PackQuaternionToStream(ByteStream, Format, Key, Mins, Ranges); },
+				[](TArray<uint8>& ByteStream, AnimationCompressionFormat Format, const FVector& Key, const float* Mins, const float* Ranges, int32 TrackIndex) { PackVectorToStream(ByteStream, Format, Key, Mins, Ranges); },
+				RawSegment, TrackRanges);
+		}
+		else
+		{
+			WriteLinearVariableTrackData(
+				RawSegment.CompressedByteStream,
+				AnimSeq,
+				[&](int32 TrackIndex) { return TargetTranslationFormat; },
+				[&](int32 TrackIndex) { return TargetRotationFormat; },
+				[&](int32 TrackIndex) { return TargetScaleFormat; },
+				[](int32 TrackIndex) { return true; },
+				[](int32 TrackIndex) { return true; },
+				[](int32 TrackIndex) { return true; },
+				[](TArray<uint8>& ByteStream, AnimationCompressionFormat Format, const FVector& Key, const float* Mins, const float* Ranges, int32 TrackIndex) { PackVectorToStream(ByteStream, Format, Key, Mins, Ranges); },
+				[](TArray<uint8>& ByteStream, AnimationCompressionFormat Format, const FQuat& Key, const float* Mins, const float* Ranges, int32 TrackIndex) { PackQuaternionToStream(ByteStream, Format, Key, Mins, Ranges); },
+				[](TArray<uint8>& ByteStream, AnimationCompressionFormat Format, const FVector& Key, const float* Mins, const float* Ranges, int32 TrackIndex) { PackVectorToStream(ByteStream, Format, Key, Mins, Ranges); },
+				RawSegment, TrackRanges);
+		}
+	}
+
+	// Make sure we have a safe alignment
+	PadByteStream(RawSegment.CompressedByteStream, 4, AnimationPadSentinel);
+
+	// Trim unused memory.
+	RawSegment.CompressedByteStream.Shrink();
+}
+
+void UAnimCompress::BitwiseCompressTrivialAnimationTracks(const UAnimSequence& AnimSeq, FAnimSegmentContext& RawSegment)
+{
+	RawSegment.CompressedTrivialTracksByteStream.Empty();
+
+	check(RawSegment.TranslationData.Num() == RawSegment.RotationData.Num());
+	const int32 NumTracks = RawSegment.RotationData.Num();
+	const bool bHasScale = RawSegment.ScaleData.Num() > 0;
+
+	const int32 NumFrames = RawSegment.NumFrames;
+	const int32 LastFrame = RawSegment.NumFrames - 1;
+	const size_t FrameSize = RawSegment.NumFrames > 0xff ? sizeof(uint16) : sizeof(uint8);
+	const float FrameRate = AnimSeq.NumFrames / AnimSeq.SequenceLength;
+
+	SanityCheckTrackData(AnimSeq, RawSegment);
+
+	for (int32 TrackIndex = 0; TrackIndex < NumTracks; ++TrackIndex)
+	{
+		const FTranslationTrack& SrcTrans = RawSegment.TranslationData[TrackIndex];
+		if (SrcTrans.PosKeys.Num() == 1)
+		{
+			// A single translation key gets written out a single uncompressed float[3].
+			UnalignedWriteToStream(RawSegment.CompressedTrivialTracksByteStream, &SrcTrans.PosKeys[0], sizeof(FVector));
+		}
+
+		const FRotationTrack& SrcRot = RawSegment.RotationData[TrackIndex];
+		if (SrcRot.RotKeys.Num() == 1)
+		{
+			// For a rotation track of n=1 keys, the single key is packed as an FQuatFloat96NoW.
+			const FQuat& Quat = SrcRot.RotKeys[0];
+			const FQuatFloat96NoW QuatFloat96NoW(Quat);
+			UnalignedWriteToStream(RawSegment.CompressedTrivialTracksByteStream, &QuatFloat96NoW, sizeof(FQuatFloat96NoW));
+		}
+
+		if (bHasScale)
+		{
+			const FScaleTrack& SrcScale = RawSegment.ScaleData[TrackIndex];
+			if (SrcScale.ScaleKeys.Num() == 1)
+			{
+				// A single Scale key gets written out a single uncompressed float[3].
+				UnalignedWriteToStream(RawSegment.CompressedTrivialTracksByteStream, &SrcScale.ScaleKeys[0], sizeof(FVector));
+			}
+		}
+	}
+
+	// Trim unused memory.
+	RawSegment.CompressedTrivialTracksByteStream.Shrink();
+}
+
+static int32 GetNumFrames(const TArray<FAnimSegmentContext>& RawSegments)
+{
+	int32 NumFrames = 0;
+	for (const FAnimSegmentContext& Segment : RawSegments)
+	{
+		NumFrames += Segment.NumFrames;
+	}
+	return NumFrames;
+}
+
+void UAnimCompress::CoalesceCompressedSegments(UAnimSequence& AnimSeq, const TArray<FAnimSegmentContext>& RawSegments, bool bIsSorted)
+{
+	AnimSeq.CompressedByteStream.Empty();
+
+	const FAnimSegmentContext& FirstRawSegment = RawSegments[0];
+	const int32 NumTracks = FirstRawSegment.RotationData.Num();
+	const int32 NumFrames = GetNumFrames(RawSegments);
+	const bool bHasScale = FirstRawSegment.ScaleData.Num() != 0;
+
+	FAnimSequenceCompressionHeader Header;
+	Header.NumTracks = NumTracks;
+	Header.NumFrames = NumFrames;
+	Header.SequenceCRC = 0;
+	Header.bHasScale = bHasScale;
+	Header.bIsSorted = bIsSorted;
+
+	AnimSeq.CompressedByteStream.AddUninitialized(sizeof(FAnimSequenceCompressionHeader));
+	memcpy(AnimSeq.CompressedByteStream.GetData(), &Header, sizeof(FAnimSequenceCompressionHeader));
+
+	// TODO: Convert to bitset!
+	int32 TrackInfoOffset = AnimSeq.CompressedByteStream.Num();
+	AnimSeq.CompressedByteStream.AddUninitialized(NumTracks);
+	for (int32 TrackIndex = 0; TrackIndex < NumTracks; ++TrackIndex)
+	{
+		bool IsTranslationTrivial = FirstRawSegment.TranslationData[TrackIndex].PosKeys.Num() <= 1;
+		bool IsRotationTrivial = FirstRawSegment.RotationData[TrackIndex].RotKeys.Num() <= 1;
+		bool IsScaleTrivial = !bHasScale || FirstRawSegment.ScaleData[TrackIndex].ScaleKeys.Num() <= 1;
+
+		// Bottom 3 bits indicate whether or not trans/rot/scale is trivial
+		uint8 TrackFlags = (IsTranslationTrivial ? 0x4 : 0x0) | (IsRotationTrivial ? 0x2 : 0x0) | (IsScaleTrivial ? 0x1 : 0x0);
+		AnimSeq.CompressedByteStream[TrackInfoOffset + TrackIndex] = TrackFlags;
+	}
+
+	PadByteStream(AnimSeq.CompressedByteStream, 4, AnimationPadSentinel);
+
+	const int32 NumSegments = RawSegments.Num();
+
+	AnimSeq.CompressedSegments.Empty(NumSegments);
+	AnimSeq.CompressedSegments.AddZeroed(NumSegments);
+	for (int32 SegmentIndex = 0; SegmentIndex < NumSegments; ++SegmentIndex)
+	{
+		const FAnimSegmentContext& RawSegment = RawSegments[SegmentIndex];
+		FCompressedSegment& CompressedSegment = AnimSeq.CompressedSegments[SegmentIndex];
+
+		if (SegmentIndex == 0)
+		{
+			int32 TrivialKeysOffset = AnimSeq.CompressedByteStream.Num();
+			AnimSeq.CompressedByteStream.AddUninitialized(RawSegment.CompressedTrivialTracksByteStream.Num());
+			memcpy(AnimSeq.CompressedByteStream.GetData() + TrivialKeysOffset, RawSegment.CompressedTrivialTracksByteStream.GetData(), RawSegment.CompressedTrivialTracksByteStream.Num());
+		}
+
+		CompressedSegment.StartFrame = RawSegment.StartFrame;
+		CompressedSegment.NumFrames = RawSegment.NumFrames;
+		CompressedSegment.ByteStreamOffset = AnimSeq.CompressedByteStream.Num();
+		CompressedSegment.TranslationCompressionFormat = RawSegment.TranslationCompressionFormat;
+		CompressedSegment.RotationCompressionFormat = RawSegment.RotationCompressionFormat;
+		CompressedSegment.ScaleCompressionFormat = RawSegment.ScaleCompressionFormat;
+
+		AnimSeq.CompressedByteStream.AddUninitialized(RawSegment.CompressedByteStream.Num());
+		memcpy(AnimSeq.CompressedByteStream.GetData() + CompressedSegment.ByteStreamOffset, RawSegment.CompressedByteStream.GetData(), RawSegment.CompressedByteStream.Num());
+	}
+
+	AnimSeq.CompressedByteStream.Shrink();
+
+	// Calculate CRC and update it
+	Header.SequenceCRC = FCrc::MemCrc32(AnimSeq.CompressedByteStream.GetData(), AnimSeq.CompressedByteStream.Num());
+	memcpy(AnimSeq.CompressedByteStream.GetData(), &Header, sizeof(FAnimSequenceCompressionHeader));
+}
+
 #if WITH_EDITOR
 
 FString UAnimCompress::MakeDDCKey()
@@ -644,6 +1606,9 @@ void UAnimCompress::PopulateDDCKey(FArchive& Ar)
 	Ar << TCF << RCF << SCF;
 
 	Ar << MaxCurveError;
+
+	bool ES = bEnableSegmenting != 0 && USE_SEGMENTING_CONTEXT;
+	Ar << ES << IdealNumFramesPerSegment << MaxNumFramesPerSegment;
 }
 
 /**
@@ -651,7 +1616,7 @@ void UAnimCompress::PopulateDDCKey(FArchive& Ar)
  * Tracks
  */
 
-bool UAnimCompress::Reduce(UAnimSequence* AnimSeq, bool bOutput)
+bool UAnimCompress::Reduce(UAnimSequence* AnimSeq, bool bOutput, const TArray<FBoneData>& BoneData)
 {
 	bool bResult = false;
 #if WITH_EDITORONLY_DATA
@@ -660,7 +1625,7 @@ bool UAnimCompress::Reduce(UAnimSequence* AnimSeq, bool bOutput)
 	if ( bSkeletonExistsIfNeeded )
 	{
 		FAnimCompressContext CompressContext(false, bOutput);
-		Reduce(AnimSeq, CompressContext);
+		Reduce(AnimSeq, CompressContext, BoneData);
 
 		bResult = true;
 	}
@@ -669,17 +1634,12 @@ bool UAnimCompress::Reduce(UAnimSequence* AnimSeq, bool bOutput)
 	return bResult;
 }
 
-bool UAnimCompress::Reduce(class UAnimSequence* AnimSeq, FAnimCompressContext& Context)
+bool UAnimCompress::Reduce(class UAnimSequence* AnimSeq, FAnimCompressContext& Context, const TArray<FBoneData>& BoneData)
 {
 	bool bResult = false;
 
 #if WITH_EDITORONLY_DATA
-	// Build skeleton metadata to use during the key reduction.
-	TArray<FBoneData> BoneData;
-	FAnimationUtils::BuildSkeletonMetaData(AnimSeq->GetSkeleton(), BoneData);
-	Context.GatherPreCompressionStats(AnimSeq);
-
-	// General key reduction.
+	
 	DoReduction(AnimSeq, BoneData);
 
 	AnimSeq->bWasCompressedWithoutTranslations = false; // @fixmelh : bAnimRotationOnly
@@ -687,8 +1647,6 @@ bool UAnimCompress::Reduce(class UAnimSequence* AnimSeq, FAnimCompressContext& C
 	AnimSeq->EncodingPkgVersion = CURRENT_ANIMATION_ENCODING_PACKAGE_VERSION;
 	AnimSeq->MarkPackageDirty();
 
-	// determine the error added by the compression
-	Context.GatherPostCompressionStats(AnimSeq, BoneData);
 	bResult = true;
 #endif // WITH_EDITORONLY_DATA
 
@@ -857,13 +1815,21 @@ void UAnimCompress::FilterTrivialKeys(
 }
 
 
-void UAnimCompress::FilterAnimRotationOnlyKeys(TArray<FTranslationTrack> & PositionTracks, UAnimSequence* AnimSeq)
+void UAnimCompress::FilterTrivialKeys(
+	TArray<FAnimSegmentContext>& RawSegments,
+	float MaxPosDelta,
+	float MaxRotDelta,
+	float MaxScaleDelta)
 {
-
+	for (FAnimSegmentContext& Segment : RawSegments)
+	{
+		FilterTrivialRotationKeys(Segment.RotationData, MaxRotDelta);
+		FilterTrivialPositionKeys(Segment.TranslationData, MaxPosDelta);
+		FilterTrivialScaleKeys(Segment.ScaleData, MaxScaleDelta);
+	}
 }
 
 
-
 void UAnimCompress::FilterIntermittentPositionKeys(
 	FTranslationTrack& Track, 
 	int32 StartIndex,
@@ -1086,3 +2052,229 @@ void UAnimCompress::SeparateRawDataIntoTracks(
 		OutScaleData.Empty();
 	}
 }
+
+static void GenerateAnimSequenceSegments(const UAnimSequence& AnimSeq, int32 NumFrames, int32 IdealNumFramesPerSegment, int32 MaxNumFramesPerSegment, TArray<FAnimSegmentContext>& OutRawSegments)
+{
+	int32 NumSegments;
+	TArray<int32> NumFramesPerSegment;
+
+	IdealNumFramesPerSegment = FMath::Max(IdealNumFramesPerSegment, 0);
+	MaxNumFramesPerSegment = FMath::Max(IdealNumFramesPerSegment, MaxNumFramesPerSegment);
+
+	if (NumFrames > MaxNumFramesPerSegment && IdealNumFramesPerSegment > 0)
+	{
+		NumSegments = (NumFrames + IdealNumFramesPerSegment - 1) / IdealNumFramesPerSegment;
+		int32 PaddedNumFrames = NumSegments * IdealNumFramesPerSegment;
+
+		NumFramesPerSegment.Init(IdealNumFramesPerSegment, NumSegments);
+
+		int32 NumLeftoverFrames = IdealNumFramesPerSegment - (PaddedNumFrames - NumFrames);
+		if (NumLeftoverFrames != 0)
+		{
+			NumFramesPerSegment[NumSegments - 1] = NumLeftoverFrames;
+		}
+
+		int32 Slack = MaxNumFramesPerSegment - IdealNumFramesPerSegment;
+		if ((NumSegments - 1) * Slack >= NumLeftoverFrames)
+		{
+			// Enough segments to distribute the leftover samples of the last segment
+			while (NumFramesPerSegment[NumSegments - 1] != 0)
+			{
+				for (int32 SegmentIndex = 0; SegmentIndex < NumSegments - 1 && NumFramesPerSegment[NumSegments - 1] != 0; ++SegmentIndex)
+				{
+					NumFramesPerSegment[SegmentIndex]++;
+					NumFramesPerSegment[NumSegments - 1]--;
+				}
+			}
+
+			NumSegments--;
+		}
+
+		checkf(NumSegments != 1, TEXT("Expected a number of segments greater than 1."));
+	}
+	else
+	{
+		// Everything fits in a single segment
+		NumSegments = 1;
+		NumFramesPerSegment.Init(NumFrames, NumSegments);
+	}
+
+	OutRawSegments.Empty(NumSegments);
+	OutRawSegments.AddZeroed(NumSegments);
+	int32 NumPreviousFrames = 0;
+	for (int32 SegmentIndex = 0; SegmentIndex < NumSegments; ++SegmentIndex)
+	{
+		int32 NumFramesInSegment = NumFramesPerSegment[SegmentIndex];
+
+		FAnimSegmentContext& Segment = OutRawSegments[SegmentIndex];
+		Segment.StartFrame = NumPreviousFrames;
+		Segment.NumFrames = NumFramesInSegment;
+
+		NumPreviousFrames += NumFramesInSegment;
+	}
+}
+
+static int32 GetNumFrames(const UAnimSequence& AnimSeq, const TArray<FTranslationTrack>& TranslationData, const TArray<FRotationTrack>& RotationData, const TArray<FScaleTrack>& ScaleData)
+{
+	for (const FTranslationTrack& track : TranslationData)
+	{
+		if (track.PosKeys.Num() > 1)
+		{
+			return track.PosKeys.Num();
+		}
+	}
+
+	for (const FRotationTrack& track : RotationData)
+	{
+		if (track.RotKeys.Num() > 1)
+		{
+			return track.RotKeys.Num();
+		}
+	}
+
+	for (const FScaleTrack& track : ScaleData)
+	{
+		if (track.ScaleKeys.Num() > 1)
+		{
+			return track.ScaleKeys.Num();
+		}
+	}
+
+	// Fallback
+	return AnimSeq.NumFrames;
+}
+
+void UAnimCompress::SeparateRawDataIntoTracks(
+	const UAnimSequence& AnimSeq,
+	const TArray<FTranslationTrack>& TranslationData,
+	const TArray<FRotationTrack>& RotationData,
+	const TArray<FScaleTrack>& ScaleData,
+	int32 IdealNumFramesPerSegment, int32 MaxNumFramesPerSegment,
+	TArray<FAnimSegmentContext>& OutRawSegments)
+{
+	int32 NumFrames = GetNumFrames(AnimSeq, TranslationData, RotationData, ScaleData);
+	GenerateAnimSequenceSegments(AnimSeq, NumFrames, IdealNumFramesPerSegment, MaxNumFramesPerSegment, OutRawSegments);
+
+	const int32 NumTracks = RotationData.Num();
+
+	for (FAnimSegmentContext& Segment : OutRawSegments)
+	{
+		Segment.TranslationData.Empty(NumTracks);
+		Segment.RotationData.Empty(NumTracks);
+		Segment.ScaleData.Empty(NumTracks);
+		Segment.TranslationData.AddZeroed(NumTracks);
+		Segment.RotationData.AddZeroed(NumTracks);
+		Segment.ScaleData.AddZeroed(NumTracks);
+
+		// only compress scale if it has valid scale keys
+		bool bCompressScaleKeys = false;
+
+		for (int32 TrackIndex = 0; TrackIndex < NumTracks; ++TrackIndex)
+		{
+			const FTranslationTrack&	SeqTranslationTrack = TranslationData[TrackIndex];
+			const FRotationTrack&		SeqRotationTrack = RotationData[TrackIndex];
+
+			FTranslationTrack&	TranslationTrack = Segment.TranslationData[TrackIndex];
+			FRotationTrack&		RotationTrack = Segment.RotationData[TrackIndex];
+
+			const int32 NumPosKeys = SeqTranslationTrack.PosKeys.Num();
+			const int32 NumRotKeys = SeqRotationTrack.RotKeys.Num();
+
+			// Do nothing if the data for this track is empty.
+			if (NumPosKeys == 0 || NumRotKeys == 0)
+			{
+				continue;
+			}
+
+			const int32 NumPosSegmentKeys = FMath::Min(NumPosKeys, Segment.NumFrames);
+			const int32 NumRotSegmentKeys = FMath::Min(NumRotKeys, Segment.NumFrames);
+
+			TranslationTrack.PosKeys.Empty(NumPosSegmentKeys);
+			TranslationTrack.PosKeys.AddUninitialized(NumPosSegmentKeys);
+			TranslationTrack.Times.Empty(NumPosSegmentKeys);
+			TranslationTrack.Times.AddUninitialized(NumPosSegmentKeys);
+
+			if (NumPosKeys == 1)
+			{
+				TranslationTrack.PosKeys[0] = SeqTranslationTrack.PosKeys[0];
+				TranslationTrack.Times[0] = 0.f;
+			}
+			else
+			{
+				// Copy over position keys.
+				memcpy(TranslationTrack.PosKeys.GetData(), SeqTranslationTrack.PosKeys.GetData() + Segment.StartFrame, NumPosSegmentKeys * sizeof(FVector));
+
+				// Set times for the translation track.
+				const float PosFrameInterval = AnimSeq.SequenceLength / static_cast<float>(NumPosKeys - 1);
+				for (int32 SeqPosIndex = Segment.StartFrame, SegmentPosIndex = 0; SegmentPosIndex < NumPosSegmentKeys; ++SeqPosIndex, ++SegmentPosIndex)
+				{
+					TranslationTrack.Times[SegmentPosIndex] = SeqPosIndex * PosFrameInterval;
+				}
+			}
+
+			RotationTrack.RotKeys.Empty(NumRotSegmentKeys);
+			RotationTrack.RotKeys.AddUninitialized(NumRotSegmentKeys);
+			RotationTrack.Times.Empty(NumRotSegmentKeys);
+			RotationTrack.Times.AddUninitialized(NumRotSegmentKeys);
+
+			if (NumRotKeys == 1)
+			{
+				RotationTrack.RotKeys[0] = SeqRotationTrack.RotKeys[0];
+				RotationTrack.Times[0] = 0.f;
+			}
+			else
+			{
+				// Copy over rotation keys.
+				memcpy(RotationTrack.RotKeys.GetData(), SeqRotationTrack.RotKeys.GetData() + Segment.StartFrame, NumRotSegmentKeys * sizeof(FQuat));
+
+				// Set times for the rotation track.
+				const float RotFrameInterval = AnimSeq.SequenceLength / static_cast<float>(NumRotKeys - 1);
+				for (int32 SeqRotIndex = Segment.StartFrame, SegmentRotIndex = 0; SegmentRotIndex < NumRotSegmentKeys; ++SeqRotIndex, ++SegmentRotIndex)
+				{
+					RotationTrack.Times[SegmentRotIndex] = SeqRotIndex * RotFrameInterval;
+				}
+			}
+
+			if (ScaleData.Num() != 0)
+			{
+				bCompressScaleKeys = true;
+
+				const FScaleTrack&	SeqScaleTrack = ScaleData[TrackIndex];
+				FScaleTrack&		ScaleTrack = Segment.ScaleData[TrackIndex];
+
+				const int32 NumScaleKeys = SeqScaleTrack.ScaleKeys.Num();
+				const int32 NumScaleSegmentKeys = FMath::Min(NumScaleKeys, Segment.NumFrames);
+
+				ScaleTrack.ScaleKeys.Empty(NumScaleSegmentKeys);
+				ScaleTrack.ScaleKeys.AddUninitialized(NumScaleSegmentKeys);
+				ScaleTrack.Times.Empty(NumScaleSegmentKeys);
+				ScaleTrack.Times.AddUninitialized(NumScaleSegmentKeys);
+
+				if (NumScaleKeys == 1)
+				{
+					ScaleTrack.ScaleKeys[0] = SeqScaleTrack.ScaleKeys[0];
+					ScaleTrack.Times[0] = 0.f;
+				}
+				else
+				{
+					// Copy over scale keys.
+					memcpy(ScaleTrack.ScaleKeys.GetData(), SeqScaleTrack.ScaleKeys.GetData() + Segment.StartFrame, NumScaleSegmentKeys * sizeof(FVector));
+
+					// Set times for the scale track.
+					const float ScaleFrameInterval = AnimSeq.SequenceLength / static_cast<float>(NumScaleKeys - 1);
+					for (int32 SeqScaleIndex = Segment.StartFrame, SegmentScaleIndex = 0; SegmentScaleIndex < NumScaleSegmentKeys; ++SeqScaleIndex, ++SegmentScaleIndex)
+					{
+						ScaleTrack.Times[SegmentScaleIndex] = SeqScaleIndex * ScaleFrameInterval;
+					}
+				}
+			}
+		}
+
+		// if nothing to compress, empty the ScaleData
+		// that way we don't have to worry about compressing scale data. 
+		if (!bCompressScaleKeys)
+		{
+			Segment.ScaleData.Empty();
+		}
+	}
+}
diff --git a/Engine/Source/Runtime/Engine/Private/Animation/AnimCompress_Automatic.cpp b/Engine/Source/Runtime/Engine/Private/Animation/AnimCompress_Automatic.cpp
--- a/Engine/Source/Runtime/Engine/Private/Animation/AnimCompress_Automatic.cpp
+++ b/Engine/Source/Runtime/Engine/Private/Animation/AnimCompress_Automatic.cpp
@@ -10,13 +10,10 @@ UAnimCompress_Automatic::UAnimCompress_Automatic(const FObjectInitializer& Objec
 	Description = TEXT("Automatic");
 	UAnimationSettings* AnimationSettings = UAnimationSettings::Get();
 	MaxEndEffectorError = AnimationSettings->AlternativeCompressionThreshold;
-	bTryFixedBitwiseCompression = AnimationSettings->bTryFixedBitwiseCompression;
-	bTryPerTrackBitwiseCompression = AnimationSettings->bTryPerTrackBitwiseCompression;
-	bTryLinearKeyRemovalCompression = AnimationSettings->bTryLinearKeyRemovalCompression;
-	bTryIntervalKeyRemoval = AnimationSettings->bTryIntervalKeyRemoval;
 	bRunCurrentDefaultCompressor = AnimationSettings->bFirstRecompressUsingCurrentOrDefault;
 	bAutoReplaceIfExistingErrorTooGreat = AnimationSettings->bForceBelowThreshold;
 	bRaiseMaxErrorToExisting = AnimationSettings->bRaiseMaxErrorToExisting;
+	bTryExhaustiveSearch = AnimationSettings->bTryExhaustiveSearch;
 }
 
 #if WITH_EDITOR
@@ -31,11 +28,11 @@ void UAnimCompress_Automatic::DoReduction(UAnimSequence* AnimSeq, const TArray<F
 		bRunCurrentDefaultCompressor,
 		bAutoReplaceIfExistingErrorTooGreat,
 		bRaiseMaxErrorToExisting,
-		bTryFixedBitwiseCompression,
-		bTryPerTrackBitwiseCompression,
-		bTryLinearKeyRemovalCompression,
-		bTryIntervalKeyRemoval);
-	AnimSeq->CompressionScheme = static_cast<UAnimCompress*>( StaticDuplicateObject( AnimSeq->CompressionScheme, AnimSeq) );
+		bTryExhaustiveSearch,
+		bEnableSegmenting,
+		IdealNumFramesPerSegment,
+		MaxNumFramesPerSegment,
+		BoneData);
 #endif // WITH_EDITORONLY_DATA
 }
 
@@ -45,13 +42,10 @@ void UAnimCompress_Automatic::PopulateDDCKey(FArchive& Ar)
 
 	Ar << MaxEndEffectorError;
 
-	uint8 Flags =	MakeBitForFlag(bTryFixedBitwiseCompression, 0) +
-					MakeBitForFlag(bTryPerTrackBitwiseCompression, 1) +
-					MakeBitForFlag(bTryLinearKeyRemovalCompression, 2) +
-					MakeBitForFlag(bTryIntervalKeyRemoval, 3) +
-					MakeBitForFlag(bRunCurrentDefaultCompressor, 4) +
-					MakeBitForFlag(bAutoReplaceIfExistingErrorTooGreat, 5) +
-					MakeBitForFlag(bRaiseMaxErrorToExisting, 6);
+	uint8 Flags =	MakeBitForFlag(bRunCurrentDefaultCompressor, 0) +
+					MakeBitForFlag(bAutoReplaceIfExistingErrorTooGreat, 1) +
+					MakeBitForFlag(bRaiseMaxErrorToExisting, 2) +
+					MakeBitForFlag(bTryExhaustiveSearch, 3);
 	Ar << Flags;
 }
 
diff --git a/Engine/Source/Runtime/Engine/Private/Animation/AnimCompress_BitwiseCompressOnly.cpp b/Engine/Source/Runtime/Engine/Private/Animation/AnimCompress_BitwiseCompressOnly.cpp
--- a/Engine/Source/Runtime/Engine/Private/Animation/AnimCompress_BitwiseCompressOnly.cpp
+++ b/Engine/Source/Runtime/Engine/Private/Animation/AnimCompress_BitwiseCompressOnly.cpp
@@ -24,26 +24,51 @@ void UAnimCompress_BitwiseCompressOnly::DoReduction(UAnimSequence* AnimSeq, cons
 	TArray<FScaleTrack> ScaleData;
 	SeparateRawDataIntoTracks( AnimSeq->GetRawAnimationData(), AnimSeq->SequenceLength, TranslationData, RotationData, ScaleData );
 
-	// Remove Translation Keys from tracks marked bAnimRotationOnly
-	FilterAnimRotationOnlyKeys(TranslationData, AnimSeq);
-
 	// remove obviously redundant keys from the source data
 	FilterTrivialKeys(TranslationData, RotationData, ScaleData, TRANSLATION_ZEROING_THRESHOLD, QUATERNION_ZEROING_THRESHOLD, SCALE_ZEROING_THRESHOLD);
 
-	// bitwise compress the tracks into the anim sequence buffers
-	BitwiseCompressAnimationTracks(
-		AnimSeq,
-		static_cast<AnimationCompressionFormat>(TranslationCompressionFormat),
-		static_cast<AnimationCompressionFormat>(RotationCompressionFormat),
-		static_cast<AnimationCompressionFormat>(ScaleCompressionFormat),
-		TranslationData,
-		RotationData, 
-		ScaleData);
-
 	// record the proper runtime decompressor to use
 	AnimSeq->KeyEncodingFormat = AKF_ConstantKeyLerp;
+	AnimSeq->RotationCompressionFormat = RotationCompressionFormat;
+	AnimSeq->TranslationCompressionFormat = TranslationCompressionFormat;
+	AnimSeq->ScaleCompressionFormat = ScaleCompressionFormat;
+	AnimationFormat_SetInterfaceLinks(*AnimSeq);
+
+#if USE_SEGMENTING_CONTEXT
+	if (bEnableSegmenting)
+	{
+		TArray<FAnimSegmentContext> RawSegments;
+		SeparateRawDataIntoTracks(*AnimSeq, TranslationData, RotationData, ScaleData, IdealNumFramesPerSegment, MaxNumFramesPerSegment, RawSegments);
+
+		BitwiseCompressAnimationTracks(
+			*AnimSeq,
+			static_cast<AnimationCompressionFormat>(TranslationCompressionFormat),
+			static_cast<AnimationCompressionFormat>(RotationCompressionFormat),
+			static_cast<AnimationCompressionFormat>(ScaleCompressionFormat),
+			RawSegments);
+
+		CoalesceCompressedSegments(*AnimSeq, RawSegments);
+
+		AnimSeq->TranslationCompressionFormat = TranslationCompressionFormat;
+		AnimSeq->RotationCompressionFormat = RotationCompressionFormat;
+		AnimSeq->ScaleCompressionFormat = ScaleCompressionFormat;
+	}
+	else
+#endif
+	{
+		// bitwise compress the tracks into the anim sequence buffers
+		BitwiseCompressAnimationTracks(
+			AnimSeq,
+			static_cast<AnimationCompressionFormat>(TranslationCompressionFormat),
+			static_cast<AnimationCompressionFormat>(RotationCompressionFormat),
+			static_cast<AnimationCompressionFormat>(ScaleCompressionFormat),
+			TranslationData,
+			RotationData,
+			ScaleData);
+	}
+
+	// We could be invalid, set the links again
 	AnimationFormat_SetInterfaceLinks(*AnimSeq);
-	AnimSeq->CompressionScheme = static_cast<UAnimCompress*>( StaticDuplicateObject( this, AnimSeq ) );
 #endif // WITH_EDITORONLY_DATA
 }
 
diff --git a/Engine/Source/Runtime/Engine/Private/Animation/AnimCompress_LeastDestructive.cpp b/Engine/Source/Runtime/Engine/Private/Animation/AnimCompress_LeastDestructive.cpp
--- a/Engine/Source/Runtime/Engine/Private/Animation/AnimCompress_LeastDestructive.cpp
+++ b/Engine/Source/Runtime/Engine/Private/Animation/AnimCompress_LeastDestructive.cpp
@@ -22,6 +22,8 @@ void UAnimCompress_LeastDestructive::DoReduction(UAnimSequence* AnimSeq, const T
 	UAnimCompress* BitwiseCompressor = NewObject<UAnimCompress_BitwiseCompressOnly>();
 	BitwiseCompressor->RotationCompressionFormat = ACF_Float96NoW;
 	BitwiseCompressor->TranslationCompressionFormat = ACF_None;
-	BitwiseCompressor->Reduce(AnimSeq, false);
+	BitwiseCompressor->Reduce(AnimSeq, false, BoneData);
+
+	AnimSeq->CompressionScheme = BitwiseCompressor;
 }
 #endif // WITH_EDITOR
