Copying //UE4/Dev-AnimPhys to Dev-Main (//UE4/Dev-Main)

diff --git a/Engine/Source/Runtime/Engine/Private/Animation/AnimCompress_PerTrackCompression.cpp b/Engine/Source/Runtime/Engine/Private/Animation/AnimCompress_PerTrackCompression.cpp
--- a/Engine/Source/Runtime/Engine/Private/Animation/AnimCompress_PerTrackCompression.cpp
+++ b/Engine/Source/Runtime/Engine/Private/Animation/AnimCompress_PerTrackCompression.cpp
@@ -7,6 +7,7 @@
 #include "Animation/AnimCompress_PerTrackCompression.h"
 #include "AnimationCompression.h"
 #include "AnimEncoding.h"
+#include "AnimEncoding_PerTrackCompression.h"
 
 struct FAnimSetMeshLinkup;
 
@@ -47,11 +48,13 @@ public:
 
 	// Results of compression
 	TArray<uint8> CompressedBytes;
-	int32 ActualCompressionMode;
-protected:
+	AnimationCompressionFormat ActualCompressionMode;
+	int32 ActualKeyFlags;
+
 	/** Does the compression scheme need a key->frame table (needed if the keys are spaced non-uniformly in time) */
 	bool bReallyNeedsFrameTable;
 
+protected:
 	/** Resets the compression buffer to defaults (no data) */
 	void Reset()
 	{
@@ -59,6 +62,7 @@ protected:
 		SumError = 0.0;
 		bReallyNeedsFrameTable = false;
 		ActualCompressionMode = ACF_None;
+		ActualKeyFlags = 0;
 		CompressedBytes.Empty();
 	}
 
@@ -71,10 +75,11 @@ protected:
 	 *
 	 *   Also updates the ActualCompressionMode field
 	 */
-	int32 MakeHeader(const int32 NumKeys, const int32 KeyFormat, const int32 KeyFlags)
+	int32 MakeHeader(const int32 NumKeys, const AnimationCompressionFormat KeyFormat, const int32 KeyFlags)
 	{
 		ActualCompressionMode = KeyFormat;
-		return FAnimationCompression_PerTrackUtils::MakeHeader(NumKeys, KeyFormat, KeyFlags, bReallyNeedsFrameTable);
+		ActualKeyFlags = KeyFlags;
+		return FAnimationCompression_PerTrackUtils::MakeHeader(NumKeys, (int32)KeyFormat, KeyFlags, bReallyNeedsFrameTable);
 	}
 
 	/** Ensures that the CompressedBytes output stream is a multiple of 4 bytes long */
@@ -113,9 +118,6 @@ protected:
 
 	void CompressTranslation_16_16_16(const FTranslationTrack& TranslationData, float ZeroingThreshold)
 	{
-		//@TODO: Explore different scales
-		const int32 LogScale = 7;
-
 		const int32 NumKeys = TranslationData.PosKeys.Num();
 
 		// Determine the bounds
@@ -598,9 +600,6 @@ protected:
 
 	void CompressScale_16_16_16(const FScaleTrack& ScaleData, float ZeroingThreshold)
 	{
-		//@TODO: Explore different scales
-		const int32 LogScale = 7;
-
 		const int32 NumKeys = ScaleData.ScaleKeys.Num();
 
 		// Determine the bounds
@@ -791,7 +790,7 @@ protected:
 		for (int32 KeyIndex = 0; KeyIndex < NumKeys; ++KeyIndex)
 		{
 			// Convert the frame time into a frame index and write it out
-			FrameIndexType FrameIndex = (FrameIndexType)FMath::Clamp(FMath::TruncToInt(Times[KeyIndex] * FramesPerSecond), 0, NumFrames - 1);
+			FrameIndexType FrameIndex = (FrameIndexType)FMath::Clamp(FMath::TruncToInt((Times[KeyIndex] * FramesPerSecond) + 0.5f), 0, NumFrames - 1);
 			AppendBytes(&FrameIndex, sizeof(FrameIndexType));
 		}
 
@@ -805,7 +804,7 @@ protected:
 		{
 			const int32 NumFrames = Params.AnimSeq->NumFrames;
 			const float SequenceLength = Params.AnimSeq->SequenceLength;
-			const float FramesPerSecond = NumFrames / SequenceLength;
+			const float FramesPerSecond = (NumFrames - 1) / SequenceLength;
 
 			if (NumFrames <= 0xFF)
 			{
@@ -925,23 +924,6 @@ public:
 	}
 };
 
-UAnimCompress_RemoveLinearKeys::UAnimCompress_RemoveLinearKeys(const FObjectInitializer& ObjectInitializer)
-	: Super(ObjectInitializer)
-{
-	bNeedsSkeleton = true;
-	Description = TEXT("Remove Linear Keys");
-	MaxPosDiff = 0.001f;
-	MaxAngleDiff = 0.00075f;
-	MaxScaleDiff = 0.000001;
-	MaxEffectorDiff = 0.001f;
-	MinEffectorDiff = 0.001f;
-	EffectorDiffSocket = 0.001f;
-	ParentKeyScale = 2.0f;
-	bRetarget = true;
-	bActuallyFilterLinearKeys = true;
-}
-
-
 UAnimCompress_PerTrackCompression::UAnimCompress_PerTrackCompression(const FObjectInitializer& ObjectInitializer)
 	: Super(ObjectInitializer)
 {
@@ -1216,6 +1198,688 @@ void UAnimCompress_PerTrackCompression::CompressUsingUnderlyingCompressor(
 	}
 }
 
+void UAnimCompress_PerTrackCompression::PackTranslationKey(TArray<uint8>& ByteStream, AnimationCompressionFormat Format, const FVector& Key, const float* Mins, const float* Ranges, const FPerTrackFormat& TrackFormat)
+{
+	const bool bHasX = TrackFormat.TranslationKeyFlags.IsComponentNeededX();
+	const bool bHasY = TrackFormat.TranslationKeyFlags.IsComponentNeededY();
+	const bool bHasZ = TrackFormat.TranslationKeyFlags.IsComponentNeededZ();
+
+	if (!bHasX && !bHasY && !bHasZ)
+	{
+		// No point in using this over the identity encoding
+		return;
+	}
+
+	switch (Format)
+	{
+	case ACF_Identity:
+		// Nothing to pack
+		break;
+	case ACF_None:
+	case ACF_Float96NoW:
+		if (bHasX)
+		{
+			UnalignedWriteToStream(ByteStream, &Key.X, sizeof(float));
+		}
+		if (bHasY)
+		{
+			UnalignedWriteToStream(ByteStream, &Key.Y, sizeof(float));
+		}
+		if (bHasZ)
+		{
+			UnalignedWriteToStream(ByteStream, &Key.Z, sizeof(float));
+		}
+		break;
+	case ACF_Fixed48NoW:
+		if (bHasX)
+		{
+			const uint16 X = FAnimationCompression_PerTrackUtils::CompressFixed16(Key.X, LogScale);
+			UnalignedWriteToStream(ByteStream, &X, sizeof(uint16));
+		}
+		if (bHasY)
+		{
+			const uint16 Y = FAnimationCompression_PerTrackUtils::CompressFixed16(Key.Y, LogScale);
+			UnalignedWriteToStream(ByteStream, &Y, sizeof(uint16));
+		}
+		if (bHasZ)
+		{
+			const uint16 Z = FAnimationCompression_PerTrackUtils::CompressFixed16(Key.Z, LogScale);
+			UnalignedWriteToStream(ByteStream, &Z, sizeof(uint16));
+		}
+		break;
+	case ACF_IntervalFixed32NoW:
+		{
+			const float MaskedMins[3] = { bHasX ? Mins[0] : 0.0f, bHasY ? Mins[1] : 0.0f, bHasZ ? Mins[2] : 0.0f };
+			const float MaskedRanges[3] = { bHasX ? Ranges[0] : 0.0f, bHasY ? Ranges[1] : 0.0f, bHasZ ? Ranges[2] : 0.0f };
+			const FVectorIntervalFixed32NoW Compressor(Key, MaskedMins, MaskedRanges);
+			UnalignedWriteToStream(ByteStream, &Compressor, sizeof(FVectorIntervalFixed32NoW));
+		}
+		break;
+		// The following two formats don't work well for translation (fixed range & low precision)
+		//case ACF_Fixed32NoW:
+		//case ACF_Float32NoW:
+	default:
+		UE_LOG(LogAnimationCompression, Fatal, TEXT("Unsupported translation compression format"));
+		break;
+	}
+}
+
+void UAnimCompress_PerTrackCompression::PackRotationKey(TArray<uint8>& ByteStream, AnimationCompressionFormat Format, const FQuat& Key, const float* Mins, const float* Ranges, const FPerTrackFormat& TrackFormat)
+{
+	const bool bHasX = TrackFormat.RotationKeyFlags.IsComponentNeededX();
+	const bool bHasY = TrackFormat.RotationKeyFlags.IsComponentNeededY();
+	const bool bHasZ = TrackFormat.RotationKeyFlags.IsComponentNeededZ();
+
+	if (!bHasX && !bHasY && !bHasZ && Format != ACF_Float96NoW && Format != ACF_Fixed32NoW && Format != ACF_Float32NoW)
+	{
+		// No point in using this over the identity encoding
+		return;
+	}
+
+	switch (Format)
+	{
+	case ACF_Identity:
+		// Nothing to pack
+		break;
+	case ACF_None:
+	case ACF_Float96NoW:
+		{
+			const FQuatFloat96NoW Compressor(Key);
+			UnalignedWriteToStream(ByteStream, &Compressor, sizeof(FQuatFloat96NoW));
+		}
+		break;
+	case ACF_Fixed48NoW:
+		{
+			FQuat MaskedKey(bHasX ? Key.X : 0.0f, bHasY ? Key.Y : 0.0f, bHasZ ? Key.Z : 0.0f, Key.W);
+			MaskedKey.Normalize();
+
+			const FQuatFloat96NoW Compressor(MaskedKey);
+
+			if (bHasX)
+			{
+				const uint16 X = FAnimationCompression_PerTrackUtils::CompressFixed16(Compressor.X);
+				UnalignedWriteToStream(ByteStream, &X, sizeof(uint16));
+			}
+			if (bHasY)
+			{
+				const uint16 Y = FAnimationCompression_PerTrackUtils::CompressFixed16(Compressor.Y);
+				UnalignedWriteToStream(ByteStream, &Y, sizeof(uint16));
+			}
+			if (bHasZ)
+			{
+				const uint16 Z = FAnimationCompression_PerTrackUtils::CompressFixed16(Compressor.Z);
+				UnalignedWriteToStream(ByteStream, &Z, sizeof(uint16));
+			}
+		}
+		break;
+	case ACF_IntervalFixed32NoW:
+		{
+			const float MaskedMins[3] = { bHasX ? Mins[0] : 0.0f, bHasY ? Mins[1] : 0.0f, bHasZ ? Mins[2] : 0.0f };
+			const float MaskedRanges[3] = { bHasX ? Ranges[0] : 0.0f, bHasY ? Ranges[1] : 0.0f, bHasZ ? Ranges[2] : 0.0f };
+			FQuat MaskedKey(bHasX ? Key.X : 0.0f, bHasY ? Key.Y : 0.0f, bHasZ ? Key.Z : 0.0f, Key.W);
+			MaskedKey.Normalize();
+
+			const FQuatIntervalFixed32NoW Compressor(MaskedKey, MaskedMins, MaskedRanges);
+			UnalignedWriteToStream(ByteStream, &Compressor, sizeof(Compressor));
+		}
+		break;
+	case ACF_Fixed32NoW:
+		{
+			const FQuatFixed32NoW Compressor(Key);
+			UnalignedWriteToStream(ByteStream, &Compressor, sizeof(FQuatFixed32NoW));
+		}
+		break;
+	case ACF_Float32NoW:
+		{
+			const FQuatFloat32NoW Compressor(Key);
+			UnalignedWriteToStream(ByteStream, &Compressor, sizeof(FQuatFloat32NoW));
+		}
+		break;
+	default:
+		UE_LOG(LogAnimationCompression, Fatal, TEXT("Unsupported rotation compression format"));
+		break;
+	}
+}
+
+void UAnimCompress_PerTrackCompression::PackScaleKey(TArray<uint8>& ByteStream, AnimationCompressionFormat Format, const FVector& Key, const float* Mins, const float* Ranges, const FPerTrackFormat& TrackFormat)
+{
+	const bool bHasX = TrackFormat.ScaleKeyFlags.IsComponentNeededX();
+	const bool bHasY = TrackFormat.ScaleKeyFlags.IsComponentNeededY();
+	const bool bHasZ = TrackFormat.ScaleKeyFlags.IsComponentNeededZ();
+
+	if (!bHasX && !bHasY && !bHasZ)
+	{
+		// No point in using this over the identity encoding
+		return;
+	}
+
+	switch (Format)
+	{
+	case ACF_Identity:
+		// Nothing to pack
+		break;
+	case ACF_None:
+	case ACF_Float96NoW:
+		if (bHasX)
+		{
+			UnalignedWriteToStream(ByteStream, &Key.X, sizeof(float));
+		}
+		if (bHasY)
+		{
+			UnalignedWriteToStream(ByteStream, &Key.Y, sizeof(float));
+		}
+		if (bHasZ)
+		{
+			UnalignedWriteToStream(ByteStream, &Key.Z, sizeof(float));
+		}
+		break;
+	case ACF_Fixed48NoW:
+		if (bHasX)
+		{
+			const uint16 X = FAnimationCompression_PerTrackUtils::CompressFixed16(Key.X, LogScale);
+			UnalignedWriteToStream(ByteStream, &X, sizeof(uint16));
+		}
+		if (bHasY)
+		{
+			const uint16 Y = FAnimationCompression_PerTrackUtils::CompressFixed16(Key.Y, LogScale);
+			UnalignedWriteToStream(ByteStream, &Y, sizeof(uint16));
+		}
+		if (bHasZ)
+		{
+			const uint16 Z = FAnimationCompression_PerTrackUtils::CompressFixed16(Key.Z, LogScale);
+			UnalignedWriteToStream(ByteStream, &Z, sizeof(uint16));
+		}
+		break;
+	case ACF_IntervalFixed32NoW:
+		{
+			const float MaskedMins[3] = { bHasX ? Mins[0] : 0.0f, bHasY ? Mins[1] : 0.0f, bHasZ ? Mins[2] : 0.0f };
+			const float MaskedRanges[3] = { bHasX ? Ranges[0] : 0.0f, bHasY ? Ranges[1] : 0.0f, bHasZ ? Ranges[2] : 0.0f };
+			const FVectorIntervalFixed32NoW Compressor(Key, MaskedMins, MaskedRanges);
+			UnalignedWriteToStream(ByteStream, &Compressor, sizeof(FVectorIntervalFixed32NoW));
+		}
+	break;
+	// The following two formats don't work well for scale (fixed range & low precision)
+	//case ACF_Fixed32NoW:
+	//case ACF_Float32NoW:
+	default:
+		UE_LOG(LogAnimationCompression, Fatal, TEXT("Unsupported scale compression format"));
+		break;
+	}
+}
+
+/**
+ * Structure that holds the necessary information for performing the per track compression.
+ * Each segment will have its own instance. Each instance is independent.
+ * This makes it safe to compress multiple segments in parallel.
+ */
+struct FOptimizeSegmentTracksContext
+{
+	const UAnimSequence& AnimSeq;
+	FAnimSegmentContext& Segment;
+
+	FOptimizeSegmentTracksContext(
+		const UAnimSequence& AnimSeq_,
+		FAnimSegmentContext& Segment_)
+		: AnimSeq(AnimSeq_)
+		, Segment(Segment_)
+	{}
+};
+
+/**
+ * Struct that holds the relevant information to optimize segment tracks in parallel.
+ * Instances of this structure are live as long as parallel task instances are live.
+ */
+struct FAsyncOptimizeSegmentTracksTaskGroupContext
+{
+	TArray<FOptimizeSegmentTracksContext*> TaskContexes;
+	volatile int32 AtomicTaskIndexCounter;
+	volatile int32 AtomicNumExecutedTasks;
+
+	FAsyncOptimizeSegmentTracksTaskGroupContext()
+		: AtomicTaskIndexCounter(0)
+		, AtomicNumExecutedTasks(0)
+	{}
+
+	~FAsyncOptimizeSegmentTracksTaskGroupContext()
+	{
+		for (FOptimizeSegmentTracksContext* Context : TaskContexes)
+		{
+			delete Context;
+		}
+	}
+
+	void ExecuteTasks(UAnimCompress_PerTrackCompression* Compressor)
+	{
+		while (true)
+		{
+			const int32 TaskIndex = FPlatformAtomics::InterlockedIncrement(&AtomicTaskIndexCounter) - 1;
+			if (TaskIndex >= TaskContexes.Num())
+			{
+				break;
+			}
+
+			FOptimizeSegmentTracksContext* JobContext = TaskContexes[TaskIndex];
+			Compressor->OptimizeSegmentTracks(*JobContext);
+
+			FPlatformAtomics::InterlockedIncrement(&AtomicNumExecutedTasks);
+		}
+	}
+
+	void WaitForAllTasks()
+	{
+		// We just spin wait until everything is done
+		// This is a decent option because segments are already sorted largest to smallest and so
+		// they should all take about the same amount of time. We should never end up waiting here for too long.
+		while (AtomicNumExecutedTasks != TaskContexes.Num());
+	}
+};
+
+/**
+ * Class that represents a parallel task to perform per track compression.
+ */
+class FAsyncOptimizeSegmentTracksTask
+{
+public:
+	FAsyncOptimizeSegmentTracksTask(UAnimCompress_PerTrackCompression* Compressor_, FAsyncOptimizeSegmentTracksTaskGroupContext* TaskGroupContext_)
+		: Compressor(Compressor_)
+		, TaskGroupContext(TaskGroupContext_)
+	{}
+
+	/** return the name of the task **/
+	static const TCHAR* GetTaskName() { return TEXT("FAsyncOptimizeSegmentTracksTask"); }
+	FORCEINLINE static TStatId GetStatId() { RETURN_QUICK_DECLARE_CYCLE_STAT(FAsyncOptimizeSegmentTracksTask, STATGROUP_TaskGraphTasks); }
+
+	static ENamedThreads::Type GetDesiredThread() { return ENamedThreads::AnyThread; }
+	static ESubsequentsMode::Type GetSubsequentsMode() { return ESubsequentsMode::TrackSubsequents; }
+
+	void DoTask(ENamedThreads::Type CurrentThread, const FGraphEventRef& MyCompletionGraphEvent)
+	{
+		TaskGroupContext->ExecuteTasks(Compressor);
+	}
+
+	UAnimCompress_PerTrackCompression* Compressor;
+	FAsyncOptimizeSegmentTracksTaskGroupContext* TaskGroupContext;
+};
+
+/**
+ * Class that represents a clean up task once every parallel task has executed.
+ * This is executed last.
+ */
+class FAsyncCleanUpOptimizeSegmentTracksTask
+{
+public:
+	FAsyncCleanUpOptimizeSegmentTracksTask(FAsyncOptimizeSegmentTracksTaskGroupContext* TaskGroupContext_)
+		: TaskGroupContext(TaskGroupContext_)
+	{}
+
+	/** return the name of the task **/
+	static const TCHAR* GetTaskName() { return TEXT("FAsyncCleanUpOptimizeSegmentTracksTask"); }
+	FORCEINLINE static TStatId GetStatId() { RETURN_QUICK_DECLARE_CYCLE_STAT(FAsyncCleanUpOptimizeSegmentTracksTask, STATGROUP_TaskGraphTasks); }
+
+	static ENamedThreads::Type GetDesiredThread() { return ENamedThreads::AnyThread; }
+	static ESubsequentsMode::Type GetSubsequentsMode() { return ESubsequentsMode::TrackSubsequents; }
+
+	void DoTask(ENamedThreads::Type CurrentThread, const FGraphEventRef& MyCompletionGraphEvent)
+	{
+		delete TaskGroupContext;
+	}
+
+	FAsyncOptimizeSegmentTracksTaskGroupContext* TaskGroupContext;
+};
+
+void UAnimCompress_PerTrackCompression::OptimizeSegmentTracks(FOptimizeSegmentTracksContext& Context) const
+{
+	// Prime the compression buffers
+	check(Context.Segment.TranslationData.Num() == Context.Segment.RotationData.Num());
+	const int32 NumTracks = Context.Segment.TranslationData.Num();
+	const bool bHasScale = Context.Segment.ScaleData.Num() > 0;
+
+	TArray<FPerTrackFormat> BestTrackFormats;
+	BestTrackFormats.Empty(NumTracks);
+	BestTrackFormats.AddUninitialized(NumTracks);
+
+	check(PerReductionCachedData != nullptr);
+	const FPerTrackCachedInfo* Cache = PerReductionCachedData;
+
+	// Compress each track independently
+	for (int32 TrackIndex = 0; TrackIndex < NumTracks; ++TrackIndex)
+	{
+		// Compression parameters / thresholds
+		FPerTrackParams Params;
+		Params.AnimSeq = &Context.AnimSeq;
+		Params.MaxZeroingThreshold = MaxZeroingThreshold;
+
+		FPerTrackFormat& TrackFormats = BestTrackFormats[TrackIndex];
+
+		// Determine the local-space error cutoffs
+		float MaxPositionErrorCutoff = MaxPosDiffBitwise;
+		float MaxAngleErrorCutoff = MaxAngleDiffBitwise;
+		float MaxScaleErrorCutoff = MaxScaleDiffBitwise;
+
+		if (bUseAdaptiveError)
+		{
+			// The height of the track is the distance from an end effector.  It's used to reduce the acceptable error the
+			// higher in the skeleton we get, since a higher bone will cause cascading errors everywhere.
+			const int32 PureTrackHeight = Cache->TrackHeights[TrackIndex];
+			const int32 EffectiveTrackHeight = FMath::Max(0, PureTrackHeight + TrackHeightBias);
+
+			const float Scaler = 1.0f / FMath::Pow(FMath::Max(ParentingDivisor, 1.0f), EffectiveTrackHeight * FMath::Max(0.0f, ParentingDivisorExponent));
+
+			MaxPositionErrorCutoff = FMath::Max<float>(MaxZeroingThreshold, MaxPosDiff * Scaler);
+			MaxAngleErrorCutoff = FMath::Max<float>(MaxZeroingThreshold, MaxAngleDiff * Scaler);
+			MaxScaleErrorCutoff = FMath::Max<float>(MaxZeroingThreshold, MaxScaleDiff * Scaler);
+
+			if (bUseOverrideForEndEffectors && (PureTrackHeight == 0))
+			{
+				MaxPositionErrorCutoff = MinEffectorDiff;
+			}
+		}
+		else if (bUseAdaptiveError2)
+		{
+			const FAnimPerturbationError& TrackError = Cache->PerTrackErrors[TrackIndex];
+
+			float ThresholdT_DueR = (TrackError.MaxErrorInTransDueToRot > SMALL_NUMBER) ? (PerturbationProbeSize / TrackError.MaxErrorInTransDueToRot) : 1.0f;
+			float ThresholdT_DueT = (TrackError.MaxErrorInTransDueToTrans > SMALL_NUMBER) ? (PerturbationProbeSize / TrackError.MaxErrorInTransDueToTrans) : 1.0f;
+			float ThresholdT_DueS = (TrackError.MaxErrorInTransDueToScale > SMALL_NUMBER) ? (PerturbationProbeSize / TrackError.MaxErrorInTransDueToScale) : 1.0f;
+
+			//@TODO: Mixing spaces (target angle error is in radians, perturbation is in quaternion component units)
+			float ThresholdR_DueR = (TrackError.MaxErrorInRotDueToRot > SMALL_NUMBER) ? (PerturbationProbeSize / TrackError.MaxErrorInRotDueToRot) : 1.0f;
+			float ThresholdR_DueT = (TrackError.MaxErrorInRotDueToTrans > SMALL_NUMBER) ? (PerturbationProbeSize / TrackError.MaxErrorInRotDueToTrans) : 1.0f;
+			float ThresholdR_DueS = (TrackError.MaxErrorInRotDueToScale > SMALL_NUMBER) ? (PerturbationProbeSize / TrackError.MaxErrorInRotDueToScale) : 1.0f;
+
+			// these values are not used, so I don't think we should calculate?
+			// 			float ThresholdS_DueR = (TrackError.MaxErrorInScaleDueToRot > SMALL_NUMBER) ? (PerturbationProbeSize / TrackError.MaxErrorInScaleDueToRot) : 1.0f;
+			// 			float ThresholdS_DueT = (TrackError.MaxErrorInScaleDueToTrans > SMALL_NUMBER) ? (PerturbationProbeSize / TrackError.MaxErrorInScaleDueToTrans) : 1.0f;
+			// 			float ThresholdS_DueS = (TrackError.MaxErrorInScaleDueToScale > SMALL_NUMBER) ? (PerturbationProbeSize / TrackError.MaxErrorInScaleDueToScale) : 1.0f;
+
+			// @Todo fix the error - this doesn't make sense
+			MaxAngleErrorCutoff = FMath::Min(MaxAngleDiffBitwise, MaxErrorPerTrackRatio * MaxAngleDiff * FMath::Lerp(ThresholdR_DueR, ThresholdT_DueR, RotationErrorSourceRatio));
+			MaxPositionErrorCutoff = FMath::Min(MaxPosDiffBitwise, MaxErrorPerTrackRatio * MaxPosDiff * FMath::Lerp(ThresholdR_DueT, ThresholdT_DueT, TranslationErrorSourceRatio));
+			MaxScaleErrorCutoff = FMath::Min(MaxScaleDiffBitwise, MaxErrorPerTrackRatio * MaxScaleDiff * FMath::Lerp(ThresholdR_DueS, ThresholdT_DueS, ScaleErrorSourceRatio));
+		}
+
+		// Start compressing translation using a totally lossless float32x3
+		const FTranslationTrack& TranslationTrack = Context.Segment.TranslationData[TrackIndex];
+
+		Params.bIncludeKeyTable = bActuallyFilterLinearKeys && TranslationTrack.PosKeys.Num() != Context.Segment.NumFrames;
+		FPerTrackCompressor BestTranslation(ACF_Float96NoW, TranslationTrack, Params);
+
+		// Try the other translation formats
+		for (int32 FormatIndex = 0; FormatIndex < AllowedTranslationFormats.Num(); ++FormatIndex)
+		{
+			FPerTrackCompressor TrialCompression(AllowedTranslationFormats[FormatIndex], TranslationTrack, Params);
+
+			if (TrialCompression.MaxError <= MaxPositionErrorCutoff)
+			{
+				// Swap if it's smaller or equal-sized but lower-max-error
+				const int32 BytesSaved = BestTranslation.CompressedBytes.Num() - TrialCompression.CompressedBytes.Num();
+				const bool bIsImprovement = (BytesSaved > 0) || ((BytesSaved == 0) && (TrialCompression.MaxError < BestTranslation.MaxError));
+
+				if (bIsImprovement)
+				{
+					BestTranslation = TrialCompression;
+				}
+			}
+		}
+		TrackFormats.TranslationFormat = BestTranslation.ActualCompressionMode;
+		TrackFormats.bHasTranslationTimeMarkers = BestTranslation.bReallyNeedsFrameTable;
+		TrackFormats.TranslationKeyFlags = FTrackKeyFlags(BestTranslation.ActualKeyFlags);
+
+		// Start compressing rotation, first using lossless float32x3
+		const FRotationTrack& RotationTrack = Context.Segment.RotationData[TrackIndex];
+
+		Params.bIncludeKeyTable = bActuallyFilterLinearKeys && RotationTrack.RotKeys.Num() != Context.Segment.NumFrames;
+		FPerTrackCompressor BestRotation(ACF_Float96NoW, RotationTrack, Params);
+
+		//bool bLeaveRotationUncompressed = (RotationTrack.Times.Num() <= 1) && (GHighQualityEmptyTracks != 0);
+		// Try the other rotation formats
+		//if (!bLeaveRotationUncompressed)
+		{
+			for (int32 FormatIndex = 0; FormatIndex < AllowedRotationFormats.Num(); ++FormatIndex)
+			{
+				FPerTrackCompressor TrialCompression(AllowedRotationFormats[FormatIndex], RotationTrack, Params);
+
+				if (TrialCompression.MaxError <= MaxAngleErrorCutoff)
+				{
+					// Swap if it's smaller or equal-sized but lower-max-error
+					const int32 BytesSaved = BestRotation.CompressedBytes.Num() - TrialCompression.CompressedBytes.Num();
+					const bool bIsImprovement = (BytesSaved > 0) || ((BytesSaved == 0) && (TrialCompression.MaxError < BestRotation.MaxError));
+
+					if (bIsImprovement)
+					{
+						BestRotation = TrialCompression;
+					}
+				}
+			}
+		}
+		TrackFormats.RotationFormat = BestRotation.ActualCompressionMode;
+		TrackFormats.bHasRotationTimeMarkers = BestRotation.bReallyNeedsFrameTable;
+		TrackFormats.RotationKeyFlags = FTrackKeyFlags(BestRotation.ActualKeyFlags);
+
+		// Start compressing Scale, first using lossless float32x3
+		TrackFormats.ScaleFormat = ACF_None;
+		TrackFormats.bHasScaleTimeMarkers = false;
+		TrackFormats.ScaleKeyFlags = FTrackKeyFlags();
+		if (bHasScale)
+		{
+			const FScaleTrack& ScaleTrack = Context.Segment.ScaleData[TrackIndex];
+
+			Params.bIncludeKeyTable = bActuallyFilterLinearKeys && ScaleTrack.ScaleKeys.Num() != Context.Segment.NumFrames;
+			FPerTrackCompressor BestScale(ACF_Float96NoW, ScaleTrack, Params);
+
+			//bool bLeaveScaleUncompressed = (ScaleTrack.Times.Num() <= 1) && (GHighQualityEmptyTracks != 0);
+			// Try the other Scale formats
+			//if (!bLeaveScaleUncompressed)
+			{
+				for (int32 FormatIndex = 0; FormatIndex < AllowedScaleFormats.Num(); ++FormatIndex)
+				{
+					FPerTrackCompressor TrialCompression(AllowedScaleFormats[FormatIndex], ScaleTrack, Params);
+
+					if (TrialCompression.MaxError <= MaxAngleErrorCutoff)
+					{
+						// Swap if it's smaller or equal-sized but lower-max-error
+						const int32 BytesSaved = BestScale.CompressedBytes.Num() - TrialCompression.CompressedBytes.Num();
+						const bool bIsImprovement = (BytesSaved > 0) || ((BytesSaved == 0) && (TrialCompression.MaxError < BestScale.MaxError));
+
+						if (bIsImprovement)
+						{
+							BestScale = TrialCompression;
+						}
+					}
+				}
+			}
+
+			TrackFormats.ScaleFormat = BestScale.ActualCompressionMode;
+			TrackFormats.bHasScaleTimeMarkers = BestScale.bReallyNeedsFrameTable;
+			TrackFormats.ScaleKeyFlags = FTrackKeyFlags(BestScale.ActualKeyFlags);
+		}
+
+#if 0
+		// This block outputs information about each individual track during compression, which is useful for debugging the compressors
+		UE_LOG(LogAnimationCompression, Warning, TEXT("   Compressed track %i, Trans=%s_%i (#keys=%i, err=%f), Rot=%s_%i (#keys=%i, err=%f)  (height=%i max pos=%f, angle=%f)"),
+			TrackIndex,
+			*FAnimationUtils::GetAnimationCompressionFormatString(static_cast<AnimationCompressionFormat>(BestTranslation.ActualCompressionMode)),
+			BestTranslation.ActualCompressionMode != ACF_Identity ? ((*((const int32*)BestTranslation.CompressedBytes.GetTypedData())) >> 24) & 0xF : 0,
+			TranslationTrack.PosKeys.Num(),
+			BestTranslation.MaxError,
+			*FAnimationUtils::GetAnimationCompressionFormatString(static_cast<AnimationCompressionFormat>(BestRotation.ActualCompressionMode)),
+			BestRotation.ActualCompressionMode != ACF_Identity ? ((*((const int32*)BestRotation.CompressedBytes.GetTypedData())) >> 24) & 0xF : 0,
+			RotationTrack.RotKeys.Num(),
+			BestRotation.MaxError,
+			(bUseAdaptiveError) ? (Cache->TrackHeights(TrackIndex)) : -1,
+			MaxPositionErrorCutoff,
+			MaxAngleErrorCutoff
+		);
+#endif
+	}
+
+	SanityCheckTrackData(Context.AnimSeq, Context.Segment);
+
+	Context.Segment.CompressedByteStream.Empty(64 * 1024);
+
+	for (int32 TrackIndex = 0; TrackIndex < NumTracks; ++TrackIndex)
+	{
+		const FPerTrackFormat& TrackFormats = BestTrackFormats[TrackIndex];
+
+		const FPerTrackFlags TranslationFlags(TrackFormats.bHasTranslationTimeMarkers, TrackFormats.TranslationFormat, TrackFormats.TranslationKeyFlags.Flags);
+		UnalignedWriteToStream(Context.Segment.CompressedByteStream, &TranslationFlags, sizeof(FPerTrackFlags));
+
+		const FPerTrackFlags RotationFlags(TrackFormats.bHasRotationTimeMarkers, TrackFormats.RotationFormat, TrackFormats.RotationKeyFlags.Flags);
+		UnalignedWriteToStream(Context.Segment.CompressedByteStream, &RotationFlags, sizeof(FPerTrackFlags));
+
+		if (bHasScale)
+		{
+			const FPerTrackFlags ScaleFlags(TrackFormats.bHasScaleTimeMarkers, TrackFormats.ScaleFormat, TrackFormats.ScaleKeyFlags.Flags);
+			UnalignedWriteToStream(Context.Segment.CompressedByteStream, &ScaleFlags, sizeof(FPerTrackFlags));
+		}
+	}
+
+	PadByteStream(Context.Segment.CompressedByteStream, 4, AnimationPadSentinel);
+
+	TArray<FAnimTrackRange> TrackRanges;
+	CalculateTrackRanges(ACF_IntervalFixed32NoW, ACF_IntervalFixed32NoW, ACF_IntervalFixed32NoW, Context.Segment, TrackRanges);
+
+	checkf((Context.Segment.CompressedByteStream.Num() % 4) == 0, TEXT("CompressedByteStream not aligned to four bytes"));
+
+	// Write track ranges
+	WriteTrackRanges(
+		Context.Segment.CompressedByteStream,
+		[&](int32 TrackIndex) { return BestTrackFormats[TrackIndex].TranslationFormat; },
+		[&](int32 TrackIndex) { return BestTrackFormats[TrackIndex].RotationFormat; },
+		[&](int32 TrackIndex) { return BestTrackFormats[TrackIndex].ScaleFormat; },
+		[&](int32 TrackIndex) { return BestTrackFormats[TrackIndex].TranslationKeyFlags; },
+		[&](int32 TrackIndex) { return BestTrackFormats[TrackIndex].RotationKeyFlags; },
+		[&](int32 TrackIndex) { return BestTrackFormats[TrackIndex].ScaleKeyFlags; },
+		Context.Segment, TrackRanges, true);
+
+	checkf((Context.Segment.CompressedByteStream.Num() % 4) == 0, TEXT("CompressedByteStream not aligned to four bytes"));
+
+	WriteUniformTrackData(
+		Context.Segment.CompressedByteStream,
+		[&](int32 TrackIndex) { return BestTrackFormats[TrackIndex].TranslationFormat; },
+		[&](int32 TrackIndex) { return BestTrackFormats[TrackIndex].RotationFormat; },
+		[&](int32 TrackIndex) { return BestTrackFormats[TrackIndex].ScaleFormat; },
+		[&](int32 TrackIndex) { return !BestTrackFormats[TrackIndex].bHasTranslationTimeMarkers; },
+		[&](int32 TrackIndex) { return !BestTrackFormats[TrackIndex].bHasRotationTimeMarkers; },
+		[&](int32 TrackIndex) { return !BestTrackFormats[TrackIndex].bHasScaleTimeMarkers; },
+		[&](TArray<uint8>& ByteStream, AnimationCompressionFormat Format, const FVector& Key, const float* Mins, const float* Ranges, int32 TrackIndex) { PackTranslationKey(ByteStream, Format, Key, Mins, Ranges, BestTrackFormats[TrackIndex]); },
+		[&](TArray<uint8>& ByteStream, AnimationCompressionFormat Format, const FQuat& Key, const float* Mins, const float* Ranges, int32 TrackIndex) { PackRotationKey(ByteStream, Format, Key, Mins, Ranges, BestTrackFormats[TrackIndex]); },
+		[&](TArray<uint8>& ByteStream, AnimationCompressionFormat Format, const FVector& Key, const float* Mins, const float* Ranges, int32 TrackIndex) { PackScaleKey(ByteStream, Format, Key, Mins, Ranges, BestTrackFormats[TrackIndex]); },
+		Context.Segment, TrackRanges);
+
+	PadByteStream(Context.Segment.CompressedByteStream, 4, AnimationPadSentinel);
+
+	if (bOptimizeForForwardPlayback)
+	{
+		WriteSortedVariableTrackData(
+			Context.Segment.CompressedByteStream,
+			Context.AnimSeq,
+			[&](int32 TrackIndex) { return BestTrackFormats[TrackIndex].TranslationFormat; },
+			[&](int32 TrackIndex) { return BestTrackFormats[TrackIndex].RotationFormat; },
+			[&](int32 TrackIndex) { return BestTrackFormats[TrackIndex].ScaleFormat; },
+			[&](int32 TrackIndex) { return BestTrackFormats[TrackIndex].bHasTranslationTimeMarkers; },
+			[&](int32 TrackIndex) { return BestTrackFormats[TrackIndex].bHasRotationTimeMarkers; },
+			[&](int32 TrackIndex) { return BestTrackFormats[TrackIndex].bHasScaleTimeMarkers; },
+			[&](TArray<uint8>& ByteStream, AnimationCompressionFormat Format, const FVector& Key, const float* Mins, const float* Ranges, int32 TrackIndex) { PackTranslationKey(ByteStream, Format, Key, Mins, Ranges, BestTrackFormats[TrackIndex]); },
+			[&](TArray<uint8>& ByteStream, AnimationCompressionFormat Format, const FQuat& Key, const float* Mins, const float* Ranges, int32 TrackIndex) { PackRotationKey(ByteStream, Format, Key, Mins, Ranges, BestTrackFormats[TrackIndex]); },
+			[&](TArray<uint8>& ByteStream, AnimationCompressionFormat Format, const FVector& Key, const float* Mins, const float* Ranges, int32 TrackIndex) { PackScaleKey(ByteStream, Format, Key, Mins, Ranges, BestTrackFormats[TrackIndex]); },
+			Context.Segment, TrackRanges);
+	}
+	else
+	{
+		WriteLinearVariableTrackData(
+			Context.Segment.CompressedByteStream,
+			Context.AnimSeq,
+			[&](int32 TrackIndex) { return BestTrackFormats[TrackIndex].TranslationFormat; },
+			[&](int32 TrackIndex) { return BestTrackFormats[TrackIndex].RotationFormat; },
+			[&](int32 TrackIndex) { return BestTrackFormats[TrackIndex].ScaleFormat; },
+			[&](int32 TrackIndex) { return BestTrackFormats[TrackIndex].bHasTranslationTimeMarkers; },
+			[&](int32 TrackIndex) { return BestTrackFormats[TrackIndex].bHasRotationTimeMarkers; },
+			[&](int32 TrackIndex) { return BestTrackFormats[TrackIndex].bHasScaleTimeMarkers; },
+			[&](TArray<uint8>& ByteStream, AnimationCompressionFormat Format, const FVector& Key, const float* Mins, const float* Ranges, int32 TrackIndex) { PackTranslationKey(ByteStream, Format, Key, Mins, Ranges, BestTrackFormats[TrackIndex]); },
+			[&](TArray<uint8>& ByteStream, AnimationCompressionFormat Format, const FQuat& Key, const float* Mins, const float* Ranges, int32 TrackIndex) { PackRotationKey(ByteStream, Format, Key, Mins, Ranges, BestTrackFormats[TrackIndex]); },
+			[&](TArray<uint8>& ByteStream, AnimationCompressionFormat Format, const FVector& Key, const float* Mins, const float* Ranges, int32 TrackIndex) { PackScaleKey(ByteStream, Format, Key, Mins, Ranges, BestTrackFormats[TrackIndex]); },
+			Context.Segment, TrackRanges);
+	}
+
+	// Make sure we have a safe alignment
+	PadByteStream(Context.Segment.CompressedByteStream, 4, AnimationPadSentinel);
+
+	// Trim unused memory.
+	Context.Segment.CompressedByteStream.Shrink();
+}
+
+void UAnimCompress_PerTrackCompression::CompressUsingUnderlyingCompressor(
+	UAnimSequence& AnimSeq,
+	const TArray<FBoneData>& BoneData,
+	TArray<FAnimSegmentContext>& RawSegments,
+	const bool bFinalPass)
+{
+	// If not doing final pass, then do the RemoveLinearKey version that is less destructive.
+	// We're potentially removing whole tracks here, and that doesn't work well with LinearKeyRemoval algorithm.
+	if (!bFinalPass)
+	{
+		UAnimCompress_RemoveLinearKeys::CompressUsingUnderlyingCompressor(
+			AnimSeq,
+			BoneData,
+			RawSegments,
+			bFinalPass);
+		return;
+	}
+
+	// record the proper runtime decompressor to use
+	AnimSeq.KeyEncodingFormat = AKF_PerTrackCompression;
+	AnimSeq.RotationCompressionFormat = ACF_Identity;
+	AnimSeq.TranslationCompressionFormat = ACF_Identity;
+	AnimSeq.ScaleCompressionFormat = ACF_Identity;
+	AnimationFormat_SetInterfaceLinks(AnimSeq);
+
+	if (bUseDecompression || !bUseMultithreading || RawSegments.Num() <= 1)
+	{
+		for (FAnimSegmentContext& Segment : RawSegments)
+		{
+			FOptimizeSegmentTracksContext Context(AnimSeq, Segment);
+			OptimizeSegmentTracks(Context);
+		}
+	}
+	else
+	{
+		// Created the context objects.
+		FAsyncOptimizeSegmentTracksTaskGroupContext* TaskGroupContext = new FAsyncOptimizeSegmentTracksTaskGroupContext();
+		for (FAnimSegmentContext& Segment : RawSegments)
+		{
+			FOptimizeSegmentTracksContext* Context = new FOptimizeSegmentTracksContext(AnimSeq, Segment);
+			TaskGroupContext->TaskContexes.Add(Context);
+		}
+
+		// Dispatch 1 task per job thread.
+		FGraphEventArray AsyncTaskCompletionEvents;
+		const int32 NumTaskThreads = FTaskGraphInterface::Get().GetNumWorkerThreads();
+		for (int32 TaskIndex = 0; TaskIndex < NumTaskThreads; ++TaskIndex)
+		{
+			AsyncTaskCompletionEvents.Add(TGraphTask<FAsyncOptimizeSegmentTracksTask>::CreateTask(NULL, ENamedThreads::AnyThread).ConstructAndDispatchWhenReady(this, TaskGroupContext));
+		}
+
+		// Execute the contexts concurrently.
+		TaskGroupContext->ExecuteTasks(this);
+
+		// Wait for all concurrent tasks to be done, we only wait for ones that were executing.
+		TaskGroupContext->WaitForAllTasks();
+
+		// Dispatch a cleanup job that will execute once all the tasks are done.
+		// The current thread/task no longer needs the TaskGroupContext beyond this point nor anything it references.
+		// All tasks that were executing actual compression work are done or queued up and pending
+		// and will do nothing since no work remains (they only touch the TaskGroupContext to find more work).
+		// Once all tasks are actually done, the clean up task will execute, freeing memory.
+		TGraphTask<FAsyncCleanUpOptimizeSegmentTracksTask>::CreateTask(&AsyncTaskCompletionEvents, ENamedThreads::AnyThread).ConstructAndDispatchWhenReady(TaskGroupContext);
+	}
+
+	// Ensure we compress the trivial tracks into our first segment
+	BitwiseCompressTrivialAnimationTracks(AnimSeq, RawSegments[0]);
+
+	CoalesceCompressedSegments(AnimSeq, RawSegments, bOptimizeForForwardPlayback);
+}
+
 /** Resamples a track of position keys */
 void ResamplePositionKeys(
 	FTranslationTrack& Track, 
diff --git a/Engine/Source/Runtime/Engine/Private/Animation/AnimCompress_RemoveEverySecondKey.cpp b/Engine/Source/Runtime/Engine/Private/Animation/AnimCompress_RemoveEverySecondKey.cpp
--- a/Engine/Source/Runtime/Engine/Private/Animation/AnimCompress_RemoveEverySecondKey.cpp
+++ b/Engine/Source/Runtime/Engine/Private/Animation/AnimCompress_RemoveEverySecondKey.cpp
@@ -28,29 +28,54 @@ void UAnimCompress_RemoveEverySecondKey::DoReduction(UAnimSequence* AnimSeq, con
 	TArray<FScaleTrack> ScaleData;
 	SeparateRawDataIntoTracks( AnimSeq->GetRawAnimationData(), AnimSeq->SequenceLength, TranslationData, RotationData, ScaleData );
 
-	// Remove Translation Keys from tracks marked bAnimRotationOnly
-	FilterAnimRotationOnlyKeys(TranslationData, AnimSeq);
-
 	// remove obviously redundant keys from the source data
 	FilterTrivialKeys(TranslationData, RotationData, ScaleData, TRANSLATION_ZEROING_THRESHOLD, QUATERNION_ZEROING_THRESHOLD, SCALE_ZEROING_THRESHOLD);
 
 	// remove intermittent keys from the source data
 	FilterIntermittentKeys(TranslationData, RotationData, StartIndex, Interval);
 
-	// bitwise compress the tracks into the anim sequence buffers
-	BitwiseCompressAnimationTracks(
-		AnimSeq,
-		static_cast<AnimationCompressionFormat>(TranslationCompressionFormat),
-		static_cast<AnimationCompressionFormat>(RotationCompressionFormat),
-		static_cast<AnimationCompressionFormat>(ScaleCompressionFormat),
-		TranslationData,
-		RotationData, 
-		ScaleData);
-
 	// record the proper runtime decompressor to use
 	AnimSeq->KeyEncodingFormat = AKF_ConstantKeyLerp;
+	AnimSeq->RotationCompressionFormat = RotationCompressionFormat;
+	AnimSeq->TranslationCompressionFormat = TranslationCompressionFormat;
+	AnimSeq->ScaleCompressionFormat = ScaleCompressionFormat;
+	AnimationFormat_SetInterfaceLinks(*AnimSeq);
+
+#if USE_SEGMENTING_CONTEXT
+	if (bEnableSegmenting)
+	{
+		TArray<FAnimSegmentContext> RawSegments;
+		SeparateRawDataIntoTracks(*AnimSeq, TranslationData, RotationData, ScaleData, IdealNumFramesPerSegment, MaxNumFramesPerSegment, RawSegments);
+
+		BitwiseCompressAnimationTracks(
+			*AnimSeq,
+			static_cast<AnimationCompressionFormat>(TranslationCompressionFormat),
+			static_cast<AnimationCompressionFormat>(RotationCompressionFormat),
+			static_cast<AnimationCompressionFormat>(ScaleCompressionFormat),
+			RawSegments);
+
+		CoalesceCompressedSegments(*AnimSeq, RawSegments);
+
+		AnimSeq->TranslationCompressionFormat = TranslationCompressionFormat;
+		AnimSeq->RotationCompressionFormat = RotationCompressionFormat;
+		AnimSeq->ScaleCompressionFormat = ScaleCompressionFormat;
+	}
+	else
+#endif
+	{
+		// bitwise compress the tracks into the anim sequence buffers
+		BitwiseCompressAnimationTracks(
+			AnimSeq,
+			static_cast<AnimationCompressionFormat>(TranslationCompressionFormat),
+			static_cast<AnimationCompressionFormat>(RotationCompressionFormat),
+			static_cast<AnimationCompressionFormat>(ScaleCompressionFormat),
+			TranslationData,
+			RotationData,
+			ScaleData);
+	}
+
+	// We could be invalid, set the links again
 	AnimationFormat_SetInterfaceLinks(*AnimSeq);
-	AnimSeq->CompressionScheme = static_cast<UAnimCompress*>( StaticDuplicateObject( this, AnimSeq ) );
 #endif // WITH_EDITORONLY_DATA
 }
 
