Copying //UE4/Dev-AnimPhys to Dev-Main (//UE4/Dev-Main)

diff --git a/Engine/Source/Runtime/Engine/Private/Animation/AnimCompress_RemoveLinearKeys.cpp b/Engine/Source/Runtime/Engine/Private/Animation/AnimCompress_RemoveLinearKeys.cpp
--- a/Engine/Source/Runtime/Engine/Private/Animation/AnimCompress_RemoveLinearKeys.cpp
+++ b/Engine/Source/Runtime/Engine/Private/Animation/AnimCompress_RemoveLinearKeys.cpp
@@ -75,38 +75,94 @@
 	return FQuat::Error(A, B);
 }
 
-/**
- * Helper function to replace a specific component in the FTransform structure
- */
-template <typename T>
-FTransform UpdateBoneAtom(int32 BoneIndex, const FTransform& Atom, const T& Component)
+UAnimCompress_RemoveLinearKeys::UAnimCompress_RemoveLinearKeys(const FObjectInitializer& ObjectInitializer)
+	: Super(ObjectInitializer)
 {
-	// only the custom instantiations below are valid
-	check(0);
-	return FTransform();
+	bNeedsSkeleton = true;
+	Description = TEXT("Remove Linear Keys");
+	MaxPosDiff = 0.001f;
+	MaxAngleDiff = 0.00075f;
+	MaxScaleDiff = 0.000001;
+	MaxEffectorDiff = 0.001f;
+	MinEffectorDiff = 0.001f;
+	EffectorDiffSocket = 0.001f;
+	ParentKeyScale = 2.0f;
+	bRetarget = true;
+	bActuallyFilterLinearKeys = true;
+	bOptimizeForForwardPlayback = false;
+	bUseDecompression = false;
+	bUseMultithreading = true;
 }
 
-/** custom instantiation of UpdateBoneAtom for FVectors */
-template <> FTransform UpdateBoneAtom<FVector>(int32 BoneIndex, const FTransform& Atom, const FVector& Component)
-{
-	return FTransform(Atom.GetRotation(), Component, FVector(1.0f));
-}
-
-/** custom instantiation of UpdateBoneAtom for FQuat */
-template <> FTransform UpdateBoneAtom<FQuat>(int32 BoneIndex, const FTransform& Atom, const FQuat& Component)
-{
-	return FTransform(Component, Atom.GetTranslation(), FVector(1.0f));
-}
-
 #if WITH_EDITOR
 /**
+ * Structure that holds the necessary information for performing linear key reduction.
+ * Each segment will have its own instance. Each instance is independent.
+ * This makes it safe to compress multiple segments in parallel.
+ */
+struct FProcessAnimationTracksContext
+{
+	UAnimSequence& AnimSeq;
+	const TArray<FBoneData>& BoneData;
+	const TArray<FTransform>& RefPose;
+	const TArray<int32>& EndEffectors;
+	const TArray<int32>& BoneIndexToTrackIndex;
+	const TArray<FTransform>& TrivialTrackValues;
+
+	FAnimSegmentContext& Segment;
+	TArray<FAnimSegmentContext>& SegmentList;
+
+	TArray<FTransform> RawWorldBones;
+	TArray<FTransform> NewWorldBones;
+
+	FProcessAnimationTracksContext(
+		UAnimSequence& AnimSeq_,
+		const TArray<FBoneData>& BoneData_,
+		const TArray<int32>& EndEffectors_,
+		const TArray<int32>& BoneIndexToTrackIndex_,
+		const TArray<FTransform>& TrivialTrackValues_,
+		FAnimSegmentContext& Segment_,
+		TArray<FAnimSegmentContext>& SegmentList_)
+		: AnimSeq(AnimSeq_)
+		, BoneData(BoneData_)
+		, RefPose(AnimSeq_.GetSkeleton()->GetRefLocalPoses())
+		, EndEffectors(EndEffectors_)
+		, BoneIndexToTrackIndex(BoneIndexToTrackIndex_)
+		, TrivialTrackValues(TrivialTrackValues_)
+		, Segment(Segment_)
+		, SegmentList(SegmentList_)
+	{}
+};
+
+struct RotationAdapter
+{
+	typedef FQuat KeyType;
+
+	static FTransform UpdateBoneAtom(const FTransform& Atom, const FQuat& Component) { return FTransform(Component, Atom.GetTranslation(), Atom.GetScale3D()); }
+};
+
+struct TranslationAdapter
+{
+	typedef FVector KeyType;
+
+	static FTransform UpdateBoneAtom(const FTransform& Atom, const FVector& Component) { return FTransform(Atom.GetRotation(), Component, Atom.GetScale3D()); }
+};
+
+struct ScaleAdapter
+{
+	typedef FVector KeyType;
+
+	static FTransform UpdateBoneAtom(const FTransform& Atom, const FVector& Component) { return FTransform(Atom.GetRotation(), Atom.GetTranslation(), Component); }
+};
+
+/**
  * Template function to reduce the keys of a given data type.
- * Used to reduce both Translation and Rotation keys using the corresponding
+ * Used to reduce both Translation, Rotation, and Scale keys using the corresponding
  * data types FVector and FQuat
  */
-template <typename T>
+template <typename AdapterType>
 void FilterLinearKeysTemplate(
-	TArray<T>& Keys,
+	TArray<typename AdapterType::KeyType>& Keys,
 	TArray<float>& Times,
 	TArray<FTransform>& BoneAtoms,
 	const TArray<float>* ParentTimes,
@@ -123,12 +179,14 @@
 	const TArray<FBoneData>& BoneData
 	)
 {
+	typedef typename AdapterType::KeyType KeyType;
+
 	const int32 KeyCount = Keys.Num();
 	check( Keys.Num() == Times.Num() );
 	check( KeyCount >= 1 );
 	
 	// generate new arrays we will fill with the final keys
-	TArray<T> NewKeys;
+	TArray<KeyType> NewKeys;
 	TArray<float> NewTimes;
 	NewKeys.Empty(KeyCount);
 	NewTimes.Empty(KeyCount);
@@ -184,8 +242,8 @@
 				// get the parameters of the window we are testing
 				const float LowTime = Times[LowKey];
 				const float HighTime = Times[HighKey];
-				const T LowValue = Keys[LowKey];
-				const T HighValue = Keys[HighKey];
+				const KeyType& LowValue = Keys[LowKey];
+				const KeyType& HighValue = Keys[HighKey];
 				const float Range = HighTime - LowTime;
 				const float InvRange = 1.0f/Range;
 
@@ -197,11 +255,11 @@
 				{
 					// get the parameters of the member being tested
 					float TestTime = Times[TestKey];
-					T TestValue = Keys[TestKey];
+					const KeyType& TestValue = Keys[TestKey];
 
 					// compute the proposed, interpolated value for the key
 					const float Alpha = (TestTime - LowTime) * InvRange;
-					const T LerpValue = Interpolate(LowValue, HighValue, Alpha);
+					const KeyType LerpValue = Interpolate(LowValue, HighValue, Alpha);
 
 					// compute the error between our interpolated value and the desired value
 					float LerpError = CalcDelta(TestValue, LerpValue);
@@ -216,7 +274,7 @@
 						const FTransform& InvRawBase = CachedInvRawBases[FrameIndex];
 						
 						// generate the proposed local bone atom and transform (local space)
-						FTransform ProposedTM = UpdateBoneAtom(BoneIndex, BoneAtoms[FrameIndex], LerpValue);
+						FTransform ProposedTM = AdapterType::UpdateBoneAtom(BoneAtoms[FrameIndex], LerpValue);
 
 						// convert the proposed local transform to world space using this bone's parent transform
 						const FTransform& CurrentParent = ParentBoneIndex != INDEX_NONE ? NewWorldBones[(ParentBoneIndex*NumFrames) + FrameIndex] : FTransform::Identity;
@@ -310,6 +368,212 @@
 	}
 }
 
+template <typename AdapterType>
+void FilterLinearKeysTemplate(
+	TArray<typename AdapterType::KeyType>& Keys,
+	TArray<float>& Times,
+	TArray<FTransform>& BoneAtoms,
+	const TArray<float>* ParentTimes,
+	const FProcessAnimationTracksContext& Context,
+	const TArray<int32>& TargetBoneIndices,
+	int32 BoneIndex,
+	int32 ParentBoneIndex,
+	float ParentScale,
+	float MaxDelta,
+	float MaxTargetDelta,
+	float EffectorDiffSocket)
+{
+	typedef typename AdapterType::KeyType KeyType;
+
+	const int32 KeyCount = Keys.Num();
+	check(Keys.Num() == Times.Num());
+	check(KeyCount >= 1);
+
+	// generate new arrays we will fill with the final keys
+	TArray<KeyType> NewKeys;
+	TArray<float> NewTimes;
+	NewKeys.Empty(KeyCount);
+	NewTimes.Empty(KeyCount);
+
+	// Only bother doing anything if we have some keys!
+	if (KeyCount > 0)
+	{
+		int32 LowKey = 0;
+		int32 HighKey = KeyCount - 1;
+
+		TArray<uint32> KnownParentTimes;
+		KnownParentTimes.Empty(KeyCount);
+		const int32 ParentKeyCount = ParentTimes ? ParentTimes->Num() : 0;
+		for (int32 TimeIndex = 0, ParentTimeIndex = 0; TimeIndex < KeyCount; TimeIndex++)
+		{
+			while ((ParentTimeIndex < ParentKeyCount) && (Times[TimeIndex] > (*ParentTimes)[ParentTimeIndex]))
+			{
+				ParentTimeIndex++;
+			}
+
+			KnownParentTimes.Add((ParentTimeIndex < ParentKeyCount) && (Times[TimeIndex] == (*ParentTimes)[ParentTimeIndex]));
+		}
+
+		TArray<FTransform> CachedInvRawBases;
+		CachedInvRawBases.Empty(KeyCount);
+		for (int32 FrameIndex = 0; FrameIndex < KeyCount; ++FrameIndex)
+		{
+			const FTransform& RawBase = Context.RawWorldBones[(BoneIndex * Context.Segment.NumFrames) + FrameIndex];
+			CachedInvRawBases.Add(RawBase.Inverse());
+		}
+
+		// copy the low key (this one is a given)
+		NewTimes.Add(Times[0]);
+		NewKeys.Add(Keys[0]);
+
+		const FTransform EndEffectorDummyBoneSocket(FQuat::Identity, FVector(END_EFFECTOR_DUMMY_BONE_LENGTH_SOCKET));
+		const FTransform EndEffectorDummyBone(FQuat::Identity, FVector(END_EFFECTOR_DUMMY_BONE_LENGTH));
+
+		const float DeltaThreshold = (Context.BoneData[BoneIndex].IsEndEffector() && (Context.BoneData[BoneIndex].bHasSocket || Context.BoneData[BoneIndex].bKeyEndEffector)) ? EffectorDiffSocket : MaxTargetDelta;
+
+		// We will test within a sliding window between LowKey and HighKey.
+		// Therefore, we are done when the LowKey exceeds the range
+		while (LowKey + 1 < KeyCount)
+		{
+			int32 GoodHighKey = LowKey + 1;
+			int32 BadHighKey = KeyCount;
+
+			// bisect until we find the lowest acceptable high key
+			while (BadHighKey - GoodHighKey >= 2)
+			{
+				HighKey = GoodHighKey + (BadHighKey - GoodHighKey) / 2;
+
+				// get the parameters of the window we are testing
+				const float LowTime = Times[LowKey];
+				const float HighTime = Times[HighKey];
+				const KeyType& LowValue = Keys[LowKey];
+				const KeyType& HighValue = Keys[HighKey];
+				const float Range = HighTime - LowTime;
+				const float InvRange = 1.0f / Range;
+
+				// iterate through all interpolated members of the window to
+				// compute the error when compared to the original raw values
+				float MaxLerpError = 0.0f;
+				float MaxTargetError = 0.0f;
+				for (int32 TestKey = LowKey + 1; TestKey < HighKey; ++TestKey)
+				{
+					// get the parameters of the member being tested
+					float TestTime = Times[TestKey];
+					const KeyType& TestValue = Keys[TestKey];
+
+					// compute the proposed, interpolated value for the key
+					const float Alpha = (TestTime - LowTime) * InvRange;
+					const KeyType LerpValue = Interpolate(LowValue, HighValue, Alpha);
+
+					// compute the error between our interpolated value and the desired value
+					float LerpError = CalcDelta(TestValue, LerpValue);
+
+					// if the local-space lerp error is within our tolerances, we will also check the
+					// effect this interpolated key will have on our target end effectors
+					float TargetError = -1.0f;
+					if (LerpError <= MaxDelta)
+					{
+						// get the raw world transform for this bone (the original world-space position)
+						const int32 FrameIndex = TestKey;
+						const FTransform& InvRawBase = CachedInvRawBases[FrameIndex];
+
+						// generate the proposed local bone atom and transform (local space)
+						FTransform ProposedTM = AdapterType::UpdateBoneAtom(BoneAtoms[FrameIndex], LerpValue);
+
+						// convert the proposed local transform to world space using this bone's parent transform
+						const FTransform& CurrentParent = ParentBoneIndex != INDEX_NONE ? Context.NewWorldBones[(ParentBoneIndex * Context.Segment.NumFrames) + FrameIndex] : FTransform::Identity;
+						FTransform ProposedBase = ProposedTM * CurrentParent;
+
+						// for each target end effector, compute the error we would introduce with our proposed key
+						for (int32 TargetIndex = 0; TargetIndex < TargetBoneIndices.Num(); ++TargetIndex)
+						{
+							// find the offset transform from the raw base to the end effector
+							const int32 TargetBoneIndex = TargetBoneIndices[TargetIndex];
+							FTransform RawTarget = Context.RawWorldBones[(TargetBoneIndex * Context.Segment.NumFrames) + FrameIndex];
+							const FTransform RelTM = RawTarget * InvRawBase;
+
+							// forecast where the new end effector would be using our proposed key
+							FTransform ProposedTarget = RelTM * ProposedBase;
+
+							// If this is an EndEffector, add a dummy bone to measure the effect of compressing the rotation.
+							// Sockets and Key EndEffectors have a longer dummy bone to maintain higher precision.
+							if (Context.BoneData[TargetIndex].bHasSocket || Context.BoneData[TargetIndex].bKeyEndEffector)
+							{
+								ProposedTarget = EndEffectorDummyBoneSocket * ProposedTarget;
+								RawTarget = EndEffectorDummyBoneSocket * RawTarget;
+							}
+							else
+							{
+								ProposedTarget = EndEffectorDummyBone * ProposedTarget;
+								RawTarget = EndEffectorDummyBone * RawTarget;
+							}
+
+							// determine the extend of error at the target end effector
+							const float ThisError = (ProposedTarget.GetTranslation() - RawTarget.GetTranslation()).Size();
+							TargetError = FMath::Max(TargetError, ThisError);
+
+							// exit early when we encounter a large delta
+							const float TargetDeltaThreshold = Context.BoneData[TargetIndex].bHasSocket ? EffectorDiffSocket : DeltaThreshold;
+							if (TargetError > TargetDeltaThreshold)
+							{
+								break;
+							}
+						}
+					}
+
+					// If the parent has a key at this time, we'll scale our error values as requested.
+					// This increases the odds that we will choose keys on the same frames as our parent bone,
+					// making the skeleton more uniform in key distribution.
+					if (ParentTimes)
+					{
+						if (KnownParentTimes[TestKey])
+						{
+							// our parent has a key at this time, 
+							// inflate our perceived error to increase our sensitivity
+							// for also retaining a key at this time
+							LerpError *= ParentScale;
+							TargetError *= ParentScale;
+						}
+					}
+
+					// keep track of the worst errors encountered for both 
+					// the local-space 'lerp' error and the end effector drift we will cause
+					MaxLerpError = FMath::Max(MaxLerpError, LerpError);
+					MaxTargetError = FMath::Max(MaxTargetError, TargetError);
+
+					// exit early if we have failed in this span
+					if (MaxLerpError > MaxDelta ||
+						MaxTargetError > DeltaThreshold)
+					{
+						break;
+					}
+				}
+
+				// determine if the span succeeded. That is, the worst errors found are within tolerances
+				if ((MaxLerpError <= MaxDelta) && (MaxTargetError <= DeltaThreshold))
+				{
+					GoodHighKey = HighKey;
+				}
+				else
+				{
+					BadHighKey = HighKey;
+				}
+			}
+
+			//UE_LOG(LogAnimation, Warning, TEXT("%u: Keeping key %u"), BoneIndex, GoodHighKey);
+
+			NewTimes.Add(Times[GoodHighKey]);
+			NewKeys.Add(Keys[GoodHighKey]);
+
+			LowKey = GoodHighKey;
+		}
+
+		// return the new key set to the caller
+		Times = NewTimes;
+		Keys = NewKeys;
+	}
+}
+
 void UAnimCompress_RemoveLinearKeys::UpdateWorldBoneTransformTable(
 	UAnimSequence* AnimSeq, 
 	const TArray<FBoneData>& BoneData, 
@@ -373,6 +637,372 @@
 	}
 }
 
+static void PackVector(
+	const FVector& Key,
+	AnimationCompressionFormat Format,
+	const float* Mins,
+	const float* Ranges,
+	uint8* Buffer)
+{
+	if (Format == ACF_None || Format == ACF_Float96NoW)
+	{
+		memcpy(Buffer, &Key, sizeof(FVector));
+	}
+	else if (Format == ACF_IntervalFixed32NoW)
+	{
+		const FVectorIntervalFixed32NoW CompressedKey(Key, Mins, Ranges);
+		memcpy(Buffer, &CompressedKey, sizeof(FVectorIntervalFixed32NoW));
+	}
+}
+
+static FVector UnpackVector(
+	const uint8* Buffer,
+	AnimationCompressionFormat Format,
+	const float* Mins,
+	const float* Ranges)
+{
+	FVector Result;
+
+	if (Format == ACF_None || Format == ACF_Float96NoW)
+	{
+		memcpy(&Result, Buffer, sizeof(FVector));
+	}
+	else if (Format == ACF_IntervalFixed32NoW)
+	{
+		FVectorIntervalFixed32NoW CompressedKey;
+		memcpy(&CompressedKey, Buffer, sizeof(FVectorIntervalFixed32NoW));
+		CompressedKey.ToVector(Result, Mins, Ranges);
+	}
+
+	return Result;
+}
+
+static void PackQuat(
+	const FQuat& Key,
+	AnimationCompressionFormat Format,
+	const float* Mins,
+	const float* Ranges,
+	uint8* Buffer)
+{
+	if (Format == ACF_None)
+	{
+		memcpy(Buffer, &Key, sizeof(FQuat));
+	}
+	else if (Format == ACF_Float96NoW)
+	{
+		const FQuatFloat96NoW QuatFloat96NoW(Key);
+		memcpy(Buffer, &QuatFloat96NoW, sizeof(FQuatFloat96NoW));
+	}
+	else if (Format == ACF_Fixed32NoW)
+	{
+		const FQuatFixed32NoW QuatFixed32NoW(Key);
+		memcpy(Buffer, &QuatFixed32NoW, sizeof(FQuatFixed32NoW));
+	}
+	else if (Format == ACF_Fixed48NoW)
+	{
+		const FQuatFixed48NoW QuatFixed48NoW(Key);
+		memcpy(Buffer, &QuatFixed48NoW, sizeof(FQuatFixed48NoW));
+	}
+	else if (Format == ACF_IntervalFixed32NoW)
+	{
+		const FQuatIntervalFixed32NoW QuatIntervalFixed32NoW(Key, Mins, Ranges);
+		memcpy(Buffer, &QuatIntervalFixed32NoW, sizeof(FQuatIntervalFixed32NoW));
+	}
+	else if (Format == ACF_Float32NoW)
+	{
+		const FQuatFloat32NoW QuatFloat32NoW(Key);
+		memcpy(Buffer, &QuatFloat32NoW, sizeof(FQuatFloat32NoW));
+	}
+}
+
+static FQuat UnpackQuat(
+	const uint8* Buffer,
+	AnimationCompressionFormat Format,
+	const float* Mins,
+	const float* Ranges)
+{
+	FQuat Result;
+
+	if (Format == ACF_None)
+	{
+		memcpy(&Result, Buffer, sizeof(FQuat));
+	}
+	else if (Format == ACF_Float96NoW)
+	{
+		FQuatFloat96NoW QuatFloat96NoW;
+		memcpy(&QuatFloat96NoW, Buffer, sizeof(FQuatFloat96NoW));
+		QuatFloat96NoW.ToQuat(Result);
+	}
+	else if (Format == ACF_Fixed32NoW)
+	{
+		FQuatFixed32NoW QuatFixed32NoW;
+		memcpy(&QuatFixed32NoW, Buffer, sizeof(FQuatFixed32NoW));
+		QuatFixed32NoW.ToQuat(Result);
+	}
+	else if (Format == ACF_Fixed48NoW)
+	{
+		FQuatFixed48NoW QuatFixed48NoW;
+		memcpy(&QuatFixed48NoW, Buffer, sizeof(FQuatFixed48NoW));
+		QuatFixed48NoW.ToQuat(Result);
+	}
+	else if (Format == ACF_IntervalFixed32NoW)
+	{
+		FQuatIntervalFixed32NoW QuatIntervalFixed32NoW;
+		memcpy(&QuatIntervalFixed32NoW, Buffer, sizeof(FQuatIntervalFixed32NoW));
+		QuatIntervalFixed32NoW.ToQuat(Result, Mins, Ranges);
+	}
+	else if (Format == ACF_Float32NoW)
+	{
+		FQuatFloat32NoW QuatFloat32NoW;
+		memcpy(&QuatFloat32NoW, Buffer, sizeof(FQuatFloat32NoW));
+		QuatFloat32NoW.ToQuat(Result);
+	}
+
+	return Result;
+}
+
+static float GetFrameRate(const UAnimSequence& AnimSeq)
+{
+	// We cannot round since some sequences have a non integral frame rate for legacy reasons
+	return float(AnimSeq.NumFrames - 1) / AnimSeq.SequenceLength;
+}
+
+static float FindKeyInterpolationData(const UAnimSequence& AnimSeq, const TArray<float>& KeyTimes, float Time, int32& OutKeyIndex0, int32& OutKeyIndex1)
+{
+	// Use the same logic as the decompression in order to ensure we match 100%
+	const int32 LastFrame = AnimSeq.NumFrames - 1;
+	const float FrameRate = GetFrameRate(AnimSeq);
+	const float RelativePos = Time / AnimSeq.SequenceLength;
+	const float FramePos = RelativePos * float(LastFrame);
+	const int32 FramePosFloor = FMath::Clamp(FMath::TruncToInt(FramePos), 0, LastFrame);
+
+	const int32 NumKeys = KeyTimes.Num();
+	int32 KeyIndex0 = NumKeys - 1;	// Default
+	for (int32 KeyIndex = 1; KeyIndex < NumKeys; ++KeyIndex)
+	{
+		const float KeyTime = KeyTimes[KeyIndex];
+		const float FrameTime = KeyTime * FrameRate;
+		const int32 FrameIndex = FMath::Clamp(FMath::TruncToInt(FrameTime + 0.5f), 0, LastFrame);
+		if (FrameIndex > FramePosFloor)
+		{
+			KeyIndex0 = KeyIndex - 1;
+			break;
+		}
+	}
+
+	const int32 KeyIndex1 = FMath::Min(KeyIndex0 + 1, NumKeys - 1);
+	const float KeyTime0 = KeyTimes[KeyIndex0];
+	const float KeyTime1 = KeyTimes[KeyIndex1];
+
+	const float FrameTime0 = KeyTime0 * FrameRate;
+	const float FrameTime1 = KeyTime1 * FrameRate;
+	const int32 FrameIndex0 = FMath::Clamp(FMath::TruncToInt(FrameTime0 + 0.5f), 0, LastFrame);
+	const int32 FrameIndex1 = FMath::Clamp(FMath::TruncToInt(FrameTime1 + 0.5f), 0, LastFrame);
+
+	const int32 Delta = FMath::Max(FrameIndex1 - FrameIndex0, 1);
+	const float Remainder = FramePos - (float)FrameIndex0;
+	const float InterpolationAlpha = AnimSeq.Interpolation == EAnimInterpolationType::Step ? 0.0f : FMath::Min(Remainder / (float)Delta, 1.0f);
+
+	OutKeyIndex0 = KeyIndex0;
+	OutKeyIndex1 = KeyIndex1;
+	return InterpolationAlpha;
+}
+
+FTransform UAnimCompress_RemoveLinearKeys::SampleSegment(const FProcessAnimationTracksContext& Context, int32 TrackIndex, float Time)
+{
+	FTransform Result;
+
+	{
+		const FTranslationTrack& TranslationData = Context.Segment.TranslationData[TrackIndex];
+		const int32 NumTransKeys = TranslationData.Times.Num();
+		if (NumTransKeys > 1)
+		{
+			int32 KeyIndex0;
+			int32 KeyIndex1;
+			const float InterpolationAlpha = FindKeyInterpolationData(Context.AnimSeq, TranslationData.Times, Time, KeyIndex0, KeyIndex1);
+
+			FVector RangeMin;
+			FVector RangeExtent;
+			CalculateTrackRange(TranslationData, Context.Segment.TranslationCompressionFormat, RangeMin, RangeExtent);
+
+			const FVector& Translation0 = TranslationData.PosKeys[KeyIndex0];
+			uint8 Buffer0[32];
+			PackVector(Translation0, Context.Segment.TranslationCompressionFormat, &RangeMin.X, &RangeExtent.X, Buffer0);
+			FVector LossyTranslation = UnpackVector(Buffer0, Context.Segment.TranslationCompressionFormat, &RangeMin.X, &RangeExtent.X);
+
+			if (KeyIndex0 != KeyIndex1)
+			{
+				const FVector& Translation1 = TranslationData.PosKeys[KeyIndex1];
+				uint8 Buffer1[32];
+				PackVector(Translation1, Context.Segment.TranslationCompressionFormat, &RangeMin.X, &RangeExtent.X, Buffer1);
+				const FVector LossyTranslation1 = UnpackVector(Buffer1, Context.Segment.TranslationCompressionFormat, &RangeMin.X, &RangeExtent.X);
+
+				LossyTranslation = FMath::Lerp(LossyTranslation, LossyTranslation1, InterpolationAlpha);
+			}
+
+			Result.SetTranslation(LossyTranslation);
+		}
+		else
+		{
+			Result.SetTranslation(Context.TrivialTrackValues[TrackIndex].GetTranslation());
+		}
+	}
+
+	{
+		const FRotationTrack& RotationData = Context.Segment.RotationData[TrackIndex];
+		const int32 NumRotKeys = RotationData.Times.Num();
+		if (NumRotKeys > 1)
+		{
+			int32 KeyIndex0;
+			int32 KeyIndex1;
+			const float InterpolationAlpha = FindKeyInterpolationData(Context.AnimSeq, RotationData.Times, Time, KeyIndex0, KeyIndex1);
+
+			FVector RangeMin;
+			FVector RangeExtent;
+			CalculateTrackRange(RotationData, Context.Segment.RotationCompressionFormat, RangeMin, RangeExtent);
+
+			const FQuat& Rotation0 = RotationData.RotKeys[KeyIndex0];
+			uint8 Buffer0[32];
+			PackQuat(Rotation0, Context.Segment.RotationCompressionFormat, &RangeMin.X, &RangeExtent.X, Buffer0);
+			FQuat LossyRotation = UnpackQuat(Buffer0, Context.Segment.RotationCompressionFormat, &RangeMin.X, &RangeExtent.X);
+
+			if (KeyIndex0 != KeyIndex1)
+			{
+				const FQuat& Rotation1 = RotationData.RotKeys[KeyIndex1];
+				uint8 Buffer1[32];
+				PackQuat(Rotation1, Context.Segment.RotationCompressionFormat, &RangeMin.X, &RangeExtent.X, Buffer1);
+				const FQuat LossyRotation1 = UnpackQuat(Buffer1, Context.Segment.RotationCompressionFormat, &RangeMin.X, &RangeExtent.X);
+
+				LossyRotation = FQuat::FastLerp(LossyRotation, LossyRotation1, InterpolationAlpha);
+				LossyRotation.Normalize();
+			}
+
+			Result.SetRotation(LossyRotation);
+		}
+		else
+		{
+			const FQuat Rotation = Context.TrivialTrackValues[TrackIndex].GetRotation();
+			uint8 Buffer[32];
+			PackQuat(Rotation, ACF_Float96NoW, nullptr, nullptr, Buffer);
+			const FQuat LossyRotation = UnpackQuat(Buffer, ACF_Float96NoW, nullptr, nullptr);
+
+			Result.SetRotation(LossyRotation);
+		}
+	}
+
+	if (Context.Segment.ScaleData.Num() != 0)
+	{
+		const FScaleTrack& ScaleData = Context.Segment.ScaleData[TrackIndex];
+		const int32 NumScaleKeys = ScaleData.Times.Num();
+		if (NumScaleKeys > 1)
+		{
+			int32 KeyIndex0;
+			int32 KeyIndex1;
+			const float InterpolationAlpha = FindKeyInterpolationData(Context.AnimSeq, ScaleData.Times, Time, KeyIndex0, KeyIndex1);
+
+			FVector RangeMin;
+			FVector RangeExtent;
+			CalculateTrackRange(ScaleData, Context.Segment.ScaleCompressionFormat, RangeMin, RangeExtent);
+
+			const FVector& Scale0 = ScaleData.ScaleKeys[KeyIndex0];
+			uint8 Buffer0[32];
+			PackVector(Scale0, Context.Segment.ScaleCompressionFormat, &RangeMin.X, &RangeExtent.X, Buffer0);
+			FVector LossyScale = UnpackVector(Buffer0, Context.Segment.ScaleCompressionFormat, &RangeMin.X, &RangeExtent.X);
+
+			if (KeyIndex0 != KeyIndex1)
+			{
+				const FVector& Scale1 = ScaleData.ScaleKeys[KeyIndex1];
+				uint8 Buffer1[32];
+				PackVector(Scale1, Context.Segment.ScaleCompressionFormat, &RangeMin.X, &RangeExtent.X, Buffer1);
+				const FVector LossyScale1 = UnpackVector(Buffer1, Context.Segment.ScaleCompressionFormat, &RangeMin.X, &RangeExtent.X);
+
+				LossyScale = FMath::Lerp(LossyScale, LossyScale1, InterpolationAlpha);
+			}
+
+			Result.SetScale3D(LossyScale);
+		}
+		else
+		{
+			Result.SetScale3D(Context.TrivialTrackValues[TrackIndex].GetScale3D());
+		}
+	}
+
+	return Result;
+}
+
+void UAnimCompress_RemoveLinearKeys::UpdateWorldBoneTransformTable(
+	const FProcessAnimationTracksContext& Context,
+	int32 BoneIndex,
+	bool UseRaw,
+	TArray<FTransform>& OutputWorldBones)
+{
+	const int32 TrackIndex = Context.BoneIndexToTrackIndex[BoneIndex];
+
+	check(OutputWorldBones.Num() >= ((BoneIndex * Context.Segment.NumFrames) + Context.Segment.NumFrames));
+
+	if (TrackIndex != INDEX_NONE)
+	{
+		const float FrameRate = GetFrameRate(Context.AnimSeq);
+		FAnimSequenceDecompressionContext DecompContext(&Context.AnimSeq);
+
+		// get the local-space bone transforms using the animation solver
+		for (int32 FrameIndex = 0; FrameIndex < Context.Segment.NumFrames; ++FrameIndex)
+		{
+			const float Time = (float)(Context.Segment.StartFrame + FrameIndex) / FrameRate;
+
+			FTransform LocalAtom;
+			if (UseRaw)
+			{
+				DecompContext.Seek(Time);
+				Context.AnimSeq.GetBoneTransform(LocalAtom, TrackIndex, DecompContext, UseRaw);
+			}
+			else
+			{
+				if (bUseDecompression)
+				{
+					DecompContext.Seek(Time);
+					Context.AnimSeq.GetBoneTransform(LocalAtom, TrackIndex, DecompContext, UseRaw);
+				}
+				else
+				{
+					LocalAtom = SampleSegment(Context, TrackIndex, Time);
+				}
+			}
+
+			FQuat Rot = LocalAtom.GetRotation();
+			LocalAtom.SetRotation(EnforceShortestArc(FQuat::Identity, Rot));
+			// Saw some crashes happening with it, so normalize here.
+			LocalAtom.NormalizeRotation();
+
+			OutputWorldBones[(BoneIndex * Context.Segment.NumFrames) + FrameIndex] = LocalAtom;
+		}
+	}
+	else
+	{
+		// get the default rotation and translation from the reference skeleton
+		FTransform DefaultTransform = Context.RefPose[BoneIndex];
+		DefaultTransform.SetRotation(EnforceShortestArc(FQuat::Identity, DefaultTransform.GetRotation()));
+
+		// copy the default transformation into the world bone table
+		for (int32 FrameIndex = 0; FrameIndex < Context.Segment.NumFrames; ++FrameIndex)
+		{
+			OutputWorldBones[(BoneIndex * Context.Segment.NumFrames) + FrameIndex] = DefaultTransform;
+		}
+	}
+
+	// apply parent transforms to bake into world space. We assume the parent transforms were previously set using this function
+	const FBoneData& Bone = Context.BoneData[BoneIndex];
+	const int32 ParentIndex = Bone.GetParent();
+	if (ParentIndex != INDEX_NONE)
+	{
+		check(ParentIndex < BoneIndex);
+		for (int32 FrameIndex = 0; FrameIndex < Context.Segment.NumFrames; ++FrameIndex)
+		{
+			OutputWorldBones[(BoneIndex * Context.Segment.NumFrames) + FrameIndex] = OutputWorldBones[(BoneIndex * Context.Segment.NumFrames) + FrameIndex] * OutputWorldBones[(ParentIndex * Context.Segment.NumFrames) + FrameIndex];
+		}
+	}
+}
+
 void UAnimCompress_RemoveLinearKeys::FilterBeforeMainKeyRemoval(
 	UAnimSequence* AnimSeq, 
 	const TArray<FBoneData>& BoneData, 
@@ -423,6 +1053,27 @@
 }
 
 
+void UAnimCompress_RemoveLinearKeys::UpdateWorldBoneTransformRange(
+	FProcessAnimationTracksContext& Context,
+	int32 StartingBoneIndex,	// this bone index should be of skeleton, not mesh
+	int32 EndingBoneIndex)		// this bone index should be of skeleton, not mesh
+{
+	if (bUseDecompression)
+	{
+		// bitwise compress the tracks into the anim sequence buffers
+		// to make sure the data we've compressed so far is ready for solving
+		CompressUsingUnderlyingCompressor(Context.AnimSeq, Context.BoneData, Context.SegmentList, false);
+	}
+
+	// build all world-space transforms from this bone to the target end effector we are monitoring
+	// all parent transforms have been built already
+	for (int32 Index = StartingBoneIndex; Index <= EndingBoneIndex; ++Index)
+	{
+		UpdateWorldBoneTransformTable(Context, Index, false, Context.NewWorldBones);
+	}
+}
+
+
 void UAnimCompress_RemoveLinearKeys::UpdateBoneAtomList(
 	UAnimSequence* AnimSeq, 
 	int32 BoneIndex,
@@ -445,6 +1096,37 @@
 }
 
 
+void UAnimCompress_RemoveLinearKeys::UpdateBoneAtomList(
+	const FProcessAnimationTracksContext& Context,
+	int32 TrackIndex,
+	TArray<FTransform>& BoneAtoms) const
+{
+	const float FrameRate = GetFrameRate(Context.AnimSeq);
+	FAnimSequenceDecompressionContext DecompContext(&Context.AnimSeq);
+
+	BoneAtoms.Reset(Context.Segment.NumFrames);
+	for (int32 FrameIndex = 0; FrameIndex < Context.Segment.NumFrames; ++FrameIndex)
+	{
+		const float Time = (float)(FrameIndex + Context.Segment.StartFrame) / FrameRate;
+
+		FTransform LocalAtom;
+		if (bUseDecompression)
+		{
+			DecompContext.Seek(Time);
+			Context.AnimSeq.GetBoneTransform(LocalAtom, TrackIndex, DecompContext, false);
+		}
+		else
+		{
+			LocalAtom = SampleSegment(Context, TrackIndex, Time);
+		}
+
+		FQuat Rot = LocalAtom.GetRotation();
+		LocalAtom.SetRotation(EnforceShortestArc(FQuat::Identity, Rot));
+		BoneAtoms.Add(LocalAtom);
+	}
+}
+
+
 bool UAnimCompress_RemoveLinearKeys::ConvertFromRelativeSpace(UAnimSequence* AnimSeq)
 {
 	// if this is an additive animation, temporarily convert it out of relative-space
@@ -496,42 +1178,76 @@
 	return bAdditiveAnimation;
 }
 
-
 void UAnimCompress_RemoveLinearKeys::ConvertToRelativeSpace(
 	UAnimSequence* AnimSeq,
 	TArray<FTranslationTrack>& TranslationData,
 	TArray<FRotationTrack>& RotationData, 
 	TArray<FScaleTrack>& ScaleData)
 {
+	ConvertToRelativeSpace(*AnimSeq);
+	ConvertToRelativeSpace(*AnimSeq, TranslationData, RotationData, ScaleData);
+}
 
+void UAnimCompress_RemoveLinearKeys::ConvertToRelativeSpace(UAnimSequence& AnimSeq) const
+{
 	// convert the raw tracks back to additive-space
-	const int32 NumTracks = AnimSeq->GetRawAnimationData().Num();
-	for ( int32 TrackIndex = 0; TrackIndex < NumTracks; ++TrackIndex )
+	const int32 NumTracks = AnimSeq.GetRawAnimationData().Num();
+	for (int32 TrackIndex = 0; TrackIndex < NumTracks; ++TrackIndex)
 	{
-		int32 const BoneIndex = AnimSeq->GetRawTrackToSkeletonMapTable()[TrackIndex].BoneTreeIndex;
+		const FRawAnimSequenceTrack& BasePoseTrack = AnimSeq.GetAdditiveBaseAnimationData()[TrackIndex];
+		FRawAnimSequenceTrack& RawTrack = AnimSeq.GetRawAnimationTrack(TrackIndex);
 
-		bool const bIsRootBone = (BoneIndex == 0);
-
-		const FRawAnimSequenceTrack& BasePoseTrack = AnimSeq->GetAdditiveBaseAnimationData()[TrackIndex];
-		FRawAnimSequenceTrack& RawTrack	= AnimSeq->GetRawAnimationTrack(TrackIndex);
 		// @note: we only extract the first frame, as we don't want to induce motion from the base pose
 		// only the motion from the additive data should matter.
-		const FQuat& InvRefBoneRotation = BasePoseTrack.RotKeys[0].Inverse();
-		const FVector& InvRefBoneTranslation = -BasePoseTrack.PosKeys[0];
+		const FQuat InvRefBoneRotation = BasePoseTrack.RotKeys[0].Inverse();
+		const FVector InvRefBoneTranslation = -BasePoseTrack.PosKeys[0];
 
 		// transform position keys.
-		for ( int32 PosIndex = 0; PosIndex < RawTrack.PosKeys.Num(); ++PosIndex )
+		for (int32 PosIndex = 0; PosIndex < RawTrack.PosKeys.Num(); ++PosIndex)
 		{
 			RawTrack.PosKeys[PosIndex] += InvRefBoneTranslation;
 		}
 
 		// transform rotation keys.
-		for ( int32 RotIndex = 0; RotIndex < RawTrack.RotKeys.Num(); ++RotIndex )
+		for (int32 RotIndex = 0; RotIndex < RawTrack.RotKeys.Num(); ++RotIndex)
 		{
 			RawTrack.RotKeys[RotIndex] = RawTrack.RotKeys[RotIndex] * InvRefBoneRotation;
 			RawTrack.RotKeys[RotIndex].Normalize();
 		}
 
+		// scale key
+		if (RawTrack.ScaleKeys.Num() > 0)
+		{
+			const FVector InvRefBoneScale = FTransform::GetSafeScaleReciprocal(BasePoseTrack.ScaleKeys[0]);
+
+			// transform scale keys.
+			for (int32 ScaleIndex = 0; ScaleIndex < RawTrack.ScaleKeys.Num(); ++ScaleIndex)
+			{
+				// to revert scale correctly, you have to - 1.f
+				// check AccumulateWithAdditiveScale
+				RawTrack.ScaleKeys[ScaleIndex] = (RawTrack.ScaleKeys[ScaleIndex] * InvRefBoneScale) - 1.f;
+			}
+		}
+	}
+}
+
+void UAnimCompress_RemoveLinearKeys::ConvertToRelativeSpace(
+	const UAnimSequence& AnimSeq,
+	TArray<FTranslationTrack>& TranslationData,
+	TArray<FRotationTrack>& RotationData,
+	TArray<FScaleTrack>& ScaleData) const
+{
+	// convert the raw tracks back to additive-space
+	const int32 NumTracks = AnimSeq.GetRawAnimationData().Num();
+	for (int32 TrackIndex = 0; TrackIndex < NumTracks; ++TrackIndex)
+	{
+		const FRawAnimSequenceTrack& BasePoseTrack = AnimSeq.GetAdditiveBaseAnimationData()[TrackIndex];
+
+		// @note: we only extract the first frame, as we don't want to induce motion from the base pose
+		// only the motion from the additive data should matter.
+		const FQuat InvRefBoneRotation = BasePoseTrack.RotKeys[0].Inverse();
+		const FVector InvRefBoneTranslation = -BasePoseTrack.PosKeys[0];
+
 		// convert the new translation tracks to additive space
 		FTranslationTrack& TranslationTrack = TranslationData[TrackIndex];
 		for (int32 KeyIndex = 0; KeyIndex < TranslationTrack.PosKeys.Num(); ++KeyIndex)
@@ -550,16 +1266,8 @@
 		// scale key
 		if (ScaleData.Num() > 0)
 		{
-			const FVector& InvRefBoneScale = FTransform::GetSafeScaleReciprocal(BasePoseTrack.ScaleKeys[0]);
+			const FVector InvRefBoneScale = FTransform::GetSafeScaleReciprocal(BasePoseTrack.ScaleKeys[0]);
 
-			// transform scale keys.
-			for (int32 ScaleIndex = 0; ScaleIndex < RawTrack.ScaleKeys.Num(); ++ScaleIndex)
-			{
-				// to revert scale correctly, you have to - 1.f
-				// check AccumulateWithAdditiveScale
-				RawTrack.ScaleKeys[ScaleIndex] = (RawTrack.ScaleKeys[ScaleIndex] * InvRefBoneScale) - 1.f;
-			}
-
 			// convert the new scale tracks to additive space
 			FScaleTrack& ScaleTrack = ScaleData[TrackIndex];
 			for (int32 KeyIndex = 0; KeyIndex < ScaleTrack.ScaleKeys.Num(); ++KeyIndex)
@@ -570,7 +1278,6 @@
 	}
 }
 
-
 void UAnimCompress_RemoveLinearKeys::ProcessAnimationTracks(
 	UAnimSequence* AnimSeq, 
 	const TArray<FBoneData>& BoneData, 
@@ -879,3 +1586,382 @@
 			// if the TargetBoneIndices array is empty, then this bone is an end effector.
 			// so we add it to the list to maintain our tolerance checks
 			if (TargetBoneIndices.Num() == 0)
+			{
+				TargetBoneIndices.Add(BoneIndex);
+			}
+
+			if (bActuallyFilterLinearKeys)
+			{
+				if (bHasScale)
+				{
+					FScaleTrack& ScaleTrack = ScaleTracks[TrackIndex];
+					// filter out translations we can approximate through interpolation
+					FilterLinearKeysTemplate<ScaleAdapter>(
+						ScaleTrack.ScaleKeys, 
+						ScaleTrack.Times, 
+						BoneAtoms,
+						GuidanceTrack, 
+						RawWorldBones,
+						NewWorldBones,
+						TargetBoneIndices,
+						NumFrames,
+						BoneIndex,
+						ParentBoneIndex,
+						GuidanceScale, 
+						MaxScaleDiff, 
+						EndEffectorTolerance,
+						EffectorDiffSocket,
+						BoneData);
+
+					// update our bone table from the current bone through the last end effector we need to test
+					UpdateWorldBoneTransformRange(
+						AnimSeq, 
+						BoneData, 
+						RefPose,
+						PositionTracks,
+						RotationTracks,
+						ScaleTracks,
+						BoneIndex,
+						HighestTargetBoneIndex,
+						false,
+						NewWorldBones);
+
+					// rebuild the BoneAtoms table using the current set of keys
+					UpdateBoneAtomList(AnimSeq, BoneIndex, TrackIndex, NumFrames, TimePerFrame, BoneAtoms); 
+				}
+
+				// filter out translations we can approximate through interpolation
+				FilterLinearKeysTemplate<TranslationAdapter>(
+					TransTrack.PosKeys, 
+					TransTrack.Times, 
+					BoneAtoms,
+					GuidanceTrack, 
+					RawWorldBones,
+					NewWorldBones,
+					TargetBoneIndices,
+					NumFrames,
+					BoneIndex,
+					ParentBoneIndex,
+					GuidanceScale, 
+					MaxPosDiff, 
+					EndEffectorTolerance,
+					EffectorDiffSocket,
+					BoneData);
+
+				// update our bone table from the current bone through the last end effector we need to test
+				UpdateWorldBoneTransformRange(
+					AnimSeq, 
+					BoneData, 
+					RefPose,
+					PositionTracks,
+					RotationTracks,
+					ScaleTracks,
+					BoneIndex,
+					HighestTargetBoneIndex,
+					false,
+					NewWorldBones);
+
+				// rebuild the BoneAtoms table using the current set of keys
+				UpdateBoneAtomList(AnimSeq, BoneIndex, TrackIndex, NumFrames, TimePerFrame, BoneAtoms); 
+
+				// filter out rotations we can approximate through interpolation
+				FilterLinearKeysTemplate<RotationAdapter>(
+					RotTrack.RotKeys, 
+					RotTrack.Times, 
+					BoneAtoms,
+					GuidanceTrack, 
+					RawWorldBones,
+					NewWorldBones,
+					TargetBoneIndices,
+					NumFrames,
+					BoneIndex,
+					ParentBoneIndex,
+					GuidanceScale, 
+					MaxAngleDiff, 
+					EndEffectorTolerance,
+					EffectorDiffSocket,
+					BoneData);
+			}
+		}
+
+		// make sure the final compressed keys are repesented in our NewWorldBones table
+		UpdateWorldBoneTransformRange(
+			AnimSeq, 
+			BoneData, 
+			RefPose,
+			PositionTracks,
+			RotationTracks,
+			ScaleTracks,
+			BoneIndex,
+			BoneIndex,
+			false,
+			NewWorldBones);
+	}
+};
+
+
+static void CalculateBoneChainInformation(const FProcessAnimationTracksContext& Context, int32 BoneIndex, int32& OutHighestTargetBoneIndex, int32& OutFurthestTargetBoneIndex, int32& OutShortestChain, TArray<int32>& OutTargetBoneIndices)
+{
+	const int32 NumBones = Context.BoneData.Num();
+
+	OutTargetBoneIndices.Reset(NumBones);
+
+	int32 HighestTargetBoneIndex = BoneIndex;
+	int32 FurthestTargetBoneIndex = BoneIndex;
+	int32 ShortestChain = 0;
+	float OffsetLength = -1.0f;
+	for (int32 EffectorIndex = 0; EffectorIndex < Context.EndEffectors.Num(); ++EffectorIndex)
+	{
+		const int32 EffectorBoneIndex = Context.EndEffectors[EffectorIndex];
+		const FBoneData& EffectorBoneData = Context.BoneData[EffectorBoneIndex];
+
+		int32 RootIndex = EffectorBoneData.BonesToRoot.Find(BoneIndex);
+		if (RootIndex != INDEX_NONE)
+		{
+			if (ShortestChain == 0 || (RootIndex + 1) < ShortestChain)
+			{
+				ShortestChain = (RootIndex + 1);
+			}
+
+			OutTargetBoneIndices.Add(EffectorBoneIndex);
+			HighestTargetBoneIndex = FMath::Max(HighestTargetBoneIndex, EffectorBoneIndex);
+			float ChainLength = 0.0f;
+			for (long FamilyIndex = 0; FamilyIndex < RootIndex; ++FamilyIndex)
+			{
+				const int32 NextParentBoneIndex = EffectorBoneData.BonesToRoot[FamilyIndex];
+				ChainLength += Context.RefPose[NextParentBoneIndex].GetTranslation().Size();
+			}
+
+			if (ChainLength > OffsetLength)
+			{
+				FurthestTargetBoneIndex = EffectorBoneIndex;
+				OffsetLength = ChainLength;
+			}
+		}
+	}
+
+	OutHighestTargetBoneIndex = HighestTargetBoneIndex;
+	OutFurthestTargetBoneIndex = FurthestTargetBoneIndex;
+	OutShortestChain = ShortestChain;
+}
+
+void UAnimCompress_RemoveLinearKeys::PerformRetargeting(
+	FProcessAnimationTracksContext& Context,
+	int32 BoneIndex,
+	int32 HighestTargetBoneIndex,
+	int32 FurthestTargetBoneIndex)
+{
+	const int32 TrackIndex = Context.BoneIndexToTrackIndex[BoneIndex];
+	const FBoneData& Bone = Context.BoneData[BoneIndex];
+	const int32 ParentBoneIndex = Bone.GetParent();
+
+	const int32 NumScaleKeys = Context.Segment.ScaleData.Num() > 0 ? Context.Segment.ScaleData[TrackIndex].ScaleKeys.Num() : 0;
+	if (NumScaleKeys > 1 && ParentBoneIndex != INDEX_NONE)
+	{
+		// update our bone table from the current bone through the last end effector we need to test
+		UpdateWorldBoneTransformRange(Context, BoneIndex, HighestTargetBoneIndex);
+
+		FScaleTrack& ScaleTrack = Context.Segment.ScaleData[TrackIndex];
+
+		// adjust all translation keys to align better with the destination
+		for (int32 KeyIndex = 0; KeyIndex < NumScaleKeys; ++KeyIndex)
+		{
+			FVector& Key = ScaleTrack.ScaleKeys[KeyIndex];
+
+			const int32 FrameIndex = FMath::Clamp(KeyIndex, 0, Context.Segment.NumFrames);
+			const FTransform& NewWorldParent = Context.NewWorldBones[(ParentBoneIndex * Context.Segment.NumFrames) + FrameIndex];
+			const FTransform& RawWorldChild = Context.RawWorldBones[(BoneIndex * Context.Segment.NumFrames) + FrameIndex];
+			const FTransform RelTM = RawWorldChild.GetRelativeTransform(NewWorldParent);
+
+			Key = RelTM.GetScale3D();
+		}
+	}
+
+	FRotationTrack& RotTrack = Context.Segment.RotationData[TrackIndex];
+	const int32 NumRotKeys = RotTrack.RotKeys.Num();
+	if (NumRotKeys > 1 && ParentBoneIndex != INDEX_NONE)
+	{
+		// update our bone table from the current bone through the last end effector we need to test
+		UpdateWorldBoneTransformRange(Context, BoneIndex, HighestTargetBoneIndex);
+
+		if (HighestTargetBoneIndex == BoneIndex)
+		{
+			for (int32 KeyIndex = 0; KeyIndex < NumRotKeys; ++KeyIndex)
+			{
+				FQuat& Key = RotTrack.RotKeys[KeyIndex];
+
+				const int32 FrameIndex = FMath::Clamp(KeyIndex, 0, Context.Segment.NumFrames);
+				FTransform NewWorldParent = Context.NewWorldBones[(ParentBoneIndex * Context.Segment.NumFrames) + FrameIndex];
+				FTransform RawWorldChild = Context.RawWorldBones[(BoneIndex * Context.Segment.NumFrames) + FrameIndex];
+				const FTransform RelTM = RawWorldChild.GetRelativeTransform(NewWorldParent);
+				FQuat Rot = RelTM.GetRotation();
+
+				const FQuat AlignedKey = EnforceShortestArc(Key, Rot);
+				Key = AlignedKey;
+			}
+		}
+		else
+		{
+			// adjust all rotation keys towards the end effector target
+			for (int32 KeyIndex = 0; KeyIndex < NumRotKeys; ++KeyIndex)
+			{
+				FQuat& Key = RotTrack.RotKeys[KeyIndex];
+
+				const int32 FrameIndex = FMath::Clamp(KeyIndex, 0, Context.Segment.NumFrames);
+
+				const FTransform& NewWorldTransform = Context.NewWorldBones[(BoneIndex * Context.Segment.NumFrames) + FrameIndex];
+
+				const FTransform DesiredChildTransform = Context.RawWorldBones[(FurthestTargetBoneIndex * Context.Segment.NumFrames) + FrameIndex].GetRelativeTransform(NewWorldTransform);
+				const FTransform CurrentChildTransform = Context.NewWorldBones[(FurthestTargetBoneIndex * Context.Segment.NumFrames) + FrameIndex].GetRelativeTransform(NewWorldTransform);
+
+				// find the two vectors which represent the angular error we are trying to correct
+				const FVector CurrentHeading = CurrentChildTransform.GetTranslation();
+				const FVector DesiredHeading = DesiredChildTransform.GetTranslation();
+
+				// if these are valid, we can continue
+				if (!CurrentHeading.IsNearlyZero() && !DesiredHeading.IsNearlyZero())
+				{
+					const float DotResult = CurrentHeading.GetSafeNormal() | DesiredHeading.GetSafeNormal();
+
+					// limit the range we will retarget to something reasonable (~60 degrees)
+					if (DotResult < 1.0f && DotResult > 0.5f)
+					{
+						FQuat Adjustment = FQuat::FindBetweenVectors(CurrentHeading, DesiredHeading);
+						Adjustment = EnforceShortestArc(FQuat::Identity, Adjustment);
+
+						const FVector Test = Adjustment.RotateVector(CurrentHeading);
+						const float DeltaSqr = (Test - DesiredHeading).SizeSquared();
+						if (DeltaSqr < FMath::Square(0.001f))
+						{
+							FQuat NewKey = Adjustment * Key;
+							NewKey.Normalize();
+
+							const FQuat AlignedKey = EnforceShortestArc(Key, NewKey);
+							Key = AlignedKey;
+						}
+					}
+				}
+			}
+		}
+	}
+
+	FTranslationTrack& TransTrack = Context.Segment.TranslationData[TrackIndex];
+	const int32 NumPosKeys = TransTrack.PosKeys.Num();
+	if (NumPosKeys > 1 && ParentBoneIndex != INDEX_NONE)
+	{
+		// update our bone table from the current bone through the last end effector we need to test
+		UpdateWorldBoneTransformRange(Context, BoneIndex, HighestTargetBoneIndex);
+
+		// adjust all translation keys to align better with the destination
+		for (int32 KeyIndex = 0; KeyIndex < NumPosKeys; ++KeyIndex)
+		{
+			FVector& Key = TransTrack.PosKeys[KeyIndex];
+
+			const int32 FrameIndex = FMath::Clamp(KeyIndex, 0, Context.Segment.NumFrames);
+			const FTransform& NewWorldParent = Context.NewWorldBones[(ParentBoneIndex * Context.Segment.NumFrames) + FrameIndex];
+			const FTransform& RawWorldChild = Context.RawWorldBones[(BoneIndex * Context.Segment.NumFrames) + FrameIndex];
+			const FTransform RelTM = RawWorldChild.GetRelativeTransform(NewWorldParent);
+
+			Key = RelTM.GetTranslation();
+		}
+	}
+}
+
+static int32 FindGuideTrackIndex(const FBoneData& Bone, const TArray<int32>& BoneIndexToTrackIndex, float ParentKeyScale)
+{
+	int32 GuideTrackIndex = INDEX_NONE;
+	if (ParentKeyScale > 1.0f)
+	{
+		for (int32 FamilyIndex = 0; FamilyIndex < Bone.BonesToRoot.Num() && GuideTrackIndex == INDEX_NONE; ++FamilyIndex)
+		{
+			const int32 NextParentBoneIndex = Bone.BonesToRoot[FamilyIndex];
+
+			GuideTrackIndex = BoneIndexToTrackIndex[NextParentBoneIndex];
+		}
+	}
+
+	return GuideTrackIndex;
+}
+
+void UAnimCompress_RemoveLinearKeys::ProcessAnimationTracks(FProcessAnimationTracksContext& Context)
+{
+	const int32 NumBones = Context.BoneData.Num();
+
+	Context.RawWorldBones.Empty(NumBones * Context.Segment.NumFrames);
+	Context.NewWorldBones.Empty(NumBones * Context.Segment.NumFrames);
+	Context.RawWorldBones.AddZeroed(NumBones * Context.Segment.NumFrames);
+	Context.NewWorldBones.AddZeroed(NumBones * Context.Segment.NumFrames);
+
+	TArray<int32> TargetBoneIndices;
+	TargetBoneIndices.Empty(NumBones);
+
+	// Create an array of FTransform to use as a workspace
+	TArray<FTransform> BoneAtoms;
+
+	// setup the raw bone transformation and find all end effectors
+	for (int32 BoneIndex = 0; BoneIndex < NumBones; ++BoneIndex)
+	{
+		// get the raw world-atoms for this bone
+		UpdateWorldBoneTransformTable(Context, BoneIndex, true, Context.RawWorldBones);
+	}
+
+	const bool bHasScale = Context.Segment.ScaleData.Num() > 0;
+
+	// for each bone...
+	for (int32 BoneIndex = 0; BoneIndex < NumBones; ++BoneIndex)
+	{
+		const FBoneData& Bone = Context.BoneData[BoneIndex];
+		const int32 ParentBoneIndex = Bone.GetParent();
+
+		const int32 TrackIndex = Context.BoneIndexToTrackIndex[BoneIndex];
+
+		if (TrackIndex != INDEX_NONE)
+		{
+			// get the tracks we will be editing for this bone
+			FRotationTrack& RotTrack = Context.Segment.RotationData[TrackIndex];
+			FTranslationTrack& TransTrack = Context.Segment.TranslationData[TrackIndex];
+			const int32 NumRotKeys = RotTrack.RotKeys.Num();
+			const int32 NumPosKeys = TransTrack.PosKeys.Num();
+			const int32 NumScaleKeys = bHasScale ? Context.Segment.ScaleData[TrackIndex].ScaleKeys.Num() : 0;
+
+			check((NumPosKeys == 1) || (NumRotKeys == 1) || (NumPosKeys == NumRotKeys));
+
+			// build an array of end effectors we need to monitor
+			int32 HighestTargetBoneIndex;
+			int32 FurthestTargetBoneIndex;
+			int32 ShortestChain;
+			CalculateBoneChainInformation(Context, BoneIndex, HighestTargetBoneIndex, FurthestTargetBoneIndex, ShortestChain, TargetBoneIndices);
+
+			// if requested, retarget the FBoneAtoms towards the target end effectors
+			if (bRetarget)
+			{
+				PerformRetargeting(Context, BoneIndex, HighestTargetBoneIndex, FurthestTargetBoneIndex);
+			}
+
+			// look for a parent track to reference as a guide
+			const int32 GuideTrackIndex = FindGuideTrackIndex(Bone, Context.BoneIndexToTrackIndex, ParentKeyScale);
+
+			// update our bone table from the current bone through the last end effector we need to test
+			UpdateWorldBoneTransformRange(Context, BoneIndex, HighestTargetBoneIndex);
+
+			// rebuild the BoneAtoms table using the current set of keys
+			UpdateBoneAtomList(Context, TrackIndex, BoneAtoms);
+
+			// determine the EndEffectorTolerance. 
+			// We use the Maximum value by default, and the Minimum value
+			// as we approach the end effectors
+			const float EndEffectorTolerance = ShortestChain <= 1 ? MinEffectorDiff : MaxEffectorDiff;
+
+			// Determine if a guidance track should be used to aid in choosing keys to retain
+			const TArray<float>* GuidanceTrack = nullptr;
+			float GuidanceScale = 1.0f;
+			if (GuideTrackIndex != INDEX_NONE)
+			{
+				const FTranslationTrack& GuideTransTrack = Context.Segment.TranslationData[GuideTrackIndex];
+				GuidanceTrack = &GuideTransTrack.Times;
+				GuidanceScale = ParentKeyScale;
+			}
+
+			// if the TargetBoneIndices array is empty, then this bone is an end effector.
+			// so we add it to the list to maintain our tolerance checks
+			if (TargetBoneIndices.Num() == 0)
@@ -887,111 +1973,262 @@
 			{
 				if (bHasScale)
 				{
-					FScaleTrack& ScaleTrack = ScaleTracks[TrackIndex];
+					FScaleTrack& ScaleTrack = Context.Segment.ScaleData[TrackIndex];
 					// filter out translations we can approximate through interpolation
-					FilterLinearKeysTemplate<FVector>(
-						ScaleTrack.ScaleKeys, 
-						ScaleTrack.Times, 
+					FilterLinearKeysTemplate<ScaleAdapter>(
+						ScaleTrack.ScaleKeys,
+						ScaleTrack.Times,
 						BoneAtoms,
-						GuidanceTrack, 
-						RawWorldBones,
-						NewWorldBones,
+						GuidanceTrack,
+						Context,
 						TargetBoneIndices,
-						NumFrames,
 						BoneIndex,
 						ParentBoneIndex,
-						GuidanceScale, 
-						MaxScaleDiff, 
+						GuidanceScale,
+						MaxScaleDiff,
 						EndEffectorTolerance,
-						EffectorDiffSocket,
-						BoneData);
+						EffectorDiffSocket);
 
 					// update our bone table from the current bone through the last end effector we need to test
-					UpdateWorldBoneTransformRange(
-						AnimSeq, 
-						BoneData, 
-						RefPose,
-						PositionTracks,
-						RotationTracks,
-						ScaleTracks,
-						BoneIndex,
-						HighestTargetBoneIndex,
-						false,
-						NewWorldBones);
+					UpdateWorldBoneTransformRange(Context, BoneIndex, HighestTargetBoneIndex);
 
 					// rebuild the BoneAtoms table using the current set of keys
-					UpdateBoneAtomList(AnimSeq, BoneIndex, TrackIndex, NumFrames, TimePerFrame, BoneAtoms); 
+					UpdateBoneAtomList(Context, TrackIndex, BoneAtoms);
 				}
 
 				// filter out translations we can approximate through interpolation
-				FilterLinearKeysTemplate<FVector>(
-					TransTrack.PosKeys, 
-					TransTrack.Times, 
+				FilterLinearKeysTemplate<TranslationAdapter>(
+					TransTrack.PosKeys,
+					TransTrack.Times,
 					BoneAtoms,
-					GuidanceTrack, 
-					RawWorldBones,
-					NewWorldBones,
+					GuidanceTrack,
+					Context,
 					TargetBoneIndices,
-					NumFrames,
 					BoneIndex,
 					ParentBoneIndex,
-					GuidanceScale, 
-					MaxPosDiff, 
+					GuidanceScale,
+					MaxPosDiff,
 					EndEffectorTolerance,
-					EffectorDiffSocket,
-					BoneData);
+					EffectorDiffSocket);
 
 				// update our bone table from the current bone through the last end effector we need to test
-				UpdateWorldBoneTransformRange(
-					AnimSeq, 
-					BoneData, 
-					RefPose,
-					PositionTracks,
-					RotationTracks,
-					ScaleTracks,
-					BoneIndex,
-					HighestTargetBoneIndex,
-					false,
-					NewWorldBones);
+				UpdateWorldBoneTransformRange(Context, BoneIndex, HighestTargetBoneIndex);
 
 				// rebuild the BoneAtoms table using the current set of keys
-				UpdateBoneAtomList(AnimSeq, BoneIndex, TrackIndex, NumFrames, TimePerFrame, BoneAtoms); 
+				UpdateBoneAtomList(Context, TrackIndex, BoneAtoms);
 
 				// filter out rotations we can approximate through interpolation
-				FilterLinearKeysTemplate<FQuat>(
-					RotTrack.RotKeys, 
-					RotTrack.Times, 
+				FilterLinearKeysTemplate<RotationAdapter>(
+					RotTrack.RotKeys,
+					RotTrack.Times,
 					BoneAtoms,
-					GuidanceTrack, 
-					RawWorldBones,
-					NewWorldBones,
+					GuidanceTrack,
+					Context,
 					TargetBoneIndices,
-					NumFrames,
 					BoneIndex,
 					ParentBoneIndex,
-					GuidanceScale, 
-					MaxAngleDiff, 
+					GuidanceScale,
+					MaxAngleDiff,
 					EndEffectorTolerance,
-					EffectorDiffSocket,
-					BoneData);
+					EffectorDiffSocket);
 			}
 		}
 
-		// make sure the final compressed keys are repesented in our NewWorldBones table
-		UpdateWorldBoneTransformRange(
-			AnimSeq, 
-			BoneData, 
-			RefPose,
-			PositionTracks,
-			RotationTracks,
-			ScaleTracks,
-			BoneIndex,
-			BoneIndex,
-			false,
-			NewWorldBones);
+		// make sure the final compressed keys are represented in our NewWorldBones table
+		UpdateWorldBoneTransformRange(Context, BoneIndex, BoneIndex);
+	}
+}
+
+/**
+ * Struct that holds the relevant information to process animation tracks in parallel.
+ * Instances of this structure are live as long as parallel task instances are live.
+ */
+struct FAsyncProcessAnimationTracksTaskGroupContext
+{
+	TArray<FProcessAnimationTracksContext*> TaskContexes;
+	volatile int32 AtomicTaskIndexCounter;
+	volatile int32 AtomicNumExecutedTasks;
+
+	FAsyncProcessAnimationTracksTaskGroupContext()
+		: AtomicTaskIndexCounter(0)
+		, AtomicNumExecutedTasks(0)
+	{}
+
+	~FAsyncProcessAnimationTracksTaskGroupContext()
+	{
+		for (FProcessAnimationTracksContext* Context : TaskContexes)
+		{
+			delete Context;
+		}
+	}
+
+	void ExecuteTasks(UAnimCompress_RemoveLinearKeys* Compressor)
+	{
+		while (true)
+		{
+			const int32 TaskIndex = FPlatformAtomics::InterlockedIncrement(&AtomicTaskIndexCounter) - 1;
+			if (TaskIndex >= TaskContexes.Num())
+			{
+				break;
+			}
+
+			FProcessAnimationTracksContext* JobContext = TaskContexes[TaskIndex];
+			Compressor->ProcessAnimationTracks(*JobContext);
+
+			FPlatformAtomics::InterlockedIncrement(&AtomicNumExecutedTasks);
+		}
+	}
+
+	void WaitForAllTasks()
+	{
+		// We just spin wait until everything is done
+		// This is a decent option because segments are already sorted largest to smallest and so
+		// they should all take about the same amount of time. We should never end up waiting here for too long.
+		while (AtomicNumExecutedTasks != TaskContexes.Num());
 	}
 };
 
+/**
+ * Class that represents a parallel task to perform linear key reduction.
+ */
+class FAsyncProcessAnimationTracksTask
+{
+public:
+	FAsyncProcessAnimationTracksTask(UAnimCompress_RemoveLinearKeys* Compressor_, FAsyncProcessAnimationTracksTaskGroupContext* TaskGroupContext_)
+		: Compressor(Compressor_)
+		, TaskGroupContext(TaskGroupContext_)
+	{}
 
+	/** return the name of the task **/
+	static const TCHAR* GetTaskName() { return TEXT("FAsyncProcessAnimationTracksTask"); }
+	FORCEINLINE static TStatId GetStatId() { RETURN_QUICK_DECLARE_CYCLE_STAT(FAsyncProcessAnimationTracksTask, STATGROUP_TaskGraphTasks); }
+
+	static ENamedThreads::Type GetDesiredThread() { return ENamedThreads::AnyThread; }
+	static ESubsequentsMode::Type GetSubsequentsMode() { return ESubsequentsMode::TrackSubsequents; }
+
+	void DoTask(ENamedThreads::Type CurrentThread, const FGraphEventRef& MyCompletionGraphEvent)
+	{
+		TaskGroupContext->ExecuteTasks(Compressor);
+	}
+
+	UAnimCompress_RemoveLinearKeys* Compressor;
+	FAsyncProcessAnimationTracksTaskGroupContext* TaskGroupContext;
+};
+
+/**
+ * Class that represents a clean up task once every parallel task has executed.
+ * This is executed last.
+ */
+class FAsyncCleanUpProcessAnimationTracksTask
+{
+public:
+	FAsyncCleanUpProcessAnimationTracksTask(FAsyncProcessAnimationTracksTaskGroupContext* TaskGroupContext_)
+		: TaskGroupContext(TaskGroupContext_)
+	{}
+
+	/** return the name of the task **/
+	static const TCHAR* GetTaskName() { return TEXT("FAsyncCleanUpProcessAnimationTracksTask"); }
+	FORCEINLINE static TStatId GetStatId() { RETURN_QUICK_DECLARE_CYCLE_STAT(FAsyncCleanUpProcessAnimationTracksTask, STATGROUP_TaskGraphTasks); }
+
+	static ENamedThreads::Type GetDesiredThread() { return ENamedThreads::AnyThread; }
+	static ESubsequentsMode::Type GetSubsequentsMode() { return ESubsequentsMode::TrackSubsequents; }
+
+	void DoTask(ENamedThreads::Type CurrentThread, const FGraphEventRef& MyCompletionGraphEvent)
+	{
+		delete TaskGroupContext;
+	}
+
+	FAsyncProcessAnimationTracksTaskGroupContext* TaskGroupContext;
+};
+
+void UAnimCompress_RemoveLinearKeys::ProcessAnimationTracks(
+	UAnimSequence& AnimSeq,
+	const TArray<FBoneData>& BoneData,
+	TArray<FAnimSegmentContext>& RawSegments)
+{
+	// extract all the data we'll need about the skeleton and animation sequence
+	const int32 NumBones = BoneData.Num();
+
+	// make sure the parent key scale is properly bound to 1.0 or more
+	ParentKeyScale = FMath::Max(ParentKeyScale, 1.0f);
+
+	TArray<int32> BoneIndexToTrackIndex;
+	BoneIndexToTrackIndex.Empty(NumBones);
+	BoneIndexToTrackIndex.AddUninitialized(NumBones);
+	for (int32 BoneIndex = 0; BoneIndex < NumBones; ++BoneIndex)
+	{
+		BoneIndexToTrackIndex[BoneIndex] = AnimSeq.GetSkeleton()->GetAnimationTrackIndex(BoneIndex, &AnimSeq, true);
+	}
+
+	// generate an array to hold the indices of our end effectors
+	TArray<int32> EndEffectors;
+	for (int32 BoneIndex = 0; BoneIndex < NumBones; ++BoneIndex)
+	{
+		const FBoneData& Bone = BoneData[BoneIndex];
+		if (Bone.IsEndEffector())
+		{
+			EndEffectors.Add(BoneIndex);
+		}
+	}
+
+	const int32 NumTracks = RawSegments[0].RotationData.Num();
+	TArray<FTransform> TrivialTrackValues;
+	TrivialTrackValues.Empty(NumTracks);
+	TrivialTrackValues.AddUninitialized(NumTracks);
+	for (int32 TrackIndex = 0; TrackIndex < NumTracks; ++TrackIndex)
+	{
+		const FQuat& Rotation = RawSegments[0].RotationData[TrackIndex].RotKeys[0];
+		const FVector& Translation = RawSegments[0].TranslationData[TrackIndex].PosKeys[0];
+		FVector Scale = FVector::OneVector;
+		if (RawSegments[0].ScaleData.Num() > 0)
+		{
+			Scale = RawSegments[0].ScaleData[TrackIndex].ScaleKeys[0];
+		}
+
+		TrivialTrackValues[TrackIndex] = FTransform(Rotation, Translation, Scale);
+	}
+
+	if (bUseDecompression || !bUseMultithreading || RawSegments.Num() <= 1)
+	{
+		for (FAnimSegmentContext& Segment : RawSegments)
+		{
+			FProcessAnimationTracksContext Context(AnimSeq, BoneData, EndEffectors, BoneIndexToTrackIndex, TrivialTrackValues, Segment, RawSegments);
+			ProcessAnimationTracks(Context);
+		}
+	}
+	else
+	{
+		// Created the context objects.
+		FAsyncProcessAnimationTracksTaskGroupContext* TaskGroupContext = new FAsyncProcessAnimationTracksTaskGroupContext();
+		for (FAnimSegmentContext& Segment : RawSegments)
+		{
+			FProcessAnimationTracksContext* Context = new FProcessAnimationTracksContext(AnimSeq, BoneData, EndEffectors, BoneIndexToTrackIndex, TrivialTrackValues, Segment, RawSegments);
+			TaskGroupContext->TaskContexes.Add(Context);
+		}
+
+		// Dispatch 1 task per job thread.
+		FGraphEventArray AsyncTaskCompletionEvents;
+		const int32 NumTaskThreads = FTaskGraphInterface::Get().GetNumWorkerThreads();
+		for (int32 TaskIndex = 0; TaskIndex < NumTaskThreads; ++TaskIndex)
+		{
+			AsyncTaskCompletionEvents.Add(TGraphTask<FAsyncProcessAnimationTracksTask>::CreateTask(NULL, ENamedThreads::AnyThread).ConstructAndDispatchWhenReady(this, TaskGroupContext));
+		}
+
+		// Execute the contexts concurrently.
+		TaskGroupContext->ExecuteTasks(this);
+
+		// Wait for all concurrent tasks to be done, we only wait for ones that were executing.
+		TaskGroupContext->WaitForAllTasks();
+
+		// Dispatch a cleanup job that will execute once all the tasks are done.
+		// The current thread/task no longer needs the TaskGroupContext beyond this point nor anything it references.
+		// All tasks that were executing actual compression work are done or queued up and pending
+		// and will do nothing since no work remains (they only touch the TaskGroupContext to find more work).
+		// Once all tasks are actually done, the clean up task will execute, freeing memory.
+		TGraphTask<FAsyncCleanUpProcessAnimationTracksTask>::CreateTask(&AsyncTaskCompletionEvents, ENamedThreads::AnyThread).ConstructAndDispatchWhenReady(TaskGroupContext);
+	}
+}
+
+
 void UAnimCompress_RemoveLinearKeys::CompressUsingUnderlyingCompressor(
 	UAnimSequence* AnimSeq,
@@ -1016,6 +2253,35 @@
 	AnimationFormat_SetInterfaceLinks(*AnimSeq);
 }
 
+void UAnimCompress_RemoveLinearKeys::CompressUsingUnderlyingCompressor(
+	UAnimSequence& AnimSeq,
+	const TArray<FBoneData>& BoneData,
+	TArray<FAnimSegmentContext>& RawSegments,
+	const bool bFinalPass)
+{
+	// record the proper runtime decompressor to use
+	AnimSeq.KeyEncodingFormat = AKF_VariableKeyLerp;
+	AnimSeq.TranslationCompressionFormat = TranslationCompressionFormat;
+	AnimSeq.RotationCompressionFormat = RotationCompressionFormat;
+	AnimSeq.ScaleCompressionFormat = ScaleCompressionFormat;
+	AnimationFormat_SetInterfaceLinks(AnimSeq);
+
+	const bool bIsSorted = bOptimizeForForwardPlayback;
+
+	BitwiseCompressAnimationTracks(
+		AnimSeq,
+		static_cast<AnimationCompressionFormat>(TranslationCompressionFormat),
+		static_cast<AnimationCompressionFormat>(RotationCompressionFormat),
+		static_cast<AnimationCompressionFormat>(ScaleCompressionFormat),
+		RawSegments,
+		bIsSorted);
+
+	CoalesceCompressedSegments(AnimSeq, RawSegments, bIsSorted);
+
+	// We could be invalid, set the links again
+	AnimationFormat_SetInterfaceLinks(AnimSeq);
+}
+
 void UAnimCompress_RemoveLinearKeys::DoReduction(UAnimSequence* AnimSeq, const TArray<FBoneData>& BoneData)
 {
 #if WITH_EDITORONLY_DATA
@@ -1023,11 +2289,6 @@
 	// One of these will always be true for the base class, but derived classes may choose to turn both off (e.g., in PerTrackCompression)
 	const bool bRunningProcessor = bRetarget || bActuallyFilterLinearKeys;
 
-	if (GIsEditor && bRunningProcessor)
-	{
-		GWarn->BeginSlowTask( NSLOCTEXT("UAnimCompress_RemoveLinearKeys", "BeginReductionTaskMessage", "Compressing animation with a RemoveLinearKeys scheme."), false);
-	}
-
 	// If the processor is to be run, then additive animations need to be converted from relative to absolute
 	const bool bNeedToConvertBackToAdditive = bRunningProcessor ? ConvertFromRelativeSpace(AnimSeq) : false;
 
@@ -1038,6 +2299,82 @@
 	SeparateRawDataIntoTracks(AnimSeq->GetRawAnimationData(), AnimSeq->SequenceLength, TranslationData, RotationData, ScaleData);
 	FilterBeforeMainKeyRemoval(AnimSeq, BoneData, TranslationData, RotationData, ScaleData);
 
+#if USE_SEGMENTING_CONTEXT
+	if (bEnableSegmenting)
+	{
+#if TIME_LINEAR_KEY_REMOVAL
+		double TimeStart = FPlatformTime::Seconds();
+#endif
+
+		TArray<FAnimSegmentContext> RawSegments;
+		SeparateRawDataIntoTracks(*AnimSeq, TranslationData, RotationData, ScaleData, IdealNumFramesPerSegment, MaxNumFramesPerSegment, RawSegments);
+
+		if (bRunningProcessor)
+		{
+			if (bUseDecompression)
+			{
+				// compress this animation without any key-reduction to prime the codec
+				CompressUsingUnderlyingCompressor(
+					*AnimSeq,
+					BoneData,
+					RawSegments,
+					false);
+			}
+			else
+			{
+				// record the proper runtime decompressor to use
+				AnimSeq->KeyEncodingFormat = AKF_VariableKeyLerp;
+				AnimSeq->TranslationCompressionFormat = TranslationCompressionFormat;
+				AnimSeq->RotationCompressionFormat = RotationCompressionFormat;
+				AnimSeq->ScaleCompressionFormat = ScaleCompressionFormat;
+				AnimationFormat_SetInterfaceLinks(*AnimSeq);
+
+				AnimSeq->CompressedByteStream.Empty();
+				AnimSeq->CompressedScaleOffsets.Empty();
+				AnimSeq->CompressedTrackOffsets.Empty();
+				AnimSeq->CompressedSegments.Empty();
+
+				for (FAnimSegmentContext& Segment : RawSegments)
+				{
+					Segment.TranslationCompressionFormat = TranslationCompressionFormat;
+					Segment.RotationCompressionFormat = RotationCompressionFormat;
+					Segment.ScaleCompressionFormat = ScaleCompressionFormat;
+				}
+			}
+
+			// now remove the keys which can be approximated with linear interpolation
+			ProcessAnimationTracks(
+				*AnimSeq,
+				BoneData,
+				RawSegments);
+
+#if TIME_LINEAR_KEY_REMOVAL
+			double ElapsedTime = FPlatformTime::Seconds() - TimeStart;
+			UE_LOG(LogAnimationCompression, Log, TEXT("ProcessAnimationTracks time is (%f) seconds"), ElapsedTime);
+#endif
+
+			// if previously additive, convert back to relative-space
+			if (bNeedToConvertBackToAdditive)
+			{
+				ConvertToRelativeSpace(*AnimSeq);
+
+				for (FAnimSegmentContext& Segment : RawSegments)
+				{
+					ConvertToRelativeSpace(*AnimSeq, Segment.TranslationData, Segment.RotationData, Segment.ScaleData);
+				}
+			}
+		}
+
+		// compress the final (possibly key-reduced) tracks into the anim sequence buffers
+		CompressUsingUnderlyingCompressor(
+			*AnimSeq,
+			BoneData,
+			RawSegments,
+			true);
+
+		return;
+	}
+#endif
 	if (bRunningProcessor)
 	{
 #if TIME_LINEAR_KEY_REMOVAL
@@ -1071,9 +2408,6 @@
 		}
 	}
 
-	// Remove Translation Keys from tracks marked bAnimRotationOnly
-	FilterAnimRotationOnlyKeys(TranslationData, AnimSeq);
-
 	// compress the final (possibly key-reduced) tracks into the anim sequence buffers
 	CompressUsingUnderlyingCompressor(
 		AnimSeq,
@@ -1083,12 +2417,6 @@
 		ScaleData,
 		true);
 
-	if (GIsEditor && bRunningProcessor)
-	{
-		GWarn->EndSlowTask();
-	}
-	AnimSeq->CompressionScheme = static_cast<UAnimCompress*>( StaticDuplicateObject( this, AnimSeq ) );
-
 #endif // WITH_EDITORONLY_DATA
 }
 
@@ -1103,7 +2431,8 @@
 	Ar << EffectorDiffSocket;
 	Ar << ParentKeyScale;
 	uint8 Flags =	MakeBitForFlag(bRetarget, 0) +
-					MakeBitForFlag(bActuallyFilterLinearKeys, 1);
+					MakeBitForFlag(bActuallyFilterLinearKeys, 1) +
+					MakeBitForFlag(bOptimizeForForwardPlayback, 2);
 	Ar << Flags;
 }
 
diff --git a/Engine/Source/Runtime/Engine/Private/Animation/AnimCompress_RemoveTrivialKeys.cpp b/Engine/Source/Runtime/Engine/Private/Animation/AnimCompress_RemoveTrivialKeys.cpp
--- a/Engine/Source/Runtime/Engine/Private/Animation/AnimCompress_RemoveTrivialKeys.cpp
+++ b/Engine/Source/Runtime/Engine/Private/Animation/AnimCompress_RemoveTrivialKeys.cpp
@@ -26,27 +26,49 @@ void UAnimCompress_RemoveTrivialKeys::DoReduction(UAnimSequence* AnimSeq, const
 	TArray<FRotationTrack> RotationData;
 	TArray<FScaleTrack> ScaleData;
 	SeparateRawDataIntoTracks( AnimSeq->GetRawAnimationData(), AnimSeq->SequenceLength, TranslationData, RotationData, ScaleData );
-
-	// Remove Translation Keys from tracks marked bAnimRotationOnly
-	FilterAnimRotationOnlyKeys(TranslationData, AnimSeq);
 	
 	// remove obviously redundant keys from the source data
 	FilterTrivialKeys(TranslationData, RotationData, ScaleData, TRANSLATION_ZEROING_THRESHOLD, QUATERNION_ZEROING_THRESHOLD, SCALE_ZEROING_THRESHOLD);
 
-	// bitwise compress the tracks into the anim sequence buffers
-	BitwiseCompressAnimationTracks(
-		AnimSeq,
-		static_cast<AnimationCompressionFormat>(TranslationCompressionFormat),
-		static_cast<AnimationCompressionFormat>(RotationCompressionFormat),
-		static_cast<AnimationCompressionFormat>(ScaleCompressionFormat),
-		TranslationData,
-		RotationData,
-		ScaleData);
-
 	// record the proper runtime decompressor to use
 	AnimSeq->KeyEncodingFormat = AKF_ConstantKeyLerp;
 	AnimationFormat_SetInterfaceLinks(*AnimSeq);
-	AnimSeq->CompressionScheme = static_cast<UAnimCompress*>( StaticDuplicateObject( this, AnimSeq ) );
+
+#if USE_SEGMENTING_CONTEXT
+	if (bEnableSegmenting)
+	{
+		TArray<FAnimSegmentContext> RawSegments;
+		SeparateRawDataIntoTracks(*AnimSeq, TranslationData, RotationData, ScaleData, IdealNumFramesPerSegment, MaxNumFramesPerSegment, RawSegments);
+
+		BitwiseCompressAnimationTracks(
+			*AnimSeq,
+			static_cast<AnimationCompressionFormat>(TranslationCompressionFormat),
+			static_cast<AnimationCompressionFormat>(RotationCompressionFormat),
+			static_cast<AnimationCompressionFormat>(ScaleCompressionFormat),
+			RawSegments);
+
+		CoalesceCompressedSegments(*AnimSeq, RawSegments);
+
+		AnimSeq->TranslationCompressionFormat = TranslationCompressionFormat;
+		AnimSeq->RotationCompressionFormat = RotationCompressionFormat;
+		AnimSeq->ScaleCompressionFormat = ScaleCompressionFormat;
+	}
+	else
+#endif
+	{
+		// bitwise compress the tracks into the anim sequence buffers
+		BitwiseCompressAnimationTracks(
+			AnimSeq,
+			static_cast<AnimationCompressionFormat>(TranslationCompressionFormat),
+			static_cast<AnimationCompressionFormat>(RotationCompressionFormat),
+			static_cast<AnimationCompressionFormat>(ScaleCompressionFormat),
+			TranslationData,
+			RotationData,
+			ScaleData);
+	}
+
+	// We could be invalid, set the links again
+	AnimationFormat_SetInterfaceLinks(*AnimSeq);
 #endif // WITH_EDITORONLY_DATA
 }
 
diff --git a/Engine/Source/Runtime/Engine/Private/Animation/AnimCompressionDerivedData.cpp b/Engine/Source/Runtime/Engine/Private/Animation/AnimCompressionDerivedData.cpp
--- a/Engine/Source/Runtime/Engine/Private/Animation/AnimCompressionDerivedData.cpp
+++ b/Engine/Source/Runtime/Engine/Private/Animation/AnimCompressionDerivedData.cpp
@@ -80,7 +80,7 @@ bool FDerivedDataAnimationCompression::Build( TArray<uint8>& OutData )
 
 	if (!bDoCompressionInPlace)
 	{
-		DuplicateSequence = DuplicateObject<UAnimSequence>(OriginalAnimSequence, GetTransientPackage());
+		DuplicateSequence = DuplicateObject<UAnimSequence>(OriginalAnimSequence, GetTransientPackage(), OriginalAnimSequence->GetFName());
 		DuplicateSequence->AddToRoot();
 		AnimToOperateOn = DuplicateSequence;
 	}
@@ -124,8 +124,6 @@ bool FDerivedDataAnimationCompression::Build( TArray<uint8>& OutData )
 		const TCHAR* OutputStr = CompressContext.Get()->bOutput ? TEXT("true") : TEXT("false");
 #endif
 
-		AnimToOperateOn->CompressedByteStream.Reset();
-		AnimToOperateOn->CompressedTrackOffsets.Reset();
 		FAnimationUtils::CompressAnimSequence(AnimToOperateOn, *CompressContext.Get());
 		bCompressionSuccessful = AnimToOperateOn->IsCompressedDataValid();
 
diff --git a/Engine/Source/Runtime/Engine/Private/Animation/AnimCompressionDerivedData.h b/Engine/Source/Runtime/Engine/Private/Animation/AnimCompressionDerivedData.h
--- a/Engine/Source/Runtime/Engine/Private/Animation/AnimCompressionDerivedData.h
+++ b/Engine/Source/Runtime/Engine/Private/Animation/AnimCompressionDerivedData.h
@@ -44,7 +44,7 @@ public:
 		// This is a version string that mimics the old versioning scheme. If you
 		// want to bump this version, generate a new guid using VS->Tools->Create GUID and
 		// return it here. Ex.
-		return TEXT("A6E18B049A324E8B8946E9CA6F32CE38");
+		return TEXT("EFEE430B29384DD38AD512590E2C4821");
 	}
 
 	virtual FString GetPluginSpecificCacheKeySuffix() const override;
diff --git a/Engine/Source/Runtime/Engine/Private/Animation/AnimEncoding.cpp b/Engine/Source/Runtime/Engine/Private/Animation/AnimEncoding.cpp
--- a/Engine/Source/Runtime/Engine/Private/Animation/AnimEncoding.cpp
+++ b/Engine/Source/Runtime/Engine/Private/Animation/AnimEncoding.cpp
@@ -182,18 +182,16 @@ void PadMemoryReader(FMemoryReader* MemoryReader, uint8*& TrackData, const int32
  * Extracts a single BoneAtom from an Animation Sequence.
  *
  * @param	OutAtom			The BoneAtom to fill with the extracted result.
- * @param	Seq				An Animation Sequence to extract the BoneAtom from.
+ * @param	DecompContext	The decompression context to use.
  * @param	TrackIndex		The index of the track desired in the Animation Sequence.
- * @param	Time			The time (in seconds) to calculate the BoneAtom for.
  */
-void AnimationFormat_GetBoneAtom(	
+void AnimationFormat_GetBoneAtom(
 	FTransform& OutAtom,
-	const UAnimSequence& Seq,
-	int32 TrackIndex,
-	float Time)
+	FAnimSequenceDecompressionContext& DecompContext,
+	int32 TrackIndex)
 {
-	checkSlow(Seq.RotationCodec != NULL);
-	((AnimEncoding*)Seq.RotationCodec)->GetBoneAtom(OutAtom, Seq, TrackIndex, Time);
+	checkSlow(DecompContext.GetRotationCodec() != nullptr);
+	DecompContext.GetRotationCodec()->GetBoneAtom(OutAtom, DecompContext, TrackIndex);
 }
 
 #if USE_ANIMATION_CODEC_BATCH_SOLVER
@@ -202,35 +200,34 @@ void AnimationFormat_GetBoneAtom(
  * Extracts an array of BoneAtoms from an Animation Sequence representing an entire pose of the skeleton.
  *
  * @param	Atoms				The BoneAtoms to fill with the extracted result.
- * @param	RotationTracks		A BoneTrackArray element for each bone requesting rotation data. 
- * @param	TranslationTracks	A BoneTrackArray element for each bone requesting translation data. 
- * @param	Seq					An Animation Sequence to extract the BoneAtom from.
- * @param	Time				The time (in seconds) to calculate the BoneAtom for.
+ * @param	RotationTracks		A BoneTrackArray element for each bone requesting rotation data.
+ * @param	TranslationTracks	A BoneTrackArray element for each bone requesting translation data.
+ * @param	ScaleTracks			A BoneTrackArray element for each bone requesting scale data.
+ * @param	DecompContext		The decompression context to use.
  */
-void AnimationFormat_GetAnimationPose(	
-	FTransformArray& Atoms, 
+void AnimationFormat_GetAnimationPose(
+	FTransformArray& Atoms,
 	const BoneTrackArray& RotationPairs,
 	const BoneTrackArray& TranslationPairs,
 	const BoneTrackArray& ScalePairs,
-	const UAnimSequence& Seq,
-	float Time)
+	FAnimSequenceDecompressionContext& DecompContext)
 {
 	// decompress the translation component using the proper method
-	checkSlow(Seq.TranslationCodec != NULL);
+	checkSlow(DecompContext.GetTranslationCodec() != nullptr);
 	if (TranslationPairs.Num() > 0)
 	{
-		((AnimEncoding*)Seq.TranslationCodec)->GetPoseTranslations(Atoms, TranslationPairs, Seq, Time);
+		DecompContext.GetTranslationCodec()->GetPoseTranslations(Atoms, TranslationPairs, DecompContext);
 	}
 
 	// decompress the rotation component using the proper method
-	checkSlow(Seq.RotationCodec != NULL);
-	((AnimEncoding*)Seq.RotationCodec)->GetPoseRotations(Atoms, RotationPairs, Seq, Time);
+	checkSlow(DecompContext.GetRotationCodec() != nullptr);
+	DecompContext.GetRotationCodec()->GetPoseRotations(Atoms, RotationPairs, DecompContext);
 
-	checkSlow(Seq.ScaleCodec != NULL);
+	checkSlow(DecompContext.GetScaleCodec() != nullptr);
 	// we allow scale key to be empty
-	if (Seq.CompressedScaleOffsets.IsValid())
+	if (DecompContext.bHasScale)
 	{
-		((AnimEncoding*)Seq.ScaleCodec)->GetPoseScales(Atoms, ScalePairs, Seq, Time);
+		DecompContext.GetScaleCodec()->GetPoseScales(Atoms, ScalePairs, DecompContext);
 	}
 }
 #endif
@@ -239,48 +236,28 @@ void AnimationFormat_GetAnimationPose(
  * Extracts a single BoneAtom from an Animation Sequence.
  *
  * @param	OutAtom			The BoneAtom to fill with the extracted result.
- * @param	Seq				An Animation Sequence to extract the BoneAtom from.
+ * @param	DecompContext	The decompression context to use.
  * @param	TrackIndex		The index of the track desired in the Animation Sequence.
- * @param	Time			The time (in seconds) to calculate the BoneAtom for.
  */
-void AnimEncodingLegacyBase::GetBoneAtom(
-	FTransform& OutAtom,
-	const UAnimSequence& Seq,
-	int32 TrackIndex,
-	float Time)
+void AnimEncodingLegacyBase::GetBoneAtom(FTransform& OutAtom, FAnimSequenceDecompressionContext& DecompContext, int32 TrackIndex)
 {
 	// Initialize to identity to set the scale and in case of a missing rotation or translation codec
 	OutAtom.SetIdentity();
 
-	// Use the CompressedTrackOffsets stream to find the data addresses
-	const int32* RESTRICT TrackData= Seq.CompressedTrackOffsets.GetData() + (TrackIndex*4);
-	int32 TransKeysOffset = *(TrackData+0);
-	int32 NumTransKeys	= *(TrackData+1);
-	int32 RotKeysOffset	= *(TrackData+2);
-	int32 NumRotKeys		= *(TrackData+3);
-	const uint8* RESTRICT TransStream	= Seq.CompressedByteStream.GetData()+TransKeysOffset;
-	const uint8* RESTRICT RotStream		= Seq.CompressedByteStream.GetData()+RotKeysOffset;
-
-	const float RelativePos = Time / (float)Seq.SequenceLength;
-
 	// decompress the translation component using the proper method
-	checkSlow(Seq.TranslationCodec != NULL);
-	((AnimEncodingLegacyBase*)Seq.TranslationCodec)->GetBoneAtomTranslation(OutAtom, Seq, TransStream, NumTransKeys, Time, RelativePos);
+	checkSlow(DecompContext.GetTranslationCodec() != nullptr);
+	((AnimEncodingLegacyBase*)DecompContext.GetTranslationCodec())->GetBoneAtomTranslation(OutAtom, DecompContext, TrackIndex);
 
 	// decompress the rotation component using the proper method
-	checkSlow(Seq.RotationCodec != NULL);
-	((AnimEncodingLegacyBase*)Seq.RotationCodec)->GetBoneAtomRotation(OutAtom, Seq, RotStream, NumRotKeys, Time, RelativePos);
+	checkSlow(DecompContext.GetRotationCodec() != nullptr);
+	((AnimEncodingLegacyBase*)DecompContext.GetRotationCodec())->GetBoneAtomRotation(OutAtom, DecompContext, TrackIndex);
 
-	// we assume scale keys can be empty, so only extrace if we have valid keys
-	bool bHasValidScale = Seq.CompressedScaleOffsets.IsValid();
-	if (bHasValidScale)
+	// we assume scale keys can be empty, so only extract if we have valid keys
+	if (DecompContext.bHasScale)
 	{
-		int32 ScaleKeyOffset = Seq.CompressedScaleOffsets.GetOffsetData(TrackIndex, 0);
-		int32 NumScaleKeys = Seq.CompressedScaleOffsets.GetOffsetData(TrackIndex, 1);
-		const uint8* RESTRICT ScaleStream		= Seq.CompressedByteStream.GetData()+ScaleKeyOffset;
 		// decompress the rotation component using the proper method
-		checkSlow(Seq.ScaleCodec != NULL);
-		((AnimEncodingLegacyBase*)Seq.ScaleCodec)->GetBoneAtomScale(OutAtom, Seq, ScaleStream, NumScaleKeys, Time, RelativePos);
+		checkSlow(DecompContext.GetScaleCodec() != nullptr);
+		((AnimEncodingLegacyBase*)DecompContext.GetScaleCodec())->GetBoneAtomScale(OutAtom, DecompContext, TrackIndex);
 	}
 }
 
@@ -305,6 +282,17 @@ void AnimEncodingLegacyBase::ByteSwapIn(
 	Seq.CompressedByteStream.Empty(OriginalNumBytes);
 	Seq.CompressedByteStream.AddUninitialized(OriginalNumBytes);
 
+	if (Seq.CompressedSegments.Num() != 0)
+	{
+#if !PLATFORM_LITTLE_ENDIAN
+#error "Byte swapping needs to be implemented here to support big-endian platforms"
+#endif
+
+		// TODO: Byte swap the new format
+		MemoryReader.Serialize(Seq.CompressedByteStream.GetData(), Seq.CompressedByteStream.Num());
+		return;
+	}
+
 	// Read and swap
 	uint8* StreamBase = Seq.CompressedByteStream.GetData();
 	bool bHasValidScale = Seq.CompressedScaleOffsets.IsValid();
@@ -374,6 +362,13 @@ void AnimEncodingLegacyBase::ByteSwapOut(
 	FMemoryWriter MemoryWriter( SerializedData, true );
 	MemoryWriter.SetByteSwapping( ForceByteSwapping );
 
+	if (Seq.CompressedSegments.Num() != 0)
+	{
+		// TODO: Byte swap the new format
+		MemoryWriter.Serialize(Seq.CompressedByteStream.GetData(), Seq.CompressedByteStream.Num());
+		return;
+	}
+
 	uint8* StreamBase		= Seq.CompressedByteStream.GetData();
 	const int32 NumTracks		= Seq.CompressedTrackOffsets.Num()/4;
 
diff --git a/Engine/Source/Runtime/Engine/Private/Animation/AnimEncoding_ConstantKeyLerp.cpp b/Engine/Source/Runtime/Engine/Private/Animation/AnimEncoding_ConstantKeyLerp.cpp
--- a/Engine/Source/Runtime/Engine/Private/Animation/AnimEncoding_ConstantKeyLerp.cpp
+++ b/Engine/Source/Runtime/Engine/Private/Animation/AnimEncoding_ConstantKeyLerp.cpp
@@ -234,3 +234,60 @@ void AEFConstantKeyLerpShared::ByteSwapScaleOut(
 		}
 	}
 }
+
+#if USE_SEGMENTING_CONTEXT
+void AEFConstantKeyLerpShared::CreateEncodingContext(FAnimSequenceDecompressionContext& DecompContext)
+{
+	checkSlow(DecompContext.EncodingContext == nullptr);
+	DecompContext.EncodingContext = new FAEConstantKeyLerpContext(DecompContext);
+}
+
+void AEFConstantKeyLerpShared::ReleaseEncodingContext(FAnimSequenceDecompressionContext& DecompContext)
+{
+	checkSlow(DecompContext.EncodingContext != nullptr);
+	delete DecompContext.EncodingContext;
+	DecompContext.EncodingContext = nullptr;
+}
+
+FAEConstantKeyLerpContext::FAEConstantKeyLerpContext(const FAnimSequenceDecompressionContext& DecompContext)
+	: KeyFrameSize(-1)
+{
+}
+
+void FAEConstantKeyLerpContext::Seek(const FAnimSequenceDecompressionContext& DecompContext, float SampleAtTime)
+{
+	if (KeyFrameSize < 0)
+	{
+		// First update, cache some stuff
+		UniformKeyOffsets.Empty(DecompContext.NumTracks * DecompContext.NumStreamsPerTrack);
+		UniformKeyOffsets.AddUninitialized(DecompContext.NumTracks * DecompContext.NumStreamsPerTrack);
+
+		const int32 PackedTranslationSize0 = CompressedTranslationStrides[DecompContext.Segment0->TranslationCompressionFormat] * CompressedTranslationNum[DecompContext.Segment0->TranslationCompressionFormat];
+		const int32 PackedRotationSize0 = CompressedRotationStrides[DecompContext.Segment0->RotationCompressionFormat] * CompressedRotationNum[DecompContext.Segment0->RotationCompressionFormat];
+		const int32 PackedScaleSize0 = DecompContext.bHasScale ? (CompressedScaleStrides[DecompContext.Segment0->ScaleCompressionFormat] * CompressedScaleNum[DecompContext.Segment0->ScaleCompressionFormat]) : 0;
+
+		int32 KeyOffset = 0;
+		for (int32 TrackIndex = 0; TrackIndex < DecompContext.NumTracks; ++TrackIndex)
+		{
+			const FTrivialTrackFlags TrackFlags(DecompContext.TrackFlags[TrackIndex]);
+
+			UniformKeyOffsets[DecompContext.GetTranslationValueOffset(TrackIndex)] = KeyOffset;
+			KeyOffset += TrackFlags.IsTranslationTrivial() ? 0 : PackedTranslationSize0;
+
+			UniformKeyOffsets[DecompContext.GetRotationValueOffset(TrackIndex)] = KeyOffset;
+			KeyOffset += TrackFlags.IsRotationTrivial() ? 0 : PackedRotationSize0;
+
+			if (DecompContext.bHasScale)
+			{
+				UniformKeyOffsets[DecompContext.GetScaleValueOffset(TrackIndex)] = KeyOffset;
+				KeyOffset += TrackFlags.IsScaleTrivial() ? 0 : PackedScaleSize0;
+			}
+		}
+
+		KeyFrameSize = KeyOffset;
+	}
+
+	FrameKeysOffset[0] = DecompContext.Segment0->ByteStreamOffset + DecompContext.RangeDataSize0 + (KeyFrameSize * DecompContext.SegmentKeyIndex0);
+	FrameKeysOffset[1] = DecompContext.Segment1->ByteStreamOffset + DecompContext.RangeDataSize0 + (KeyFrameSize * DecompContext.SegmentKeyIndex1);
+}
+#endif
\ No newline at end of file
