Copying //UE4/Dev-AnimPhys to Dev-Main (//UE4/Dev-Main)

diff --git a/Engine/Source/Runtime/Engine/Private/Animation/AnimEncoding_PerTrackCompression.cpp b/Engine/Source/Runtime/Engine/Private/Animation/AnimEncoding_PerTrackCompression.cpp
--- a/Engine/Source/Runtime/Engine/Private/Animation/AnimEncoding_PerTrackCompression.cpp
+++ b/Engine/Source/Runtime/Engine/Private/Animation/AnimEncoding_PerTrackCompression.cpp
@@ -7,8 +7,73 @@
 #include "AnimEncoding_PerTrackCompression.h"
 #include "Serialization/MemoryWriter.h"
 #include "Serialization/MemoryReader.h"
+#include "Animation/AnimEncodingHeapAllocator.h"
+#include "Animation/AnimEncodingDecompressionContext.h"
 #include "AnimationCompression.h"
 
+
+class FAEPerTrackKeyLerpContext : public FAnimEncodingDecompressionContext
+{
+public:
+	static constexpr int32 OffsetNumKeysPairSize = sizeof(uint32) + sizeof(uint16);
+
+	FAEPerTrackKeyLerpContext(const FAnimSequenceDecompressionContext& DecompContext);
+	virtual void Seek(const FAnimSequenceDecompressionContext& DecompContext, float SampleAtTime) override;
+
+	inline void GetRotation(FTransform& OutAtom, const FAnimSequenceDecompressionContext& DecompContext, int32 TrackIndex) const;
+	inline void GetTranslation(FTransform& OutAtom, const FAnimSequenceDecompressionContext& DecompContext, int32 TrackIndex) const;
+	inline void GetScale(FTransform& OutAtom, const FAnimSequenceDecompressionContext& DecompContext, int32 TrackIndex) const;
+
+	// Common
+	TArray<int32, FAnimEncodingHeapAllocator> RangeOffsets[2];
+	int32 PerTrackStreamFlagOffsets[2];
+	int32 RangeDataSize[2];
+	uint16 PreviousSegmentIndex[2];
+
+	// Uniform
+	TArray<int32, FAnimEncodingHeapAllocator> UniformKeyOffsets[2];
+	int32 UniformKeyFrameSize[2];
+	int32 UniformDataOffsets[2];
+
+	// Variable common
+	TArray<uint8, FAnimEncodingHeapAllocator> TrackStreamKeySizes[2];
+
+	// Variable linear
+	float SegmentRelativePos0;
+	uint8 TimeMarkerSize[2];	// sizeof(uint8) or sizeof(uint16)
+	int32 OffsetNumKeysPairOffsets[2];
+	TArray<int32, FAnimEncodingHeapAllocator> NumAnimatedTrackStreams[2];
+
+	// Variable sorted
+	struct FCachedKey
+	{
+		// Linear interpolation only requires 2 keys
+		// Index 0 is always the oldest key, index 1 the newest
+		int32 RotOffsets[2];
+		int32 RotFrameIndices[2];
+		int32 TransOffsets[2];
+		int32 TransFrameIndices[2];
+		int32 ScaleOffsets[2];
+		int32 ScaleFrameIndices[2];
+	};
+
+	TArray<FCachedKey, FAnimEncodingHeapAllocator> CachedKeys;			// 1 Entry per track
+
+	int32 SegmentStartFrame[2];
+	float FramePos;
+	float PreviousSampleAtTime;
+
+	const uint8_t* PackedSampleData;		// The current pointer into our data stream
+	int32 PreviousFrameIndex;				// The previously read frame index
+	int32 CurrentFrameIndex;				// The current frame index
+	uint8 CurrentSegmentIndex;				// The current segment index of the sorted data, either 0 or 1
+
+private:
+	void CacheSegmentValues(const FAnimSequenceDecompressionContext& DecompContext, const FCompressedSegment& Segment, uint8 SegmentIndex);
+	void ResetSortedCache(const FAnimSequenceDecompressionContext& DecompContext);
+	void AdvanceSortedCachedKeys(const FAnimSequenceDecompressionContext& DecompContext);
+};
+
 // This define controls whether scalar or vector code is used to decompress keys.  Note that not all key decompression code
 // is vectorized yet, so some (seldom used) formats will actually get slower (due to extra LHS stalls) when enabled.
 // The code also relies on a flexible permute instruction being available (e.g., PPC vperm)
@@ -393,6 +458,17 @@ void AEFPerTrackCompressionCodec::ByteSwapIn(
 	Seq.CompressedByteStream.Empty(OriginalNumBytes);
 	Seq.CompressedByteStream.AddUninitialized(OriginalNumBytes);
 
+	if (Seq.CompressedSegments.Num() != 0)
+	{
+#if !PLATFORM_LITTLE_ENDIAN
+#error "Byte swapping needs to be implemented here to support big-endian platforms"
+#endif
+
+		// TODO: Byte swap the new format
+		MemoryReader.Serialize(Seq.CompressedByteStream.GetData(), Seq.CompressedByteStream.Num());
+		return;
+	}
+
 	const int32 NumTracks = Seq.CompressedTrackOffsets.Num() / 2;
 	const bool bHasScaleData = Seq.CompressedScaleOffsets.IsValid();
 
@@ -427,6 +503,13 @@ void AEFPerTrackCompressionCodec::ByteSwapOut(
 	FMemoryWriter MemoryWriter(SerializedData, true);
 	MemoryWriter.SetByteSwapping(ForceByteSwapping);
 
+	if (Seq.CompressedSegments.Num() != 0)
+	{
+		// TODO: Byte swap the new format
+		MemoryWriter.Serialize(Seq.CompressedByteStream.GetData(), Seq.CompressedByteStream.Num());
+		return;
+	}
+
 	const int32 NumTracks = Seq.CompressedTrackOffsets.Num() / 2;
 	const bool bHasScaleData = Seq.CompressedScaleOffsets.IsValid();
 	for ( int32 TrackIndex = 0; TrackIndex < NumTracks; ++TrackIndex )
@@ -450,51 +533,56 @@ void AEFPerTrackCompressionCodec::ByteSwapOut(
  * Extracts a single BoneAtom from an Animation Sequence.
  *
  * @param	OutAtom			The BoneAtom to fill with the extracted result.
- * @param	Seq				An Animation Sequence to extract the BoneAtom from.
+ * @param	DecompContext	The decompression context to use.
  * @param	TrackIndex		The index of the track desired in the Animation Sequence.
- * @param	Time			The time (in seconds) to calculate the BoneAtom for.
  */
 void AEFPerTrackCompressionCodec::GetBoneAtom(
 	FTransform& OutAtom,
-	const UAnimSequence& Seq,
-	int32 TrackIndex,
-	float Time)
+	FAnimSequenceDecompressionContext& DecompContext,
+	int32 TrackIndex)
 {
 	// Initialize to identity to set the scale and in case of a missing rotation or translation codec
 	OutAtom.SetIdentity();
 
-	// Use the CompressedTrackOffsets stream to find the data addresses
-	const int32* RESTRICT TrackData= Seq.CompressedTrackOffsets.GetData() + (TrackIndex * 2);
-	const int32 TransKeysOffset = TrackData[0];
-	const int32 RotKeysOffset = TrackData[1];
-	const float RelativePos = Time / (float)Seq.SequenceLength;
+	GetBoneAtomTranslation(OutAtom, DecompContext, TrackIndex);
+	GetBoneAtomRotation(OutAtom, DecompContext, TrackIndex);
 
-	GetBoneAtomTranslation(OutAtom, Seq, TransKeysOffset, Time, RelativePos);
-	GetBoneAtomRotation(OutAtom, Seq, RotKeysOffset, Time, RelativePos);
-	const bool bHasScaleData = Seq.CompressedScaleOffsets.IsValid();
-	if (bHasScaleData)
+	if (DecompContext.bHasScale)
 	{
-		const int32 ScaleKeysOffset = Seq.CompressedScaleOffsets.GetOffsetData(TrackIndex, 0);
-		GetBoneAtomScale(OutAtom, Seq, ScaleKeysOffset, Time, RelativePos);
+		GetBoneAtomScale(OutAtom, DecompContext, TrackIndex);
 	}
 }
-	
-
-
-
-
 
 void AEFPerTrackCompressionCodec::GetBoneAtomRotation(
 	FTransform& OutAtom,
-	const UAnimSequence& Seq,
-	int32 Offset,
-	float Time,
-	float RelativePos)
+	FAnimSequenceDecompressionContext& DecompContext,
+	int32 TrackIndex)
 {
-	if (Offset != INDEX_NONE)
+#if USE_SEGMENTING_CONTEXT
+	if (DecompContext.AnimSeq->CompressedSegments.Num() != 0)
 	{
-		const uint8* RESTRICT TrackData = Seq.CompressedByteStream.GetData() + Offset + 4;
-		const int32 Header = *((int32*)(Seq.CompressedByteStream.GetData() + Offset));
+		const FTrivialAnimKeyHandle TrivialKeyHandle = DecompContext.GetTrivialRotationKeyHandle(TrackIndex);
+		if (TrivialKeyHandle.IsValid())
+		{
+			DecompContext.GetTrivialRotation(OutAtom, TrivialKeyHandle);
+		}
+		else
+		{
+			const FAEPerTrackKeyLerpContext& EncodingContext = *static_cast<const FAEPerTrackKeyLerpContext*>(DecompContext.EncodingContext);
+			EncodingContext.GetRotation(OutAtom, DecompContext, TrackIndex);
+		}
+
+		return;
+	}
+#endif
+
+	const int32* RESTRICT TrackOffsetData = DecompContext.GetCompressedTrackOffsets() + (TrackIndex * 2);
+	const int32 RotKeysOffset = TrackOffsetData[1];
+
+	if (RotKeysOffset != INDEX_NONE)
+	{
+		const uint8* RESTRICT TrackData = DecompContext.GetCompressedByteStream() + RotKeysOffset + 4;
+		const int32 Header = *((int32*)(DecompContext.GetCompressedByteStream() + RotKeysOffset));
 
 		int32 KeyFormat;
 		int32 NumKeys;
@@ -514,12 +602,12 @@ void AEFPerTrackCompressionCodec::GetBoneAtomRotation(
 		{
 			if ((FormatFlags & 0x8) == 0)
 			{
-				Alpha = TimeToIndex(Seq, RelativePos, NumKeys, Index0, Index1);
+				Alpha = TimeToIndex(*DecompContext.AnimSeq, DecompContext.RelativePos, NumKeys, Index0, Index1);
 			}
 			else
 			{
 				const uint8* RESTRICT FrameTable = Align(TrackData + FixedBytes + BytesPerKey * NumKeys, 4);
-				Alpha = TimeToIndex(Seq, FrameTable, RelativePos, NumKeys, Index0, Index1);
+				Alpha = TimeToIndex(*DecompContext.AnimSeq, FrameTable, DecompContext.RelativePos, NumKeys, Index0, Index1);
 			}
 		}
 
@@ -570,18 +658,36 @@ void AEFPerTrackCompressionCodec::GetBoneAtomRotation(
 }
 
 
-
 void AEFPerTrackCompressionCodec::GetBoneAtomTranslation(
 	FTransform& OutAtom,
-	const UAnimSequence& Seq,
-	int32 Offset,
-	float Time,
-	float RelativePos)
+	FAnimSequenceDecompressionContext& DecompContext,
+	int32 TrackIndex)
 {
-	if (Offset != INDEX_NONE)
+#if USE_SEGMENTING_CONTEXT
+	if (DecompContext.AnimSeq->CompressedSegments.Num() != 0)
+	{
+		const FTrivialAnimKeyHandle TrivialKeyHandle = DecompContext.GetTrivialTranslationKeyHandle(TrackIndex);
+		if (TrivialKeyHandle.IsValid())
+		{
+			DecompContext.GetTrivialTranslation(OutAtom, TrivialKeyHandle);
+		}
+		else
+		{
+			const FAEPerTrackKeyLerpContext& EncodingContext = *static_cast<const FAEPerTrackKeyLerpContext*>(DecompContext.EncodingContext);
+			EncodingContext.GetTranslation(OutAtom, DecompContext, TrackIndex);
+		}
+
+		return;
+	}
+#endif
+
+	const int32* RESTRICT TrackOffsetData = DecompContext.GetCompressedTrackOffsets() + (TrackIndex * 2);
+	const int32 PosKeysOffset = TrackOffsetData[0];
+
+	if (PosKeysOffset != INDEX_NONE)
 	{
-		const uint8* RESTRICT TrackData = Seq.CompressedByteStream.GetData() + Offset + 4;
-		const int32 Header = *((int32*)(Seq.CompressedByteStream.GetData() + Offset));
+		const uint8* RESTRICT TrackData = DecompContext.GetCompressedByteStream() + PosKeysOffset + 4;
+		const int32 Header = *((int32*)(DecompContext.GetCompressedByteStream() + PosKeysOffset));
 
 		int32 KeyFormat;
 		int32 NumKeys;
@@ -590,7 +696,7 @@ void AEFPerTrackCompressionCodec::GetBoneAtomTranslation(
 		int32 FixedBytes;
 		FAnimationCompression_PerTrackUtils::DecomposeHeader(Header, /*OUT*/ KeyFormat, /*OUT*/ NumKeys, /*OUT*/ FormatFlags, /*OUT*/BytesPerKey, /*OUT*/ FixedBytes);
 
-		checkf(KeyFormat != ACF_None, TEXT("[%s] contians invalid keyformat. NumKeys (%d), FormatFlags (%d), BytesPerKeys (%d), FixedBytes (%d)"), *Seq.GetName(), NumKeys, FormatFlags, BytesPerKey, FixedBytes);
+		checkf(KeyFormat != ACF_None, TEXT("[%s] contians invalid keyformat. NumKeys (%d), FormatFlags (%d), BytesPerKeys (%d), FixedBytes (%d)"), *DecompContext.AnimSeq->GetName(), NumKeys, FormatFlags, BytesPerKey, FixedBytes);
 
 		// Figure out the key indexes
 		int32 Index0 = 0;
@@ -603,12 +709,12 @@ void AEFPerTrackCompressionCodec::GetBoneAtomTranslation(
 		{
 			if ((FormatFlags & 0x8) == 0)
 			{
-				Alpha = TimeToIndex(Seq, RelativePos, NumKeys, Index0, Index1);
+				Alpha = TimeToIndex(*DecompContext.AnimSeq, DecompContext.RelativePos, NumKeys, Index0, Index1);
 			}
 			else
 			{
 				const uint8* RESTRICT FrameTable = Align(TrackData + FixedBytes + BytesPerKey * NumKeys, 4);
-				Alpha = TimeToIndex(Seq, FrameTable, RelativePos, NumKeys, Index0, Index1);
+				Alpha = TimeToIndex(*DecompContext.AnimSeq, FrameTable, DecompContext.RelativePos, NumKeys, Index0, Index1);
 			}
 		}
 
@@ -660,15 +766,34 @@ void AEFPerTrackCompressionCodec::GetBoneAtomTranslation(
 
 void AEFPerTrackCompressionCodec::GetBoneAtomScale(
 	FTransform& OutAtom,
-	const UAnimSequence& Seq,
-	int32 Offset,
-	float Time,
-	float RelativePos)
+	FAnimSequenceDecompressionContext& DecompContext,
+	int32 TrackIndex)
 {
-	if (Offset != INDEX_NONE)
+#if USE_SEGMENTING_CONTEXT
+	if (DecompContext.AnimSeq->CompressedSegments.Num() != 0)
 	{
-		const uint8* RESTRICT TrackData = Seq.CompressedByteStream.GetData() + Offset + 4;
-		const int32 Header = *((int32*)(Seq.CompressedByteStream.GetData() + Offset));
+		const FTrivialAnimKeyHandle TrivialKeyHandle = DecompContext.GetTrivialScaleKeyHandle(TrackIndex);
+		if (TrivialKeyHandle.IsValid())
+		{
+			DecompContext.GetTrivialScale(OutAtom, TrivialKeyHandle);
+		}
+		else
+		{
+			const FAEPerTrackKeyLerpContext& EncodingContext = *static_cast<const FAEPerTrackKeyLerpContext*>(DecompContext.EncodingContext);
+			EncodingContext.GetScale(OutAtom, DecompContext, TrackIndex);
+		}
+
+		return;
+	}
+#endif
+
+
+	const int32 ScaleKeysOffset = DecompContext.GetCompressedScaleOffsets()->GetOffsetData( TrackIndex, 0 );
+
+	if (ScaleKeysOffset != INDEX_NONE)
+	{
+		const uint8* RESTRICT TrackData = DecompContext.GetCompressedByteStream() + ScaleKeysOffset + 4;
+		const int32 Header = *((int32*)(DecompContext.GetCompressedByteStream() + ScaleKeysOffset));
 
 		int32 KeyFormat;
 		int32 NumKeys;
@@ -688,12 +813,12 @@ void AEFPerTrackCompressionCodec::GetBoneAtomScale(
 		{
 			if ((FormatFlags & 0x8) == 0)
 			{
-				Alpha = TimeToIndex(Seq, RelativePos, NumKeys, Index0, Index1);
+				Alpha = TimeToIndex(*DecompContext.AnimSeq, DecompContext.RelativePos, NumKeys, Index0, Index1);
 			}
 			else
 			{
 				const uint8* RESTRICT FrameTable = Align(TrackData + FixedBytes + BytesPerKey * NumKeys, 4);
-				Alpha = TimeToIndex(Seq, FrameTable, RelativePos, NumKeys, Index0, Index1);
+				Alpha = TimeToIndex(*DecompContext.AnimSeq, FrameTable, DecompContext.RelativePos, NumKeys, Index0, Index1);
 			}
 		}
 
@@ -745,22 +870,18 @@ void AEFPerTrackCompressionCodec::GetBoneAtomScale(
 #if USE_ANIMATION_CODEC_BATCH_SOLVER
 
 /**
-* Decompress all requested rotation components from an Animation Sequence
-*
-* @param	Atoms			The FTransform array to fill in.
-* @param	DesiredPairs	Array of requested bone information
-* @param	Seq				The animation sequence to use.
-* @param	Time			Current time to solve for.
-* @return					None.
-*/
+ * Decompress all requested rotation components from an Animation Sequence
+ *
+ * @param	Atoms			The FTransform array to fill in.
+ * @param	DesiredPairs	Array of requested bone information
+ * @param	DecompContext	The decompression context to use.
+ */
 void AEFPerTrackCompressionCodec::GetPoseRotations(
 	FTransformArray& Atoms,
 	const BoneTrackArray& DesiredPairs,
-	const UAnimSequence& Seq,
-	float Time)
+	FAnimSequenceDecompressionContext& DecompContext)
 {
 	const int32 PairCount = DesiredPairs.Num();
-	const float RelativePos = Time / Seq.SequenceLength;
 
 	for( int32 PairIndex = 0; PairIndex < PairCount; ++PairIndex )
 	{
@@ -769,30 +890,23 @@ void AEFPerTrackCompressionCodec::GetPoseRotations(
 		const int32 AtomIndex = Pair.AtomIndex;
 		FTransform& BoneAtom = Atoms[AtomIndex];
 
-		const int32* RESTRICT TrackData = Seq.CompressedTrackOffsets.GetData() + (TrackIndex * 2);
-		const int32 RotKeysOffset = *(TrackData + 1);
-
-		GetBoneAtomRotation( BoneAtom, Seq, RotKeysOffset, Time, RelativePos );
+		GetBoneAtomRotation(BoneAtom, DecompContext, TrackIndex);
 	}
 }
 
 /**
-* Decompress all requested translation components from an Animation Sequence
-*
-* @param	Atoms			The FTransform array to fill in.
-* @param	DesiredPairs	Array of requested bone information
-* @param	Seq				The animation sequence to use.
-* @param	Time			Current time to solve for.
-* @return					None.
-*/
+ * Decompress all requested translation components from an Animation Sequence
+ *
+ * @param	Atoms			The FTransform array to fill in.
+ * @param	DesiredPairs	Array of requested bone information
+ * @param	DecompContext	The decompression context to use.
+ */
 void AEFPerTrackCompressionCodec::GetPoseTranslations(
 	FTransformArray& Atoms,
 	const BoneTrackArray& DesiredPairs,
-	const UAnimSequence& Seq,
-	float Time)
+	FAnimSequenceDecompressionContext& DecompContext)
 {
 	const int32 PairCount = DesiredPairs.Num();
-	const float RelativePos = Time / Seq.SequenceLength;
 
 	for( int32 PairIndex = 0; PairIndex < PairCount; ++PairIndex )
 	{
@@ -801,32 +915,25 @@ void AEFPerTrackCompressionCodec::GetPoseTranslations(
 		const int32 AtomIndex = Pair.AtomIndex;
 		FTransform& BoneAtom = Atoms[AtomIndex];
 
-		const int32* RESTRICT TrackData = Seq.CompressedTrackOffsets.GetData() + (TrackIndex * 2);
-		const int32 PosKeysOffset = *(TrackData + 0);
-
-		GetBoneAtomTranslation( BoneAtom, Seq, PosKeysOffset, Time, RelativePos );
+		GetBoneAtomTranslation(BoneAtom, DecompContext, TrackIndex);
 	}
 }
 
 /**
-* Decompress all requested Scale components from an Animation Sequence
-*
-* @param	Atoms			The FTransform array to fill in.
-* @param	DesiredPairs	Array of requested bone information
-* @param	Seq				The animation sequence to use.
-* @param	Time			Current time to solve for.
-* @return					None.
-*/
+ * Decompress all requested Scale components from an Animation Sequence
+ *
+ * @param	Atoms			The FTransform array to fill in.
+ * @param	DesiredPairs	Array of requested bone information
+ * @param	DecompContext	The decompression context to use.
+ */
 void AEFPerTrackCompressionCodec::GetPoseScales(
 	FTransformArray& Atoms,
 	const BoneTrackArray& DesiredPairs,
-	const UAnimSequence& Seq,
-	float Time)
+	FAnimSequenceDecompressionContext& DecompContext)
 {
-	check( Seq.CompressedScaleOffsets.IsValid() );
+	checkSlow(DecompContext.bHasScale);
 
 	const int32 PairCount = DesiredPairs.Num();
-	const float RelativePos = Time / Seq.SequenceLength;
 
 	for( int32 PairIndex = 0; PairIndex < PairCount; ++PairIndex )
 	{
@@ -835,10 +942,990 @@ void AEFPerTrackCompressionCodec::GetPoseScales(
 		const int32 AtomIndex = Pair.AtomIndex;
 		FTransform& BoneAtom = Atoms[AtomIndex];
 
-		const int32 ScaleKeysOffset = Seq.CompressedScaleOffsets.GetOffsetData( TrackIndex, 0 );
-
-		GetBoneAtomScale( BoneAtom, Seq, ScaleKeysOffset, Time, RelativePos );
+		GetBoneAtomScale(BoneAtom, DecompContext, TrackIndex);
 	}
 }
 
 #endif // USE_ANIMATION_CODEC_BATCH_SOLVER
+
+#if USE_SEGMENTING_CONTEXT
+void AEFPerTrackCompressionCodec::CreateEncodingContext(FAnimSequenceDecompressionContext& DecompContext)
+{
+	checkSlow(DecompContext.EncodingContext == nullptr);
+	DecompContext.EncodingContext = new FAEPerTrackKeyLerpContext(DecompContext);
+}
+
+void AEFPerTrackCompressionCodec::ReleaseEncodingContext(FAnimSequenceDecompressionContext& DecompContext)
+{
+	checkSlow(DecompContext.EncodingContext != nullptr);
+	delete DecompContext.EncodingContext;
+	DecompContext.EncodingContext = nullptr;
+}
+
+FAEPerTrackKeyLerpContext::FAEPerTrackKeyLerpContext(const FAnimSequenceDecompressionContext& DecompContext)
+	: PreviousSampleAtTime(FLT_MAX)	// Very large to trigger a reset on the first Seek()
+{
+	PreviousSegmentIndex[0] = -1;
+	PreviousSegmentIndex[1] = -1;
+
+	const int32 NumEntries = DecompContext.NumTracks * DecompContext.NumStreamsPerTrack;
+	for (int32 SegmentIndex = 0; SegmentIndex < 2; ++SegmentIndex)
+	{
+		UniformKeyOffsets[SegmentIndex].Empty(NumEntries);
+		UniformKeyOffsets[SegmentIndex].AddUninitialized(NumEntries);
+		RangeOffsets[SegmentIndex].Empty(NumEntries);
+		RangeOffsets[SegmentIndex].AddUninitialized(NumEntries);
+
+		if (!DecompContext.bIsSorted)
+		{
+			NumAnimatedTrackStreams[SegmentIndex].Empty(NumEntries);
+			NumAnimatedTrackStreams[SegmentIndex].AddUninitialized(NumEntries);
+			TrackStreamKeySizes[SegmentIndex].Empty(NumEntries);
+			TrackStreamKeySizes[SegmentIndex].AddUninitialized(NumEntries);
+		}
+	}
+}
+
+void FAEPerTrackKeyLerpContext::CacheSegmentValues(const FAnimSequenceDecompressionContext& DecompContext, const FCompressedSegment& Segment, uint8 SegmentIndex)
+{
+	PerTrackStreamFlagOffsets[SegmentIndex] = Segment.ByteStreamOffset;
+
+	const int32 PerTrackFlagsSize = Align(DecompContext.NumTracks * DecompContext.NumStreamsPerTrack * sizeof(uint8), 4);
+	const int32 RangeBaseOffset = Segment.ByteStreamOffset + PerTrackFlagsSize;
+	const uint8* PerTrackStreamFlags = DecompContext.CompressedByteStream + Segment.ByteStreamOffset;
+
+	int32 KeyOffset = 0;
+	int32 RangeOffset = RangeBaseOffset;
+	int32 TotalNumAnimatedTrackStreams = 0;
+	for (int32 TrackIndex = 0; TrackIndex < DecompContext.NumTracks; ++TrackIndex)
+	{
+		const FTrivialTrackFlags TrivialTrackFlags(DecompContext.TrackFlags[TrackIndex]);
+
+		const int32 TranslationValueOffset = DecompContext.GetTranslationValueOffset(TrackIndex);
+		UniformKeyOffsets[SegmentIndex][TranslationValueOffset] = KeyOffset;
+		RangeOffsets[SegmentIndex][TranslationValueOffset] = RangeOffset;
+		NumAnimatedTrackStreams[SegmentIndex][TranslationValueOffset] = TotalNumAnimatedTrackStreams;
+
+		int32 BytesPerKey = 0;
+		if (!TrivialTrackFlags.IsTranslationTrivial())
+		{
+			const FPerTrackFlags TranslationTrackFlags(PerTrackStreamFlags[TranslationValueOffset]);
+			const uint8 Format = TranslationTrackFlags.GetFormat();
+			const uint8 FormatFlags = TranslationTrackFlags.GetFormatFlags();
+
+			int32 BytesPerRange;
+			FAnimationCompression_PerTrackUtils::GetByteSizesFromFormat(Format, FormatFlags, BytesPerKey, BytesPerRange);
+
+			if (TranslationTrackFlags.IsUniform())
+			{
+				KeyOffset += BytesPerKey;
+			}
+			else
+			{
+				TotalNumAnimatedTrackStreams++;
+			}
+
+			if (Format == ACF_IntervalFixed32NoW)
+			{
+				RangeOffset += BytesPerRange;
+			}
+		}
+
+		TrackStreamKeySizes[SegmentIndex][TranslationValueOffset] = static_cast<uint8>(BytesPerKey);
+
+		const int32 RotationValueOffset = DecompContext.GetRotationValueOffset(TrackIndex);
+		UniformKeyOffsets[SegmentIndex][RotationValueOffset] = KeyOffset;
+		RangeOffsets[SegmentIndex][RotationValueOffset] = RangeOffset;
+		NumAnimatedTrackStreams[SegmentIndex][RotationValueOffset] = TotalNumAnimatedTrackStreams;
+
+		BytesPerKey = 0;
+		if (!TrivialTrackFlags.IsRotationTrivial())
+		{
+			const FPerTrackFlags RotationTrackFlags(PerTrackStreamFlags[RotationValueOffset]);
+			const uint8 Format = RotationTrackFlags.GetFormat();
+			const uint8 FormatFlags = RotationTrackFlags.GetFormatFlags();
+
+			int32 BytesPerRange;
+			FAnimationCompression_PerTrackUtils::GetByteSizesFromFormat(Format, FormatFlags, BytesPerKey, BytesPerRange);
+
+			if (RotationTrackFlags.IsUniform())
+			{
+				KeyOffset += BytesPerKey;
+			}
+			else
+			{
+				TotalNumAnimatedTrackStreams++;
+			}
+
+			if (Format == ACF_IntervalFixed32NoW)
+			{
+				RangeOffset += BytesPerRange;
+			}
+		}
+
+		TrackStreamKeySizes[SegmentIndex][RotationValueOffset] = static_cast<uint8>(BytesPerKey);
+
+		if (DecompContext.bHasScale)
+		{
+			const int32 ScaleValueOffset = DecompContext.GetScaleValueOffset(TrackIndex);
+			UniformKeyOffsets[SegmentIndex][ScaleValueOffset] = KeyOffset;
+			RangeOffsets[SegmentIndex][ScaleValueOffset] = RangeOffset;
+			NumAnimatedTrackStreams[SegmentIndex][ScaleValueOffset] = TotalNumAnimatedTrackStreams;
+
+			BytesPerKey = 0;
+			if (!TrivialTrackFlags.IsScaleTrivial())
+			{
+				const FPerTrackFlags ScaleTrackFlags(PerTrackStreamFlags[ScaleValueOffset]);
+				const uint8 Format = ScaleTrackFlags.GetFormat();
+				const uint8 FormatFlags = ScaleTrackFlags.GetFormatFlags();
+
+				int32 BytesPerRange;
+				FAnimationCompression_PerTrackUtils::GetByteSizesFromFormat(Format, FormatFlags, BytesPerKey, BytesPerRange);
+
+				if (ScaleTrackFlags.IsUniform())
+				{
+					KeyOffset += BytesPerKey;
+				}
+				else
+				{
+					TotalNumAnimatedTrackStreams++;
+				}
+
+				if (Format == ACF_IntervalFixed32NoW)
+				{
+					RangeOffset += BytesPerRange;
+				}
+			}
+
+			TrackStreamKeySizes[SegmentIndex][ScaleValueOffset] = static_cast<uint8>(BytesPerKey);
+		}
+	}
+
+	const int32 SegmentUniformKeyFrameSize = KeyOffset;
+	UniformKeyFrameSize[SegmentIndex] = SegmentUniformKeyFrameSize;
+	RangeDataSize[SegmentIndex] = RangeOffset - RangeBaseOffset;
+
+	UniformDataOffsets[SegmentIndex] = Segment.ByteStreamOffset + PerTrackFlagsSize + RangeDataSize[SegmentIndex];
+
+	if (!DecompContext.bIsSorted)
+	{
+		// Variable linear
+		TimeMarkerSize[SegmentIndex] = Segment.NumFrames < 256 ? sizeof(uint8) : sizeof(uint16);
+
+		const int32 UniformDataSize = Align(SegmentUniformKeyFrameSize * Segment.NumFrames, 4);
+
+		OffsetNumKeysPairOffsets[SegmentIndex] = Segment.ByteStreamOffset + PerTrackFlagsSize + RangeDataSize[SegmentIndex] + UniformDataSize;
+	}
+}
+
+void FAEPerTrackKeyLerpContext::ResetSortedCache(const FAnimSequenceDecompressionContext& DecompContext)
+{
+	CachedKeys.Reset(DecompContext.NumTracks);
+	CachedKeys.AddZeroed(DecompContext.NumTracks);
+
+	const int32 PerTrackFlagsSize = Align(DecompContext.NumTracks * DecompContext.NumStreamsPerTrack * sizeof(uint8), 4);
+	const int32 UniformDataSize = Align(UniformKeyFrameSize[0] * DecompContext.Segment0->NumFrames, 4);
+
+	PackedSampleData = DecompContext.CompressedByteStream + DecompContext.Segment0->ByteStreamOffset + PerTrackFlagsSize + RangeDataSize[0] + UniformDataSize;
+	PreviousFrameIndex = 0;
+	CurrentSegmentIndex = 0;
+}
+
+static constexpr uint8 AEPerTrackKeyLerpContextCachedKeyStructOffsets[3][2] =
+{
+	{ offsetof(FAEPerTrackKeyLerpContext::FCachedKey, TransOffsets), offsetof(FAEPerTrackKeyLerpContext::FCachedKey, TransFrameIndices) },
+	{ offsetof(FAEPerTrackKeyLerpContext::FCachedKey, RotOffsets), offsetof(FAEPerTrackKeyLerpContext::FCachedKey, RotFrameIndices) },
+	{ offsetof(FAEPerTrackKeyLerpContext::FCachedKey, ScaleOffsets), offsetof(FAEPerTrackKeyLerpContext::FCachedKey, ScaleFrameIndices) },
+};
+
+void FAEPerTrackKeyLerpContext::AdvanceSortedCachedKeys(const FAnimSequenceDecompressionContext& DecompContext)
+{
+	while (true)
+	{
+		const uint8* SampleData = PackedSampleData;
+		const FSortedKeyHeader KeyHeader(SampleData);
+		if (KeyHeader.IsEndOfStream())
+		{
+			break;	// Reached the end of the stream
+		}
+		checkSlow(KeyHeader.TrackIndex < DecompContext.NumTracks);
+
+		const uint8 SampleType = KeyHeader.GetKeyType();
+		checkSlow(SampleType <= 2);
+
+		const int32 TimeDelta = KeyHeader.GetTimeDelta();
+		const int32 FrameIndex = PreviousFrameIndex + TimeDelta;
+
+		// Swap and update
+		FCachedKey& CachedKey = CachedKeys[KeyHeader.TrackIndex];
+		int32* DataOffset = reinterpret_cast<int32*>(reinterpret_cast<uint8*>(&CachedKey) + AEPerTrackKeyLerpContextCachedKeyStructOffsets[SampleType][0]);
+		int32* FrameIndicesOffset = reinterpret_cast<int32*>(reinterpret_cast<uint8*>(&CachedKey) + AEPerTrackKeyLerpContextCachedKeyStructOffsets[SampleType][1]);
+		if (FrameIndex > CurrentFrameIndex && FrameIndicesOffset[1] >= CurrentFrameIndex)
+		{
+			break;		// Reached a sample we don't need yet, stop for now
+		}
+
+		SampleData += KeyHeader.GetSize();
+
+		DataOffset[0] = DataOffset[1];
+		DataOffset[1] = static_cast<int32>(SampleData - DecompContext.CompressedByteStream);
+		FrameIndicesOffset[0] = FrameIndicesOffset[1];
+		FrameIndicesOffset[1] = FrameIndex;
+
+		const uint8 BytesPerKey = TrackStreamKeySizes[CurrentSegmentIndex][(KeyHeader.TrackIndex * DecompContext.NumStreamsPerTrack) + SampleType];
+
+		PreviousFrameIndex = FrameIndex;
+		SampleData += BytesPerKey;
+		PackedSampleData = SampleData;	// Update the pointer since we consumed the sample
+	}
+}
+
+void FAEPerTrackKeyLerpContext::Seek(const FAnimSequenceDecompressionContext& DecompContext, float SampleAtTime)
+{
+	const bool IsSegmentCacheStale0 = PreviousSegmentIndex[0] != DecompContext.SegmentIndex0;
+	const bool IsSegmentCacheStale1 = PreviousSegmentIndex[1] != DecompContext.SegmentIndex1;
+	if (IsSegmentCacheStale0 || IsSegmentCacheStale1)
+	{
+		if (IsSegmentCacheStale0 && PreviousSegmentIndex[1] == DecompContext.SegmentIndex0)
+		{
+			// Forward playback, the new segment 0 is our old segment 1, swap the data and refresh the segment 1
+			PerTrackStreamFlagOffsets[0] = PerTrackStreamFlagOffsets[1];
+			UniformKeyOffsets[0] = UniformKeyOffsets[1];
+			UniformKeyFrameSize[0] = UniformKeyFrameSize[1];
+			RangeOffsets[0] = RangeOffsets[1];
+			RangeDataSize[0] = RangeDataSize[1];
+			UniformDataOffsets[0] = UniformDataOffsets[1];
+
+			// Variable linear
+			TimeMarkerSize[0] = TimeMarkerSize[1];
+			OffsetNumKeysPairOffsets[0] = OffsetNumKeysPairOffsets[1];
+			NumAnimatedTrackStreams[0] = NumAnimatedTrackStreams[1];
+			TrackStreamKeySizes[0] = TrackStreamKeySizes[1];
+
+			if (IsSegmentCacheStale1)
+			{
+				CacheSegmentValues(DecompContext, *DecompContext.Segment1, 1);
+			}
+		}
+		else if (IsSegmentCacheStale1 && PreviousSegmentIndex[0] == DecompContext.SegmentIndex1)
+		{
+			// Backward playback, the new segment 1 is our old segment 0, swap the data and refresh the segment 0
+			PerTrackStreamFlagOffsets[1] = PerTrackStreamFlagOffsets[0];
+			UniformKeyOffsets[1] = UniformKeyOffsets[0];
+			UniformKeyFrameSize[1] = UniformKeyFrameSize[0];
+			RangeOffsets[1] = RangeOffsets[0];
+			RangeDataSize[1] = RangeDataSize[0];
+			UniformDataOffsets[1] = UniformDataOffsets[0];
+
+			// Variable linear
+			TimeMarkerSize[1] = TimeMarkerSize[0];
+			OffsetNumKeysPairOffsets[1] = OffsetNumKeysPairOffsets[0];
+			NumAnimatedTrackStreams[1] = NumAnimatedTrackStreams[0];
+			TrackStreamKeySizes[1] = TrackStreamKeySizes[0];
+
+			if (IsSegmentCacheStale0)
+			{
+				CacheSegmentValues(DecompContext, *DecompContext.Segment0, 0);
+			}
+		}
+		else
+		{
+			if (IsSegmentCacheStale0)
+			{
+				CacheSegmentValues(DecompContext, *DecompContext.Segment0, 0);
+			}
+
+			if (IsSegmentCacheStale1)
+			{
+				CacheSegmentValues(DecompContext, *DecompContext.Segment1, 1);
+			}
+		}
+
+		PreviousSegmentIndex[0] = DecompContext.SegmentIndex0;
+		PreviousSegmentIndex[1] = DecompContext.SegmentIndex1;
+	}
+
+	FramePos = DecompContext.RelativePos * float(DecompContext.AnimSeq->NumFrames - 1);
+
+	if (DecompContext.bIsSorted)
+	{
+		if (SampleAtTime < PreviousSampleAtTime)
+		{
+			// Seeking backwards is terribly slow because we start over from the start
+			ResetSortedCache(DecompContext);
+		}
+		else if (IsSegmentCacheStale0)
+		{
+			// We are seeking forward into a new segment, start over
+			ResetSortedCache(DecompContext);
+		}
+
+		SegmentStartFrame[0] = DecompContext.Segment0->StartFrame;
+		SegmentStartFrame[1] = DecompContext.Segment1->StartFrame;
+
+		if (DecompContext.NeedsTwoSegments)
+		{
+			CurrentFrameIndex = CurrentSegmentIndex == 0 ? DecompContext.SegmentKeyIndex0 : DecompContext.SegmentKeyIndex1;
+		}
+		else
+		{
+			CurrentFrameIndex = FMath::Max(DecompContext.SegmentKeyIndex1, 1);
+		}
+
+		AdvanceSortedCachedKeys(DecompContext);
+
+		if (DecompContext.NeedsTwoSegments && CurrentSegmentIndex == 0)
+		{
+			// Switch to our segment 1
+			const int32 PerTrackFlagsSize = Align(DecompContext.NumTracks * DecompContext.NumStreamsPerTrack * sizeof(uint8), 4);
+			const int32 UniformDataSize = Align(UniformKeyFrameSize[1] * DecompContext.Segment1->NumFrames, 4);
+
+			PackedSampleData = DecompContext.CompressedByteStream + DecompContext.Segment1->ByteStreamOffset + PerTrackFlagsSize + RangeDataSize[1] + UniformDataSize;
+			PreviousFrameIndex = 0;
+			CurrentFrameIndex = DecompContext.SegmentKeyIndex1;
+			CurrentSegmentIndex = 1;
+
+			AdvanceSortedCachedKeys(DecompContext);
+
+			// Any track that is variable in segment 0 but not in segment 1 needs to be manually rotated in the cache
+			const uint8* PerTrackStreamFlags0 = DecompContext.CompressedByteStream + DecompContext.Segment0->ByteStreamOffset;
+			const uint8* PerTrackStreamFlags1 = DecompContext.CompressedByteStream + DecompContext.Segment1->ByteStreamOffset;
+
+			for (int32 TrackIndex = 0; TrackIndex < DecompContext.NumTracks; ++TrackIndex)
+			{
+				const FTrivialTrackFlags TrivialTrackFlags(DecompContext.TrackFlags[TrackIndex]);
+
+				if (!TrivialTrackFlags.IsTranslationTrivial())
+				{
+					const int32 TranslationValueOffset = DecompContext.GetTranslationValueOffset(TrackIndex);
+					const FPerTrackFlags TranslationTrackFlags0(PerTrackStreamFlags0[TranslationValueOffset]);
+					const FPerTrackFlags TranslationTrackFlags1(PerTrackStreamFlags1[TranslationValueOffset]);
+					if (!TranslationTrackFlags0.IsUniform() && TranslationTrackFlags1.IsUniform())
+					{
+						FCachedKey& CachedKey = CachedKeys[TrackIndex];
+						CachedKey.TransFrameIndices[0] = CachedKey.TransFrameIndices[1];
+						CachedKey.TransOffsets[0] = CachedKey.TransOffsets[1];
+					}
+				}
+
+				if (!TrivialTrackFlags.IsRotationTrivial())
+				{
+					const int32 RotationValueOffset = DecompContext.GetRotationValueOffset(TrackIndex);
+					const FPerTrackFlags RotationTrackFlags0(PerTrackStreamFlags0[RotationValueOffset]);
+					const FPerTrackFlags RotationTrackFlags1(PerTrackStreamFlags1[RotationValueOffset]);
+					if (!RotationTrackFlags0.IsUniform() && RotationTrackFlags1.IsUniform())
+					{
+						FCachedKey& CachedKey = CachedKeys[TrackIndex];
+						CachedKey.RotFrameIndices[0] = CachedKey.RotFrameIndices[1];
+						CachedKey.RotOffsets[0] = CachedKey.RotOffsets[1];
+					}
+				}
+
+				if (DecompContext.bHasScale && !TrivialTrackFlags.IsScaleTrivial())
+				{
+					const int32 ScaleValueOffset = DecompContext.GetScaleValueOffset(TrackIndex);
+					const FPerTrackFlags ScaleTrackFlags0(PerTrackStreamFlags0[ScaleValueOffset]);
+					const FPerTrackFlags ScaleTrackFlags1(PerTrackStreamFlags1[ScaleValueOffset]);
+					if (!ScaleTrackFlags0.IsUniform() && ScaleTrackFlags1.IsUniform())
+					{
+						FCachedKey& CachedKey = CachedKeys[TrackIndex];
+						CachedKey.ScaleFrameIndices[0] = CachedKey.ScaleFrameIndices[1];
+						CachedKey.ScaleOffsets[0] = CachedKey.ScaleOffsets[1];
+					}
+				}
+			}
+		}
+
+		PreviousSampleAtTime = SampleAtTime;
+	}
+	else
+	{
+		const float SegmentFramePos = FramePos - float(DecompContext.Segment0->StartFrame);
+		SegmentRelativePos0 = SegmentFramePos / float(DecompContext.Segment0->NumFrames - 1);
+	}
+}
+
+void FAEPerTrackKeyLerpContext::GetRotation(FTransform& OutAtom, const FAnimSequenceDecompressionContext& DecompContext, int32 TrackIndex) const
+{
+	const int32 RotationValueOffset = DecompContext.GetRotationValueOffset(TrackIndex);
+
+	const uint8* PerTrackStreamFlags0 = DecompContext.CompressedByteStream + PerTrackStreamFlagOffsets[0];
+	const FPerTrackFlags RotationFlags0(PerTrackStreamFlags0[RotationValueOffset]);
+
+	const uint8 KeyFormat0 = RotationFlags0.GetFormat();
+	const uint8 FormatFlags0 = RotationFlags0.GetFormatFlags();
+
+	const int32 RangeOffset0 = RangeOffsets[0][RotationValueOffset];
+	const uint8* RangeData0 = DecompContext.CompressedByteStream + RangeOffset0;
+
+	if (DecompContext.NeedsTwoSegments)
+	{
+		FQuat Rotation0;
+		if (RotationFlags0.IsUniform())
+		{
+			const int32 KeyOffset0 = UniformKeyOffsets[0][RotationValueOffset];
+			const int32 FrameStartOffset0 = UniformKeyFrameSize[0] * DecompContext.SegmentKeyIndex0;
+
+			const uint8* KeyData0 = DecompContext.CompressedByteStream + UniformDataOffsets[0] + FrameStartOffset0 + KeyOffset0;
+
+			FAnimationCompression_PerTrackUtils::DecompressRotation<false>(KeyFormat0, FormatFlags0, Rotation0, RangeData0, KeyData0);
+		}
+		else if (DecompContext.bIsSorted)
+		{
+			const FCachedKey& CachedKey = CachedKeys[TrackIndex];
+
+			const uint8* KeyData0 = DecompContext.CompressedByteStream + CachedKey.RotOffsets[0];
+
+			FAnimationCompression_PerTrackUtils::DecompressRotation<false>(KeyFormat0, FormatFlags0, Rotation0, RangeData0, KeyData0);
+		}
+		else
+		{
+			const int32 NumTrackStreams0 = NumAnimatedTrackStreams[0][RotationValueOffset];
+
+			const uint8* OffsetNumKeysPairs0 = DecompContext.CompressedByteStream + OffsetNumKeysPairOffsets[0];
+			const uint8* OffsetNumKeysPair0 = OffsetNumKeysPairs0 + (OffsetNumKeysPairSize * NumTrackStreams0);
+			const uint16 NumKeys0 = *reinterpret_cast<const uint16*>(OffsetNumKeysPair0 + sizeof(uint32));
+
+			const uint32 KeysOffset0 = AnimationCompressionUtils::UnalignedRead<uint32>(OffsetNumKeysPair0);
+			const int32 TimeMarkersOffset0 = DecompContext.Segment0->ByteStreamOffset + KeysOffset0;
+
+			const int32 TrackDataOffset0 = Align(TimeMarkersOffset0 + (NumKeys0 * TimeMarkerSize[0]), 4);
+			const uint8 BytesPerKey0 = TrackStreamKeySizes[0][RotationValueOffset];
+			const int32 KeyDataOffset0 = TrackDataOffset0 + ((NumKeys0 - 1) * BytesPerKey0);	// We need the last key of segment 0 and the first key of segment 1
+			const uint8* KeyData0 = DecompContext.CompressedByteStream + KeyDataOffset0;
+
+			FAnimationCompression_PerTrackUtils::DecompressRotation<false>(KeyFormat0, FormatFlags0, Rotation0, RangeData0, KeyData0);
+		}
+
+		const uint8* PerTrackStreamFlags1 = DecompContext.CompressedByteStream + PerTrackStreamFlagOffsets[1];
+		const FPerTrackFlags RotationFlags1(PerTrackStreamFlags1[RotationValueOffset]);
+
+		const uint8 KeyFormat1 = RotationFlags1.GetFormat();
+		const uint8 FormatFlags1 = RotationFlags1.GetFormatFlags();
+
+		const int32 RangeOffset1 = RangeOffsets[1][RotationValueOffset];
+		const uint8* RangeData1 = DecompContext.CompressedByteStream + RangeOffset1;
+
+		FQuat Rotation1;
+		if (RotationFlags1.IsUniform())
+		{
+			const int32 KeyOffset1 = UniformKeyOffsets[1][RotationValueOffset];
+			const int32 FrameStartOffset1 = UniformKeyFrameSize[1] * DecompContext.SegmentKeyIndex1;
+
+			const uint8* KeyData1 = DecompContext.CompressedByteStream + UniformDataOffsets[1] + FrameStartOffset1 + KeyOffset1;
+
+			FAnimationCompression_PerTrackUtils::DecompressRotation<false>(KeyFormat1, FormatFlags1, Rotation1, RangeData1, KeyData1);
+		}
+		else if (DecompContext.bIsSorted)
+		{
+			const FCachedKey& CachedKey = CachedKeys[TrackIndex];
+
+			const uint8* KeyData1 = DecompContext.CompressedByteStream + CachedKey.RotOffsets[1];
+
+			FAnimationCompression_PerTrackUtils::DecompressRotation<false>(KeyFormat1, FormatFlags1, Rotation1, RangeData1, KeyData1);
+		}
+		else
+		{
+			const int32 NumTrackStreams1 = NumAnimatedTrackStreams[1][RotationValueOffset];
+
+			const uint8* OffsetNumKeysPairs1 = DecompContext.CompressedByteStream + OffsetNumKeysPairOffsets[1];
+			const uint8* OffsetNumKeysPair1 = OffsetNumKeysPairs1 + (OffsetNumKeysPairSize * NumTrackStreams1);
+			const uint16 NumKeys1 = *reinterpret_cast<const uint16*>(OffsetNumKeysPair1 + sizeof(uint32));
+
+			const uint32 KeysOffset1 = AnimationCompressionUtils::UnalignedRead<uint32>(OffsetNumKeysPair1);
+			const int32 TimeMarkersOffset1 = DecompContext.Segment1->ByteStreamOffset + KeysOffset1;
+
+			const int32 TrackDataOffset1 = Align(TimeMarkersOffset1 + (NumKeys1 * TimeMarkerSize[1]), 4);
+			// We need the last key of segment 0 and the first key of segment 1
+			const int32 KeyDataOffset1 = TrackDataOffset1;	// First key!
+			const uint8* KeyData1 = DecompContext.CompressedByteStream + KeyDataOffset1;
+
+			FAnimationCompression_PerTrackUtils::DecompressRotation<false>(KeyFormat1, FormatFlags1, Rotation1, RangeData1, KeyData1);
+		}
+
+		// Fast linear quaternion interpolation.
+		FQuat BlendedQuat = FQuat::FastLerp(Rotation0, Rotation1, DecompContext.KeyAlpha);
+		BlendedQuat.Normalize();
+
+		OutAtom.SetRotation(BlendedQuat);
+	}
+	else
+	{
+		if (RotationFlags0.IsUniform())
+		{
+			const int32 KeyOffset0 = UniformKeyOffsets[0][RotationValueOffset];
+			const int32 UniformKeyFrameSize0 = UniformKeyFrameSize[0];
+			const int32 FrameStartOffset0 = UniformKeyFrameSize0 * DecompContext.SegmentKeyIndex0;
+
+			const uint8* KeyData0 = DecompContext.CompressedByteStream + UniformDataOffsets[0] + FrameStartOffset0 + KeyOffset0;
+
+			FQuat Rotation;
+			FAnimationCompression_PerTrackUtils::DecompressRotation<false>(KeyFormat0, FormatFlags0, Rotation, RangeData0, KeyData0);
+
+			if (DecompContext.NeedsInterpolation)
+			{
+				const uint8* KeyData1 = KeyData0 + UniformKeyFrameSize0;
+
+				FQuat Rotation1;
+				FAnimationCompression_PerTrackUtils::DecompressRotation<false>(KeyFormat0, FormatFlags0, Rotation1, RangeData0, KeyData1);
+
+				// Fast linear quaternion interpolation.
+				FQuat BlendedQuat = FQuat::FastLerp(Rotation, Rotation1, DecompContext.KeyAlpha);
+				BlendedQuat.Normalize();
+				Rotation = BlendedQuat;
+			}
+
+			OutAtom.SetRotation(Rotation);
+		}
+		else if (DecompContext.bIsSorted)
+		{
+			const FCachedKey& CachedKey = CachedKeys[TrackIndex];
+
+			// compute the blend parameters for the keys we have found
+			const int32 FrameIndex0 = SegmentStartFrame[0] + CachedKey.RotFrameIndices[0];
+			const int32 FrameIndex1 = SegmentStartFrame[1] + CachedKey.RotFrameIndices[1];
+
+			const int32 Delta = FMath::Max(FrameIndex1 - FrameIndex0, 1);
+			const float Remainder = FramePos - float(FrameIndex0);
+			const float Alpha = DecompContext.AnimSeq->Interpolation == EAnimInterpolationType::Step ? 0.0f : (Remainder / float(Delta));
+
+			const uint8* KeyData0 = DecompContext.CompressedByteStream + CachedKey.RotOffsets[0];
+
+			FQuat Rotation0;
+			FAnimationCompression_PerTrackUtils::DecompressRotation<false>(KeyFormat0, FormatFlags0, Rotation0, RangeData0, KeyData0);
+
+			const uint8* KeyData1 = DecompContext.CompressedByteStream + CachedKey.RotOffsets[1];
+
+			FQuat Rotation1;
+			FAnimationCompression_PerTrackUtils::DecompressRotation<false>(KeyFormat0, FormatFlags0, Rotation1, RangeData0, KeyData1);
+
+			// Fast linear quaternion interpolation.
+			FQuat BlendedQuat = FQuat::FastLerp(Rotation0, Rotation1, Alpha);
+			BlendedQuat.Normalize();
+
+			OutAtom.SetRotation(BlendedQuat);
+		}
+		else
+		{
+			const int32 NumTrackStreams0 = NumAnimatedTrackStreams[0][RotationValueOffset];
+
+			const uint8* OffsetNumKeysPairs0 = DecompContext.CompressedByteStream + OffsetNumKeysPairOffsets[0];
+			const uint8* OffsetNumKeysPair0 = OffsetNumKeysPairs0 + (OffsetNumKeysPairSize * NumTrackStreams0);
+			const uint16 NumKeys0 = *reinterpret_cast<const uint16*>(OffsetNumKeysPair0 + sizeof(uint32));
+
+			const uint32 KeysOffset0 = AnimationCompressionUtils::UnalignedRead<uint32>(OffsetNumKeysPair0);
+			const int32 TimeMarkersOffset0 = DecompContext.Segment0->ByteStreamOffset + KeysOffset0;
+			const uint8* TimeMarkers0 = DecompContext.CompressedByteStream + TimeMarkersOffset0;
+
+			int32 FrameIndex0;
+			int32 FrameIndex1;
+			float Alpha = AnimEncoding::TimeToIndex(DecompContext, TimeMarkers0, NumKeys0, DecompContext.Segment0->NumFrames, TimeMarkerSize[0], SegmentRelativePos0, FrameIndex0, FrameIndex1);
+
+			const int32 TrackDataOffset0 = Align(TimeMarkersOffset0 + (NumKeys0 * TimeMarkerSize[0]), 4);
+			const uint8 BytesPerKey0 = TrackStreamKeySizes[0][RotationValueOffset];
+			const int32 KeyDataOffset0 = TrackDataOffset0 + (FrameIndex0 * BytesPerKey0);
+			const uint8* KeyData0 = DecompContext.CompressedByteStream + KeyDataOffset0;
+
+			FQuat Rotation;
+			FAnimationCompression_PerTrackUtils::DecompressRotation<false>(KeyFormat0, FormatFlags0, Rotation, RangeData0, KeyData0);
+
+			if (DecompContext.NeedsInterpolation)
+			{
+				const uint8* KeyData1 = KeyData0 + BytesPerKey0;
+
+				FQuat Rotation1;
+				FAnimationCompression_PerTrackUtils::DecompressRotation<false>(KeyFormat0, FormatFlags0, Rotation1, RangeData0, KeyData1);
+
+				// Fast linear quaternion interpolation.
+				FQuat BlendedQuat = FQuat::FastLerp(Rotation, Rotation1, Alpha);
+				BlendedQuat.Normalize();
+				Rotation = BlendedQuat;
+			}
+
+			OutAtom.SetRotation(Rotation);
+		}
+	}
+}
+
+void FAEPerTrackKeyLerpContext::GetTranslation(FTransform& OutAtom, const FAnimSequenceDecompressionContext& DecompContext, int32 TrackIndex) const
+{
+	const int32 TranslationValueOffset = DecompContext.GetTranslationValueOffset(TrackIndex);
+
+	const uint8* PerTrackStreamFlags0 = DecompContext.CompressedByteStream + PerTrackStreamFlagOffsets[0];
+	const FPerTrackFlags TranslationFlags0(PerTrackStreamFlags0[TranslationValueOffset]);
+
+	const uint8 KeyFormat0 = TranslationFlags0.GetFormat();
+	const uint8 FormatFlags0 = TranslationFlags0.GetFormatFlags();
+
+	const int32 RangeOffset0 = RangeOffsets[0][TranslationValueOffset];
+	const uint8* RangeData0 = DecompContext.CompressedByteStream + RangeOffset0;
+
+	if (DecompContext.NeedsTwoSegments)
+	{
+		FVector Translation0;
+		if (TranslationFlags0.IsUniform())
+		{
+			const int32 KeyOffset0 = UniformKeyOffsets[0][TranslationValueOffset];
+			const int32 FrameStartOffset0 = UniformKeyFrameSize[0] * DecompContext.SegmentKeyIndex0;
+
+			const uint8* KeyData0 = DecompContext.CompressedByteStream + UniformDataOffsets[0] + FrameStartOffset0 + KeyOffset0;
+
+			FAnimationCompression_PerTrackUtils::DecompressTranslation<false>(KeyFormat0, FormatFlags0, Translation0, RangeData0, KeyData0);
+		}
+		else if (DecompContext.bIsSorted)
+		{
+			const FCachedKey& CachedKey = CachedKeys[TrackIndex];
+
+			const uint8* KeyData0 = DecompContext.CompressedByteStream + CachedKey.TransOffsets[0];
+
+			FAnimationCompression_PerTrackUtils::DecompressTranslation<false>(KeyFormat0, FormatFlags0, Translation0, RangeData0, KeyData0);
+		}
+		else
+		{
+			const int32 NumTrackStreams0 = NumAnimatedTrackStreams[0][TranslationValueOffset];
+
+			const uint8* OffsetNumKeysPairs0 = DecompContext.CompressedByteStream + OffsetNumKeysPairOffsets[0];
+			const uint8* OffsetNumKeysPair0 = OffsetNumKeysPairs0 + (OffsetNumKeysPairSize * NumTrackStreams0);
+			const uint16 NumKeys0 = *reinterpret_cast<const uint16*>(OffsetNumKeysPair0 + sizeof(uint32));
+
+			const uint32 KeysOffset0 = AnimationCompressionUtils::UnalignedRead<uint32>(OffsetNumKeysPair0);
+			const int32 TimeMarkersOffset0 = DecompContext.Segment0->ByteStreamOffset + KeysOffset0;
+
+			const int32 TrackDataOffset0 = Align(TimeMarkersOffset0 + (NumKeys0 * TimeMarkerSize[0]), 4);
+			const uint8 BytesPerKey0 = TrackStreamKeySizes[0][TranslationValueOffset];
+			const int32 KeyDataOffset0 = TrackDataOffset0 + ((NumKeys0 - 1) * BytesPerKey0);	// We need the last key of segment 0 and the first key of segment 1
+			const uint8* KeyData0 = DecompContext.CompressedByteStream + KeyDataOffset0;
+
+			FAnimationCompression_PerTrackUtils::DecompressTranslation<false>(KeyFormat0, FormatFlags0, Translation0, RangeData0, KeyData0);
+		}
+
+		const uint8* PerTrackStreamFlags1 = DecompContext.CompressedByteStream + PerTrackStreamFlagOffsets[1];
+		const FPerTrackFlags TranslationFlags1(PerTrackStreamFlags1[TranslationValueOffset]);
+
+		const uint8 KeyFormat1 = TranslationFlags1.GetFormat();
+		const uint8 FormatFlags1 = TranslationFlags1.GetFormatFlags();
+
+		const int32 RangeOffset1 = RangeOffsets[1][TranslationValueOffset];
+		const uint8* RangeData1 = DecompContext.CompressedByteStream + RangeOffset1;
+
+		FVector Translation1;
+		if (TranslationFlags1.IsUniform())
+		{
+			const int32 KeyOffset1 = UniformKeyOffsets[1][TranslationValueOffset];
+			const int32 FrameStartOffset1 = UniformKeyFrameSize[1] * DecompContext.SegmentKeyIndex1;
+
+			const uint8* KeyData1 = DecompContext.CompressedByteStream + UniformDataOffsets[1] + FrameStartOffset1 + KeyOffset1;
+
+			FAnimationCompression_PerTrackUtils::DecompressTranslation<false>(KeyFormat1, FormatFlags1, Translation1, RangeData1, KeyData1);
+		}
+		else if (DecompContext.bIsSorted)
+		{
+			const FCachedKey& CachedKey = CachedKeys[TrackIndex];
+
+			const uint8* KeyData1 = DecompContext.CompressedByteStream + CachedKey.TransOffsets[1];
+
+			FAnimationCompression_PerTrackUtils::DecompressTranslation<false>(KeyFormat1, FormatFlags1, Translation1, RangeData1, KeyData1);
+		}
+		else
+		{
+			const int32 NumTrackStreams1 = NumAnimatedTrackStreams[1][TranslationValueOffset];
+
+			const uint8* OffsetNumKeysPairs1 = DecompContext.CompressedByteStream + OffsetNumKeysPairOffsets[1];
+			const uint8* OffsetNumKeysPair1 = OffsetNumKeysPairs1 + (OffsetNumKeysPairSize * NumTrackStreams1);
+			const uint16 NumKeys1 = *reinterpret_cast<const uint16*>(OffsetNumKeysPair1 + sizeof(uint32));
+
+			const uint32 KeysOffset1 = AnimationCompressionUtils::UnalignedRead<uint32>(OffsetNumKeysPair1);
+			const int32 TimeMarkersOffset1 = DecompContext.Segment1->ByteStreamOffset + KeysOffset1;
+
+			const int32 TrackDataOffset1 = Align(TimeMarkersOffset1 + (NumKeys1 * TimeMarkerSize[1]), 4);
+			// We need the last key of segment 0 and the first key of segment 1
+			const int32 KeyDataOffset1 = TrackDataOffset1;	// First key!
+			const uint8* KeyData1 = DecompContext.CompressedByteStream + KeyDataOffset1;
+
+			FAnimationCompression_PerTrackUtils::DecompressTranslation<false>(KeyFormat1, FormatFlags1, Translation1, RangeData1, KeyData1);
+		}
+
+		const FVector Translation = FMath::Lerp(Translation0, Translation1, DecompContext.KeyAlpha);
+		OutAtom.SetTranslation(Translation);
+	}
+	else
+	{
+		if (TranslationFlags0.IsUniform())
+		{
+			const int32 KeyOffset0 = UniformKeyOffsets[0][TranslationValueOffset];
+			const int32 UniformKeyFrameSize0 = UniformKeyFrameSize[0];
+			const int32 FrameStartOffset0 = UniformKeyFrameSize0 * DecompContext.SegmentKeyIndex0;
+
+			const uint8* KeyData0 = DecompContext.CompressedByteStream + UniformDataOffsets[0] + FrameStartOffset0 + KeyOffset0;
+
+			FVector Translation;
+			FAnimationCompression_PerTrackUtils::DecompressTranslation<false>(KeyFormat0, FormatFlags0, Translation, RangeData0, KeyData0);
+
+			if (DecompContext.NeedsInterpolation)
+			{
+				const uint8* KeyData1 = KeyData0 + UniformKeyFrameSize0;
+
+				FVector Translation1;
+				FAnimationCompression_PerTrackUtils::DecompressTranslation<false>(KeyFormat0, FormatFlags0, Translation1, RangeData0, KeyData1);
+
+				Translation = FMath::Lerp(Translation, Translation1, DecompContext.KeyAlpha);
+			}
+
+			OutAtom.SetTranslation(Translation);
+		}
+		else if (DecompContext.bIsSorted)
+		{
+			const FCachedKey& CachedKey = CachedKeys[TrackIndex];
+
+			// compute the blend parameters for the keys we have found
+			const int32 FrameIndex0 = SegmentStartFrame[0] + CachedKey.TransFrameIndices[0];
+			const int32 FrameIndex1 = SegmentStartFrame[1] + CachedKey.TransFrameIndices[1];
+
+			const int32 Delta = FMath::Max(FrameIndex1 - FrameIndex0, 1);
+			const float Remainder = FramePos - float(FrameIndex0);
+			const float Alpha = DecompContext.AnimSeq->Interpolation == EAnimInterpolationType::Step ? 0.0f : (Remainder / float(Delta));
+
+			const uint8* KeyData0 = DecompContext.CompressedByteStream + CachedKey.TransOffsets[0];
+
+			FVector Translation0;
+			FAnimationCompression_PerTrackUtils::DecompressTranslation<false>(KeyFormat0, FormatFlags0, Translation0, RangeData0, KeyData0);
+
+			const uint8* KeyData1 = DecompContext.CompressedByteStream + CachedKey.TransOffsets[1];
+
+			FVector Translation1;
+			FAnimationCompression_PerTrackUtils::DecompressTranslation<false>(KeyFormat0, FormatFlags0, Translation1, RangeData0, KeyData1);
+
+			FVector Translation = FMath::Lerp(Translation0, Translation1, Alpha);
+
+			OutAtom.SetTranslation(Translation);
+		}
+		else
+		{
+			const int32 NumTrackStreams0 = NumAnimatedTrackStreams[0][TranslationValueOffset];
+
+			const uint8* OffsetNumKeysPairs0 = DecompContext.CompressedByteStream + OffsetNumKeysPairOffsets[0];
+			const uint8* OffsetNumKeysPair0 = OffsetNumKeysPairs0 + (OffsetNumKeysPairSize * NumTrackStreams0);
+			const uint16 NumKeys0 = *reinterpret_cast<const uint16*>(OffsetNumKeysPair0 + sizeof(uint32));
+
+			const uint32 KeysOffset0 = AnimationCompressionUtils::UnalignedRead<uint32>(OffsetNumKeysPair0);
+			const int32 TimeMarkersOffset0 = DecompContext.Segment0->ByteStreamOffset + KeysOffset0;
+			const uint8* TimeMarkers0 = DecompContext.CompressedByteStream + TimeMarkersOffset0;
+
+			int32 FrameIndex0;
+			int32 FrameIndex1;
+			float Alpha = AnimEncoding::TimeToIndex(DecompContext, TimeMarkers0, NumKeys0, DecompContext.Segment0->NumFrames, TimeMarkerSize[0], SegmentRelativePos0, FrameIndex0, FrameIndex1);
+
+			const int32 TrackDataOffset0 = Align(TimeMarkersOffset0 + (NumKeys0 * TimeMarkerSize[0]), 4);
+			const uint8 BytesPerKey0 = TrackStreamKeySizes[0][TranslationValueOffset];
+			const int32 KeyDataOffset0 = TrackDataOffset0 + (FrameIndex0 * BytesPerKey0);
+			const uint8* KeyData0 = DecompContext.CompressedByteStream + KeyDataOffset0;
+
+			FVector Translation;
+			FAnimationCompression_PerTrackUtils::DecompressTranslation<false>(KeyFormat0, FormatFlags0, Translation, RangeData0, KeyData0);
+
+			if (DecompContext.NeedsInterpolation)
+			{
+				const uint8* KeyData1 = KeyData0 + BytesPerKey0;
+
+				FVector Translation1;
+				FAnimationCompression_PerTrackUtils::DecompressTranslation<false>(KeyFormat0, FormatFlags0, Translation1, RangeData0, KeyData1);
+
+				Translation = FMath::Lerp(Translation, Translation1, Alpha);
+			}
+
+			OutAtom.SetTranslation(Translation);
+		}
+	}
+}
+
+void FAEPerTrackKeyLerpContext::GetScale(FTransform& OutAtom, const FAnimSequenceDecompressionContext& DecompContext, int32 TrackIndex) const
+{
+	const int32 ScaleValueOffset = DecompContext.GetScaleValueOffset(TrackIndex);
+
+	const uint8* PerTrackStreamFlags0 = DecompContext.CompressedByteStream + PerTrackStreamFlagOffsets[0];
+	const FPerTrackFlags ScaleFlags0(PerTrackStreamFlags0[ScaleValueOffset]);
+
+	const uint8 KeyFormat0 = ScaleFlags0.GetFormat();
+	const uint8 FormatFlags0 = ScaleFlags0.GetFormatFlags();
+
+	const int32 RangeOffset0 = RangeOffsets[0][ScaleValueOffset];
+	const uint8* RangeData0 = DecompContext.CompressedByteStream + RangeOffset0;
+
+	if (DecompContext.NeedsTwoSegments)
+	{
+		FVector Scale0;
+		if (ScaleFlags0.IsUniform())
+		{
+			const int32 KeyOffset0 = UniformKeyOffsets[0][ScaleValueOffset];
+			const int32 FrameStartOffset0 = UniformKeyFrameSize[0] * DecompContext.SegmentKeyIndex0;
+
+			const uint8* KeyData0 = DecompContext.CompressedByteStream + UniformDataOffsets[0] + FrameStartOffset0 + KeyOffset0;
+
+			FAnimationCompression_PerTrackUtils::DecompressScale<false>(KeyFormat0, FormatFlags0, Scale0, RangeData0, KeyData0);
+		}
+		else if (DecompContext.bIsSorted)
+		{
+			const FCachedKey& CachedKey = CachedKeys[TrackIndex];
+
+			const uint8* KeyData0 = DecompContext.CompressedByteStream + CachedKey.ScaleOffsets[0];
+
+			FAnimationCompression_PerTrackUtils::DecompressScale<false>(KeyFormat0, FormatFlags0, Scale0, RangeData0, KeyData0);
+		}
+		else
+		{
+			const int32 NumTrackStreams0 = NumAnimatedTrackStreams[0][ScaleValueOffset];
+
+			const uint8* OffsetNumKeysPairs0 = DecompContext.CompressedByteStream + OffsetNumKeysPairOffsets[0];
+			const uint8* OffsetNumKeysPair0 = OffsetNumKeysPairs0 + (OffsetNumKeysPairSize * NumTrackStreams0);
+			const uint16 NumKeys0 = *reinterpret_cast<const uint16*>(OffsetNumKeysPair0 + sizeof(uint32));
+
+			const uint32 KeysOffset0 = AnimationCompressionUtils::UnalignedRead<uint32>(OffsetNumKeysPair0);
+			const int32 TimeMarkersOffset0 = DecompContext.Segment0->ByteStreamOffset + KeysOffset0;
+
+			const int32 TrackDataOffset0 = Align(TimeMarkersOffset0 + (NumKeys0 * TimeMarkerSize[0]), 4);
+			const uint8 BytesPerKey0 = TrackStreamKeySizes[0][ScaleValueOffset];
+			const int32 KeyDataOffset0 = TrackDataOffset0 + ((NumKeys0 - 1) * BytesPerKey0);	// We need the last key of segment 0 and the first key of segment 1
+			const uint8* KeyData0 = DecompContext.CompressedByteStream + KeyDataOffset0;
+
+			FAnimationCompression_PerTrackUtils::DecompressScale<false>(KeyFormat0, FormatFlags0, Scale0, RangeData0, KeyData0);
+		}
+
+		const uint8* PerTrackStreamFlags1 = DecompContext.CompressedByteStream + PerTrackStreamFlagOffsets[1];
+		const FPerTrackFlags ScaleFlags1(PerTrackStreamFlags1[ScaleValueOffset]);
+
+		const uint8 KeyFormat1 = ScaleFlags1.GetFormat();
+		const uint8 FormatFlags1 = ScaleFlags1.GetFormatFlags();
+
+		const int32 RangeOffset1 = RangeOffsets[1][ScaleValueOffset];
+		const uint8* RangeData1 = DecompContext.CompressedByteStream + RangeOffset1;
+
+		FVector Scale1;
+		if (ScaleFlags1.IsUniform())
+		{
+			const int32 KeyOffset1 = UniformKeyOffsets[1][ScaleValueOffset];
+			const int32 FrameStartOffset1 = UniformKeyFrameSize[1] * DecompContext.SegmentKeyIndex1;
+
+			const uint8* KeyData1 = DecompContext.CompressedByteStream + UniformDataOffsets[1] + FrameStartOffset1 + KeyOffset1;
+
+			FAnimationCompression_PerTrackUtils::DecompressScale<false>(KeyFormat1, FormatFlags1, Scale1, RangeData1, KeyData1);
+		}
+		else if (DecompContext.bIsSorted)
+		{
+			const FCachedKey& CachedKey = CachedKeys[TrackIndex];
+
+			const uint8* KeyData1 = DecompContext.CompressedByteStream + CachedKey.ScaleOffsets[1];
+
+			FAnimationCompression_PerTrackUtils::DecompressScale<false>(KeyFormat1, FormatFlags1, Scale1, RangeData1, KeyData1);
+		}
+		else
+		{
+			const int32 NumTrackStreams1 = NumAnimatedTrackStreams[1][ScaleValueOffset];
+
+			const uint8* OffsetNumKeysPairs1 = DecompContext.CompressedByteStream + OffsetNumKeysPairOffsets[1];
+			const uint8* OffsetNumKeysPair1 = OffsetNumKeysPairs1 + (OffsetNumKeysPairSize * NumTrackStreams1);
+			const uint16 NumKeys1 = *reinterpret_cast<const uint16*>(OffsetNumKeysPair1 + sizeof(uint32));
+
+			const uint32 KeysOffset1 = AnimationCompressionUtils::UnalignedRead<uint32>(OffsetNumKeysPair1);
+			const int32 TimeMarkersOffset1 = DecompContext.Segment1->ByteStreamOffset + KeysOffset1;
+
+			const int32 TrackDataOffset1 = Align(TimeMarkersOffset1 + (NumKeys1 * TimeMarkerSize[1]), 4);
+			// We need the last key of segment 0 and the first key of segment 1
+			const int32 KeyDataOffset1 = TrackDataOffset1;	// First key!
+			const uint8* KeyData1 = DecompContext.CompressedByteStream + KeyDataOffset1;
+
+			FAnimationCompression_PerTrackUtils::DecompressScale<false>(KeyFormat1, FormatFlags1, Scale1, RangeData1, KeyData1);
+		}
+
+		const FVector Scale = FMath::Lerp(Scale0, Scale1, DecompContext.KeyAlpha);
+		OutAtom.SetScale3D(Scale);
+	}
+	else
+	{
+		if (ScaleFlags0.IsUniform())
+		{
+			const int32 KeyOffset0 = UniformKeyOffsets[0][ScaleValueOffset];
+			const int32 UniformKeyFrameSize0 = UniformKeyFrameSize[0];
+			const int32 FrameStartOffset0 = UniformKeyFrameSize0 * DecompContext.SegmentKeyIndex0;
+
+			const uint8* KeyData0 = DecompContext.CompressedByteStream + UniformDataOffsets[0] + FrameStartOffset0 + KeyOffset0;
+
+			FVector Scale;
+			FAnimationCompression_PerTrackUtils::DecompressScale<false>(KeyFormat0, FormatFlags0, Scale, RangeData0, KeyData0);
+
+			if (DecompContext.NeedsInterpolation)
+			{
+				const uint8* KeyData1 = KeyData0 + UniformKeyFrameSize0;
+
+				FVector Scale1;
+				FAnimationCompression_PerTrackUtils::DecompressScale<false>(KeyFormat0, FormatFlags0, Scale1, RangeData0, KeyData1);
+
+				Scale = FMath::Lerp(Scale, Scale1, DecompContext.KeyAlpha);
+			}
+
+			OutAtom.SetScale3D(Scale);
+		}
+		else if (DecompContext.bIsSorted)
+		{
+			const FCachedKey& CachedKey = CachedKeys[TrackIndex];
+
+			// compute the blend parameters for the keys we have found
+			const int32 FrameIndex0 = SegmentStartFrame[0] + CachedKey.ScaleFrameIndices[0];
+			const int32 FrameIndex1 = SegmentStartFrame[1] + CachedKey.ScaleFrameIndices[1];
+
+			const int32 Delta = FMath::Max(FrameIndex1 - FrameIndex0, 1);
+			const float Remainder = FramePos - float(FrameIndex0);
+			const float Alpha = DecompContext.AnimSeq->Interpolation == EAnimInterpolationType::Step ? 0.0f : (Remainder / float(Delta));
+
+			const uint8* KeyData0 = DecompContext.CompressedByteStream + CachedKey.ScaleOffsets[0];
+
+			FVector Scale0;
+			FAnimationCompression_PerTrackUtils::DecompressScale<false>(KeyFormat0, FormatFlags0, Scale0, RangeData0, KeyData0);
+
+			const uint8* KeyData1 = DecompContext.CompressedByteStream + CachedKey.ScaleOffsets[1];
+
+			FVector Scale1;
+			FAnimationCompression_PerTrackUtils::DecompressScale<false>(KeyFormat0, FormatFlags0, Scale1, RangeData0, KeyData1);
+
+			FVector Scale = FMath::Lerp(Scale0, Scale1, Alpha);
+
+			OutAtom.SetScale3D(Scale);
+		}
+		else
+		{
+			const int32 NumTrackStreams0 = NumAnimatedTrackStreams[0][ScaleValueOffset];
+
+			const uint8* OffsetNumKeysPairs0 = DecompContext.CompressedByteStream + OffsetNumKeysPairOffsets[0];
+			const uint8* OffsetNumKeysPair0 = OffsetNumKeysPairs0 + (OffsetNumKeysPairSize * NumTrackStreams0);
+			const uint16 NumKeys0 = *reinterpret_cast<const uint16*>(OffsetNumKeysPair0 + sizeof(uint32));
+
+			const uint32 KeysOffset0 = AnimationCompressionUtils::UnalignedRead<uint32>(OffsetNumKeysPair0);
+			const int32 TimeMarkersOffset0 = DecompContext.Segment0->ByteStreamOffset + KeysOffset0;
+			const uint8* TimeMarkers0 = DecompContext.CompressedByteStream + TimeMarkersOffset0;
+
+			int32 FrameIndex0;
+			int32 FrameIndex1;
+			float Alpha = AnimEncoding::TimeToIndex(DecompContext, TimeMarkers0, NumKeys0, DecompContext.Segment0->NumFrames, TimeMarkerSize[0], SegmentRelativePos0, FrameIndex0, FrameIndex1);
+
+			const int32 TrackDataOffset0 = Align(TimeMarkersOffset0 + (NumKeys0 * TimeMarkerSize[0]), 4);
+			const uint8 BytesPerKey0 = TrackStreamKeySizes[0][ScaleValueOffset];
+			const int32 KeyDataOffset0 = TrackDataOffset0 + (FrameIndex0 * BytesPerKey0);
+			const uint8* KeyData0 = DecompContext.CompressedByteStream + KeyDataOffset0;
+
+			FVector Scale;
+			FAnimationCompression_PerTrackUtils::DecompressScale<false>(KeyFormat0, FormatFlags0, Scale, RangeData0, KeyData0);
+
+			if (DecompContext.NeedsInterpolation)
+			{
+				const uint8* KeyData1 = KeyData0 + BytesPerKey0;
+
+				FVector Scale1;
+				FAnimationCompression_PerTrackUtils::DecompressScale<false>(KeyFormat0, FormatFlags0, Scale1, RangeData0, KeyData1);
+
+				Scale = FMath::Lerp(Scale, Scale1, Alpha);
+			}
+
+			OutAtom.SetScale3D(Scale);
+		}
+	}
+}
+#endif
\ No newline at end of file
diff --git a/Engine/Source/Runtime/Engine/Private/Animation/AnimEncoding_VariableKeyLerp.cpp b/Engine/Source/Runtime/Engine/Private/Animation/AnimEncoding_VariableKeyLerp.cpp
--- a/Engine/Source/Runtime/Engine/Private/Animation/AnimEncoding_VariableKeyLerp.cpp
+++ b/Engine/Source/Runtime/Engine/Private/Animation/AnimEncoding_VariableKeyLerp.cpp
@@ -200,3 +200,169 @@ void AEFVariableKeyLerpShared::ByteSwapScaleOut(
 	}
 }
 
+#if USE_SEGMENTING_CONTEXT
+void AEFVariableKeyLerpShared::CreateEncodingContext(FAnimSequenceDecompressionContext& DecompContext)
+{
+	checkSlow(DecompContext.EncodingContext == nullptr);
+	const FAnimSequenceCompressionHeader* Header = reinterpret_cast<const FAnimSequenceCompressionHeader*>(DecompContext.CompressedByteStream);
+	if (Header->bIsSorted)
+	{
+		DecompContext.EncodingContext = new FAEVariableKeyLerpSortedContext(DecompContext);
+	}
+	else
+	{
+		DecompContext.EncodingContext = new FAEVariableKeyLerpLinearContext(DecompContext);
+	}
+}
+
+void AEFVariableKeyLerpShared::ReleaseEncodingContext(FAnimSequenceDecompressionContext& DecompContext)
+{
+	checkSlow(DecompContext.EncodingContext != nullptr);
+	delete DecompContext.EncodingContext;
+	DecompContext.EncodingContext = nullptr;
+}
+
+static constexpr uint8 AEVariableKeyLerpSortedContextCachedKeyStructOffsets[3][2] =
+{
+	{ offsetof(FAEVariableKeyLerpSortedContext::FCachedKey, TransOffsets), offsetof(FAEVariableKeyLerpSortedContext::FCachedKey, TransFrameIndices) },
+	{ offsetof(FAEVariableKeyLerpSortedContext::FCachedKey, RotOffsets), offsetof(FAEVariableKeyLerpSortedContext::FCachedKey, RotFrameIndices) },
+	{ offsetof(FAEVariableKeyLerpSortedContext::FCachedKey, ScaleOffsets), offsetof(FAEVariableKeyLerpSortedContext::FCachedKey, ScaleFrameIndices) },
+};
+
+static void AdvanceCachedKeys(const FAnimSequenceDecompressionContext& DecompContext, FAEVariableKeyLerpSortedContext& EncodingContext)
+{
+	const int32 SampleSizes[3] = { DecompContext.PackedTranslationSize0, DecompContext.PackedRotationSize0, DecompContext.PackedScaleSize0 };
+
+	while (true)
+	{
+		const uint8* PackedSampleData = EncodingContext.PackedSampleData;
+		const FSortedKeyHeader KeyHeader(PackedSampleData);
+		if (KeyHeader.IsEndOfStream())
+		{
+			break;	// Reached the end of the stream
+		}
+		checkSlow(KeyHeader.TrackIndex < DecompContext.NumTracks);
+
+		const uint8 SampleType = KeyHeader.GetKeyType();
+		checkSlow(SampleType <= 2);
+
+		const int32 TimeDelta = KeyHeader.GetTimeDelta();
+		const int32 FrameIndex = EncodingContext.PreviousFrameIndex + TimeDelta;
+
+		// Swap and update
+		FAEVariableKeyLerpSortedContext::FCachedKey& CachedKey = EncodingContext.CachedKeys[KeyHeader.TrackIndex];
+		int32* DataOffset = reinterpret_cast<int32*>(reinterpret_cast<uint8*>(&CachedKey) + AEVariableKeyLerpSortedContextCachedKeyStructOffsets[SampleType][0]);
+		int32* FrameIndicesOffset = reinterpret_cast<int32*>(reinterpret_cast<uint8*>(&CachedKey) + AEVariableKeyLerpSortedContextCachedKeyStructOffsets[SampleType][1]);
+		if (FrameIndex > EncodingContext.CurrentFrameIndex && FrameIndicesOffset[1] >= EncodingContext.CurrentFrameIndex)
+		{
+			break;		// Reached a sample we don't need yet, stop for now
+		}
+
+		PackedSampleData += KeyHeader.GetSize();
+
+		DataOffset[0] = DataOffset[1];
+		DataOffset[1] = static_cast<int32>(PackedSampleData - DecompContext.CompressedByteStream);
+		FrameIndicesOffset[0] = FrameIndicesOffset[1];
+		FrameIndicesOffset[1] = FrameIndex;
+
+		EncodingContext.PreviousFrameIndex = FrameIndex;
+		PackedSampleData += SampleSizes[SampleType];
+		EncodingContext.PackedSampleData = PackedSampleData;	// Update the pointer since we consumed the sample
+	}
+}
+
+FAEVariableKeyLerpSortedContext::FAEVariableKeyLerpSortedContext(const FAnimSequenceDecompressionContext& DecompContext)
+	: PreviousSampleAtTime(FLT_MAX)	// Very large to trigger a reset on the first Seek()
+{
+	CachedKeys.Empty(DecompContext.NumTracks);
+}
+
+static void Reset(const FAnimSequenceDecompressionContext& DecompContext, FAEVariableKeyLerpSortedContext& EncodingContext)
+{
+	EncodingContext.CachedKeys.Empty(DecompContext.NumTracks);
+	EncodingContext.CachedKeys.AddZeroed(DecompContext.NumTracks);
+
+	EncodingContext.PackedSampleData = DecompContext.CompressedByteStream + DecompContext.Segment0->ByteStreamOffset + DecompContext.RangeDataSize0;
+	EncodingContext.PreviousFrameIndex = 0;
+	EncodingContext.PreviousSegmentIndex = DecompContext.SegmentIndex0;
+}
+
+void FAEVariableKeyLerpSortedContext::Seek(const FAnimSequenceDecompressionContext& DecompContext, float SampleAtTime)
+{
+	if (SampleAtTime < PreviousSampleAtTime)
+	{
+		// Seeking backwards is terribly slow because we start over from the start
+		Reset(DecompContext, *this);
+	}
+	else if (PreviousSegmentIndex != DecompContext.SegmentIndex0)
+	{
+		// We are seeking forward into a new segment, start over
+		Reset(DecompContext, *this);
+	}
+
+	SegmentStartFrame0 = DecompContext.Segment0->StartFrame;
+	SegmentStartFrame1 = DecompContext.Segment1->StartFrame;
+
+	FramePos = DecompContext.RelativePos * float(DecompContext.AnimSeq->NumFrames - 1);
+
+	if (DecompContext.NeedsTwoSegments)
+	{
+		CurrentFrameIndex = PreviousSegmentIndex == 0 ? DecompContext.SegmentKeyIndex0 : DecompContext.SegmentKeyIndex1;
+	}
+	else
+	{
+		CurrentFrameIndex = FMath::Max(DecompContext.SegmentKeyIndex1, 1);
+	}
+
+	AdvanceCachedKeys(DecompContext, *this);
+
+	if (DecompContext.NeedsTwoSegments)
+	{
+		PackedSampleData = DecompContext.CompressedByteStream + DecompContext.Segment1->ByteStreamOffset + DecompContext.RangeDataSize0;
+		PreviousFrameIndex = 0;
+		CurrentFrameIndex = DecompContext.SegmentKeyIndex1;
+		PreviousSegmentIndex = DecompContext.SegmentIndex1;
+
+		AdvanceCachedKeys(DecompContext, *this);
+	}
+
+	PreviousSampleAtTime = SampleAtTime;
+}
+
+FAEVariableKeyLerpLinearContext::FAEVariableKeyLerpLinearContext(const FAnimSequenceDecompressionContext& DecompContext)
+{
+	NumAnimatedTrackStreams.Empty(DecompContext.NumTracks * DecompContext.NumStreamsPerTrack);
+	NumAnimatedTrackStreams.AddUninitialized(DecompContext.NumTracks * DecompContext.NumStreamsPerTrack);
+
+	int32 TotalNumAnimatedTrackStreams = 0;
+	for (int32 TrackIndex = 0; TrackIndex < DecompContext.NumTracks; ++TrackIndex)
+	{
+		const FTrivialTrackFlags TrackFlags(DecompContext.TrackFlags[TrackIndex]);
+
+		NumAnimatedTrackStreams[DecompContext.GetTranslationValueOffset(TrackIndex)] = TotalNumAnimatedTrackStreams;
+		TotalNumAnimatedTrackStreams += TrackFlags.IsTranslationTrivial() ? 0 : 1;
+
+		NumAnimatedTrackStreams[DecompContext.GetRotationValueOffset(TrackIndex)] = TotalNumAnimatedTrackStreams;
+		TotalNumAnimatedTrackStreams += TrackFlags.IsRotationTrivial() ? 0 : 1;
+
+		if (DecompContext.bHasScale)
+		{
+			NumAnimatedTrackStreams[DecompContext.GetScaleValueOffset(TrackIndex)] = TotalNumAnimatedTrackStreams;
+			TotalNumAnimatedTrackStreams += TrackFlags.IsScaleTrivial() ? 0 : 1;
+		}
+	}
+}
+
+void FAEVariableKeyLerpLinearContext::Seek(const FAnimSequenceDecompressionContext& DecompContext, float SampleAtTime)
+{
+	float FramePos = DecompContext.RelativePos * float(DecompContext.AnimSeq->NumFrames - 1);
+	float SegmentFramePos = FramePos - float(DecompContext.Segment0->StartFrame);
+	SegmentRelativePos0 = SegmentFramePos / float(DecompContext.Segment0->NumFrames - 1);
+
+	TimeMarkerSize[0] = DecompContext.Segment0->NumFrames < 256 ? sizeof(uint8) : sizeof(uint16);
+	TimeMarkerSize[1] = DecompContext.Segment1->NumFrames < 256 ? sizeof(uint8) : sizeof(uint16);
+
+	OffsetNumKeysPairs[0] = DecompContext.CompressedByteStream + DecompContext.Segment0->ByteStreamOffset + DecompContext.RangeDataSize0;
+	OffsetNumKeysPairs[1] = DecompContext.CompressedByteStream + DecompContext.Segment1->ByteStreamOffset + DecompContext.RangeDataSize0;
+}
+#endif
\ No newline at end of file
diff --git a/Engine/Source/Runtime/Engine/Private/Animation/AnimNode_SequencePlayer.cpp b/Engine/Source/Runtime/Engine/Private/Animation/AnimNode_SequencePlayer.cpp
--- a/Engine/Source/Runtime/Engine/Private/Animation/AnimNode_SequencePlayer.cpp
+++ b/Engine/Source/Runtime/Engine/Private/Animation/AnimNode_SequencePlayer.cpp
@@ -2,6 +2,7 @@
 
 #include "Animation/AnimNode_SequencePlayer.h"
 #include "Animation/AnimInstanceProxy.h"
+#include "AnimEncoding.h"
 
 #define LOCTEXT_NAMESPACE "AnimNode_SequencePlayer"
 
@@ -64,11 +65,6 @@ void FAnimNode_SequencePlayer::UpdateAssetPlayer(const FAnimationUpdateContext&
 
 void FAnimNode_SequencePlayer::Evaluate_AnyThread(FPoseContext& Output)
 {
-	auto BoolToYesNo = [](const bool& bBool) -> const TCHAR*
-	{
-		return bBool ? TEXT("Yes") : TEXT("No");
-	};
-
 	if ((Sequence != nullptr) && (Output.AnimInstanceProxy->IsSkeletonCompatible(Sequence->GetSkeleton())))
 	{
 		const bool bExpectedAdditive = Output.ExpectsAdditivePose();
diff --git a/Engine/Source/Runtime/Engine/Private/Animation/AnimSequence.cpp b/Engine/Source/Runtime/Engine/Private/Animation/AnimSequence.cpp
--- a/Engine/Source/Runtime/Engine/Private/Animation/AnimSequence.cpp
+++ b/Engine/Source/Runtime/Engine/Private/Animation/AnimSequence.cpp
@@ -29,10 +29,15 @@
 #include "Animation/AnimCompressionDerivedData.h"
 #include "UObject/UObjectThreadContext.h"
 #include "Animation/AnimNotifies/AnimNotifyState.h"
+#include "Widgets/Notifications/SNotificationList.h"
+#include "Framework/Notifications/NotificationManager.h"
+#include "ProfilingDebugging/CsvProfiler.h"
 
 #define USE_SLERP 0
 #define LOCTEXT_NAMESPACE "AnimSequence"
 
+CSV_DEFINE_CATEGORY(Animation, true);
+
 DECLARE_CYCLE_STAT(TEXT("AnimSeq GetBonePose"), STAT_AnimSeq_GetBonePose, STATGROUP_Anim);
 DECLARE_CYCLE_STAT(TEXT("Build Anim Track Pairs"), STAT_BuildAnimTrackPairs, STATGROUP_Anim);
 DECLARE_CYCLE_STAT(TEXT("Extract Pose From Anim Data"), STAT_ExtractPoseFromAnimData, STATGROUP_Anim);
@@ -159,7 +164,6 @@ void UAnimSequence::GetAssetRegistryTags(TArray<FAssetRegistryTag>& OutTags) con
 	}
 #endif
 
-
 	OutTags.Add(FAssetRegistryTag(TEXT("Compression Ratio"), FString::Printf(TEXT("%.03f"), (float)GetApproxCompressedSize() / (float)GetUncompressedRawSize()), FAssetRegistryTag::TT_Numerical));
 	OutTags.Add(FAssetRegistryTag(TEXT("Compressed Size (KB)"), FString::Printf(TEXT("%.02f"), (float)GetApproxCompressedSize() / 1024.0f), FAssetRegistryTag::TT_Numerical));
 	OutTags.Add(FAssetRegistryTag(TEXT("FrameRate"), FString::Printf(TEXT("%.2f"), GetFrameRate()), FAssetRegistryTag::TT_Numerical));
@@ -187,7 +191,7 @@ int32 UAnimSequence::GetApproxRawSize() const
 
 int32 UAnimSequence::GetApproxCompressedSize() const
 {
-	const int32 Total = sizeof(int32)*CompressedTrackOffsets.Num() + CompressedByteStream.Num() + CompressedScaleOffsets.GetMemorySize();
+	const int32 Total = sizeof(int32)*CompressedTrackOffsets.Num() + CompressedByteStream.Num() + CompressedScaleOffsets.GetMemorySize() + sizeof(FCompressedSegment)*CompressedSegments.Num();
 	return Total;
 }
 
@@ -477,7 +481,7 @@ void UAnimSequence::PostLoad()
 		SequenceLength = MINIMUM_ANIMATION_LENGTH;
 	}
 	// Raw data exists, but missing compress animation data
-	else if( !bCompressionInProgress && GetSkeleton() && CompressedTrackOffsets.Num() == 0 && RawAnimationData.Num() > 0)
+	else if( !bCompressionInProgress && GetSkeleton() && !IsCompressedDataValid())
 	{
 		UE_LOG(LogAnimation, Fatal, TEXT("No animation compression exists for sequence %s (%s)"), *GetName(), (GetOuter() ? *GetOuter()->GetFullName() : *GetFullName()) );
 	}
@@ -489,7 +493,7 @@ void UAnimSequence::PostLoad()
 		{
 			SetSkeletonVirtualBoneGuid(GetSkeleton()->GetVirtualBoneGuid());
 		}
-		if( RawAnimationData.Num() > 0  && CompressedTrackOffsets.Num() > 0 )
+		if( RawAnimationData.Num() > 0  && CompressedByteStream.Num() > 0 )
 		{
 #if [REDACTED_EMAIL]/Package...
 			// Don't do this on consoles; raw animation data should have been stripped during cook!
@@ -725,7 +729,7 @@ void UAnimSequence::PostEditChangeProperty(FPropertyChangedEvent& PropertyChange
 		
 	}
 	// @Todo fix me: This is temporary fix to make sure they always have compressed data
-	if (RawAnimationData.Num() > 0 && (CompressedTrackOffsets.Num() == 0 || bAdditiveSettingsChanged))
+	if (RawAnimationData.Num() > 0 && (!IsCompressedDataValid() || bAdditiveSettingsChanged))
 	{
 		PostProcessSequence();
 	}
@@ -766,18 +770,29 @@ static int32 FindKeyIndex(float Time, const TimeArray& Times)
 void UAnimSequence::GetBoneTransform(FTransform& OutAtom, int32 TrackIndex, float Time, bool bUseRawData) const
 {
 	// If the caller didn't request that raw animation data be used . . .
-	if ( !bUseRawData )
+	if ( !bUseRawData && IsCompressedDataValid() )
 	{
-		if ( CompressedTrackOffsets.Num() > 0 )
-		{
-			AnimationFormat_GetBoneAtom( OutAtom, *this, TrackIndex, Time );
-			return;
-		}
+		FAnimSequenceDecompressionContext DecompContext(this);
+		DecompContext.Seek(Time);
+		AnimationFormat_GetBoneAtom( OutAtom, DecompContext, TrackIndex );
+		return;
 	}
 
 	ExtractBoneTransform(RawAnimationData, OutAtom, TrackIndex, Time);
 }
 
+void UAnimSequence::GetBoneTransform(FTransform& OutAtom, int32 TrackIndex, FAnimSequenceDecompressionContext& DecompContext, bool bUseRawData) const
+{
+	// If the caller didn't request that raw animation data be used . . .
+	if (!bUseRawData && IsCompressedDataValid())
+	{
+		AnimationFormat_GetBoneAtom(OutAtom, DecompContext, TrackIndex);
+		return;
+	}
+
+	ExtractBoneTransform(RawAnimationData, OutAtom, TrackIndex, DecompContext.Time);
+}
+
 void UAnimSequence::ExtractBoneTransform(const TArray<struct FRawAnimSequenceTrack>& InRawAnimationData, FTransform& OutAtom, int32 TrackIndex, float Time) const
 {
 	// Bail out if the animation data doesn't exists (e.g. was stripped by the cooker).
@@ -1482,6 +1497,12 @@ void UAnimSequence::GetBonePose(FCompactPose& OutPose, FBlendedCurve& OutCurve,
 
 	{
 		SCOPE_CYCLE_COUNTER(STAT_ExtractPoseFromAnimData);
+		CSV_SCOPED_TIMING_STAT(Animation, ExtractPoseFromAnimData);
+		CSV_CUSTOM_STAT(Animation, NumberOfExtractedAnimations, 1, ECsvCustomStatOp::Accumulate);
+
+		FAnimSequenceDecompressionContext EvalDecompContext(this);
+		EvalDecompContext.Seek(ExtractionContext.CurrentTime);
+
 		// Handle Root Bone separately
 		if (bFirstTrackIsRootBone)
 		{
@@ -1491,9 +1512,8 @@ void UAnimSequence::GetBonePose(FCompactPose& OutPose, FBlendedCurve& OutCurve,
 
 			AnimationFormat_GetBoneAtom(
 				RootAtom,
-				*this,
-				TrackIndex,
-				ExtractionContext.CurrentTime);
+				EvalDecompContext,
+				TrackIndex);
 
 			// @laurent - we should look into splitting rotation and translation tracks, so we don't have to process translation twice.
 			RetargetBoneTransform(RootAtom, 0, RootBone, RequiredBones, bIsBakedAdditive);
@@ -1502,20 +1522,17 @@ void UAnimSequence::GetBonePose(FCompactPose& OutPose, FBlendedCurve& OutCurve,
 		if (RotationScalePairs.Num() > 0)
 		{
 			// get the remaining bone atoms
-			OutPose.PopulateFromAnimation( //@TODO:@ANIMATION: Nasty hack, be good to not have this function on the pose
-				*this,
-				RotationScalePairs,
-				TranslationPairs,
-				RotationScalePairs,
-				ExtractionContext.CurrentTime);
+			FTransformArray LocalBones;
+			OutPose.CopyBonesTo(LocalBones);
 
-			/*AnimationFormat_GetAnimationPose(
-				*((FTransformArray*)&OutAtoms), 
+			AnimationFormat_GetAnimationPose(
+				LocalBones,
 				RotationScalePairs,
 				TranslationPairs,
 				RotationScalePairs,
-				*this,
-				ExtractionContext.CurrentTime);*/
+				EvalDecompContext);
+
+			OutPose.CopyBonesFrom(LocalBones);
 		}
 	}
 
@@ -1618,6 +1635,26 @@ int32 UAnimSequence::AddNewRawTrack(FName TrackName, FRawAnimSequenceTrack* Trac
 			return TrackIndex;
 		}
 
+		// During compression, we store the track indices on 16 bits
+		const int32 MAX_NUM_TRACKS = 65535;
+		if (RawAnimationData.Num() >= MAX_NUM_TRACKS)
+		{
+			FFormatNamedArguments Args;
+			Args.Add(TEXT("TrackName"), FText::FromName(TrackName));
+			FNotificationInfo Info(FText::Format(LOCTEXT("TooManyRawTracks", "Cannot add track with name \"{TrackName}\". An animation sequence cannot contain more than 65535 tracks."), Args));
+
+			Info.bUseLargeFont = false;
+			Info.ExpireDuration = 5.0f;
+
+			TSharedPtr<SNotificationItem> Notification = FSlateNotificationManager::Get().AddNotification(Info);
+			if (Notification.IsValid())
+			{
+				Notification->SetCompletionState(SNotificationItem::CS_Fail);
+			}
+
+			UE_LOG(LogAnimation, Error, TEXT("Cannot add track with name \"%s\". An animation sequence cannot contain more than 65535 tracks."), *TrackName.ToString());
+		}
+
 		check(AnimationTrackNames.Num() == RawAnimationData.Num());
 		TrackIndex = AnimationTrackNames.Add(TrackName);
 		TrackToSkeletonMapTable.Add(FTrackToSkeletonMap(SkeletonIndex));
@@ -2194,16 +2231,13 @@ bool UAnimSequence::CompressRawAnimData(float MaxPosDiff, float MaxAngleDiff)
 				}
 			}
 		}
-
-		CompressedTrackOffsets.Empty();
-		CompressedScaleOffsets.Empty();
 	}
-	else
-	{
-		CompressedTrackOffsets.Empty();
-		CompressedScaleOffsets.Empty();
-	}
-	
+
+	CompressedTrackOffsets.Empty();
+	CompressedScaleOffsets.Empty();
+	CompressedSegments.Empty();
+	CompressedByteStream.Empty();
+
 #endif
 	return bRemovedKeys;
 }
@@ -2376,6 +2410,7 @@ void UAnimSequence::SerializeCompressedData(FArchive& Ar, bool bDDCData)
 
 	Ar << CompressedTrackOffsets;
 	Ar << CompressedScaleOffsets;
+	Ar << CompressedSegments;
 
 	Ar << CompressedTrackToSkeletonMapTable;
 	Ar << CompressedCurveData;
@@ -2613,6 +2648,7 @@ void UAnimSequence::BakeOutVirtualBoneTracks()
 	//Pose evaluation data
 	FCompactPose Pose;
 	Pose.SetBoneContainer(&EvalContext.RequiredBones);
+
 	FAnimExtractContext ExtractContext;
 
 	const TArray<FVirtualBoneRefData>& VBRefData = GetSkeleton()->GetReferenceSkeleton().GetVirtualBoneRefData();
@@ -2718,6 +2754,7 @@ void UAnimSequence::BakeOutAdditiveIntoRawData()
 	Pose.SetBoneContainer(&EvalContext.RequiredBones);
 	FCompactPose BasePose;
 	BasePose.SetBoneContainer(&EvalContext.RequiredBones);
+
 	FAnimExtractContext ExtractContext;
 
 	for (int Frame = 0; Frame < NumFrames; ++Frame)
@@ -2845,6 +2882,7 @@ void UAnimSequence::RecycleAnimSequence()
 	CompressedTrackOffsets.Empty(0);
 	CompressedByteStream.Empty(0);
 	CompressedScaleOffsets.Empty(0);
+	CompressedSegments.Empty(0);
 	SourceRawAnimationData.Empty(0);
 	RawCurveData.Empty();
 	CompressedCurveData.Empty();
@@ -3647,7 +3685,7 @@ void UAnimSequence::RemapTracksToNewSkeleton( USkeleton* NewSkeleton, bool bConv
 	}
 
 	AddAnimLoadingDebugEntry(TEXT("RemapTracksToNewSkeletonPreSetSkeleton"));
-	SetSkeleton(NewSkeleton);
+	Super::RemapTracksToNewSkeleton(NewSkeleton, bConvertSpaces);
 	AddAnimLoadingDebugEntry(TEXT("RemapTracksToNewSkeletonEnd"));
 }
 
@@ -4626,6 +4664,7 @@ void UAnimSequence::ResetAnimation()
 	TrackToSkeletonMapTable.Empty();
 	CompressedTrackOffsets.Empty();
 	CompressedScaleOffsets.Empty();
+	CompressedSegments.Empty();
 	CompressedByteStream.Empty();
 
 	Notifies.Empty();
@@ -5570,14 +5609,14 @@ void UAnimSequence::OnRawDataChanged()
 	//RequestDependentAnimRecompression();
 }
 
+#endif
+
 bool UAnimSequence::IsCompressedDataValid() const
 {
 	return CompressedByteStream.Num() > 0 || RawAnimationData.Num() == 0 ||
 		   (TranslationCompressionFormat == ACF_Identity && RotationCompressionFormat == ACF_Identity && ScaleCompressionFormat == ACF_Identity);
 }
 
-#endif
-
 /*-----------------------------------------------------------------------------
 	AnimNotify& subclasses
 -----------------------------------------------------------------------------*/
diff --git a/Engine/Source/Runtime/Engine/Private/Animation/AnimSequenceBase.cpp b/Engine/Source/Runtime/Engine/Private/Animation/AnimSequenceBase.cpp
--- a/Engine/Source/Runtime/Engine/Private/Animation/AnimSequenceBase.cpp
+++ b/Engine/Source/Runtime/Engine/Private/Animation/AnimSequenceBase.cpp
@@ -281,6 +281,15 @@ void UAnimSequenceBase::GetAnimNotifiesFromDeltaPositions(const float& PreviousP
 	}
 }
 
+#if WITH_EDITOR
+void UAnimSequenceBase::RemapTracksToNewSkeleton(USkeleton* NewSkeleton, bool bConvertSpaces)
+{
+	Super::RemapTracksToNewSkeleton(NewSkeleton, bConvertSpaces);
+	VerifyCurveNames<FFloatCurve>(*NewSkeleton, USkeleton::AnimCurveMappingName, RawCurveData.FloatCurves);
+	VerifyCurveNames<FTransformCurve>(*NewSkeleton, USkeleton::AnimTrackCurveMappingName, RawCurveData.TransformCurves);
+}
+#endif
+
 void UAnimSequenceBase::TickAssetPlayer(FAnimTickRecord& Instance, struct FAnimNotifyQueue& NotifyQueue, FAnimAssetTickContext& Context) const
 {
 	float& CurrentTime = *(Instance.TimeAccumulator);
diff --git a/Engine/Source/Runtime/Engine/Private/Animation/AnimSequenceDecompressionContext.cpp b/Engine/Source/Runtime/Engine/Private/Animation/AnimSequenceDecompressionContext.cpp
new file mode 100644
--- /dev/null
+++ b/Engine/Source/Runtime/Engine/Private/Animation/AnimSequenceDecompressionContext.cpp
@@ -0,0 +1,273 @@
+#include "Animation/AnimSequenceDecompressionContext.h"
+#include "Animation/AnimEncodingDecompressionContext.h"
+#include "Animation/AnimSequence.h"
+#include "AnimEncoding.h"
+
+#include <algorithm>
+
+#if USE_SEGMENTING_CONTEXT 
+FAnimSequenceDecompressionContext::FAnimSequenceDecompressionContext()
+{
+	memset(this, 0, sizeof(FAnimSequenceDecompressionContext));
+}
+
+FAnimSequenceDecompressionContext::FAnimSequenceDecompressionContext(const UAnimSequence* AnimSeq_)
+	: FAnimSequenceDecompressionContext()
+{
+	Bind(AnimSeq_);
+}
+
+FAnimSequenceDecompressionContext::FAnimSequenceDecompressionContext(const FAnimSequenceDecompressionContext& Other)
+{
+	// We don't disable copying but we also don't retain any state
+	memset(this, 0, sizeof(FAnimSequenceDecompressionContext));
+}
+
+FAnimSequenceDecompressionContext& FAnimSequenceDecompressionContext::operator=(const FAnimSequenceDecompressionContext& Other)
+{
+	// We don't disable copying but we also don't retain any state
+	if (EncodingContext != nullptr)
+	{
+		checkSlow(RotationCodec != nullptr);
+		RotationCodec->ReleaseEncodingContext(*this);
+		checkSlow(EncodingContext == nullptr);
+	}
+
+	memset(this, 0, sizeof(FAnimSequenceDecompressionContext));
+	return *this;
+}
+
+FAnimSequenceDecompressionContext::~FAnimSequenceDecompressionContext()
+{
+	if (EncodingContext != nullptr)
+	{
+		checkSlow(RotationCodec != nullptr);
+		RotationCodec->ReleaseEncodingContext(*this);
+		checkSlow(EncodingContext == nullptr);
+	}
+}
+
+void FAnimSequenceDecompressionContext::Seek(float SampleAtTime)
+{
+	if (AnimSeq == nullptr)
+	{
+		return;	// Context is not bound and valid
+	}
+
+	// Make sure to clamp
+	SampleAtTime = FMath::Clamp(SampleAtTime, 0.0f, AnimSeq->SequenceLength);
+
+	Time = SampleAtTime;
+	RelativePos = SampleAtTime / AnimSeq->SequenceLength;
+
+	if (!AnimSeq->IsCompressedDataValid())
+	{
+		return;	// No compressed data
+	}
+
+	if (HasSegments)
+	{
+		const FAnimSequenceCompressionHeader* Header = reinterpret_cast<const FAnimSequenceCompressionHeader*>(CompressedByteStream);
+		const int32 NumFrames = Header->NumFrames;
+
+		KeyAlpha = AnimEncoding::TimeToIndex(*AnimSeq, RelativePos, NumFrames, KeyIndex0, KeyIndex1);
+
+		SegmentIndex0 = INDEX_NONE;
+		SegmentIndex1 = INDEX_NONE;
+		const uint16 NumSegments = static_cast<uint16>(AnimSeq->CompressedSegments.Num());
+		for (uint16 SegmentIndex = 0; SegmentIndex < NumSegments; ++SegmentIndex)
+		{
+			const FCompressedSegment& Segment = AnimSeq->CompressedSegments[SegmentIndex];
+			if (KeyIndex0 >= Segment.StartFrame && KeyIndex0 < Segment.StartFrame + Segment.NumFrames)
+			{
+				SegmentIndex0 = SegmentIndex;
+
+				if (KeyIndex1 >= Segment.StartFrame && KeyIndex1 < Segment.StartFrame + Segment.NumFrames)
+				{
+					SegmentIndex1 = SegmentIndex;
+				}
+				else
+				{
+					check(SegmentIndex + 1 < NumSegments);
+					const FCompressedSegment& NextSegment = AnimSeq->CompressedSegments[SegmentIndex + 1];
+					check(KeyIndex1 >= NextSegment.StartFrame && KeyIndex1 < NextSegment.StartFrame + NextSegment.NumFrames);
+					SegmentIndex1 = SegmentIndex + 1;
+				}
+				break;
+			}
+		}
+
+		check(SegmentIndex0 != INDEX_NONE && SegmentIndex1 != INDEX_NONE);
+
+		Segment0 = &AnimSeq->CompressedSegments[SegmentIndex0];
+		Segment1 = &AnimSeq->CompressedSegments[SegmentIndex1];
+		NeedsTwoSegments = SegmentIndex0 != SegmentIndex1;
+		NeedsInterpolation = KeyIndex0 != KeyIndex1;
+
+		SegmentKeyIndex0 = KeyIndex0 - Segment0->StartFrame;
+		SegmentKeyIndex1 = KeyIndex1 - (SegmentIndex0 == SegmentIndex1 ? Segment0->StartFrame : Segment1->StartFrame);
+
+		// When we need two segments, we always need the last key from segment 0 and the first key from segment 1
+		checkSlow(!NeedsTwoSegments || SegmentKeyIndex0 == Segment0->NumFrames - 1);
+		checkSlow(!NeedsTwoSegments || SegmentKeyIndex1 == 0);
+
+		PackedTranslationSize0 = CompressedTranslationStrides[Segment0->TranslationCompressionFormat] * CompressedTranslationNum[Segment0->TranslationCompressionFormat];
+		PackedRotationSize0 = CompressedRotationStrides[Segment0->RotationCompressionFormat] * CompressedRotationNum[Segment0->RotationCompressionFormat];
+		PackedScaleSize0 = bHasScale ? (CompressedScaleStrides[Segment0->ScaleCompressionFormat] * CompressedScaleNum[Segment0->ScaleCompressionFormat]) : 0;
+
+		// TODO: Use bitset popcnt to calculate this
+		RangeDataSize0 = 0;
+		for (int32 CompressedTrackIndex = 0; CompressedTrackIndex < NumTracks; ++CompressedTrackIndex)
+		{
+			const FTrivialTrackFlags TrivialTrackFlags(TrackFlags[CompressedTrackIndex]);
+
+			RangeDataSize0 += !TrivialTrackFlags.IsTranslationTrivial() && Segment0->TranslationCompressionFormat == ACF_IntervalFixed32NoW ? (sizeof(float) * 6) : 0;
+			RangeDataSize0 += !TrivialTrackFlags.IsRotationTrivial() && Segment0->RotationCompressionFormat == ACF_IntervalFixed32NoW ? (sizeof(float) * 6) : 0;
+			RangeDataSize0 += !TrivialTrackFlags.IsScaleTrivial() && Segment0->ScaleCompressionFormat == ACF_IntervalFixed32NoW ? (sizeof(float) * 6) : 0;
+		}
+
+		TrackRangeData[0] = CompressedByteStream + Segment0->ByteStreamOffset;
+		TrackRangeData[1] = CompressedByteStream + Segment1->ByteStreamOffset;
+	}
+	else
+	{
+		// Legacy sequences need to refresh every update since we can't detect a stale context without a CRC
+		RotationCodec = AnimSeq->RotationCodec;
+		TranslationCodec = AnimSeq->TranslationCodec;
+		ScaleCodec = AnimSeq->ScaleCodec;
+		CompressedTrackOffsets = AnimSeq->CompressedTrackOffsets.GetData();
+		CompressedByteStream = AnimSeq->CompressedByteStream.GetData();
+		CompressedScaleOffsets = AnimSeq->CompressedScaleOffsets.IsValid() ? &AnimSeq->CompressedScaleOffsets : nullptr;
+	}
+
+	if (EncodingContext != nullptr)
+	{
+		EncodingContext->Seek(*this, SampleAtTime);
+	}
+}
+
+bool FAnimSequenceDecompressionContext::IsStale(const UAnimSequence* AnimSeq_) const
+{
+	if (AnimSeq_ != AnimSeq)
+	{
+		return true;
+	}
+
+	if (AnimSeq_ == nullptr)
+	{
+		return false;
+	}
+
+	uint32 NewCRC = 0;
+	if (AnimSeq_->CompressedSegments.Num() != 0)
+	{
+		const FAnimSequenceCompressionHeader* Header = reinterpret_cast<const FAnimSequenceCompressionHeader*>(AnimSeq_->CompressedByteStream.GetData());
+		NewCRC = Header->SequenceCRC;
+	}
+
+	if (NewCRC != SequenceCRC || AnimSeq_->CompressedByteStream.GetData() != CompressedByteStream)
+	{
+		return true;
+	}
+
+	return false;
+}
+
+void FAnimSequenceDecompressionContext::Bind(const UAnimSequence* AnimSeq_)
+{
+	if (!IsStale(AnimSeq_))
+	{
+		return;	// Nothing to do
+	}
+
+#if WITH_EDITOR
+	// Check if our context is bound too often, if probably means that it is being shared between multiple anim sequences which is bad for performance.
+	const double CurrentTime = FPlatformTime::Seconds();
+	const double ElapsedTime = CurrentTime - PreviousBindTimeStamps[0];
+
+	const double BIND_FREQUENCY_CHECK_THRESHOLD = 1.0;		// In seconds
+	if (ElapsedTime < BIND_FREQUENCY_CHECK_THRESHOLD)
+	{
+		// PreviousBindTimeStamps.Num() bind calls were executed on a stale context within BIND_FREQUENCY_CHECK_THRESHOLD seconds
+		UE_LOG(LogAnimation, Warning, TEXT("Decompression context is bound too often. Consider reusing the same context for an anim sequence every frame or performance will be degraded."));
+	}
+
+	std::copy_n(&PreviousBindTimeStamps[1], PreviousBindTimeStamps.Num() - 1, &PreviousBindTimeStamps[0]);
+	PreviousBindTimeStamps[PreviousBindTimeStamps.Num() - 1] = CurrentTime;
+#endif
+
+	// Was previously bound to something else, reset everything since we need to refresh
+	if (EncodingContext != nullptr)
+	{
+		checkSlow(RotationCodec != nullptr);
+		RotationCodec->ReleaseEncodingContext(*this);
+		checkSlow(EncodingContext == nullptr);
+	}
+
+	memset(this, 0, sizeof(FAnimSequenceDecompressionContext));
+
+	AnimSeq = AnimSeq_;
+
+	if (AnimSeq_ == nullptr || !AnimSeq_->IsCompressedDataValid())
+	{
+		return;	// No compressed data
+	}
+
+	RotationCodec = AnimSeq_->RotationCodec;
+	TranslationCodec = AnimSeq_->TranslationCodec;
+	ScaleCodec = AnimSeq_->ScaleCodec;
+	CompressedTrackOffsets = AnimSeq_->CompressedTrackOffsets.GetData();
+	CompressedByteStream = AnimSeq_->CompressedByteStream.GetData();
+	CompressedScaleOffsets = AnimSeq_->CompressedScaleOffsets.IsValid() ? &AnimSeq_->CompressedScaleOffsets : nullptr;
+	EncodingContext = nullptr;
+
+	HasSegments = AnimSeq_->CompressedSegments.Num() != 0;
+	if (HasSegments)
+	{
+		const FAnimSequenceCompressionHeader* Header = reinterpret_cast<const FAnimSequenceCompressionHeader*>(CompressedByteStream);
+		NumTracks = Header->NumTracks;
+		SequenceCRC = Header->SequenceCRC;
+		bHasScale = Header->bHasScale != 0;
+		bIsSorted = Header->bIsSorted != 0;
+		NumStreamsPerTrack = Header->bHasScale ? 3 : 2;
+
+		const SIZE_T TrackFlagsOffset = sizeof(FAnimSequenceCompressionHeader);
+		TrackFlags = CompressedByteStream + TrackFlagsOffset;
+
+		const SIZE_T TrivialTrackKeysOffset = TrackFlagsOffset + Align(NumTracks * sizeof(uint8), 4);
+		TrivialTrackKeys = CompressedByteStream + TrivialTrackKeysOffset;
+
+		TrivialTrackStreamOffsets.Empty(NumTracks * NumStreamsPerTrack);
+		TrivialTrackStreamOffsets.AddUninitialized(NumTracks * NumStreamsPerTrack);
+
+		int32 TrivialTrackStreamKeyOffset = 0;
+		for (int32 TrackIndex = 0; TrackIndex < NumTracks; ++TrackIndex)
+		{
+			const FTrivialTrackFlags TrivialTrackFlags(TrackFlags[TrackIndex]);
+			const bool IsTranslationTrivial = TrivialTrackFlags.IsTranslationTrivial();
+			const bool IsRotationTrivial = TrivialTrackFlags.IsRotationTrivial();
+
+			TrivialTrackStreamOffsets[GetTranslationValueOffset(TrackIndex)] = IsTranslationTrivial ? TrivialTrackStreamKeyOffset : -1;
+			TrivialTrackStreamKeyOffset += IsTranslationTrivial ? (CompressedTranslationStrides[ACF_None] * CompressedTranslationNum[ACF_None]) : 0;
+
+			TrivialTrackStreamOffsets[GetRotationValueOffset(TrackIndex)] = IsRotationTrivial ? TrivialTrackStreamKeyOffset : -1;
+			TrivialTrackStreamKeyOffset += IsRotationTrivial ? (CompressedRotationStrides[ACF_Float96NoW] * CompressedRotationNum[ACF_Float96NoW]) : 0;
+
+			if (Header->bHasScale)
+			{
+				const bool IsScaleTrivial = TrivialTrackFlags.IsScaleTrivial();
+
+				TrivialTrackStreamOffsets[GetScaleValueOffset(TrackIndex)] = IsScaleTrivial ? TrivialTrackStreamKeyOffset : -1;
+				TrivialTrackStreamKeyOffset += IsScaleTrivial ? (CompressedScaleStrides[ACF_None] * CompressedScaleNum[ACF_None]) : 0;
+			}
+		}
+
+		RotationCodec->CreateEncodingContext(*this);
+	}
+	else
+	{
+		bHasScale = CompressedScaleOffsets != nullptr;
+		SequenceCRC = 0;
+	}
+}
+#endif
\ No newline at end of file
diff --git a/Engine/Source/Runtime/Engine/Private/Animation/AnimSingleNodeInstanceProxy.cpp b/Engine/Source/Runtime/Engine/Private/Animation/AnimSingleNodeInstanceProxy.cpp
--- a/Engine/Source/Runtime/Engine/Private/Animation/AnimSingleNodeInstanceProxy.cpp
+++ b/Engine/Source/Runtime/Engine/Private/Animation/AnimSingleNodeInstanceProxy.cpp
@@ -6,12 +6,17 @@
 #include "Animation/BlendSpaceBase.h"
 #include "Animation/PoseAsset.h"
 #include "Animation/AnimSingleNodeInstance.h"
+#include "AnimEncoding.h"
+
+FAnimSingleNodeInstanceProxy::~FAnimSingleNodeInstanceProxy()
+{
+}
 
 void FAnimSingleNodeInstanceProxy::Initialize(UAnimInstance* InAnimInstance)
 {
 	FAnimInstanceProxy::Initialize(InAnimInstance);
 
-	CurrentAsset = NULL;
+	CurrentAsset = nullptr;
 #if WITH_EDITORONLY_DATA
 	PreviewPoseCurrentTime = 0.0f;
 #endif
diff --git a/Engine/Source/Runtime/Engine/Private/Animation/AnimationSettings.cpp b/Engine/Source/Runtime/Engine/Private/Animation/AnimationSettings.cpp
--- a/Engine/Source/Runtime/Engine/Private/Animation/AnimationSettings.cpp
+++ b/Engine/Source/Runtime/Engine/Private/Animation/AnimationSettings.cpp
@@ -16,10 +16,8 @@ UAnimationSettings::UAnimationSettings(const FObjectInitializer& ObjectInitializ
 	, bForceBelowThreshold(false)
 	, bFirstRecompressUsingCurrentOrDefault(true)
 	, bRaiseMaxErrorToExisting(false)
-	, bTryFixedBitwiseCompression(true)
-	, bTryPerTrackBitwiseCompression(true)
-	, bTryLinearKeyRemovalCompression(true)
-	, bTryIntervalKeyRemoval(true)
+	, bTryExhaustiveSearch(false)
+	, bEnableSegmenting(true)
 	, bEnablePerformanceLog(false)
 	, bTickAnimationOnSkeletalMeshInit(true)
 {
