Copying //UE4/Dev-AnimPhys to Dev-Main (//UE4/Dev-Main)

diff --git a/Engine/Source/Runtime/Engine/Private/Animation/AnimationUtils.cpp b/Engine/Source/Runtime/Engine/Private/Animation/AnimationUtils.cpp
--- a/Engine/Source/Runtime/Engine/Private/Animation/AnimationUtils.cpp
+++ b/Engine/Source/Runtime/Engine/Private/Animation/AnimationUtils.cpp
@@ -231,7 +231,7 @@ void FAnimationUtils::ComputeCompressionError(const UAnimSequence* AnimSeq, cons
 
 	if (AnimSeq->NumFrames > 0)
 	{
-		const bool bCanUseCompressedData = (AnimSeq->CompressedTrackOffsets.Num() > 0);
+		const bool bCanUseCompressedData = (AnimSeq->CompressedByteStream.Num() > 0);
 		if (!bCanUseCompressedData)
 		{
 			// If we can't use CompressedData, there's not much point in being here.
@@ -284,9 +284,12 @@ void FAnimationUtils::ComputeCompressionError(const UAnimSequence* AnimSeq, cons
 		const FAnimKeyHelper Helper(AnimSeq->SequenceLength, AnimSeq->NumFrames);
 		const float KeyLength = Helper.TimePerKey() + SMALL_NUMBER;
 
+		FAnimSequenceDecompressionContext DecompContext(AnimSeq);
+
 		for (int32 FrameIndex = 0; FrameIndex<AnimSeq->NumFrames; FrameIndex++)
 		{
 			const float Time = (float)FrameIndex * KeyLength;
+			DecompContext.Seek(Time);
 
 			// get the raw and compressed atom for each bone
 			for (int32 BoneIndex = 0; BoneIndex < NumBones; ++BoneIndex)
@@ -313,8 +316,8 @@ void FAnimationUtils::ComputeCompressionError(const UAnimSequence* AnimSeq, cons
 
 						FTransform AdditiveRawTransform;
 						FTransform AdditiveNewTransform;
-						AnimSeq->GetBoneTransform(AdditiveRawTransform, BoneIndexData.RawTrackIndex, Time, true);
-						AnimSeq->GetBoneTransform(AdditiveNewTransform, BoneIndexData.CompressedTrackIndex, Time, false);
+						AnimSeq->GetBoneTransform(AdditiveRawTransform, BoneIndexData.RawTrackIndex, DecompContext, true);
+						AnimSeq->GetBoneTransform(AdditiveNewTransform, BoneIndexData.CompressedTrackIndex, DecompContext, false);
 
 						const ScalarRegister VBlendWeight(1.f);
 						RawTransforms[BoneIndex].AccumulateWithAdditiveScale(AdditiveRawTransform, VBlendWeight);
@@ -322,8 +325,8 @@ void FAnimationUtils::ComputeCompressionError(const UAnimSequence* AnimSeq, cons
 					}
 					else
 					{
-						AnimSeq->GetBoneTransform(RawTransforms[BoneIndex], BoneIndexData.RawTrackIndex, Time, true);
-						AnimSeq->GetBoneTransform(NewTransforms[BoneIndex], BoneIndexData.CompressedTrackIndex, Time, false);
+						AnimSeq->GetBoneTransform(RawTransforms[BoneIndex], BoneIndexData.RawTrackIndex, DecompContext, true);
+						AnimSeq->GetBoneTransform(NewTransforms[BoneIndex], BoneIndexData.CompressedTrackIndex, DecompContext, false);
 					}
 				}
 
@@ -476,95 +479,277 @@ bool FAnimationUtils::GetForcedRecompressionSetting()
 	return ForcedRecompressionSetting;
 }
 
-#if !WITH_EDITOR
-	#define TRYCOMPRESSION_INNER(compressionname,winningcompressor_count,winningcompressor_error,winningcompressor_margin,compressionalgorithm)	
-#else
-	#define TRYCOMPRESSION_INNER(compressionname,winningcompressor_count,winningcompressor_error,winningcompressor_margin,compressionalgorithm)				\
-{																																							\
-	/* try the alternative compressor	*/																													\
-	AnimSeq->CompressionScheme = (compressionalgorithm);																									\
-	(compressionalgorithm)->Reduce(AnimSeq, CompressContext);																								\
-	AnimSeq->SetUseRawDataOnly(false);																														\
-	const SIZE_T NewSize = AnimSeq->GetApproxCompressedSize();																		\
-																																							\
-	/* compute the savings and compression error*/																											\
-	const int64 MemorySavingsFromOriginal = OriginalSize - NewSize;																						\
-	const int64 MemorySavingsFromPrevious = CurrentSize - NewSize;																							\
-	PctSaving = 0.f;																																		\
-	/* figure out our new compression error*/																												\
-	FAnimationUtils::ComputeCompressionError(AnimSeq, BoneData, NewErrorStats);																				\
-																																							\
-	const bool bLowersError = NewErrorStats.MaxError < WinningCompressorError;																				\
-	const bool bErrorUnderThreshold = NewErrorStats.MaxError <= MasterTolerance;																			\
-																																							\
-	/* keep it if it we want to force the error below the threshold and it reduces error */																	\
-	bKeepNewCompressionMethod = false;																														\
-	const bool bReducesErrorBelowThreshold = (bLowersError && (WinningCompressorError > MasterTolerance) && bForceBelowThreshold);							\
-	bKeepNewCompressionMethod |= bReducesErrorBelowThreshold;																								\
-	/* or if has an acceptable error and saves space  */																									\
-	const bool bHasAcceptableErrorAndSavesSpace = bErrorUnderThreshold && (MemorySavingsFromPrevious > 0);													\
-	bKeepNewCompressionMethod |= bHasAcceptableErrorAndSavesSpace;																							\
-	/* or if saves the same amount and an acceptable error that is lower than the previous best */															\
-	const bool bLowersErrorAndSavesSameOrBetter = bErrorUnderThreshold && bLowersError && (MemorySavingsFromPrevious >= 0);									\
-	bKeepNewCompressionMethod |= bLowersErrorAndSavesSameOrBetter;																							\
-																																							\
-	PctSaving = (OriginalSize > 0) ? (100.f - (100.f * float(NewSize) / float(OriginalSize))) : 0.f;														\
-	UE_LOG(LogAnimationCompression, Log, TEXT("- %s - bytes saved(%i) (%.1f%%) from previous(%i) MaxError(%.2f) bLowersError(%d) %s"),								\
-		compressionname, MemorySavingsFromOriginal, PctSaving, MemorySavingsFromPrevious, NewErrorStats.MaxError, bLowersError, bKeepNewCompressionMethod ? TEXT("(**Best so far**)") : TEXT(""));									\
-																																							\
-	UE_LOG(LogAnimationCompression, Log, TEXT("    bReducesErrorBelowThreshold(%d) bHasAcceptableErrorAndSavesSpace(%d) bLowersErrorAndSavesSameOrBetter(%d)"),	\
-		bReducesErrorBelowThreshold, bHasAcceptableErrorAndSavesSpace, bLowersErrorAndSavesSameOrBetter);													\
-																																							\
-	UE_LOG(LogAnimationCompression, Log, TEXT("    WinningCompressorError(%f) MasterTolerance(%f) bForceBelowThreshold(%d) bErrorUnderThreshold(%d)"),				\
-		WinningCompressorError, MasterTolerance, bForceBelowThreshold, bErrorUnderThreshold);																\
-																																							\
-	if (bKeepNewCompressionMethod)																															\
-	{																																						\
-		WinningCompressorMarginalSavings = MemorySavingsFromPrevious;																						\
-		WinningCompressorCounter = &(winningcompressor_count);																								\
-		WinningCompressorErrorSum = &(winningcompressor_error);																								\
-		WinningCompressorMarginalSavingsSum = &(winningcompressor_margin);																					\
-		WinningCompressorName = compressionname;																											\
-		CurrentSize = NewSize;																																\
-		WinningCompressorSavings = MemorySavingsFromOriginal;																								\
-		WinningCompressorError = NewErrorStats.MaxError;																									\
-																																							\
-		/* backup key information from the sequence */																										\
-		SavedCompressionScheme = AnimSeq->CompressionScheme;																								\
-		SavedTranslationCompressionFormat = AnimSeq->TranslationCompressionFormat;																			\
-		SavedRotationCompressionFormat = AnimSeq->RotationCompressionFormat;																				\
-		SavedKeyEncodingFormat = AnimSeq->KeyEncodingFormat;																								\
-		SavedCompressedTrackOffsets = AnimSeq->CompressedTrackOffsets;																						\
-		SavedCompressedScaleOffsets = AnimSeq->CompressedScaleOffsets;																						\
-		SavedCompressedByteStream = AnimSeq->CompressedByteStream;																							\
-		SavedTranslationCodec = AnimSeq->TranslationCodec;																									\
-		SavedRotationCodec = AnimSeq->RotationCodec;																										\
-		SavedScaleCodec = AnimSeq->ScaleCodec;																												\
-		bSavedUseRawDataOnly = false;																														\
-	}																																						\
-	else																																					\
-	{																																						\
-		/* revert back to the old method by copying back the data we cached */																				\
-		AnimSeq->CompressionScheme = SavedCompressionScheme;																								\
-		AnimSeq->TranslationCompressionFormat = SavedTranslationCompressionFormat;																			\
-		AnimSeq->RotationCompressionFormat = SavedRotationCompressionFormat;																				\
-		AnimSeq->KeyEncodingFormat = SavedKeyEncodingFormat;																								\
-		AnimSeq->CompressedTrackOffsets = SavedCompressedTrackOffsets;																						\
-		AnimSeq->CompressedByteStream = SavedCompressedByteStream;																							\
-		AnimSeq->CompressedScaleOffsets = SavedCompressedScaleOffsets;																						\
-		AnimSeq->TranslationCodec = SavedTranslationCodec;																									\
-		AnimSeq->RotationCodec = SavedRotationCodec;																										\
-		AnimSeq->ScaleCodec = SavedScaleCodec;																												\
-		AnimSeq->SetUseRawDataOnly(bSavedUseRawDataOnly);																									\
-		AnimationFormat_SetInterfaceLinks(*AnimSeq);																										\
-																																							\
-		const SIZE_T RestoredSize = AnimSeq->GetApproxCompressedSize();															\
-		check(RestoredSize == CurrentSize);																													\
-	}																																						\
+struct FAnimCompressionJobContext
+{
+	// Inputs
+	int64 OriginalSize;
+	float MasterTolerance;
+	bool bForceBelowThreshold;
+	const TArray<FBoneData>* BoneData;
+	FAnimCompressContext* CompressContext;
+	UAnimCompress* CompressionAlgorithm;
+	UAnimSequence* AnimSeq;
+	const TCHAR* CompressionName;
+	int32* WinningCompressor_Count;
+	float* WinningCompressor_Error;
+	int64* WinningCompressor_Margin;
+
+	// Input/Output
+	int64 CurrentSize;
+	bool bSavedUseRawDataOnly;
+
+	int64 WinningCompressorMarginalSavings;
+	int32* WinningCompressorCounter;
+	float* WinningCompressorErrorSum;
+	int64* WinningCompressorMarginalSavingsSum;
+	FString WinningCompressorName;
+	int32 WinningCompressorSavings;
+	float WinningCompressorError;
+
+	UAnimCompress* SavedCompressionScheme;
+	AnimationCompressionFormat SavedTranslationCompressionFormat;
+	AnimationCompressionFormat SavedRotationCompressionFormat;
+	AnimationCompressionFormat SavedScaleCompressionFormat;
+	AnimationKeyFormat SavedKeyEncodingFormat;
+	TArray<int32> SavedCompressedTrackOffsets;
+	FCompressedOffsetData SavedCompressedScaleOffsets;
+	TArray<FCompressedSegment> SavedCompressedSegments;
+	TArray<uint8> SavedCompressedByteStream;
+	AnimEncoding* SavedTranslationCodec;
+	AnimEncoding* SavedRotationCodec;
+	AnimEncoding* SavedScaleCodec;
+
+	// Outputs
+	AnimationErrorStats NewErrorStats;
+	float PctSaving;
+
+	FAnimCompressionJobContext()
+	{
+		memset(this, 0, sizeof(FAnimCompressionJobContext));
+	}
+};
+
+static void TryCompressionInner(FAnimCompressionJobContext& JobContext, bool is_async);
+
+class FAsyncAnimCompressionTask
+{
+public:
+	FAsyncAnimCompressionTask(FAnimCompressionJobContext* JobContext_)
+		: JobContext(JobContext_)
+	{}
+
+	/** return the name of the task **/
+	static const TCHAR* GetTaskName() { return TEXT("FAsyncAnimCompressionTask"); }
+	FORCEINLINE static TStatId GetStatId() { RETURN_QUICK_DECLARE_CYCLE_STAT(FAsyncAnimCompressionTask, STATGROUP_TaskGraphTasks); }
+
+	static ENamedThreads::Type GetDesiredThread() { return ENamedThreads::AnyThread; }
+	static ESubsequentsMode::Type GetSubsequentsMode() { return ESubsequentsMode::TrackSubsequents; }
+
+	void DoTask(ENamedThreads::Type CurrentThread, const FGraphEventRef& MyCompletionGraphEvent)
+	{
+		TryCompressionInner(*JobContext, true);
+	}
+
+	FAnimCompressionJobContext* JobContext;
+};
+
+static void TryCompressionInner(FAnimCompressionJobContext& JobContext, bool is_async)
+{
+#if WITH_EDITOR
+	UAnimSequence* AnimSeq = JobContext.AnimSeq;
+
+	/* try the alternative compressor	*/
+	AnimSeq->CompressionScheme = JobContext.CompressionAlgorithm;
+	JobContext.CompressionAlgorithm->Reduce(AnimSeq, *JobContext.CompressContext, *JobContext.BoneData);
+	AnimSeq->SetUseRawDataOnly(false);
+	const SIZE_T NewSize = AnimSeq->GetApproxCompressedSize();
+
+	if (is_async)
+	{
+		JobContext.CurrentSize = NewSize;
+	}
+	else
+	{
+		/* compute the savings and compression error*/
+		const int64 MemorySavingsFromOriginal = JobContext.OriginalSize - NewSize;
+		const int64 MemorySavingsFromPrevious = JobContext.CurrentSize - NewSize;
+
+		/* figure out our new compression error*/
+		FAnimationUtils::ComputeCompressionError(AnimSeq, *JobContext.BoneData, JobContext.NewErrorStats);
+
+		const bool bLowersError = JobContext.NewErrorStats.MaxError < JobContext.WinningCompressorError;
+		const bool bErrorUnderThreshold = JobContext.NewErrorStats.MaxError <= JobContext.MasterTolerance;
+
+		/* keep it if it we want to force the error below the threshold and it reduces error */
+		bool bKeepNewCompressionMethod = false;
+		const bool bReducesErrorBelowThreshold = (bLowersError && (JobContext.WinningCompressorError > JobContext.MasterTolerance) && JobContext.bForceBelowThreshold);
+		bKeepNewCompressionMethod |= bReducesErrorBelowThreshold;
+		/* or if has an acceptable error and saves space  */
+		const bool bHasAcceptableErrorAndSavesSpace = bErrorUnderThreshold && (MemorySavingsFromPrevious > 0);
+		bKeepNewCompressionMethod |= bHasAcceptableErrorAndSavesSpace;
+		/* or if saves the same amount and an acceptable error that is lower than the previous best */
+		const bool bLowersErrorAndSavesSameOrBetter = bErrorUnderThreshold && bLowersError && (MemorySavingsFromPrevious >= 0);
+		bKeepNewCompressionMethod |= bLowersErrorAndSavesSameOrBetter;
+
+		JobContext.PctSaving = (JobContext.OriginalSize > 0) ? (100.f - (100.f * float(NewSize) / float(JobContext.OriginalSize))) : 0.f;
+		UE_LOG(LogAnimationCompression, Verbose, TEXT("- %s - bytes saved(%i) (%.1f%%) from previous(%i) MaxError(%.2f) bLowersError(%d) %s"),
+			JobContext.CompressionName, MemorySavingsFromOriginal, JobContext.PctSaving, MemorySavingsFromPrevious, JobContext.NewErrorStats.MaxError, bLowersError, bKeepNewCompressionMethod ? TEXT("(**Best so far**)") : TEXT(""));
+
+		UE_LOG(LogAnimationCompression, Verbose, TEXT("    bReducesErrorBelowThreshold(%d) bHasAcceptableErrorAndSavesSpace(%d) bLowersErrorAndSavesSameOrBetter(%d)"),
+			bReducesErrorBelowThreshold, bHasAcceptableErrorAndSavesSpace, bLowersErrorAndSavesSameOrBetter);
+
+		UE_LOG(LogAnimationCompression, Verbose, TEXT("    WinningCompressorError(%f) MasterTolerance(%f) bForceBelowThreshold(%d) bErrorUnderThreshold(%d)"),
+			JobContext.WinningCompressorError, JobContext.MasterTolerance, JobContext.bForceBelowThreshold, bErrorUnderThreshold);
+
+		if (bKeepNewCompressionMethod)
+		{
+			JobContext.WinningCompressorMarginalSavings = MemorySavingsFromPrevious;
+			JobContext.WinningCompressorCounter = JobContext.WinningCompressor_Count;
+			JobContext.WinningCompressorErrorSum = JobContext.WinningCompressor_Error;
+			JobContext.WinningCompressorMarginalSavingsSum = JobContext.WinningCompressor_Margin;
+			JobContext.WinningCompressorName = JobContext.CompressionName;
+			JobContext.CurrentSize = NewSize;
+			JobContext.WinningCompressorSavings = MemorySavingsFromOriginal;
+			JobContext.WinningCompressorError = JobContext.NewErrorStats.MaxError;
+
+			/* backup key information from the sequence */
+			JobContext.SavedCompressionScheme = AnimSeq->CompressionScheme;
+			JobContext.SavedTranslationCompressionFormat = AnimSeq->TranslationCompressionFormat;
+			JobContext.SavedRotationCompressionFormat = AnimSeq->RotationCompressionFormat;
+			JobContext.SavedScaleCompressionFormat = AnimSeq->ScaleCompressionFormat;
+			JobContext.SavedKeyEncodingFormat = AnimSeq->KeyEncodingFormat;
+			JobContext.SavedCompressedTrackOffsets = AnimSeq->CompressedTrackOffsets;
+			JobContext.SavedCompressedScaleOffsets = AnimSeq->CompressedScaleOffsets;
+			JobContext.SavedCompressedSegments = AnimSeq->CompressedSegments;
+			JobContext.SavedCompressedByteStream = AnimSeq->CompressedByteStream;
+			JobContext.SavedTranslationCodec = AnimSeq->TranslationCodec;
+			JobContext.SavedRotationCodec = AnimSeq->RotationCodec;
+			JobContext.SavedScaleCodec = AnimSeq->ScaleCodec;
+			JobContext.bSavedUseRawDataOnly = false;
+		}
+		else
+		{
+			/* revert back to the old method by copying back the data we cached */
+			AnimSeq->CompressionScheme = JobContext.SavedCompressionScheme;
+			AnimSeq->TranslationCompressionFormat = JobContext.SavedTranslationCompressionFormat;
+			AnimSeq->RotationCompressionFormat = JobContext.SavedRotationCompressionFormat;
+			AnimSeq->ScaleCompressionFormat = JobContext.SavedScaleCompressionFormat;
+			AnimSeq->KeyEncodingFormat = JobContext.SavedKeyEncodingFormat;
+			AnimSeq->CompressedTrackOffsets = JobContext.SavedCompressedTrackOffsets;
+			AnimSeq->CompressedByteStream = JobContext.SavedCompressedByteStream;
+			AnimSeq->CompressedScaleOffsets = JobContext.SavedCompressedScaleOffsets;
+			AnimSeq->CompressedSegments = JobContext.SavedCompressedSegments;
+			AnimSeq->TranslationCodec = JobContext.SavedTranslationCodec;
+			AnimSeq->RotationCodec = JobContext.SavedRotationCodec;
+			AnimSeq->ScaleCodec = JobContext.SavedScaleCodec;
+			AnimSeq->SetUseRawDataOnly(JobContext.bSavedUseRawDataOnly);
+			AnimationFormat_SetInterfaceLinks(*AnimSeq);
+
+			const SIZE_T RestoredSize = AnimSeq->GetApproxCompressedSize();
+			check(RestoredSize == JobContext.CurrentSize);
+		}
+	}
+#endif	// WITH_EDITOR
 }
-#endif
 
-#define TRYCOMPRESSION(Name, CompressionAlgorithm) TRYCOMPRESSION_INNER(TEXT(#Name), Name ## CompressorWins, Name ## CompressorSumError, Name ## CompressorWinMargin, CompressionAlgorithm)
+#define POPULATE_ANIM_COMPRESSION_JOB_CONTEXT(Name_, CompressionAlgorithm_, AnimSeq_, ParentAnimSeq_, CompressContext_, JobContext_) \
+	(JobContext_).OriginalSize = OriginalSize; \
+	(JobContext_).MasterTolerance = MasterTolerance; \
+	(JobContext_).bForceBelowThreshold = bForceBelowThreshold; \
+	(JobContext_).BoneData = &BoneData; \
+	(JobContext_).CompressContext = &CompressContext_; \
+	(JobContext_).CompressionAlgorithm = DuplicateObject<UAnimCompress>((CompressionAlgorithm_), ParentAnimSeq_); \
+	(JobContext_).CompressionAlgorithm->bEnableSegmenting = bEnableSegmenting; \
+	(JobContext_).CompressionAlgorithm->IdealNumFramesPerSegment = IdealNumFramesPerSegment; \
+	(JobContext_).CompressionAlgorithm->MaxNumFramesPerSegment = MaxNumFramesPerSegment; \
+	(JobContext_).AnimSeq = AnimSeq_; \
+	(JobContext_).CompressionName = TEXT(#Name_); \
+	(JobContext_).WinningCompressor_Count = &(Name_ ## CompressorWins); \
+	(JobContext_).WinningCompressor_Error = &(Name_ ## CompressorSumError); \
+	(JobContext_).WinningCompressor_Margin = &(Name_ ## CompressorWinMargin);
+
+#define TRYCOMPRESSION(Name, CompressionAlgorithm_) \
+	do \
+	{ \
+		FAnimCompressionJobContext JobContext; \
+		POPULATE_ANIM_COMPRESSION_JOB_CONTEXT(Name, CompressionAlgorithm_, AnimSeq, AnimSeq, CompressContext, JobContext); \
+		\
+		JobContext.CurrentSize = CompressorStats.CurrentSize; \
+		JobContext.bSavedUseRawDataOnly = bSavedUseRawDataOnly; \
+		JobContext.WinningCompressorMarginalSavings = CompressorStats.WinningCompressorMarginalSavings; \
+		JobContext.WinningCompressorCounter = CompressorStats.WinningCompressorCounter; \
+		JobContext.WinningCompressorErrorSum = CompressorStats.WinningCompressorErrorSum; \
+		JobContext.WinningCompressorMarginalSavingsSum = CompressorStats.WinningCompressorMarginalSavingsSum; \
+		JobContext.WinningCompressorName = CompressorStats.WinningCompressorName; \
+		JobContext.WinningCompressorSavings = CompressorStats.WinningCompressorSavings; \
+		JobContext.WinningCompressorError = CompressorStats.WinningCompressorError; \
+		JobContext.SavedCompressionScheme = SavedCompressionScheme; \
+		JobContext.SavedTranslationCompressionFormat = SavedTranslationCompressionFormat; \
+		JobContext.SavedRotationCompressionFormat = SavedRotationCompressionFormat; \
+		JobContext.SavedScaleCompressionFormat = SavedScaleCompressionFormat; \
+		JobContext.SavedKeyEncodingFormat = SavedKeyEncodingFormat; \
+		JobContext.SavedCompressedTrackOffsets = SavedCompressedTrackOffsets; \
+		JobContext.SavedCompressedScaleOffsets = SavedCompressedScaleOffsets; \
+		JobContext.SavedCompressedByteStream = SavedCompressedByteStream; \
+		JobContext.SavedCompressedSegments = SavedCompressedSegments; \
+		JobContext.SavedTranslationCodec = SavedTranslationCodec; \
+		JobContext.SavedRotationCodec = SavedRotationCodec; \
+		JobContext.SavedScaleCodec = SavedScaleCodec; \
+		\
+		TryCompressionInner(JobContext, false); \
+		\
+		CompressorStats.CurrentSize = JobContext.CurrentSize; \
+		bSavedUseRawDataOnly = JobContext.bSavedUseRawDataOnly; \
+		CompressorStats.WinningCompressorMarginalSavings = JobContext.WinningCompressorMarginalSavings; \
+		CompressorStats.WinningCompressorCounter = JobContext.WinningCompressorCounter; \
+		CompressorStats.WinningCompressorErrorSum = JobContext.WinningCompressorErrorSum; \
+		CompressorStats.WinningCompressorMarginalSavingsSum = JobContext.WinningCompressorMarginalSavingsSum; \
+		CompressorStats.WinningCompressorName = JobContext.WinningCompressorName; \
+		CompressorStats.WinningCompressorSavings = JobContext.WinningCompressorSavings; \
+		CompressorStats.WinningCompressorError = JobContext.WinningCompressorError; \
+		SavedCompressionScheme = JobContext.SavedCompressionScheme; \
+		SavedTranslationCompressionFormat = JobContext.SavedTranslationCompressionFormat; \
+		SavedRotationCompressionFormat = JobContext.SavedRotationCompressionFormat; \
+		SavedScaleCompressionFormat = JobContext.SavedScaleCompressionFormat; \
+		SavedKeyEncodingFormat = JobContext.SavedKeyEncodingFormat; \
+		SavedCompressedTrackOffsets = JobContext.SavedCompressedTrackOffsets; \
+		SavedCompressedScaleOffsets = JobContext.SavedCompressedScaleOffsets; \
+		SavedCompressedByteStream = JobContext.SavedCompressedByteStream; \
+		SavedCompressedSegments = JobContext.SavedCompressedSegments; \
+		SavedTranslationCodec = JobContext.SavedTranslationCodec; \
+		SavedRotationCodec = JobContext.SavedRotationCodec; \
+		SavedScaleCodec = JobContext.SavedScaleCodec; \
+		\
+		NewErrorStats = JobContext.NewErrorStats; \
+		CompressorStats.PctSaving = JobContext.PctSaving; \
+	} \
+	while (false)
+
+#define TRYCOMPRESSION_SYNC(Name, CompressionAlgorithm_) TRYCOMPRESSION(Name, CompressionAlgorithm_)
+
+// TODO: Async compression is DISABLED for additive sequences because UAnimCompress_RemoveLinearKeys::ConvertFromRelativeSpace()
+// modifies the RAW data! This is bad... Even though we duplicate the anim sequence, some additive information isn't
+// copied over and it doesn't seem safe to generate it by calling UAnimSequence::BakeOutAdditiveIntoRawData()
+#define TRYCOMPRESSION_ASYNC(Name, CompressionAlgorithm_) \
+	do \
+	{ \
+		if (!AnimSeq->IsValidAdditive()) \
+		{ \
+			FAnimCompressionJobContext* JobContext = new FAnimCompressionJobContext(); \
+			FAnimCompressContext CompressContextCopy(CompressContext); \
+			CompressContextCopy.CompressionSummary = FCompressionMemorySummary(false); \
+			UAnimSequence* AnimSeqCopy = DuplicateObject<UAnimSequence>(AnimSeq, GetTransientPackage()); \
+			POPULATE_ANIM_COMPRESSION_JOB_CONTEXT(Name, CompressionAlgorithm_, AnimSeqCopy, AnimSeq, CompressContextCopy, *JobContext); \
+			AnimCompressionJobContexes.Add(JobContext); \
+			\
+			AnimCompressionTask_CompletionEvents.Add(TGraphTask<FAsyncAnimCompressionTask>::CreateTask(NULL, ENamedThreads::GameThread).ConstructAndDispatchWhenReady(JobContext)); \
+		} \
+		else \
+		{ \
+			TRYCOMPRESSION_SYNC(Name, CompressionAlgorithm_); \
+		} \
+	} \
+	while (false)
 
 #define LOG_COMPRESSION_STATUS(Name) \
 	UE_LOG(LogAnimationCompression, Log, TEXT("\t\tWins for '%32s': %4i\t\t%f\t%i bytes"), TEXT(#Name), Name ## CompressorWins, (Name ## CompressorWins > 0) ? Name ## CompressorSumError / Name ## CompressorWins : 0.0f, Name ## CompressorWinMargin)
@@ -574,6 +759,161 @@ bool FAnimationUtils::GetForcedRecompressionSetting()
 	static float Algorithm ## CompressorSumError = 0.0f; \
 	static int64 Algorithm ## CompressorWinMargin = 0
 
+#if WITH_EDITORONLY_DATA
+struct WinningCompressorStats
+{
+	SIZE_T CurrentSize;
+
+	int32* WinningCompressorCounter;
+	float* WinningCompressorErrorSum;
+	int64* WinningCompressorMarginalSavingsSum;
+	int64 WinningCompressorMarginalSavings;
+	FString WinningCompressorName;
+	int32 WinningCompressorSavings;
+	float PctSaving;
+	float WinningCompressorError;
+
+	WinningCompressorStats()
+	{
+		memset(this, 0, sizeof(WinningCompressorStats));
+	}
+};
+
+static void WaitForAnimCompressionJobs(const FGraphEventArray& AnimCompressionTask_CompletionEvents)
+{
+	FTaskGraphInterface::Get().WaitUntilTasksComplete(AnimCompressionTask_CompletionEvents, ENamedThreads::GameThread);
+}
+
+static void ClearAnimCompressionJobs(FGraphEventArray& AnimCompressionTask_CompletionEvents, TArray<FAnimCompressionJobContext*>& AnimCompressionJobContexes)
+{
+	for (FAnimCompressionJobContext* Context : AnimCompressionJobContexes)
+	{
+		Context->AnimSeq->RecycleAnimSequence();
+		delete Context;
+	}
+	AnimCompressionJobContexes.Reset();
+	AnimCompressionTask_CompletionEvents.Reset();
+}
+
+static FAnimCompressionJobContext* FindBestAnimCompression(TArray<FAnimCompressionJobContext*>& AnimCompressionJobContexes, SIZE_T OriginalSize, SIZE_T CurrentSize, float WinningCompressorError, float MasterTolerance)
+{
+	FAnimCompressionJobContext* BestJobContext = nullptr;
+
+	for (FAnimCompressionJobContext* context : AnimCompressionJobContexes)
+	{
+		FAnimCompressionJobContext& JobContext = *context;
+
+		const SIZE_T NewSize = JobContext.AnimSeq->GetApproxCompressedSize();
+
+		/* compute the savings and compression error*/
+		const int64 MemorySavingsFromOriginal = OriginalSize - NewSize;
+		const int64 MemorySavingsFromPrevious = CurrentSize - NewSize;
+
+		/* figure out our new compression error*/
+		FAnimationUtils::ComputeCompressionError(JobContext.AnimSeq, *JobContext.BoneData, JobContext.NewErrorStats);
+
+		const bool bLowersError = JobContext.NewErrorStats.MaxError < WinningCompressorError;
+		const bool bErrorUnderThreshold = JobContext.NewErrorStats.MaxError <= MasterTolerance;
+
+		/* keep it if it we want to force the error below the threshold and it reduces error */
+		bool bKeepNewCompressionMethod = false;
+		const bool bReducesErrorBelowThreshold = (bLowersError && (WinningCompressorError > MasterTolerance) && JobContext.bForceBelowThreshold);
+		bKeepNewCompressionMethod |= bReducesErrorBelowThreshold;
+		/* or if has an acceptable error and saves space  */
+		const bool bHasAcceptableErrorAndSavesSpace = bErrorUnderThreshold && (MemorySavingsFromPrevious > 0);
+		bKeepNewCompressionMethod |= bHasAcceptableErrorAndSavesSpace;
+		/* or if saves the same amount and an acceptable error that is lower than the previous best */
+		const bool bLowersErrorAndSavesSameOrBetter = bErrorUnderThreshold && bLowersError && (MemorySavingsFromPrevious >= 0);
+		bKeepNewCompressionMethod |= bLowersErrorAndSavesSameOrBetter;
+
+		JobContext.PctSaving = (OriginalSize > 0) ? (100.f - (100.f * float(NewSize) / float(OriginalSize))) : 0.f;
+		UE_LOG(LogAnimationCompression, Verbose, TEXT("- %s - bytes saved(%i) (%.1f%%) from previous(%i) MaxError(%.2f) bLowersError(%d) %s"),
+			JobContext.CompressionName, MemorySavingsFromOriginal, JobContext.PctSaving, MemorySavingsFromPrevious, JobContext.NewErrorStats.MaxError, bLowersError, bKeepNewCompressionMethod ? TEXT("(**Best so far**)") : TEXT(""));
+
+		UE_LOG(LogAnimationCompression, Verbose, TEXT("    bReducesErrorBelowThreshold(%d) bHasAcceptableErrorAndSavesSpace(%d) bLowersErrorAndSavesSameOrBetter(%d)"),
+			bReducesErrorBelowThreshold, bHasAcceptableErrorAndSavesSpace, bLowersErrorAndSavesSameOrBetter);
+
+		UE_LOG(LogAnimationCompression, Verbose, TEXT("    WinningCompressorError(%f) MasterTolerance(%f) bForceBelowThreshold(%d) bErrorUnderThreshold(%d)"),
+			WinningCompressorError, JobContext.MasterTolerance, JobContext.bForceBelowThreshold, bErrorUnderThreshold);
+
+		if (bKeepNewCompressionMethod)
+		{
+			BestJobContext = context;
+			WinningCompressorError = JobContext.NewErrorStats.MaxError;
+			CurrentSize = NewSize;
+		}
+	}
+
+	return BestJobContext;
+}
+
+static void UpdateAnimCompressionFromAsyncJobs(UAnimSequence* AnimSeq, FGraphEventArray& AnimCompressionTask_CompletionEvents, TArray<FAnimCompressionJobContext*>& AnimCompressionJobContexes, SIZE_T OriginalSize, WinningCompressorStats& CompressorStats, float MasterTolerance)
+{
+	// Pick the best
+	const FAnimCompressionJobContext* BestJobContext = FindBestAnimCompression(AnimCompressionJobContexes, OriginalSize, CompressorStats.CurrentSize, CompressorStats.WinningCompressorError, MasterTolerance);
+
+	if (BestJobContext != nullptr)
+	{
+		const FAnimCompressionJobContext& JobContext = *BestJobContext;
+
+		/* Copy our data */
+		AnimSeq->CompressionScheme = JobContext.AnimSeq->CompressionScheme;
+		AnimSeq->TranslationCompressionFormat = JobContext.AnimSeq->TranslationCompressionFormat;
+		AnimSeq->RotationCompressionFormat = JobContext.AnimSeq->RotationCompressionFormat;
+		AnimSeq->KeyEncodingFormat = JobContext.AnimSeq->KeyEncodingFormat;
+		AnimSeq->CompressedTrackOffsets = JobContext.AnimSeq->CompressedTrackOffsets;
+		AnimSeq->CompressedByteStream = JobContext.AnimSeq->CompressedByteStream;
+		AnimSeq->CompressedScaleOffsets = JobContext.AnimSeq->CompressedScaleOffsets;
+		AnimSeq->CompressedSegments = JobContext.AnimSeq->CompressedSegments;
+		AnimSeq->TranslationCodec = JobContext.AnimSeq->TranslationCodec;
+		AnimSeq->RotationCodec = JobContext.AnimSeq->RotationCodec;
+		AnimSeq->ScaleCodec = JobContext.AnimSeq->ScaleCodec;
+		AnimSeq->SetUseRawDataOnly(false);
+		AnimationFormat_SetInterfaceLinks(*AnimSeq);
+
+		const SIZE_T RestoredSize = AnimSeq->GetApproxCompressedSize();
+		check(RestoredSize == JobContext.CurrentSize);
+
+		const int64 MemorySavingsFromOriginal = OriginalSize - RestoredSize;
+		const SIZE_T MemorySavingsFromPrevious = CompressorStats.CurrentSize - RestoredSize;
+
+		CompressorStats.WinningCompressorMarginalSavings = MemorySavingsFromPrevious;
+		CompressorStats.WinningCompressorCounter = JobContext.WinningCompressor_Count;
+		CompressorStats.WinningCompressorErrorSum = JobContext.WinningCompressor_Error;
+		CompressorStats.WinningCompressorMarginalSavingsSum = JobContext.WinningCompressor_Margin;
+		CompressorStats.WinningCompressorName = JobContext.CompressionName;
+		CompressorStats.CurrentSize = RestoredSize;
+		CompressorStats.WinningCompressorSavings = MemorySavingsFromOriginal;
+		CompressorStats.WinningCompressorError = JobContext.NewErrorStats.MaxError;
+	}
+
+	ClearAnimCompressionJobs(AnimCompressionTask_CompletionEvents, AnimCompressionJobContexes);
+}
+#endif	// WITH_EDITORONLY_DATA
+
+// Simple struct to calc duration of of scope
+struct FCompressionTimeElapsed
+{
+private:
+
+	// Time we started tracking
+	double StartTime;
+
+	// Where to store the elapsed time
+	double& Result;
+
+public:
+	FCompressionTimeElapsed(double& InResult)
+		: StartTime(FPlatformTime::Seconds())
+		, Result(InResult)
+	{}
+
+	~FCompressionTimeElapsed()
+	{
+		Result = FPlatformTime::Seconds() - StartTime;
+	}
+};
+
 /**
  * Utility function to compress an animation. If the animation is currently associated with a codec, it will be used to 
  * compress the animation. Otherwise, the default codec will be used. If AllowAlternateCompressor is true, an
@@ -616,22 +956,57 @@ void FAnimationUtils::CompressAnimSequence(UAnimSequence* AnimSeq, FAnimCompress
 				bFirstRecompressUsingCurrentOrDefault = true;
 			}
 
-			bool bTryFixedBitwiseCompression = AnimSetting->bTryFixedBitwiseCompression;
-			bool bTryPerTrackBitwiseCompression = AnimSetting->bTryPerTrackBitwiseCompression;
-			bool bTryLinearKeyRemovalCompression = AnimSetting->bTryLinearKeyRemovalCompression;
-			bool bTryIntervalKeyRemoval = AnimSetting->bTryIntervalKeyRemoval;
+			bool bTryExhaustiveSearch = AnimSetting->bTryExhaustiveSearch;
+			bool bEnableSegmenting = AnimSetting->bEnableSegmenting;
+			int32 IdealNumFramesPerSegment = 64;
+			int32 MaxNumFramesPerSegment = (IdealNumFramesPerSegment * 2) - 1;
 
-			CompressAnimSequenceExplicit(
-				AnimSeq,
-				CompressContext,
-				CompressContext.bAllowAlternateCompressor ? MasterTolerance : 0.0f,
-				bFirstRecompressUsingCurrentOrDefault,
-				bForceBelowThreshold,
-				bRaiseMaxErrorToExisting,
-				bTryFixedBitwiseCompression,
-				bTryPerTrackBitwiseCompression,
-				bTryLinearKeyRemovalCompression,
-				bTryIntervalKeyRemoval);
+			// Build skeleton metadata to use during the key reduction.
+			TArray<FBoneData> BoneData;
+			FAnimationUtils::BuildSkeletonMetaData(AnimSeq->GetSkeleton(), BoneData);
+
+#if WITH_EDITORONLY_DATA
+			UAnimCompress_Automatic* AutoCompressionScheme = Cast<UAnimCompress_Automatic>(AnimSeq->CompressionScheme);
+
+			CompressContext.GatherPreCompressionStats(AnimSeq);
+
+			AnimSeq->CompressedByteStream.Reset();
+			AnimSeq->CompressedTrackOffsets.Reset();
+			AnimSeq->CompressedScaleOffsets.OffsetData.Reset();
+			AnimSeq->CompressedSegments.Reset();
+
+			if (AutoCompressionScheme != nullptr)
+			{
+				bTryExhaustiveSearch = AutoCompressionScheme->bTryExhaustiveSearch;
+				bEnableSegmenting = AutoCompressionScheme->bEnableSegmenting;
+				IdealNumFramesPerSegment = AutoCompressionScheme->IdealNumFramesPerSegment;
+				MaxNumFramesPerSegment = AutoCompressionScheme->MaxNumFramesPerSegment;
+			}
+			else if (AnimSeq->CompressionScheme != nullptr)
+			{
+				bTryExhaustiveSearch = AnimSetting->bTryExhaustiveSearch;
+			}
+#endif
+
+			double CompressionTime = 0.0;
+			{
+				//Scoped timing of compression, make sure nothing else is added to this scope
+				FCompressionTimeElapsed TimeTracker(CompressionTime);
+				CompressAnimSequenceExplicit(
+					AnimSeq,
+					CompressContext,
+					CompressContext.bAllowAlternateCompressor ? MasterTolerance : 0.0f,
+					bFirstRecompressUsingCurrentOrDefault,
+					bForceBelowThreshold,
+					bRaiseMaxErrorToExisting,
+					bTryExhaustiveSearch,
+					bEnableSegmenting,
+					IdealNumFramesPerSegment,
+					MaxNumFramesPerSegment,
+					BoneData);
+			}
+
+			CompressContext.GatherPostCompressionStats(AnimSeq, BoneData, CompressionTime);
 		}
 	}
 }
@@ -652,13 +1027,14 @@ void FAnimationUtils::CompressAnimSequenceExplicit(
 	UAnimSequence* AnimSeq,
 	FAnimCompressContext& CompressContext,
 	float MasterTolerance,
-	const bool bFirstRecompressUsingCurrentOrDefault,
-	const bool bForceBelowThreshold,
-	const bool bRaiseMaxErrorToExisting,
-	const bool bTryFixedBitwiseCompression,
-	const bool bTryPerTrackBitwiseCompression,
-	const bool bTryLinearKeyRemovalCompression,
-	const bool bTryIntervalKeyRemoval)
+	bool bFirstRecompressUsingCurrentOrDefault,
+	bool bForceBelowThreshold,
+	bool bRaiseMaxErrorToExisting,
+	bool bTryExhaustiveSearch,
+	bool bEnableSegmenting,
+	int32 IdealNumFramesPerSegment,
+	int32 MaxNumFramesPerSegment,
+	const TArray<FBoneData>& BoneData)
 {
 #if WITH_EDITORONLY_DATA
 	DECLARE_ANIM_COMP_ALGORITHM(BitwiseACF_Float96);
@@ -751,10 +1127,7 @@ void FAnimationUtils::CompressAnimSequenceExplicit(
 
 		AnimationErrorStats OriginalErrorStats;
 		AnimationErrorStats TrueOriginalErrorStats;
-		TArray<FBoneData> BoneData;
 
-		// Build skeleton metadata to use during the key reduction.
-		FAnimationUtils::BuildSkeletonMetaData( Skeleton, BoneData );
 		FAnimationUtils::ComputeCompressionError(AnimSeq, BoneData, TrueOriginalErrorStats);
 
 		int32 AfterOriginalRecompression = 0;
@@ -765,8 +1138,16 @@ void FAnimationUtils::CompressAnimSequenceExplicit(
 			// Automatic compression brings us back here, so don't create an infinite loop and pick bitwise compress instead.
 			if (!OriginalCompressionAlgorithm || OriginalCompressionAlgorithm->IsA(UAnimCompress_Automatic::StaticClass()))
 			{
-				UAnimCompress* CompressionAlgorithm = NewObject<UAnimCompress_BitwiseCompressOnly>();
-				OriginalCompressionAlgorithm = static_cast<UAnimCompress*>(StaticDuplicateObject(CompressionAlgorithm, AnimSeq));
+				UAnimCompress* CompressionAlgorithm = NewObject<UAnimCompress_BitwiseCompressOnly>(AnimSeq);
+				if (OriginalCompressionAlgorithm != nullptr)
+				{
+					// Keep the same segmenting settings
+					CompressionAlgorithm->bEnableSegmenting = OriginalCompressionAlgorithm->bEnableSegmenting;
+					CompressionAlgorithm->IdealNumFramesPerSegment = OriginalCompressionAlgorithm->IdealNumFramesPerSegment;
+					CompressionAlgorithm->MaxNumFramesPerSegment = OriginalCompressionAlgorithm->MaxNumFramesPerSegment;
+				}
+
+				OriginalCompressionAlgorithm = CompressionAlgorithm;
 			}
 
 			UE_LOG(LogAnimationCompression, Log, TEXT("Recompressing (%s) using current/default (%s) bFirstRecompressUsingCurrentOrDefault(%d) bTryAlternateCompressor(%d) IsCompressedDataValid(%d)"),
@@ -776,8 +1157,8 @@ void FAnimationUtils::CompressAnimSequenceExplicit(
 				bTryAlternateCompressor,
 				AnimSeq->IsCompressedDataValid());
 
-			AnimSeq->CompressionScheme = OriginalCompressionAlgorithm;
-			OriginalCompressionAlgorithm->Reduce(AnimSeq, CompressContext);
+			AnimSeq->CompressionScheme = DuplicateObject<UAnimCompress>(OriginalCompressionAlgorithm, AnimSeq);
+			OriginalCompressionAlgorithm->Reduce(AnimSeq, CompressContext, BoneData);
 			AnimSeq->SetUseRawDataOnly(false);
 			AfterOriginalRecompression = AnimSeq->GetApproxCompressedSize();
 
@@ -795,7 +1176,7 @@ void FAnimationUtils::CompressAnimSequenceExplicit(
 		TotalSizeBefore += OriginalSize;
 
 		// Estimate total uncompressed
-		TotalUncompressed += AnimSeq->GetUncompressedRawSize();
+		TotalUncompressed += ((sizeof(FVector) + sizeof(FQuat) + sizeof(FVector)) * NumRawDataTracks * AnimSeq->NumFrames);
 
 		// start with the current technique, or the default if none exists.
 		// this will serve as our fallback if no better technique can be found
@@ -820,35 +1201,142 @@ void FAnimationUtils::CompressAnimSequenceExplicit(
 				}
 			}
 
+			// count all attempts for debugging
+			++TotalRecompressions;
+
+			WinningCompressorStats CompressorStats;
+			CompressorStats.CurrentSize = OriginalSize;
+			CompressorStats.WinningCompressorError = OriginalErrorStats.MaxError;
+
+			FGraphEventArray AnimCompressionTask_CompletionEvents;
+			TArray<FAnimCompressionJobContext*> AnimCompressionJobContexes;
+
+			// backup key information from the sequence
+			class UAnimCompress* SavedCompressionScheme = AnimSeq->CompressionScheme;
+			AnimationCompressionFormat SavedTranslationCompressionFormat = AnimSeq->TranslationCompressionFormat;
+			AnimationCompressionFormat SavedRotationCompressionFormat = AnimSeq->RotationCompressionFormat;
+			AnimationCompressionFormat SavedScaleCompressionFormat = AnimSeq->ScaleCompressionFormat;
+			AnimationKeyFormat SavedKeyEncodingFormat = AnimSeq->KeyEncodingFormat;
+			TArray<int32> SavedCompressedTrackOffsets = AnimSeq->CompressedTrackOffsets;
+			TArray<uint8> SavedCompressedByteStream = AnimSeq->CompressedByteStream;
+			FCompressedOffsetData SavedCompressedScaleOffsets = AnimSeq->CompressedScaleOffsets;
+			TArray<FCompressedSegment> SavedCompressedSegments = AnimSeq->CompressedSegments;
+			AnimEncoding* SavedTranslationCodec = AnimSeq->TranslationCodec;
+			AnimEncoding* SavedRotationCodec = AnimSeq->RotationCodec;
+			AnimEncoding* SavedScaleCodec = AnimSeq->ScaleCodec;
+			bool bSavedUseRawDataOnly = AnimSeq->OnlyUseRawData();
+
+			if (!bTryExhaustiveSearch)
 			{
-				// backup key information from the sequence
-				class UAnimCompress* SavedCompressionScheme = AnimSeq->CompressionScheme;
-				AnimationCompressionFormat SavedTranslationCompressionFormat = AnimSeq->TranslationCompressionFormat;
-				AnimationCompressionFormat SavedRotationCompressionFormat = AnimSeq->RotationCompressionFormat;
-				AnimationKeyFormat SavedKeyEncodingFormat = AnimSeq->KeyEncodingFormat;
-				TArray<int32> SavedCompressedTrackOffsets = AnimSeq->CompressedTrackOffsets;
-				TArray<uint8> SavedCompressedByteStream = AnimSeq->CompressedByteStream;
-				FCompressedOffsetData SavedCompressedScaleOffsets = AnimSeq->CompressedScaleOffsets;
-				AnimEncoding* SavedTranslationCodec = AnimSeq->TranslationCodec;
-				AnimEncoding* SavedRotationCodec = AnimSeq->RotationCodec;
-				AnimEncoding* SavedScaleCodec = AnimSeq->ScaleCodec;
-				bool bSavedUseRawDataOnly = AnimSeq->OnlyUseRawData();
-
-				// count all attempts for debugging
-				++TotalRecompressions;
+				// Dispatch our async compression
+				{
+					{
+						// Adaptive error through probing the effect of perturbations at each track
+						UAnimCompress_PerTrackCompression* NewPerTrackCompressor = NewObject<UAnimCompress_PerTrackCompression>();
+						NewPerTrackCompressor->bUseAdaptiveError2 = true;
+						NewPerTrackCompressor->MaxPosDiffBitwise = 0.05f;
+						NewPerTrackCompressor->MaxAngleDiffBitwise = 0.02f;
+						NewPerTrackCompressor->MaxScaleDiffBitwise = 0.00005f;
 
-				// Prepare to compress
-				SIZE_T CurrentSize = AnimSeq->GetApproxCompressedSize();
-				int32* WinningCompressorCounter = NULL;
-				float* WinningCompressorErrorSum = NULL;
-				int64* WinningCompressorMarginalSavingsSum = NULL;
-				int64 WinningCompressorMarginalSavings = 0;
-				FString WinningCompressorName;
-				int32 WinningCompressorSavings = 0;
-				float PctSaving = 0.f;
-				float WinningCompressorError = OriginalErrorStats.MaxError;
-				bool bKeepNewCompressionMethod = false;
+						TRYCOMPRESSION_ASYNC(Adaptive2_PerTrack, NewPerTrackCompressor);
+
+						NewPerTrackCompressor->bActuallyFilterLinearKeys = true;
+						NewPerTrackCompressor->bRetarget = true;
+						TRYCOMPRESSION_ASYNC(Adaptive2_LinPerTrack, NewPerTrackCompressor);
+
+						NewPerTrackCompressor->bActuallyFilterLinearKeys = true;
+						NewPerTrackCompressor->bRetarget = false;
+						TRYCOMPRESSION_ASYNC(Adaptive2_LinPerTrackNoRT, NewPerTrackCompressor);
+					}
+
+					{
+						UAnimCompress_PerTrackCompression* PerTrackCompressor = NewObject<UAnimCompress_PerTrackCompression>();
+						PerTrackCompressor->bUseAdaptiveError = true;
+
+						if (AnimSeq->NumFrames > 1)
+						{
+							PerTrackCompressor->bActuallyFilterLinearKeys = true;
+							PerTrackCompressor->bRetarget = true;
+
+							PerTrackCompressor->MaxPosDiff = 0.1f;
+							// 						PerTrackCompressor->MaxAngleDiff = 0.1;
+							PerTrackCompressor->MaxScaleDiff = 0.00001f;
+							PerTrackCompressor->ParentingDivisor = 2.0f;
+							PerTrackCompressor->ParentingDivisorExponent = 1.0f;
+							TRYCOMPRESSION_ASYNC(Linear_PerTrackExp1, PerTrackCompressor);
+						}
+					}
+
+					{
+						UAnimCompress_PerTrackCompression* PerTrackCompressor = NewObject<UAnimCompress_PerTrackCompression>();
+
+						// Straight PerTrackCompression, no key decimation and no linear key removal
+						TRYCOMPRESSION_ASYNC(Bitwise_PerTrack, PerTrackCompressor);
+						PerTrackCompressor->bUseAdaptiveError = true;
+
+						// Full blown linear
+						PerTrackCompressor->bActuallyFilterLinearKeys = true;
+						PerTrackCompressor->bRetarget = true;
+						TRYCOMPRESSION_ASYNC(Linear_PerTrack, PerTrackCompressor);
+
+						// Adaptive retargetting based on height within the skeleton
+						PerTrackCompressor->bActuallyFilterLinearKeys = true;
+						PerTrackCompressor->bRetarget = false;
+						PerTrackCompressor->ParentingDivisor = 2.0f;
+						PerTrackCompressor->ParentingDivisorExponent = 1.6f;
+						TRYCOMPRESSION_ASYNC(Adaptive1_LinPerTrackNoRT, PerTrackCompressor);
+					}
+
+					{
+						UAnimCompress_PerTrackCompression* PerTrackCompressor = NewObject<UAnimCompress_PerTrackCompression>();
+						PerTrackCompressor->bUseAdaptiveError = true;
+
+						// Try the decimation algorithms
+						if (AnimSeq->NumFrames >= PerTrackCompressor->MinKeysForResampling)
+						{
+							PerTrackCompressor->bActuallyFilterLinearKeys = false;
+							PerTrackCompressor->bRetarget = false;
+							PerTrackCompressor->bUseAdaptiveError = false;
+							PerTrackCompressor->bResampleAnimation = true;
+
+							// Try PerTrackCompression, downsample to 5 Hz
+							PerTrackCompressor->ResampledFramerate = 5.0f;
+							TRYCOMPRESSION_ASYNC(Downsample5Hz_PerTrack, PerTrackCompressor);
+						}
+					}
+
+					if (AnimSeq->NumFrames > 1)
+					{
+						UAnimCompress_RemoveLinearKeys* LinearKeyRemover = NewObject<UAnimCompress_RemoveLinearKeys>();
+						{
+							// Try ACF_Float96NoW
+							LinearKeyRemover->RotationCompressionFormat = ACF_Float96NoW;
+							LinearKeyRemover->TranslationCompressionFormat = ACF_None;
+							TRYCOMPRESSION_ASYNC(LinearACF_Float96, LinearKeyRemover);
+						}
+					}
+
+					{
+						UAnimCompress_BitwiseCompressOnly* BitwiseCompressor = NewObject<UAnimCompress_BitwiseCompressOnly>();
+
+						// Try ACF_Float96NoW
+						BitwiseCompressor->RotationCompressionFormat = ACF_Float96NoW;
+						BitwiseCompressor->TranslationCompressionFormat = ACF_None;
+						TRYCOMPRESSION_ASYNC(BitwiseACF_Float96, BitwiseCompressor);
+
+						// Try ACF_Fixed48NoW
+						BitwiseCompressor->RotationCompressionFormat = ACF_Fixed48NoW;
+						BitwiseCompressor->TranslationCompressionFormat = ACF_None;
+						TRYCOMPRESSION_ASYNC(BitwiseACF_Fixed48, BitwiseCompressor);
+					}
+				}
 
+				WaitForAnimCompressionJobs(AnimCompressionTask_CompletionEvents);
+				UpdateAnimCompressionFromAsyncJobs(AnimSeq, AnimCompressionTask_CompletionEvents, AnimCompressionJobContexes, OriginalSize, CompressorStats, MasterTolerance);
+			}
+			else
+			{
+				// Prepare to compress
 				UE_LOG(LogAnimationCompression, Log, TEXT("Compressing %s (%s)\n\tSkeleton: %s\n\tOriginal Size: %i   MaxDiff: %f"),
 					*AnimSeq->GetName(),
 					*AnimSeq->GetFullName(),
@@ -862,7 +1350,7 @@ void FAnimationUtils::CompressAnimSequenceExplicit(
 					*FAnimationUtils::GetAnimationCompressionFormatString(static_cast<AnimationCompressionFormat>(OriginalTranslationFormat)),
 					AnimSeq->NumFrames,
 					AnimSeq->SequenceLength,
-					(AnimSeq->NumFrames > 1) ? AnimSeq->NumFrames / AnimSeq->SequenceLength : DEFAULT_SAMPLERATE);
+					(AnimSeq->NumFrames > 1) ? (AnimSeq->NumFrames - 1) / AnimSeq->SequenceLength : DEFAULT_SAMPLERATE);
 
 				if (bFirstRecompressUsingCurrentOrDefault)
 				{
@@ -873,7 +1361,6 @@ void FAnimationUtils::CompressAnimSequenceExplicit(
 				}
 
 				// Progressive Algorithm
-				if( bTryPerTrackBitwiseCompression )
 				{
 					UAnimCompress_PerTrackCompression* PerTrackCompressor = NewObject<UAnimCompress_PerTrackCompression>();
 
@@ -896,7 +1383,7 @@ void FAnimationUtils::CompressAnimSequenceExplicit(
 					else
 					{
 						// First, start by finding most downsampling factor.
-						if( bTryIntervalKeyRemoval && (AnimSeq->NumFrames >= PerTrackCompressor->MinKeysForResampling) )
+						if( AnimSeq->NumFrames >= PerTrackCompressor->MinKeysForResampling )
 						{
 							PerTrackCompressor->bResampleAnimation = true;
 				
@@ -1015,19 +1502,18 @@ void FAnimationUtils::CompressAnimSequenceExplicit(
 				}
 
 				// Start with Bitwise Compress only
-				if( bTryFixedBitwiseCompression )
 				{
 					UAnimCompress_BitwiseCompressOnly* BitwiseCompressor = NewObject<UAnimCompress_BitwiseCompressOnly>();
 
 					// Try ACF_Float96NoW
 					BitwiseCompressor->RotationCompressionFormat = ACF_Float96NoW;
 					BitwiseCompressor->TranslationCompressionFormat = ACF_None;
-					TRYCOMPRESSION(BitwiseACF_Float96, BitwiseCompressor);
+					TRYCOMPRESSION_ASYNC(BitwiseACF_Float96, BitwiseCompressor);
 
 					// Try ACF_Fixed48NoW
 					BitwiseCompressor->RotationCompressionFormat = ACF_Fixed48NoW;
 					BitwiseCompressor->TranslationCompressionFormat = ACF_None;
-					TRYCOMPRESSION(BitwiseACF_Fixed48,BitwiseCompressor);
+					TRYCOMPRESSION_ASYNC(BitwiseACF_Fixed48,BitwiseCompressor);
 
 // 32bits currently unusable due to creating too much error
 // 					// Try ACF_IntervalFixed32NoW
@@ -1043,7 +1529,6 @@ void FAnimationUtils::CompressAnimSequenceExplicit(
 
 				// Start with Bitwise Compress only
 				// this compressor has a minimum number of frames requirement. So no need to go there if we don't meet that...
-				if( bTryFixedBitwiseCompression && bTryIntervalKeyRemoval )
 				{
 					UAnimCompress_RemoveEverySecondKey* RemoveEveryOtherKeyCompressor = NewObject<UAnimCompress_RemoveEverySecondKey>();
 					if( AnimSeq->NumFrames > RemoveEveryOtherKeyCompressor->MinKeys )
@@ -1053,12 +1538,12 @@ void FAnimationUtils::CompressAnimSequenceExplicit(
 							// Try ACF_Float96NoW
 							RemoveEveryOtherKeyCompressor->RotationCompressionFormat = ACF_Float96NoW;	
 							RemoveEveryOtherKeyCompressor->TranslationCompressionFormat = ACF_None;		
-							TRYCOMPRESSION(HalfOddACF_Float96, RemoveEveryOtherKeyCompressor);
+							TRYCOMPRESSION_ASYNC(HalfOddACF_Float96, RemoveEveryOtherKeyCompressor);
 
 							// Try ACF_Fixed48NoW
 							RemoveEveryOtherKeyCompressor->RotationCompressionFormat = ACF_Fixed48NoW;	
 							RemoveEveryOtherKeyCompressor->TranslationCompressionFormat = ACF_None;		
-							TRYCOMPRESSION(HalfOddACF_Fixed48, RemoveEveryOtherKeyCompressor);
+							TRYCOMPRESSION_ASYNC(HalfOddACF_Fixed48, RemoveEveryOtherKeyCompressor);
 
 // 32bits currently unusable due to creating too much error
 // 							// Try ACF_IntervalFixed32NoW
@@ -1076,12 +1561,12 @@ void FAnimationUtils::CompressAnimSequenceExplicit(
 							// Try ACF_Float96NoW
 							RemoveEveryOtherKeyCompressor->RotationCompressionFormat = ACF_Float96NoW;	
 							RemoveEveryOtherKeyCompressor->TranslationCompressionFormat = ACF_None;		
-							TRYCOMPRESSION(HalfEvenACF_Float96,RemoveEveryOtherKeyCompressor);
+							TRYCOMPRESSION_ASYNC(HalfEvenACF_Float96,RemoveEveryOtherKeyCompressor);
 
 							// Try ACF_Fixed48NoW
 							RemoveEveryOtherKeyCompressor->RotationCompressionFormat = ACF_Fixed48NoW;	
 							RemoveEveryOtherKeyCompressor->TranslationCompressionFormat = ACF_None;		
-							TRYCOMPRESSION(HalfEvenACF_Fixed48,RemoveEveryOtherKeyCompressor);
+							TRYCOMPRESSION_ASYNC(HalfEvenACF_Fixed48,RemoveEveryOtherKeyCompressor);
 
 // 32bits currently unusable due to creating too much error
 // 							// Try ACF_IntervalFixed32NoW
@@ -1097,20 +1582,20 @@ void FAnimationUtils::CompressAnimSequenceExplicit(
 					}
 				}
 
-				// construct the proposed compressor		
-				if( bTryLinearKeyRemovalCompression && AnimSeq->NumFrames > 1 )
+				// construct the proposed compressor
+				if( AnimSeq->NumFrames > 1 )
 				{
 					UAnimCompress_RemoveLinearKeys* LinearKeyRemover = NewObject<UAnimCompress_RemoveLinearKeys>();
 					{
 						// Try ACF_Float96NoW
 						LinearKeyRemover->RotationCompressionFormat = ACF_Float96NoW;
 						LinearKeyRemover->TranslationCompressionFormat = ACF_None;	
-						TRYCOMPRESSION(LinearACF_Float96,LinearKeyRemover);
+						TRYCOMPRESSION_ASYNC(LinearACF_Float96,LinearKeyRemover);
 
 						// Try ACF_Fixed48NoW
 						LinearKeyRemover->RotationCompressionFormat = ACF_Fixed48NoW;
 						LinearKeyRemover->TranslationCompressionFormat = ACF_None;	
-						TRYCOMPRESSION(LinearACF_Fixed48,LinearKeyRemover);
+						TRYCOMPRESSION_ASYNC(LinearACF_Fixed48,LinearKeyRemover);
 
 // Error is too bad w/ 32bits
 // 						// Try ACF_IntervalFixed32NoW
@@ -1125,25 +1610,24 @@ void FAnimationUtils::CompressAnimSequenceExplicit(
 					}
 				}
 
-				if( bTryPerTrackBitwiseCompression )
 				{
 					UAnimCompress_PerTrackCompression* PerTrackCompressor = NewObject<UAnimCompress_PerTrackCompression>();
 
 					// Straight PerTrackCompression, no key decimation and no linear key removal
-					TRYCOMPRESSION(Bitwise_PerTrack, PerTrackCompressor);
+					TRYCOMPRESSION_ASYNC(Bitwise_PerTrack, PerTrackCompressor);
 					PerTrackCompressor->bUseAdaptiveError = true;
 
 					// Full blown linear
 					PerTrackCompressor->bActuallyFilterLinearKeys = true;
 					PerTrackCompressor->bRetarget = true;
-					TRYCOMPRESSION(Linear_PerTrack, PerTrackCompressor);
+					TRYCOMPRESSION_ASYNC(Linear_PerTrack, PerTrackCompressor);
 
 					// Adaptive retargetting based on height within the skeleton
 					PerTrackCompressor->bActuallyFilterLinearKeys = true;
 					PerTrackCompressor->bRetarget = false;
 					PerTrackCompressor->ParentingDivisor = 2.0f;
 					PerTrackCompressor->ParentingDivisorExponent = 1.6f;
-					TRYCOMPRESSION(Adaptive1_LinPerTrackNoRT, PerTrackCompressor);
+					TRYCOMPRESSION_ASYNC(Adaptive1_LinPerTrackNoRT, PerTrackCompressor);
 					PerTrackCompressor->ParentingDivisor = 1.0f;
 					PerTrackCompressor->ParentingDivisorExponent = 1.0f;
 
@@ -1151,13 +1635,12 @@ void FAnimationUtils::CompressAnimSequenceExplicit(
 					PerTrackCompressor->bRetarget = true;
 					PerTrackCompressor->ParentingDivisor = 2.0f;
 					PerTrackCompressor->ParentingDivisorExponent = 1.6f;
-					TRYCOMPRESSION(Adaptive1_LinPerTrack, PerTrackCompressor);
+					TRYCOMPRESSION_ASYNC(Adaptive1_LinPerTrack, PerTrackCompressor);
 					PerTrackCompressor->ParentingDivisor = 1.0f;
 					PerTrackCompressor->ParentingDivisorExponent = 1.0f;
 				}
 
 
-				if( bTryPerTrackBitwiseCompression )
 				{
 					UAnimCompress_PerTrackCompression* PerTrackCompressor = NewObject<UAnimCompress_PerTrackCompression>();
 					PerTrackCompressor->bUseAdaptiveError = true;
@@ -1167,36 +1650,35 @@ void FAnimationUtils::CompressAnimSequenceExplicit(
 						PerTrackCompressor->bActuallyFilterLinearKeys = true;
 						PerTrackCompressor->bRetarget = true;
 
-						PerTrackCompressor->MaxPosDiff = 0.1;
-// 						PerTrackCompressor->MaxAngleDiff = 0.1;
-						PerTrackCompressor->MaxScaleDiff = 0.00001;
+						PerTrackCompressor->MaxPosDiff = 0.1f;
+// 						PerTrackCompressor->MaxAngleDiff = 0.1f;
+						PerTrackCompressor->MaxScaleDiff = 0.00001f;
 						PerTrackCompressor->ParentingDivisor = 2.0f;
 						PerTrackCompressor->ParentingDivisorExponent = 1.0f;
-						TRYCOMPRESSION(Linear_PerTrackExp1, PerTrackCompressor);
+						TRYCOMPRESSION_ASYNC(Linear_PerTrackExp1, PerTrackCompressor);
 
-						PerTrackCompressor->MaxPosDiff = 0.01;
-// 						PerTrackCompressor->MaxAngleDiff = 0.025;
-						PerTrackCompressor->MaxScaleDiff = 0.000001;
+						PerTrackCompressor->MaxPosDiff = 0.01f;
+// 						PerTrackCompressor->MaxAngleDiff = 0.025f;
+						PerTrackCompressor->MaxScaleDiff = 0.000001f;
 						PerTrackCompressor->ParentingDivisor = 2.0f;
 						PerTrackCompressor->ParentingDivisorExponent = 1.0f;
-						TRYCOMPRESSION(Linear_PerTrackExp2, PerTrackCompressor);
+						TRYCOMPRESSION_ASYNC(Linear_PerTrackExp2, PerTrackCompressor);
 
 						PerTrackCompressor->bRetarget = false;
-						PerTrackCompressor->MaxPosDiff = 0.1;
-// 						PerTrackCompressor->MaxAngleDiff = 0.025;
-						PerTrackCompressor->MaxScaleDiff = 0.00001;
+						PerTrackCompressor->MaxPosDiff = 0.1f;
+// 						PerTrackCompressor->MaxAngleDiff = 0.025f;
+						PerTrackCompressor->MaxScaleDiff = 0.00001f;
 						PerTrackCompressor->ParentingDivisor = 1.0f;
 						PerTrackCompressor->ParentingDivisorExponent = 1.0f;
 					}
 				}
 
-				if( bTryPerTrackBitwiseCompression )
 				{
 					UAnimCompress_PerTrackCompression* PerTrackCompressor = NewObject<UAnimCompress_PerTrackCompression>();
 					PerTrackCompressor->bUseAdaptiveError = true;
 
 					// Try the decimation algorithms
-					if (bTryIntervalKeyRemoval && (AnimSeq->NumFrames >= PerTrackCompressor->MinKeysForResampling))
+					if (AnimSeq->NumFrames >= PerTrackCompressor->MinKeysForResampling)
 					{
 						PerTrackCompressor->bActuallyFilterLinearKeys = false;
 						PerTrackCompressor->bRetarget = false;
@@ -1205,19 +1687,19 @@ void FAnimationUtils::CompressAnimSequenceExplicit(
 
 						// Try PerTrackCompression, downsample to 20 Hz
 						PerTrackCompressor->ResampledFramerate = 20.0f;
-						TRYCOMPRESSION(Downsample20Hz_PerTrack, PerTrackCompressor);
+						TRYCOMPRESSION_ASYNC(Downsample20Hz_PerTrack, PerTrackCompressor);
 
 						// Try PerTrackCompression, downsample to 15 Hz
 						PerTrackCompressor->ResampledFramerate = 15.0f;
-						TRYCOMPRESSION(Downsample15Hz_PerTrack, PerTrackCompressor);
+						TRYCOMPRESSION_ASYNC(Downsample15Hz_PerTrack, PerTrackCompressor);
 
 						// Try PerTrackCompression, downsample to 10 Hz
 						PerTrackCompressor->ResampledFramerate = 10.0f;
-						TRYCOMPRESSION(Downsample10Hz_PerTrack, PerTrackCompressor);
+						TRYCOMPRESSION_ASYNC(Downsample10Hz_PerTrack, PerTrackCompressor);
 
 						// Try PerTrackCompression, downsample to 5 Hz
 						PerTrackCompressor->ResampledFramerate = 5.0f;
-						TRYCOMPRESSION(Downsample5Hz_PerTrack, PerTrackCompressor);
+						TRYCOMPRESSION_ASYNC(Downsample5Hz_PerTrack, PerTrackCompressor);
 
 
 						// Downsampling with linear key removal and adaptive error metrics
@@ -1228,18 +1710,17 @@ void FAnimationUtils::CompressAnimSequenceExplicit(
 						PerTrackCompressor->ParentingDivisorExponent = 1.6f;
 
 						PerTrackCompressor->ResampledFramerate = 15.0f;
-						TRYCOMPRESSION(Adaptive1_15Hz_LinPerTrack, PerTrackCompressor);
+						TRYCOMPRESSION_ASYNC(Adaptive1_15Hz_LinPerTrack, PerTrackCompressor);
 
 						PerTrackCompressor->ResampledFramerate = 10.0f;
-						TRYCOMPRESSION(Adaptive1_10Hz_LinPerTrack, PerTrackCompressor);
+						TRYCOMPRESSION_ASYNC(Adaptive1_10Hz_LinPerTrack, PerTrackCompressor);
 
 						PerTrackCompressor->ResampledFramerate = 5.0f;
-						TRYCOMPRESSION(Adaptive1_5Hz_LinPerTrack, PerTrackCompressor);
+						TRYCOMPRESSION_ASYNC(Adaptive1_5Hz_LinPerTrack, PerTrackCompressor);
 					}
 				}
 
 
-				if( bTryPerTrackBitwiseCompression && bTryIntervalKeyRemoval)
 				{
 					// Try the decimation algorithms
 					if (AnimSeq->NumFrames >= 3)
@@ -1249,155 +1730,152 @@ void FAnimationUtils::CompressAnimSequenceExplicit(
 						// Downsampling with linear key removal and adaptive error metrics v2
 						NewPerTrackCompressor->MinKeysForResampling = 3;
 						NewPerTrackCompressor->bUseAdaptiveError2 = true;
-						NewPerTrackCompressor->MaxPosDiffBitwise = 0.05;
-						NewPerTrackCompressor->MaxAngleDiffBitwise = 0.02;
-						NewPerTrackCompressor->MaxScaleDiffBitwise = 0.00005;
+						NewPerTrackCompressor->MaxPosDiffBitwise = 0.05f;
+						NewPerTrackCompressor->MaxAngleDiffBitwise = 0.02f;
+						NewPerTrackCompressor->MaxScaleDiffBitwise = 0.00005f;
 						NewPerTrackCompressor->bActuallyFilterLinearKeys = true;
 						NewPerTrackCompressor->bRetarget = true;
 
 						NewPerTrackCompressor->ResampledFramerate = 15.0f;
-						TRYCOMPRESSION(Adaptive2_15Hz_LinPerTrack, NewPerTrackCompressor);
+						TRYCOMPRESSION_ASYNC(Adaptive2_15Hz_LinPerTrack, NewPerTrackCompressor);
 
 						NewPerTrackCompressor->ResampledFramerate = 10.0f;
-						TRYCOMPRESSION(Adaptive2_10Hz_LinPerTrack, NewPerTrackCompressor);
+						TRYCOMPRESSION_ASYNC(Adaptive2_10Hz_LinPerTrack, NewPerTrackCompressor);
 					}
 				}
 
 
-				if( bTryPerTrackBitwiseCompression)
 				{
 					// Adaptive error through probing the effect of perturbations at each track
 					UAnimCompress_PerTrackCompression* NewPerTrackCompressor = NewObject<UAnimCompress_PerTrackCompression>();
 					NewPerTrackCompressor->bUseAdaptiveError2 = true;
-					NewPerTrackCompressor->MaxPosDiffBitwise = 0.05;
-					NewPerTrackCompressor->MaxAngleDiffBitwise = 0.02;
-					NewPerTrackCompressor->MaxScaleDiffBitwise = 0.00005;
+					NewPerTrackCompressor->MaxPosDiffBitwise = 0.05f;
+					NewPerTrackCompressor->MaxAngleDiffBitwise = 0.02f;
+					NewPerTrackCompressor->MaxScaleDiffBitwise = 0.00005f;
 
-					TRYCOMPRESSION(Adaptive2_PerTrack, NewPerTrackCompressor);
+					TRYCOMPRESSION_ASYNC(Adaptive2_PerTrack, NewPerTrackCompressor);
 
 					NewPerTrackCompressor->bActuallyFilterLinearKeys = true;
 					NewPerTrackCompressor->bRetarget = true;
-					TRYCOMPRESSION(Adaptive2_LinPerTrack, NewPerTrackCompressor);
+					TRYCOMPRESSION_ASYNC(Adaptive2_LinPerTrack, NewPerTrackCompressor);
 
 					NewPerTrackCompressor->bActuallyFilterLinearKeys = true;
 					NewPerTrackCompressor->bRetarget = false;
-					TRYCOMPRESSION(Adaptive2_LinPerTrackNoRT, NewPerTrackCompressor);
+					TRYCOMPRESSION_ASYNC(Adaptive2_LinPerTrackNoRT, NewPerTrackCompressor);
 				}
 
-				// Increase winning compressor.
+				WaitForAnimCompressionJobs(AnimCompressionTask_CompletionEvents);
+				UpdateAnimCompressionFromAsyncJobs(AnimSeq, AnimCompressionTask_CompletionEvents, AnimCompressionJobContexes, OriginalSize, CompressorStats, MasterTolerance);
+			}
+
+			// Increase winning compressor.
+			{
+				int32 SizeDecrease = OriginalSize - CompressorStats.CurrentSize;
+				if (CompressorStats.WinningCompressorCounter)
 				{
-					int32 SizeDecrease = OriginalSize - CurrentSize;
-					if (WinningCompressorCounter)
-					{
-						++(*WinningCompressorCounter);
-						(*WinningCompressorErrorSum) += WinningCompressorError;
-						AlternativeCompressorSavings += WinningCompressorSavings;
-						*WinningCompressorMarginalSavingsSum += WinningCompressorMarginalSavings;
-						check(WinningCompressorSavings == SizeDecrease);
+					++(*CompressorStats.WinningCompressorCounter);
+					(*CompressorStats.WinningCompressorErrorSum) += CompressorStats.WinningCompressorError;
+					AlternativeCompressorSavings += CompressorStats.WinningCompressorSavings;
+					*CompressorStats.WinningCompressorMarginalSavingsSum += CompressorStats.WinningCompressorMarginalSavings;
+					check(CompressorStats.WinningCompressorSavings == SizeDecrease);
 
 					UE_LOG(LogAnimationCompression, Log, TEXT("  Recompressing(%s) with compressor('%s') saved %i bytes (%i -> %i -> %i) (max diff=%f)\n"),
-							*AnimSeq->GetName(),
-							*WinningCompressorName,
-							SizeDecrease,
-							OriginalSize, AfterOriginalRecompression, CurrentSize,
-							WinningCompressorError);
-					}
-					else
-					{
-						UE_LOG(LogAnimationCompression, Log, TEXT("  No compressor suitable! Recompressing(%s) with original/default compressor(%s) saved %i bytes (%i -> %i -> %i) (max diff=%f)\n"),
-							*AnimSeq->GetName(),
-							*AnimSeq->CompressionScheme->GetName(),
-							SizeDecrease,
-							OriginalSize, AfterOriginalRecompression, CurrentSize,
-							WinningCompressorError);
-
-						UE_LOG(LogAnimationCompression, Log, TEXT("  CompressedTrackOffsets(%d) CompressedByteStream(%d) CompressedScaleOffsets(%d)"),
-							AnimSeq->CompressedTrackOffsets.Num(),
-							AnimSeq->CompressedByteStream.Num(),
-							AnimSeq->CompressedScaleOffsets.GetMemorySize());
-
-						TotalNoWinnerRounds++;
-					}
+						*AnimSeq->GetName(),
+						*CompressorStats.WinningCompressorName,
+						SizeDecrease,
+						OriginalSize, AfterOriginalRecompression,CompressorStats. CurrentSize,
+						CompressorStats.WinningCompressorError);
+				}
+				else
+				{
+					UE_LOG(LogAnimationCompression, Log, TEXT("  No compressor suitable! Recompressing(%s) with original/default compressor(%s) saved %i bytes (%i -> %i -> %i) (max diff=%f)\n"), 
+						*AnimSeq->GetName(),
+						*AnimSeq->CompressionScheme->GetName(),
+						SizeDecrease,
+						OriginalSize, AfterOriginalRecompression, CompressorStats.CurrentSize,
+						CompressorStats.WinningCompressorError);
+
+					UE_LOG(LogAnimation, Warning, TEXT("  CompressedTrackOffsets(%d) CompressedByteStream(%d) CompressedScaleOffsets(%d) CompressedSegments(%d)"),
+						AnimSeq->CompressedTrackOffsets.Num(),
+						AnimSeq->CompressedByteStream.Num(),
+						AnimSeq->CompressedScaleOffsets.GetMemorySize(),
+						AnimSeq->CompressedSegments.Num());
+
+					TotalNoWinnerRounds++;
 				}
+			}
 
-				// Make sure we got that right.
-				check(CurrentSize == AnimSeq->GetApproxCompressedSize());
-				TotalSizeNow += CurrentSize;
+			// Make sure we got that right.
+			check(CompressorStats.CurrentSize == AnimSeq->GetApproxCompressedSize());
+			TotalSizeNow += CompressorStats.CurrentSize;
 
-				PctSaving = TotalSizeBefore > 0 ? 100.f - (100.f * float(TotalSizeNow) / float(TotalSizeBefore)) : 0.f;
-				UE_LOG(LogAnimationCompression, Log, TEXT("Compression Stats Summary [Recompressions(%i) Bytes saved(%i) before(%i) now(%i) savings(%3.1f%%) Uncompressed(%i) TotalRatio(%i:1)]"),
+			CompressorStats.PctSaving = TotalSizeBefore > 0 ? 100.f - (100.f * float(TotalSizeNow) / float(TotalSizeBefore)) : 0.f;
+			UE_LOG(LogAnimationCompression, Log, TEXT("Compression Stats Summary [Recompressions(%i) Bytes saved(%i) before(%i) now(%i) savings(%3.1f%%) Uncompressed(%i) TotalRatio(%i:1)]"), 
 				TotalRecompressions,
 				AlternativeCompressorSavings,
 				TotalSizeBefore, 
 				TotalSizeNow,
-				PctSaving,
+				CompressorStats.PctSaving,
 				TotalUncompressed,
 				(TotalUncompressed / TotalSizeNow));
 
-				UE_LOG(LogAnimationCompression, Log, TEXT("\t\tDefault compressor wins:                      %i"), TotalNoWinnerRounds);
+			UE_LOG(LogAnimationCompression, Log, TEXT("\t\tDefault compressor wins:                      %i"), TotalNoWinnerRounds);
 
-				if (bTryFixedBitwiseCompression)
-				{
-					LOG_COMPRESSION_STATUS(BitwiseACF_Float96);
-					LOG_COMPRESSION_STATUS(BitwiseACF_Fixed48);
-// 					LOG_COMPRESSION_STATUS(BitwiseACF_IntervalFixed32);
-// 					LOG_COMPRESSION_STATUS(BitwiseACF_Fixed32);
-				}
+			{
+				LOG_COMPRESSION_STATUS(BitwiseACF_Float96);
+				LOG_COMPRESSION_STATUS(BitwiseACF_Fixed48);
+// 				LOG_COMPRESSION_STATUS(BitwiseACF_IntervalFixed32);
+// 				LOG_COMPRESSION_STATUS(BitwiseACF_Fixed32);
+			}
 
-				if (bTryFixedBitwiseCompression && bTryIntervalKeyRemoval)
-				{
-					LOG_COMPRESSION_STATUS(HalfOddACF_Float96);
-					LOG_COMPRESSION_STATUS(HalfOddACF_Fixed48);
-// 					LOG_COMPRESSION_STATUS(HalfOddACF_IntervalFixed32);
-// 					LOG_COMPRESSION_STATUS(HalfOddACF_Fixed32);
-
-					LOG_COMPRESSION_STATUS(HalfEvenACF_Float96);
-					LOG_COMPRESSION_STATUS(HalfEvenACF_Fixed48);
-// 					LOG_COMPRESSION_STATUS(HalfEvenACF_IntervalFixed32);
-// 					LOG_COMPRESSION_STATUS(HalfEvenACF_Fixed32);
-				}
+			{
+				LOG_COMPRESSION_STATUS(HalfOddACF_Float96);
+				LOG_COMPRESSION_STATUS(HalfOddACF_Fixed48);
+// 				LOG_COMPRESSION_STATUS(HalfOddACF_IntervalFixed32);
+// 				LOG_COMPRESSION_STATUS(HalfOddACF_Fixed32);
+
+				LOG_COMPRESSION_STATUS(HalfEvenACF_Float96);
+				LOG_COMPRESSION_STATUS(HalfEvenACF_Fixed48);
+// 				LOG_COMPRESSION_STATUS(HalfEvenACF_IntervalFixed32);
+// 				LOG_COMPRESSION_STATUS(HalfEvenACF_Fixed32);
+			}
 
-				if (bTryLinearKeyRemovalCompression)
-				{
-					LOG_COMPRESSION_STATUS(LinearACF_Float96);
-					LOG_COMPRESSION_STATUS(LinearACF_Fixed48);
-// 					LOG_COMPRESSION_STATUS(LinearACF_IntervalFixed32);
-// 					LOG_COMPRESSION_STATUS(LinearACF_Fixed32);
-				}
+			{
+				LOG_COMPRESSION_STATUS(LinearACF_Float96);
+				LOG_COMPRESSION_STATUS(LinearACF_Fixed48);
+// 				LOG_COMPRESSION_STATUS(LinearACF_IntervalFixed32);
+// 				LOG_COMPRESSION_STATUS(LinearACF_Fixed32);
+			}
 
-				if (bTryPerTrackBitwiseCompression)
-				{
-					LOG_COMPRESSION_STATUS(Progressive_PerTrack);
-					LOG_COMPRESSION_STATUS(Bitwise_PerTrack);
-					LOG_COMPRESSION_STATUS(Linear_PerTrack);
-					LOG_COMPRESSION_STATUS(Adaptive1_LinPerTrackNoRT);
-					LOG_COMPRESSION_STATUS(Adaptive1_LinPerTrack);
+			{
+				LOG_COMPRESSION_STATUS(Progressive_PerTrack);
+				LOG_COMPRESSION_STATUS(Bitwise_PerTrack);
+				LOG_COMPRESSION_STATUS(Linear_PerTrack);
+				LOG_COMPRESSION_STATUS(Adaptive1_LinPerTrackNoRT);
+				LOG_COMPRESSION_STATUS(Adaptive1_LinPerTrack);
 					
-					LOG_COMPRESSION_STATUS(Linear_PerTrackExp1);
-					LOG_COMPRESSION_STATUS(Linear_PerTrackExp2);
-				}
+				LOG_COMPRESSION_STATUS(Linear_PerTrackExp1);
+				LOG_COMPRESSION_STATUS(Linear_PerTrackExp2);
+			}
 
-				if (bTryPerTrackBitwiseCompression && bTryIntervalKeyRemoval)
-				{
-					LOG_COMPRESSION_STATUS(Downsample20Hz_PerTrack);
-					LOG_COMPRESSION_STATUS(Downsample15Hz_PerTrack);
-					LOG_COMPRESSION_STATUS(Downsample10Hz_PerTrack);
-					LOG_COMPRESSION_STATUS(Downsample5Hz_PerTrack);
+			{
+				LOG_COMPRESSION_STATUS(Downsample20Hz_PerTrack);
+				LOG_COMPRESSION_STATUS(Downsample15Hz_PerTrack);
+				LOG_COMPRESSION_STATUS(Downsample10Hz_PerTrack);
+				LOG_COMPRESSION_STATUS(Downsample5Hz_PerTrack);
 
-					LOG_COMPRESSION_STATUS(Adaptive1_15Hz_LinPerTrack);
-					LOG_COMPRESSION_STATUS(Adaptive1_10Hz_LinPerTrack);
-					LOG_COMPRESSION_STATUS(Adaptive1_5Hz_LinPerTrack);
+				LOG_COMPRESSION_STATUS(Adaptive1_15Hz_LinPerTrack);
+				LOG_COMPRESSION_STATUS(Adaptive1_10Hz_LinPerTrack);
+				LOG_COMPRESSION_STATUS(Adaptive1_5Hz_LinPerTrack);
 
-					LOG_COMPRESSION_STATUS(Adaptive2_15Hz_LinPerTrack);
-					LOG_COMPRESSION_STATUS(Adaptive2_10Hz_LinPerTrack);
-				}
+				LOG_COMPRESSION_STATUS(Adaptive2_15Hz_LinPerTrack);
+				LOG_COMPRESSION_STATUS(Adaptive2_10Hz_LinPerTrack);
+			}
 
-				if (bTryPerTrackBitwiseCompression)
-				{
-					LOG_COMPRESSION_STATUS(Adaptive2_PerTrack);
-					LOG_COMPRESSION_STATUS(Adaptive2_LinPerTrack);
-					LOG_COMPRESSION_STATUS(Adaptive2_LinPerTrackNoRT);
-				}
+			{
+				LOG_COMPRESSION_STATUS(Adaptive2_PerTrack);
+				LOG_COMPRESSION_STATUS(Adaptive2_LinPerTrack);
+				LOG_COMPRESSION_STATUS(Adaptive2_LinPerTrackNoRT);
 			}
 		}
 		// Do not recompress - Still take into account size for stats.
diff --git a/Engine/Source/Runtime/Engine/Private/Animation/BlendSpaceBase.cpp b/Engine/Source/Runtime/Engine/Private/Animation/BlendSpaceBase.cpp
--- a/Engine/Source/Runtime/Engine/Private/Animation/BlendSpaceBase.cpp
+++ b/Engine/Source/Runtime/Engine/Private/Animation/BlendSpaceBase.cpp
@@ -562,7 +562,7 @@ void UBlendSpaceBase::GetAnimationPose(TArray<FBlendSampleData>& BlendSampleData
 				const float Time = FMath::Clamp<float>(BlendSampleDataCache[I].Time, 0.f, Sample.Animation->SequenceLength);
 
 				// first one always fills up the source one
-Sample.Animation->GetAnimationPose(Pose, ChildrenCurves[I], FAnimExtractContext(Time, true));
+				Sample.Animation->GetAnimationPose(Pose, ChildrenCurves[I], FAnimExtractContext(Time, true));
 			}
 			else
 			{
@@ -1283,3 +1283,8 @@ void UBlendSpaceBase::UpdateBlendSpacesUsingAnimSequence(UAnimSequenceBase* Sequ
 }
 
 #endif // WITH_EDITOR
+
+TArray<FName>* UBlendSpaceBase::GetUniqueMarkerNames()
+{
+	return (SampleIndexWithMarkers != INDEX_NONE && SampleData.Num() > 0) ? SampleData[SampleIndexWithMarkers].Animation->GetUniqueMarkerNames() : nullptr;
+}
diff --git a/Engine/Source/Runtime/Engine/Private/Animation/PoseAsset.cpp b/Engine/Source/Runtime/Engine/Private/Animation/PoseAsset.cpp
--- a/Engine/Source/Runtime/Engine/Private/Animation/PoseAsset.cpp
+++ b/Engine/Source/Runtime/Engine/Private/Animation/PoseAsset.cpp
@@ -4,7 +4,7 @@
 #include "UObject/FrameworkObjectVersion.h"
 #include "AnimationRuntime.h"
 #include "Components/SkeletalMeshComponent.h"
-
+#include "Animation/AnimSequence.h"
 #define LOCTEXT_NAMESPACE "PoseAsset"
 
 // utility function 
@@ -407,6 +407,8 @@ void UPoseAsset::GetBaseAnimationPose(struct FCompactPose& OutPose, FBlendedCurv
 
 bool UPoseAsset::GetAnimationPose(struct FCompactPose& OutPose, FBlendedCurve& OutCurve, const FAnimExtractContext& ExtractionContext) const
 {
+	ANIM_MT_SCOPE_CYCLE_COUNTER(PoseAssetGetAnimationPose, !IsInGameThread());
+
 	const FBoneContainer& RequiredBones = OutPose.GetBoneContainer();
 	USkeleton* MySkeleton = GetSkeleton();
 
diff --git a/Engine/Source/Runtime/Engine/Private/AnimationAsset.cpp b/Engine/Source/Runtime/Engine/Private/AnimationAsset.cpp
--- a/Engine/Source/Runtime/Engine/Private/AnimationAsset.cpp
+++ b/Engine/Source/Runtime/Engine/Private/AnimationAsset.cpp
@@ -343,6 +343,10 @@ bool UAnimationAsset::ReplaceSkeleton(USkeleton* NewSkeleton, bool bConvertSpace
 		}
 
 		RemapTracksToNewSkeleton(NewSkeleton, bConvertSpaces);
+		if (UAnimSequence* Seq = Cast<UAnimSequence>(this))
+		{
+			Seq->PostProcessSequence(false);
+		}
 
 		PostEditChange();
 		MarkPackageDirty();
diff --git a/Engine/Source/Runtime/Engine/Private/EnginePrivatePCH.h b/Engine/Source/Runtime/Engine/Private/EnginePrivatePCH.h
--- a/Engine/Source/Runtime/Engine/Private/EnginePrivatePCH.h
+++ b/Engine/Source/Runtime/Engine/Private/EnginePrivatePCH.h
@@ -503,7 +503,6 @@
 #include "AI/NavigationSystemBase.h"
 #include "Sound/SoundConcurrency.h"
 #include "ClothSimData.h"
-#include "AnimEncoding.h"
 #include "DynamicMeshBuilder.h"
 #include "EngineStats.h"
 #include "Components/SkeletalMeshComponent.h"
diff --git a/Engine/Source/Runtime/Engine/Public/AnimEncoding.h b/Engine/Source/Runtime/Engine/Public/AnimEncoding.h
--- a/Engine/Source/Runtime/Engine/Public/AnimEncoding.h
+++ b/Engine/Source/Runtime/Engine/Public/AnimEncoding.h
@@ -10,6 +10,8 @@
 #include "Misc/MemStack.h"
 #include "Animation/AnimTypes.h"
 #include "Animation/AnimSequence.h"
+#include "Animation/AnimSequenceDecompressionContext.h"
+#include "AnimationCompression.h"
 
 // switches to toggle subsets of the new animation codec system
 #define USE_ANIMATION_CODEC_BATCH_SOLVER 1
@@ -59,18 +61,129 @@ typedef TArray<BoneTrackPair> BoneTrackArray;
 /** Array of FTransform using the game memory stack */
 typedef TArray< FTransform, TMemStackAllocator<> > FTransformArray;
 
+/**
+ * Structure to wrap trivial track flags for easier and safer handling.
+ */
+struct FTrivialTrackFlags
+{
+	constexpr explicit FTrivialTrackFlags(uint8 InFlags) : Flags(InFlags) {}
+
+	constexpr bool IsTranslationTrivial() const { return (Flags & 0x4) != 0; }
+	constexpr bool IsRotationTrivial() const { return (Flags & 0x2) != 0; }
+	constexpr bool IsScaleTrivial() const { return (Flags & 0x1) != 0; }
+
+	uint8 Flags;
+};
+
+/**
+ * Structure that represents a sorted key header.
+ *
+ * Our time delta can be positive or negative. To make packing easier, we offset it in order to always make it positive
+ * A small header has 5 bits to store the time delta. 5 bits == 0x1F == 31 == [0 .. 31]
+ * A large header has 5 + 8 = 13 bits to store the time delta. 13 bits == 0x1FFF == 8191 == [0 .. 8191]
+ * A small header has an offset of 16. e.g.:
+ *     -5 + 16 = 11
+ *     15 + 16 = 31
+ *    -16 + 16 = 0
+ * A small header this has a signed range of [-16 .. 15] and maps to [0 .. 31]
+ * A large header has an offset of 4096. The signed range is [-4096 .. 4095] and maps to [0 ... 8191]
+ */
+struct FSortedKeyHeader
+{
+	static constexpr uint8 KeyTypeMask = 0x3;
+	static constexpr uint8 KeyTypeShift = 5;
+	static constexpr uint8 KeyTimeDeltaMask = (1 << 5) - 1;
+	static constexpr uint8 KeyTimeDeltaShift = 8;
+	static constexpr uint8 KeyHeaderSizeMask = 0x80;
+
+	static constexpr uint32 LargestSmallHeaderTimeDelta = (1 << 5) - 1;
+	static constexpr uint32 LargestLargeHeaderTimeDelta = (1 << 13) - 1;
+	static constexpr int32 SmallHeaderTimeDeltaOffset = (LargestSmallHeaderTimeDelta + 1) / 2;
+	static constexpr int32 LargeHeaderTimeDeltaOffset = (LargestLargeHeaderTimeDelta + 1) / 2;
+
+	FSortedKeyHeader()
+	{
+		TrackIndex = 0xFFFF;
+		PackedData[0] = 0;
+		PackedData[1] = 0;
+	}
+
+	explicit FSortedKeyHeader(const uint8* InData)
+	{
+		TrackIndex = AnimationCompressionUtils::UnalignedRead<uint16>(InData);
+		PackedData[0] = InData[sizeof(uint16) + 0];
+		PackedData[1] = IsLargeHeader() ? InData[sizeof(uint16) + 1] : 0;
+	}
+
+	FSortedKeyHeader(uint16 InTrackIndex, uint8 KeyType, int32 TimeDelta)
+	{
+		TrackIndex = InTrackIndex;
+		if (TimeDelta + SmallHeaderTimeDeltaOffset >= 0 && TimeDelta + SmallHeaderTimeDeltaOffset <= LargestSmallHeaderTimeDelta)
+		{
+			// Small header
+			const int32 OffsetTimeDelta = TimeDelta + SmallHeaderTimeDeltaOffset;
+			PackedData[0] = 0x00 | (static_cast<uint8>(KeyType) << 5) | static_cast<uint8>(OffsetTimeDelta);
+			PackedData[1] = 0;
+		}
+		else
+		{
+			// Large header
+			// Pack each byte separately to avoid issues with little endian
+			const int32 OffsetTimeDelta = TimeDelta + LargeHeaderTimeDeltaOffset;
+			check(OffsetTimeDelta >= 0 && OffsetTimeDelta <= LargestLargeHeaderTimeDelta);
+			PackedData[0] = 0x80 | (static_cast<uint16>(KeyType) << 5) | static_cast<uint16>(OffsetTimeDelta >> 8);
+			PackedData[1] = OffsetTimeDelta & 0xFF;
+		}
+	}
+
+	constexpr bool IsEndOfStream() const { return TrackIndex == 0xFFFF; }
+	constexpr bool IsLargeHeader() const { return (PackedData[0] & KeyHeaderSizeMask) != 0; }
+	constexpr uint8 GetKeyType() const { return (PackedData[0] >> KeyTypeShift) & KeyTypeMask; }
+
+	int32 GetTimeDelta() const
+	{
+		int32 TimeDelta = static_cast<uint32>(PackedData[0] & KeyTimeDeltaMask);
+		if (IsLargeHeader())
+		{
+			TimeDelta = static_cast<int32>((static_cast<uint32>(TimeDelta) << KeyTimeDeltaShift) | PackedData[1]);
+			TimeDelta -= LargeHeaderTimeDeltaOffset;
+		}
+		else
+		{
+			TimeDelta -= SmallHeaderTimeDeltaOffset;
+		}
+		return TimeDelta;
+	}
+
+	constexpr uint8 GetSize() const { return sizeof(uint16) + (IsLargeHeader() ? sizeof(uint16) : sizeof(uint8)); }
+
+	uint16 TrackIndex;
+	uint8 PackedData[2];
+};
+
+/**
+* Small header in the anim sequence compressed stream.
+*/
+struct FAnimSequenceCompressionHeader
+{
+	uint32 NumTracks;
+	uint32 NumFrames;
+	uint32 SequenceCRC;
+	uint8 bHasScale;
+	uint8 bIsSorted;						// For variable interpolation
+};
+
+
 /**
  * Extracts a single BoneAtom from an Animation Sequence.
  *
  * @param	OutAtom			The BoneAtom to fill with the extracted result.
- * @param	Seq				An Animation Sequence to extract the BoneAtom from.
+ * @param	DecompContext	The decompression context to use.
  * @param	TrackIndex		The index of the track desired in the Animation Sequence.
- * @param	Time			The time (in seconds) to calculate the BoneAtom for.
  */
 void AnimationFormat_GetBoneAtom(	FTransform& OutAtom,
-									const UAnimSequence& Seq,
-									int32 TrackIndex,
-									float Time);
+									FAnimSequenceDecompressionContext& DecompContext,
+									int32 TrackIndex);
 
 #if USE_ANIMATION_CODEC_BATCH_SOLVER
 
@@ -78,18 +191,17 @@ void AnimationFormat_GetBoneAtom(	FTransform& OutAtom,
  * Extracts an array of BoneAtoms from an Animation Sequence representing an entire pose of the skeleton.
  *
  * @param	Atoms				The BoneAtoms to fill with the extracted result.
- * @param	RotationTracks		A BoneTrackArray element for each bone requesting rotation data. 
- * @param	TranslationTracks	A BoneTrackArray element for each bone requesting translation data. 
- * @param	Seq					An Animation Sequence to extract the BoneAtom from.
- * @param	Time				The time (in seconds) to calculate the BoneAtom for.
+ * @param	RotationTracks		A BoneTrackArray element for each bone requesting rotation data.
+ * @param	TranslationTracks	A BoneTrackArray element for each bone requesting translation data.
+ * @param	ScaleTracks			A BoneTrackArray element for each bone requesting scale data.
+ * @param	DecompContext		The decompression context to use.
  */
-void AnimationFormat_GetAnimationPose(	
-	FTransformArray& Atoms, 
+void AnimationFormat_GetAnimationPose(
+	FTransformArray& Atoms,
 	const BoneTrackArray& RotationTracks,
 	const BoneTrackArray& TranslationTracks,
 	const BoneTrackArray& ScaleTracks,
-	const UAnimSequence& Seq,
-	float Time);
+	FAnimSequenceDecompressionContext& DecompContext);
 
 #endif
 
@@ -185,15 +297,18 @@ public:
 	 * Extracts a single BoneAtom from an Animation Sequence.
 	 *
 	 * @param	OutAtom			The BoneAtom to fill with the extracted result.
-	 * @param	Seq				An Animation Sequence to extract the BoneAtom from.
+	 * @param	DecompContext	The decompression context to use.
 	 * @param	TrackIndex		The index of the track desired in the Animation Sequence.
-	 * @param	Time			The time (in seconds) to calculate the BoneAtom for.
 	 */
 	virtual void GetBoneAtom(
 		FTransform& OutAtom,
-		const UAnimSequence& Seq,
-		int32 TrackIndex,
-		float Time) PURE_VIRTUAL(AnimEncoding::GetBoneAtom,);
+		FAnimSequenceDecompressionContext& DecompContext,
+		int32 TrackIndex) PURE_VIRTUAL(AnimEncoding::GetBoneAtom,);
+
+#if USE_SEGMENTING_CONTEXT
+	virtual void CreateEncodingContext(FAnimSequenceDecompressionContext& DecompContext) {}
+	virtual void ReleaseEncodingContext(FAnimSequenceDecompressionContext& DecompContext) {}
+#endif
 
 #if USE_ANIMATION_CODEC_BATCH_SOLVER
 
@@ -202,46 +317,48 @@ public:
 	 *
 	 * @param	Atoms			The FTransform array to fill in.
 	 * @param	DesiredPairs	Array of requested bone information
-	 * @param	Seq				The animation sequence to use.
-	 * @param	Time			Current time to solve for.
-	 * @return					None. 
+	 * @param	DecompContext	The decompression context to use.
 	 */
-	virtual void GetPoseRotations(	
-		FTransformArray& Atoms, 
+	virtual void GetPoseRotations(
+		FTransformArray& Atoms,
 		const BoneTrackArray& DesiredPairs,
-		const UAnimSequence& Seq,
-		float Time) PURE_VIRTUAL(AnimEncoding::GetPoseRotations,);
+		FAnimSequenceDecompressionContext& DecompContext) PURE_VIRTUAL(AnimEncoding::GetPoseRotations,);
 
 	/**
 	 * Decompress all requested translation components from an Animation Sequence
 	 *
 	 * @param	Atoms			The FTransform array to fill in.
 	 * @param	DesiredPairs	Array of requested bone information
-	 * @param	Seq				The animation sequence to use.
-	 * @param	Time			Current time to solve for.
-	 * @return					None. 
+	 * @param	DecompContext	The decompression context to use.
 	 */
-	virtual void GetPoseTranslations(	
-		FTransformArray& Atoms, 
+	virtual void GetPoseTranslations(
+		FTransformArray& Atoms,
 		const BoneTrackArray& DesiredPairs,
-		const UAnimSequence& Seq,
-		float Time) PURE_VIRTUAL(AnimEncoding::GetPoseTranslations,);
+		FAnimSequenceDecompressionContext& DecompContext) PURE_VIRTUAL(AnimEncoding::GetPoseTranslations,);
 
 	/**
 	 * Decompress all requested translation components from an Animation Sequence
 	 *
 	 * @param	Atoms			The FTransform array to fill in.
 	 * @param	DesiredPairs	Array of requested bone information
-	 * @param	Seq				The animation sequence to use.
-	 * @param	Time			Current time to solve for.
-	 * @return					None. 
+	 * @param	DecompContext	The decompression context to use.
 	 */
-	virtual void GetPoseScales(	
-		FTransformArray& Atoms, 
+	virtual void GetPoseScales(
+		FTransformArray& Atoms,
 		const BoneTrackArray& DesiredPairs,
-		const UAnimSequence& Seq,
-		float Time) PURE_VIRTUAL(AnimEncoding::GetPoseScales,);
+		FAnimSequenceDecompressionContext& DecompContext) PURE_VIRTUAL(AnimEncoding::GetPoseScales,);
 #endif
+
+	static FORCEINLINE float TimeToIndex(
+		const FAnimSequenceDecompressionContext& DecompContext,
+		const uint8* TimeMarkers,
+		uint32 NumKeys,
+		uint32 NumFrames,
+		uint8 TimeMarkerSize,
+		float SegmentRelativePos,
+		int32& FrameIndex0Out,
+		int32& FrameIndex1Out);
+
 protected:
 
 	/**
@@ -279,6 +396,8 @@ protected:
 		int32 NumKeys,
 		int32 &PosIndex0Out,
 		int32 &PosIndex1Out);
+
+	friend FAnimSequenceDecompressionContext;	// For the TimeToIndex functions
 };
 
 
@@ -294,72 +413,49 @@ public:
 	 * Extracts a single BoneAtom from an Animation Sequence.
 	 *
 	 * @param	OutAtom			The BoneAtom to fill with the extracted result.
-	 * @param	Seq				An Animation Sequence to extract the BoneAtom from.
+	 * @param	DecompContext	The decompression context to use.
 	 * @param	TrackIndex		The index of the track desired in the Animation Sequence.
-	 * @param	Time			The time (in seconds) to calculate the BoneAtom for.
 	 */
 	virtual void GetBoneAtom(
 		FTransform& OutAtom,
-		const UAnimSequence& Seq,
-		int32 TrackIndex,
-		float Time) override;
+		FAnimSequenceDecompressionContext& DecompContext,
+		int32 TrackIndex) override;
 
 	/**
 	 * Decompress the Rotation component of a BoneAtom
 	 *
 	 * @param	OutAtom			The FTransform to fill in.
-	 * @param	Seq				The animation sequence to use.
-	 * @param	Stream			The compressed animation data.
-	 * @param	NumKeys			The number of keys present in Stream.
-	 * @param	Time			Current time to solve for.
-	 * @param	RelativePos		Current position within the animation to solve for in the range [0.0,1.0].
-	 * @return					None. 
+	 * @param	DecompContext	The decompression context to use.
+	 * @param	TrackIndex		The index of the track desired in the Animation Sequence.
 	 */
-	virtual void GetBoneAtomRotation(	
+	virtual void GetBoneAtomRotation(
 		FTransform& OutAtom,
-		const UAnimSequence& Seq,
-		const uint8* RESTRICT Stream,
-		int32 NumKeys,
-		float Time,
-		float RelativePos) PURE_VIRTUAL(AnimEncoding::GetBoneAtomRotation,);
+		const FAnimSequenceDecompressionContext& DecompContext,
+		int32 TrackIndex) PURE_VIRTUAL(AnimEncoding::GetBoneAtomRotation,);
 
 	/**
 	 * Decompress the Translation component of a BoneAtom
 	 *
 	 * @param	OutAtom			The FTransform to fill in.
-	 * @param	Seq				The animation sequence to use.
-	 * @param	Stream			The compressed animation data.
-	 * @param	NumKeys			The number of keys present in Stream.
-	 * @param	Time			Current time to solve for.
-	 * @param	RelativePos		Current position within the animation to solve for in the range [0.0,1.0].
-	 * @return					None. 
+	 * @param	DecompContext	The decompression context to use.
+	 * @param	TrackIndex		The index of the track desired in the Animation Sequence.
 	 */
-	virtual void GetBoneAtomTranslation(	
+	virtual void GetBoneAtomTranslation(
 		FTransform& OutAtom,
-		const UAnimSequence& Seq,
-		const uint8* RESTRICT Stream,
-		int32 NumKeys,
-		float Time,
-		float RelativePos) PURE_VIRTUAL(AnimEncoding::GetBoneAtomTranslation,);
+		const FAnimSequenceDecompressionContext& DecompContext,
+		int32 TrackIndex) PURE_VIRTUAL(AnimEncoding::GetBoneAtomTranslation,);
 
 	/**
 	 * Decompress the Scale component of a BoneAtom
 	 *
 	 * @param	OutAtom			The FTransform to fill in.
-	 * @param	Seq				The animation sequence to use.
-	 * @param	Stream			The compressed animation data.
-	 * @param	NumKeys			The number of keys present in Stream.
-	 * @param	Time			Current time to solve for.
-	 * @param	RelativePos		Current position within the animation to solve for in the range [0.0,1.0].
-	 * @return					None. 
+	 * @param	DecompContext	The decompression context to use.
+	 * @param	TrackIndex		The index of the track desired in the Animation Sequence.
 	 */
-	virtual void GetBoneAtomScale(	
+	virtual void GetBoneAtomScale(
 		FTransform& OutAtom,
-		const UAnimSequence& Seq,
-		const uint8* RESTRICT Stream,
-		int32 NumKeys,
-		float Time,
-		float RelativePos) PURE_VIRTUAL(AnimEncoding::GetBoneAtomScale,);
+		const FAnimSequenceDecompressionContext& DecompContext,
+		int32 TrackIndex) PURE_VIRTUAL(AnimEncoding::GetBoneAtomScale,);
 
 	/**
 	 * Handles Byte-swapping incoming animation data from a MemoryReader
@@ -478,10 +574,6 @@ public:
 };
 
 
-
-
-
-
 /**
  * Utility function to determine the two key indices to interpolate given a relative position in the animation
  *
@@ -686,3 +778,74 @@ FORCEINLINE float AnimEncoding::TimeToIndex(
 	
 	return Alpha;
 }
+
+float AnimEncoding::TimeToIndex(
+	const FAnimSequenceDecompressionContext& DecompContext,
+	const uint8* TimeMarkers,
+	uint32 NumKeys,
+	uint32 NumFrames,
+	uint8 TimeMarkerSize,
+	float SegmentRelativePos,
+	int32& FrameIndex0Out,
+	int32& FrameIndex1Out)
+{
+	float Alpha = 0.0f;
+
+	checkSlow(NumKeys != 0);
+
+	const int32 LastKey = NumKeys - 1;
+
+	if (NumKeys < 2 || DecompContext.RelativePos <= 0.f)
+	{
+		// return the first key
+		FrameIndex0Out = 0;
+		FrameIndex1Out = 0;
+		Alpha = 0.0f;
+	}
+	else if (DecompContext.RelativePos >= 1.0f)
+	{
+		// return the ending key
+		FrameIndex0Out = LastKey;
+		FrameIndex1Out = LastKey;
+		Alpha = 0.0f;
+	}
+	else
+	{
+		// find the proper key range to return
+		const int32 LastFrame = NumFrames - 1;
+		const float KeyPos = SegmentRelativePos * (float)LastKey;
+		const float FramePos = SegmentRelativePos * (float)LastFrame;
+		const int32 FramePosFloor = FMath::Clamp(FMath::TruncToInt(FramePos), 0, LastFrame);
+		const int32 KeyEstimate = FMath::Clamp(FMath::TruncToInt(KeyPos), 0, LastKey);
+
+		int32 LowFrame = 0;
+		int32 HighFrame = 0;
+
+		// find the pair of keys which surround our target frame index
+		if (TimeMarkerSize == sizeof(uint16))
+		{
+			const uint16* Frames = reinterpret_cast<const uint16*>(TimeMarkers);
+			FrameIndex0Out = FindLowKeyIndex<uint16>(Frames, NumKeys, FramePosFloor, KeyEstimate);
+			LowFrame = Frames[FrameIndex0Out];
+
+			FrameIndex1Out = FMath::Min(FrameIndex0Out + 1, LastKey);
+			HighFrame = Frames[FrameIndex1Out];
+		}
+		else
+		{
+			FrameIndex0Out = FindLowKeyIndex<uint8>(TimeMarkers, NumKeys, FramePosFloor, KeyEstimate);
+			LowFrame = TimeMarkers[FrameIndex0Out];
+
+			FrameIndex1Out = FMath::Min(FrameIndex0Out + 1, LastKey);
+			HighFrame = TimeMarkers[FrameIndex1Out];
+		}
+
+		// compute the blend parameters for the keys we have found
+		const int32 Delta = FMath::Max(HighFrame - LowFrame, 1);
+		const float Remainder = FramePos - (float)LowFrame;
+		Alpha = DecompContext.AnimSeq->Interpolation == EAnimInterpolationType::Step ? 0.f : (Remainder / (float)Delta);
+	}
+
+	return Alpha;
+}
+
