Copying //UE4/Dev-AnimPhys to Dev-Main (//UE4/Dev-Main)

diff --git a/Engine/Source/Runtime/Engine/Public/AnimEncoding_ConstantKeyLerp.h b/Engine/Source/Runtime/Engine/Public/AnimEncoding_ConstantKeyLerp.h
--- a/Engine/Source/Runtime/Engine/Public/AnimEncoding_ConstantKeyLerp.h
+++ b/Engine/Source/Runtime/Engine/Public/AnimEncoding_ConstantKeyLerp.h
@@ -9,9 +9,28 @@
 #include "CoreMinimal.h"
 #include "AnimEncoding.h"
 #include "AnimationCompression.h"
+#include "Animation/AnimEncodingDecompressionContext.h"
 
 class FMemoryWriter;
 
+#if USE_SEGMENTING_CONTEXT
+class FAEConstantKeyLerpContext : public FAnimEncodingDecompressionContext
+{
+public:
+	FAEConstantKeyLerpContext(const FAnimSequenceDecompressionContext& DecompContext);
+	virtual void Seek(const FAnimSequenceDecompressionContext& DecompContext, float SampleAtTime) override;
+
+	template<int32 Format> inline FQuat GetUniformRotation(const FAnimSequenceDecompressionContext& DecompContext, int32 TrackIndex, uint8 SegmentIndex) const;
+	template<int32 Format> inline FVector GetUniformTranslation(const FAnimSequenceDecompressionContext& DecompContext, int32 TrackIndex, uint8 SegmentIndex) const;
+	template<int32 Format> inline FVector GetUniformScale(const FAnimSequenceDecompressionContext& DecompContext, int32 TrackIndex, uint8 SegmentIndex) const;
+
+	TArray<int32> UniformKeyOffsets;
+	int32 KeyFrameSize;
+
+	int32 FrameKeysOffset[2];
+};
+#endif
+
 /**
  * Base class for all Animation Encoding Formats using consistently-spaced key interpolation.
  */
@@ -102,6 +121,11 @@ public:
 		FMemoryWriter& MemoryWriter,
 		uint8*& ScaleTrackData,
 		int32 NumKeysScale) override;
+
+#if USE_SEGMENTING_CONTEXT
+	virtual void CreateEncodingContext(FAnimSequenceDecompressionContext& DecompContext) override;
+	virtual void ReleaseEncodingContext(FAnimSequenceDecompressionContext& DecompContext) override;
+#endif
 };
 
 template<int32 FORMAT>
@@ -112,55 +136,29 @@ public:
 	 * Decompress the Rotation component of a BoneAtom
 	 *
 	 * @param	OutAtom			The FTransform to fill in.
-	 * @param	Stream			The compressed animation data.
-	 * @param	NumKeys			The number of keys present in Stream.
-	 * @param	Time			Current time to solve for.
-	 * @param	RelativePos		Current position within the animation to solve for in the range [0.0,1.0].
-	 * @return					None. 
+	 * @param	DecompContext	The decompression context to use.
+	 * @param	TrackIndex		The index of the track desired in the Animation Sequence.
 	 */
-	void GetBoneAtomRotation(	
-		FTransform& OutAtom,
-		const UAnimSequence& Seq,
-		const uint8* RESTRICT Stream,
-		int32 NumKeys,
-		float Time,
-		float RelativePos);
+	virtual void GetBoneAtomRotation(FTransform& OutAtom, const FAnimSequenceDecompressionContext& DecompContext, int32 TrackIndex) override;
 
 	/**
 	 * Decompress the Translation component of a BoneAtom
 	 *
 	 * @param	OutAtom			The FTransform to fill in.
-	 * @param	Stream			The compressed animation data.
-	 * @param	NumKeys			The number of keys present in Stream.
-	 * @param	Time			Current time to solve for.
-	 * @param	RelativePos		Current position within the animation to solve for in the range [0.0,1.0].
-	 * @return					None. 
+	 * @param	DecompContext	The decompression context to use.
+	 * @param	TrackIndex		The index of the track desired in the Animation Sequence.
 	 */
-	void GetBoneAtomTranslation(	
-		FTransform& OutAtom,
-		const UAnimSequence& Seq,
-		const uint8* RESTRICT Stream,
-		int32 NumKeys,
-		float Time,
-		float RelativePos);
+	virtual void GetBoneAtomTranslation(FTransform& OutAtom, const FAnimSequenceDecompressionContext& DecompContext, int32 TrackIndex) override;
 
 	/**
 	 * Decompress the Scale component of a BoneAtom
 	 *
 	 * @param	OutAtom			The FTransform to fill in.
-	 * @param	Stream			The compressed animation data.
-	 * @param	NumKeys			The number of keys present in Stream.
-	 * @param	Time			Current time to solve for.
-	 * @param	RelativePos		Current position within the animation to solve for in the range [0.0,1.0].
-	 * @return					None. 
+	 * @param	DecompContext	The decompression context to use.
+	 * @param	TrackIndex		The index of the track desired in the Animation Sequence.
 	 */
-	void GetBoneAtomScale(	
-		FTransform& OutAtom,
-		const UAnimSequence& Seq,
-		const uint8* RESTRICT Stream,
-		int32 NumKeys,
-		float Time,
-		float RelativePos);
+	virtual void GetBoneAtomScale(FTransform& OutAtom, const FAnimSequenceDecompressionContext& DecompContext, int32 TrackIndex) override;
+
 #if USE_ANIMATION_CODEC_BATCH_SOLVER
 
 	/**
@@ -168,45 +166,36 @@ public:
 	 *
 	 * @param	Atoms			The FTransform array to fill in.
 	 * @param	DesiredPairs	Array of requested bone information
-	 * @param	Seq				The animation sequence to use.
-	 * @param	Time			Current time to solve for.
-	 * @return					None. 
+	 * @param	DecompContext	The decompression context to use.
 	 */
-	void GetPoseRotations(	
-		FTransformArray& Atoms, 
+	void GetPoseRotations(
+		FTransformArray& Atoms,
 		const BoneTrackArray& DesiredPairs,
-		const UAnimSequence& Seq,
-		float RelativePos);
+		FAnimSequenceDecompressionContext& DecompContext);
 
 	/**
 	 * Decompress all requested translation components from an Animation Sequence
 	 *
 	 * @param	Atoms			The FTransform array to fill in.
 	 * @param	DesiredPairs	Array of requested bone information
-	 * @param	Seq				The animation sequence to use.
-	 * @param	Time			Current time to solve for.
-	 * @return					None. 
+	 * @param	DecompContext	The decompression context to use.
 	 */
-	void GetPoseTranslations(	
+	void GetPoseTranslations(
 		FTransformArray& Atoms,
 		const BoneTrackArray& DesiredPairs,
-		const UAnimSequence& Seq,
-		float RelativePos);
+		FAnimSequenceDecompressionContext& DecompContext);
 
 	/**
 	 * Decompress all requested Scale components from an Animation Sequence
 	 *
 	 * @param	Atoms			The FTransform array to fill in.
 	 * @param	DesiredPairs	Array of requested bone information
-	 * @param	Seq				The animation sequence to use.
-	 * @param	Time			Current time to solve for.
-	 * @return					None. 
+	 * @param	DecompContext	The decompression context to use.
 	 */
-	void GetPoseScales(	
+	void GetPoseScales(
 		FTransformArray& Atoms,
 		const BoneTrackArray& DesiredPairs,
-		const UAnimSequence& Seq,
-		float RelativePos);
+		FAnimSequenceDecompressionContext& DecompContext);
 #endif
 
 };
@@ -217,21 +206,50 @@ public:
  * Decompress the Rotation component of a BoneAtom
  *
  * @param	OutAtom			The FTransform to fill in.
- * @param	Stream			The compressed animation data.
- * @param	NumKeys			The number of keys present in Stream.
- * @param	Time			Current time to solve for.
- * @param	RelativePos		Current position within the animation to solve for in the range [0.0,1.0].
- * @return					None. 
+ * @param	DecompContext	The decompression context to use.
+ * @param	TrackIndex		The index of the track desired in the Animation Sequence.
  */
 template<int32 FORMAT>
-FORCEINLINE void AEFConstantKeyLerp<FORMAT>::GetBoneAtomRotation(	
-	FTransform& OutAtom,
-	const UAnimSequence& Seq,
-	const uint8* RESTRICT RotStream,
-	int32 NumRotKeys,
-	float Time,
-	float RelativePos)
+FORCEINLINE void AEFConstantKeyLerp<FORMAT>::GetBoneAtomRotation(FTransform& OutAtom, const FAnimSequenceDecompressionContext& DecompContext, int32 TrackIndex)
 {
+#if USE_SEGMENTING_CONTEXT
+	if (DecompContext.AnimSeq->CompressedSegments.Num() != 0)
+	{
+		const FTrivialAnimKeyHandle TrivialKeyHandle = DecompContext.GetTrivialRotationKeyHandle(TrackIndex);
+		if (TrivialKeyHandle.IsValid())
+		{
+			DecompContext.GetTrivialRotation(OutAtom, TrivialKeyHandle);
+		}
+		else
+		{
+			const FAEConstantKeyLerpContext& EncodingContext = *static_cast<const FAEConstantKeyLerpContext*>(DecompContext.EncodingContext);
+
+			const FQuat R0 = EncodingContext.GetUniformRotation<FORMAT>(DecompContext, TrackIndex, 0);
+
+			if (DecompContext.NeedsInterpolation)
+			{
+				const FQuat R1 = EncodingContext.GetUniformRotation<FORMAT>(DecompContext, TrackIndex, 1);
+
+				// Fast linear quaternion interpolation.
+				FQuat BlendedQuat = FQuat::FastLerp(R0, R1, DecompContext.KeyAlpha);
+				BlendedQuat.Normalize();
+				OutAtom.SetRotation(BlendedQuat);
+			}
+			else // (Index0 == Index1)
+			{
+				OutAtom.SetRotation(R0);
+			}
+		}
+
+		return;
+	}
+#endif
+
+	const int32* RESTRICT TrackData = DecompContext.GetCompressedTrackOffsets() + (TrackIndex * 4);
+	int32 RotKeysOffset = TrackData[2];
+	int32 NumRotKeys = TrackData[3];
+	const uint8* RESTRICT RotStream = DecompContext.GetCompressedByteStream() + RotKeysOffset;
+
 	if (NumRotKeys == 1)
 	{
 		// For a rotation track of n=1 keys, the single key is packed as an FQuatFloat96NoW.
@@ -243,7 +261,7 @@ FORCEINLINE void AEFConstantKeyLerp<FORMAT>::GetBoneAtomRotation(
 	{
 		int32 Index0;
 		int32 Index1;
-		float Alpha = TimeToIndex(Seq,RelativePos,NumRotKeys,Index0,Index1);
+		float Alpha = TimeToIndex(*DecompContext.AnimSeq, DecompContext.RelativePos, NumRotKeys, Index0, Index1);
 
 		const int32 RotationStreamOffset = (FORMAT == ACF_IntervalFixed32NoW) ? (sizeof(float)*6) : 0; // offset past Min and Range data
 
@@ -279,24 +297,50 @@ FORCEINLINE void AEFConstantKeyLerp<FORMAT>::GetBoneAtomRotation(
  * Decompress the Translation component of a BoneAtom
  *
  * @param	OutAtom			The FTransform to fill in.
- * @param	Stream			The compressed animation data.
- * @param	NumKeys			The number of keys present in Stream.
- * @param	Time			Current time to solve for.
- * @param	RelativePos		Current position within the animation to solve for in the range [0.0,1.0].
- * @return					None. 
+ * @param	DecompContext	The decompression context to use.
+ * @param	TrackIndex		The index of the track desired in the Animation Sequence.
  */
 template<int32 FORMAT>
-FORCEINLINE void AEFConstantKeyLerp<FORMAT>::GetBoneAtomTranslation(	
-	FTransform& OutAtom,
-	const UAnimSequence& Seq,
-	const uint8* RESTRICT TransStream,
-	int32 NumTransKeys,
-	float Time,
-	float RelativePos)
+FORCEINLINE void AEFConstantKeyLerp<FORMAT>::GetBoneAtomTranslation(FTransform& OutAtom, const FAnimSequenceDecompressionContext& DecompContext, int32 TrackIndex)
 {
+#if USE_SEGMENTING_CONTEXT
+	if (DecompContext.AnimSeq->CompressedSegments.Num() != 0)
+	{
+		const FTrivialAnimKeyHandle TrivialKeyHandle = DecompContext.GetTrivialTranslationKeyHandle(TrackIndex);
+		if (TrivialKeyHandle.IsValid())
+		{
+			DecompContext.GetTrivialTranslation(OutAtom, TrivialKeyHandle);
+		}
+		else
+		{
+			const FAEConstantKeyLerpContext& EncodingContext = *static_cast<const FAEConstantKeyLerpContext*>(DecompContext.EncodingContext);
+
+			const FVector P0 = EncodingContext.GetUniformTranslation<FORMAT>(DecompContext, TrackIndex, 0);
+
+			if (DecompContext.NeedsInterpolation)
+			{
+				const FVector P1 = EncodingContext.GetUniformTranslation<FORMAT>(DecompContext, TrackIndex, 1);
+
+				OutAtom.SetTranslation(FMath::Lerp(P0, P1, DecompContext.KeyAlpha));
+			}
+			else // (Index0 == Index1)
+			{
+				OutAtom.SetTranslation(P0);
+			}
+		}
+
+		return;
+	}
+#endif
+
+	const int32* RESTRICT TrackData = DecompContext.GetCompressedTrackOffsets() + (TrackIndex * 4);
+	int32 TransKeysOffset = TrackData[0];
+	int32 NumTransKeys = TrackData[1];
+	const uint8* RESTRICT TransStream = DecompContext.GetCompressedByteStream() + TransKeysOffset;
+
 	int32 Index0;
 	int32 Index1;
-	float Alpha = TimeToIndex(Seq,RelativePos,NumTransKeys,Index0,Index1);
+	float Alpha = TimeToIndex(*DecompContext.AnimSeq, DecompContext.RelativePos, NumTransKeys, Index0, Index1);
 
 	const int32 TransStreamOffset = ((FORMAT == ACF_IntervalFixed32NoW) && NumTransKeys > 1) ? (sizeof(float)*6) : 0; // offset past Min and Range data
 
@@ -323,25 +367,50 @@ FORCEINLINE void AEFConstantKeyLerp<FORMAT>::GetBoneAtomTranslation(
 /**
  * Decompress the Scale component of a BoneAtom
  *
- * @param	OutAtom			The FScaleform to fill in.
- * @param	Stream			The compressed animation data.
- * @param	NumKeys			The number of keys present in Stream.
- * @param	Time			Current time to solve for.
- * @param	RelativePos		Current position within the animation to solve for in the range [0.0,1.0].
- * @return					None. 
+ * @param	OutAtom			The FTransform to fill in.
+ * @param	DecompContext	The decompression context to use.
+ * @param	TrackIndex		The index of the track desired in the Animation Sequence.
  */
 template<int32 FORMAT>
-FORCEINLINE void AEFConstantKeyLerp<FORMAT>::GetBoneAtomScale(	
-	FTransform& OutAtom,
-	const UAnimSequence& Seq,
-	const uint8* RESTRICT ScaleStream,
-	int32 NumScaleKeys,
-	float Time,
-	float RelativePos)
+FORCEINLINE void AEFConstantKeyLerp<FORMAT>::GetBoneAtomScale(FTransform& OutAtom, const FAnimSequenceDecompressionContext& DecompContext, int32 TrackIndex)
 {
+#if USE_SEGMENTING_CONTEXT
+	if (DecompContext.AnimSeq->CompressedSegments.Num() != 0)
+	{
+		const FTrivialAnimKeyHandle TrivialKeyHandle = DecompContext.GetTrivialScaleKeyHandle(TrackIndex);
+		if (TrivialKeyHandle.IsValid())
+		{
+			DecompContext.GetTrivialScale(OutAtom, TrivialKeyHandle);
+		}
+		else
+		{
+			const FAEConstantKeyLerpContext& EncodingContext = *static_cast<const FAEConstantKeyLerpContext*>(DecompContext.EncodingContext);
+
+			const FVector S0 = EncodingContext.GetUniformScale<FORMAT>(DecompContext, TrackIndex, 0);
+
+			if (DecompContext.NeedsInterpolation)
+			{
+				const FVector S1 = EncodingContext.GetUniformScale<FORMAT>(DecompContext, TrackIndex, 1);
+
+				OutAtom.SetScale3D(FMath::Lerp(S0, S1, DecompContext.KeyAlpha));
+			}
+			else // (Index0 == Index1)
+			{
+				OutAtom.SetScale3D(S0);
+			}
+		}
+
+		return;
+	}
+#endif
+
+	const int32 ScaleKeysOffset = DecompContext.GetCompressedScaleOffsets()->GetOffsetData(TrackIndex, 0);
+	const int32 NumScaleKeys = DecompContext.GetCompressedScaleOffsets()->GetOffsetData(TrackIndex, 1);
+	const uint8* RESTRICT ScaleStream = DecompContext.GetCompressedByteStream() + ScaleKeysOffset;
+
 	int32 Index0;
 	int32 Index1;
-	float Alpha = TimeToIndex(Seq,RelativePos,NumScaleKeys,Index0,Index1);
+	float Alpha = TimeToIndex(*DecompContext.AnimSeq, DecompContext.RelativePos, NumScaleKeys, Index0, Index1);
 
 	const int32 ScaleStreamOffset = ((FORMAT == ACF_IntervalFixed32NoW) && NumScaleKeys > 1) ? (sizeof(float)*6) : 0; // offset past Min and Range data
 
@@ -372,19 +441,15 @@ FORCEINLINE void AEFConstantKeyLerp<FORMAT>::GetBoneAtomScale(
  *
  * @param	Atoms			The FTransform array to fill in.
  * @param	DesiredPairs	Array of requested bone information
- * @param	Seq				The animation sequence to use.
- * @param	Time			Current time to solve for.
- * @return					None. 
+ * @param	DecompContext	The decompression context to use.
  */
 template<int32 FORMAT>
-inline void AEFConstantKeyLerp<FORMAT>::GetPoseRotations(	
-	FTransformArray& Atoms, 
+inline void AEFConstantKeyLerp<FORMAT>::GetPoseRotations(
+	FTransformArray& Atoms,
 	const BoneTrackArray& DesiredPairs,
-	const UAnimSequence& Seq,
-	float Time)
+	FAnimSequenceDecompressionContext& DecompContext)
 {
 	const int32 PairCount = DesiredPairs.Num();
-	const float RelativePos = Time / (float)Seq.SequenceLength;
 
 	for (int32 PairIndex=0; PairIndex<PairCount; ++PairIndex)
 	{
@@ -393,13 +458,8 @@ inline void AEFConstantKeyLerp<FORMAT>::GetPoseRotations(
 		const int32 AtomIndex = Pair.AtomIndex;
 		FTransform& BoneAtom = Atoms[AtomIndex];
 
-		const int32* RESTRICT TrackData = Seq.CompressedTrackOffsets.GetData() + (TrackIndex*4);
-		const int32 RotKeysOffset	= *(TrackData+2);
-		const int32 NumRotKeys	= *(TrackData+3);
-		const uint8* RESTRICT RotStream		= Seq.CompressedByteStream.GetData()+RotKeysOffset;
-
 		// call the decoder directly (not through the vtable)
-		AEFConstantKeyLerp<FORMAT>::GetBoneAtomRotation(BoneAtom, Seq, RotStream, NumRotKeys, Time, RelativePos);
+		AEFConstantKeyLerp<FORMAT>::GetBoneAtomRotation(BoneAtom, DecompContext, TrackIndex);
 	}
 }
 
@@ -408,19 +468,15 @@ inline void AEFConstantKeyLerp<FORMAT>::GetPoseRotations(
  *
  * @param	Atoms			The FTransform array to fill in.
  * @param	DesiredPairs	Array of requested bone information
- * @param	Seq				The animation sequence to use.
- * @param	Time			Current time to solve for.
- * @return					None. 
+ * @param	DecompContext	The decompression context to use.
  */
 template<int32 FORMAT>
-inline void AEFConstantKeyLerp<FORMAT>::GetPoseTranslations(	
-	FTransformArray& Atoms, 
+inline void AEFConstantKeyLerp<FORMAT>::GetPoseTranslations(
+	FTransformArray& Atoms,
 	const BoneTrackArray& DesiredPairs,
-	const UAnimSequence& Seq,
-	float Time)
+	FAnimSequenceDecompressionContext& DecompContext)
 {
-	const int32 PairCount= DesiredPairs.Num();
-	const float RelativePos = Time / (float)Seq.SequenceLength;
+	const int32 PairCount = DesiredPairs.Num();
 
 	//@TODO: Verify that this prefetch is helping
 	// Prefetch the desired pairs array and 2 destination spots; the loop will prefetch one 2 out each iteration
@@ -445,13 +501,8 @@ inline void AEFConstantKeyLerp<FORMAT>::GetPoseTranslations(
 		const int32 AtomIndex = Pair.AtomIndex;
 		FTransform& BoneAtom = Atoms[AtomIndex];
 
-		const int32* RESTRICT TrackData = Seq.CompressedTrackOffsets.GetData() + (TrackIndex*4);
-		const int32 TransKeysOffset	= *(TrackData+0);
-		const int32 NumTransKeys		= *(TrackData+1);
-		const uint8* RESTRICT TransStream = Seq.CompressedByteStream.GetData()+TransKeysOffset;
-
 		// call the decoder directly (not through the vtable)
-		AEFConstantKeyLerp<FORMAT>::GetBoneAtomTranslation(BoneAtom, Seq, TransStream, NumTransKeys, Time, RelativePos);
+		AEFConstantKeyLerp<FORMAT>::GetBoneAtomTranslation(BoneAtom, DecompContext, TrackIndex);
 	}
 }
 
@@ -460,21 +511,17 @@ inline void AEFConstantKeyLerp<FORMAT>::GetPoseTranslations(
  *
  * @param	Atoms			The FTransform array to fill in.
  * @param	DesiredPairs	Array of requested bone information
- * @param	Seq				The animation sequence to use.
- * @param	Time			Current time to solve for.
- * @return					None. 
+ * @param	DecompContext	The decompression context to use.
  */
 template<int32 FORMAT>
-inline void AEFConstantKeyLerp<FORMAT>::GetPoseScales(	
-	FTransformArray& Atoms, 
+inline void AEFConstantKeyLerp<FORMAT>::GetPoseScales(
+	FTransformArray& Atoms,
 	const BoneTrackArray& DesiredPairs,
-	const UAnimSequence& Seq,
-	float Time)
+	FAnimSequenceDecompressionContext& DecompContext)
 {
-	check (Seq.CompressedScaleOffsets.IsValid());
+	checkSlow(DecompContext.bHasScale);
 
 	const int32 PairCount= DesiredPairs.Num();
-	const float RelativePos = Time / (float)Seq.SequenceLength;
 
 	//@TODO: Verify that this prefetch is helping
 	// Prefetch the desired pairs array and 2 destination spots; the loop will prefetch one 2 out each iteration
@@ -499,12 +546,49 @@ inline void AEFConstantKeyLerp<FORMAT>::GetPoseScales(
 		const int32 AtomIndex = Pair.AtomIndex;
 		FTransform& BoneAtom = Atoms[AtomIndex];
 
-		const int32 ScaleKeysOffset		= Seq.CompressedScaleOffsets.GetOffsetData(TrackIndex, 0);
-		const int32 NumScaleKeys		= Seq.CompressedScaleOffsets.GetOffsetData(TrackIndex, 1);
-		const uint8* RESTRICT ScaleStream = Seq.CompressedByteStream.GetData()+ScaleKeysOffset;
-
 		// call the decoder directly (not through the vtable)
-		AEFConstantKeyLerp<FORMAT>::GetBoneAtomScale(BoneAtom, Seq, ScaleStream, NumScaleKeys, Time, RelativePos);
+		AEFConstantKeyLerp<FORMAT>::GetBoneAtomScale(BoneAtom, DecompContext, TrackIndex);
 	}
 }
 #endif
+
+#if USE_SEGMENTING_CONTEXT
+template<int32 Format>
+FQuat FAEConstantKeyLerpContext::GetUniformRotation(const FAnimSequenceDecompressionContext& DecompContext, int32 TrackIndex, uint8 SegmentIndex) const
+{
+	const int32 FrameKeyOffset = UniformKeyOffsets[DecompContext.GetRotationValueOffset(TrackIndex)];
+	const int32 SegmentKeyOffset = FrameKeysOffset[SegmentIndex] + FrameKeyOffset;
+	const uint8* KeyData = DecompContext.CompressedByteStream + SegmentKeyOffset;
+
+	FQuat Rotation;
+	DecompressRotation<Format>(Rotation, DecompContext.TrackRangeData[SegmentIndex], KeyData);
+
+	return Rotation;
+}
+
+template<int32 Format>
+FVector FAEConstantKeyLerpContext::GetUniformTranslation(const FAnimSequenceDecompressionContext& DecompContext, int32 TrackIndex, uint8 SegmentIndex) const
+{
+	const int32 FrameKeyOffset = UniformKeyOffsets[DecompContext.GetTranslationValueOffset(TrackIndex)];
+	const int32 SegmentKeyOffset = FrameKeysOffset[SegmentIndex] + FrameKeyOffset;
+	const uint8* KeyData = DecompContext.CompressedByteStream + SegmentKeyOffset;
+
+	FVector Translation;
+	DecompressTranslation<Format>(Translation, DecompContext.TrackRangeData[SegmentIndex], KeyData);
+
+	return Translation;
+}
+
+template<int32 Format>
+FVector FAEConstantKeyLerpContext::GetUniformScale(const FAnimSequenceDecompressionContext& DecompContext, int32 TrackIndex, uint8 SegmentIndex) const
+{
+	const int32 FrameKeyOffset = UniformKeyOffsets[DecompContext.GetScaleValueOffset(TrackIndex)];
+	const int32 SegmentKeyOffset = FrameKeysOffset[SegmentIndex] + FrameKeyOffset;
+	const uint8* KeyData = DecompContext.CompressedByteStream + SegmentKeyOffset;
+
+	FVector Scale;
+	DecompressScale<Format>(Scale, DecompContext.TrackRangeData[SegmentIndex], KeyData);
+
+	return Scale;
+}
+#endif
\ No newline at end of file
diff --git a/Engine/Source/Runtime/Engine/Public/AnimEncoding_PerTrackCompression.h b/Engine/Source/Runtime/Engine/Public/AnimEncoding_PerTrackCompression.h
--- a/Engine/Source/Runtime/Engine/Public/AnimEncoding_PerTrackCompression.h
+++ b/Engine/Source/Runtime/Engine/Public/AnimEncoding_PerTrackCompression.h
@@ -42,15 +42,13 @@ public:
 	 * Extracts a single BoneAtom from an Animation Sequence.
 	 *
 	 * @param	OutAtom			The BoneAtom to fill with the extracted result.
-	 * @param	Seq				An Animation Sequence to extract the BoneAtom from.
+	 * @param	DecompContext	The decompression context to use.
 	 * @param	TrackIndex		The index of the track desired in the Animation Sequence.
-	 * @param	Time			The time (in seconds) to calculate the BoneAtom for.
 	 */
 	virtual void GetBoneAtom(
 		FTransform& OutAtom,
-		const UAnimSequence& Seq,
-		int32 TrackIndex,
-		float Time) override;
+		FAnimSequenceDecompressionContext& DecompContext,
+		int32 TrackIndex) override;
 
 #if USE_ANIMATION_CODEC_BATCH_SOLVER
 
@@ -59,45 +57,41 @@ public:
 	 *
 	 * @param	Atoms			The FTransform array to fill in.
 	 * @param	DesiredPairs	Array of requested bone information
-	 * @param	Seq				The animation sequence to use.
-	 * @param	Time			Current time to solve for.
-	 * @return					None. 
+	 * @param	DecompContext	The decompression context to use.
 	 */
-	virtual void GetPoseRotations(	
-		FTransformArray& Atoms, 
+	virtual void GetPoseRotations(
+		FTransformArray& Atoms,
 		const BoneTrackArray& DesiredPairs,
-		const UAnimSequence& Seq,
-		float Time) override;
+		FAnimSequenceDecompressionContext& DecompContext) override;
 
 	/**
 	 * Decompress all requested translation components from an Animation Sequence
 	 *
 	 * @param	Atoms			The FTransform array to fill in.
 	 * @param	DesiredPairs	Array of requested bone information
-	 * @param	Seq				The animation sequence to use.
-	 * @param	Time			Current time to solve for.
-	 * @return					None. 
+	 * @param	DecompContext	The decompression context to use.
 	 */
-	virtual void GetPoseTranslations(	
-		FTransformArray& Atoms, 
+	virtual void GetPoseTranslations(
+		FTransformArray& Atoms,
 		const BoneTrackArray& DesiredPairs,
-		const UAnimSequence& Seq,
-		float Time) override;
+		FAnimSequenceDecompressionContext& DecompContext) override;
 
 	/**
 	 * Decompress all requested Scale components from an Animation Sequence
 	 *
 	 * @param	Atoms			The FTransform array to fill in.
 	 * @param	DesiredPairs	Array of requested bone information
-	 * @param	Seq				The animation sequence to use.
-	 * @param	Time			Current time to solve for.
-	 * @return					None. 
+	 * @param	DecompContext	The decompression context to use.
 	 */
-	virtual void GetPoseScales(	
-		FTransformArray& Atoms, 
+	virtual void GetPoseScales(
+		FTransformArray& Atoms,
 		const BoneTrackArray& DesiredPairs,
-		const UAnimSequence& Seq,
-		float Time) override;
+		FAnimSequenceDecompressionContext& DecompContext) override;
+#endif
+
+#if USE_SEGMENTING_CONTEXT
+	virtual void CreateEncodingContext(FAnimSequenceDecompressionContext& DecompContext) override;
+	virtual void ReleaseEncodingContext(FAnimSequenceDecompressionContext& DecompContext) override;
 #endif
 
 protected:
@@ -123,54 +117,55 @@ protected:
 	 * Decompress the Rotation component of a BoneAtom
 	 *
 	 * @param	OutAtom			The FTransform to fill in.
-	 * @param	Seq				The animation sequence to use.
-	 * @param	Stream			The compressed animation data.
-	 * @param	NumKeys			The number of keys present in Stream.
-	 * @param	Time			Current time to solve for.
-	 * @param	RelativePos		Current position within the animation to solve for in the range [0.0,1.0].
-	 * @return					None. 
+	 * @param	DecompContext	The decompression context to use.
+	 * @param	TrackIndex		The index of the track desired in the Animation Sequence.
 	 */
-	static void GetBoneAtomRotation(	
+	static void GetBoneAtomRotation(
 		FTransform& OutAtom,
-		const UAnimSequence& Seq,
-		int32 Offset,
-		float Time,
-		float RelativePos);
+		FAnimSequenceDecompressionContext& DecompContext,
+		int32 TrackIndex);
 
 	/**
 	 * Decompress the Translation component of a BoneAtom
 	 *
 	 * @param	OutAtom			The FTransform to fill in.
-	 * @param	Seq				The animation sequence to use.
-	 * @param	Stream			The compressed animation data.
-	 * @param	NumKeys			The number of keys present in Stream.
-	 * @param	Time			Current time to solve for.
-	 * @param	RelativePos		Current position within the animation to solve for in the range [0.0,1.0].
-	 * @return					None. 
+	 * @param	DecompContext	The decompression context to use.
+	 * @param	TrackIndex		The index of the track desired in the Animation Sequence.
 	 */
-	static void GetBoneAtomTranslation(	
+	static void GetBoneAtomTranslation(
 		FTransform& OutAtom,
-		const UAnimSequence& Seq,
-		int32 Offset,
-		float Time,
-		float RelativePos);
+		FAnimSequenceDecompressionContext& DecompContext,
+		int32 TrackIndex);
 
 	/**
 	 * Decompress the Scale component of a BoneAtom
 	 *
 	 * @param	OutAtom			The FTransform to fill in.
-	 * @param	Seq				The animation sequence to use.
-	 * @param	Stream			The compressed animation data.
-	 * @param	NumKeys			The number of keys present in Stream.
-	 * @param	Time			Current time to solve for.
-	 * @param	RelativePos		Current position within the animation to solve for in the range [0.0,1.0].
-	 * @return					None. 
+	 * @param	DecompContext	The decompression context to use.
+	 * @param	TrackIndex		The index of the track desired in the Animation Sequence.
 	 */
-	static void GetBoneAtomScale(	
+	static void GetBoneAtomScale(
 		FTransform& OutAtom,
-		const UAnimSequence& Seq,
-		int32 Offset,
-		float Time,
-		float RelativePos);
+		FAnimSequenceDecompressionContext& DecompContext,
+		int32 TrackIndex);
 };
 
+/**
+ * Structure to wrap per track flags.
+ */
+struct FPerTrackFlags
+{
+	constexpr explicit FPerTrackFlags(uint8 InFlags) : Flags(InFlags) {}
+	FPerTrackFlags(bool bHasTimeMarkers, AnimationCompressionFormat Format, uint8 FormatFlags)
+		: Flags((bHasTimeMarkers ? 0x80 : 0) | (FormatFlags << 4) | (uint8)Format)
+	{
+		check((FormatFlags & ~0x7) == 0);
+		check(((uint8)Format & ~0xF) == 0);
+	}
+
+	constexpr bool IsUniform() const { return (Flags & 0x80) == 0; }
+	constexpr uint8 GetFormatFlags() const { return (Flags >> 4) & 0x7; }
+	constexpr uint8 GetFormat() const { return Flags & 0xF; }
+
+	uint8 Flags;
+};
diff --git a/Engine/Source/Runtime/Engine/Public/AnimEncoding_VariableKeyLerp.h b/Engine/Source/Runtime/Engine/Public/AnimEncoding_VariableKeyLerp.h
--- a/Engine/Source/Runtime/Engine/Public/AnimEncoding_VariableKeyLerp.h
+++ b/Engine/Source/Runtime/Engine/Public/AnimEncoding_VariableKeyLerp.h
@@ -8,10 +8,71 @@
 
 #include "CoreMinimal.h"
 #include "AnimEncoding.h"
+#include "AnimationCompression.h"
+#include "Animation/AnimEncodingDecompressionContext.h"
 #include "AnimEncoding_ConstantKeyLerp.h"
 
 class FMemoryWriter;
 
+#if USE_SEGMENTING_CONTEXT
+class FAEVariableKeyLerpSortedContext : public FAnimEncodingDecompressionContext
+{
+public:
+	struct FCachedKey
+	{
+		// Linear interpolation only requires 2 keys
+		// Index 0 is always the oldest key, index 1 the newest
+		int32 RotOffsets[2];
+		int32 RotFrameIndices[2];
+		int32 TransOffsets[2];
+		int32 TransFrameIndices[2];
+		int32 ScaleOffsets[2];
+		int32 ScaleFrameIndices[2];
+	};
+
+	TArray<FCachedKey> CachedKeys;			// 1 Entry per track
+
+	int32 SegmentStartFrame0;
+	int32 SegmentStartFrame1;
+	float FramePos;
+
+	const uint8_t* PackedSampleData;		// The current pointer into our data stream
+	int32 PreviousFrameIndex;				// The previously read frame index
+	int32 CurrentFrameIndex;				// The current frame index
+
+	float PreviousSampleAtTime;
+	uint16 PreviousSegmentIndex;
+
+	FAEVariableKeyLerpSortedContext(const FAnimSequenceDecompressionContext& DecompContext);
+	virtual void Seek(const FAnimSequenceDecompressionContext& DecompContext, float SampleAtTime) override;
+
+	template<int32 Format> inline FQuat GetSortedRotation(const FAnimSequenceDecompressionContext& DecompContext, int32 TrackIndex) const;
+	template<int32 Format> inline FVector GetSortedTranslation(const FAnimSequenceDecompressionContext& DecompContext, int32 TrackIndex) const;
+	template<int32 Format> inline FVector GetSortedScale(const FAnimSequenceDecompressionContext& DecompContext, int32 TrackIndex) const;
+};
+
+class FAEVariableKeyLerpLinearContext : public FAnimEncodingDecompressionContext
+{
+public:
+	static constexpr int32 OffsetNumKeysPairSize = sizeof(uint32) + sizeof(uint16);
+
+	float SegmentRelativePos0;
+
+	uint8 TimeMarkerSize[2];	// sizeof(uint8) or sizeof(uint16)
+
+	const uint8* OffsetNumKeysPairs[2];
+
+	TArray<int32> NumAnimatedTrackStreams;
+
+	FAEVariableKeyLerpLinearContext(const FAnimSequenceDecompressionContext& DecompContext);
+	virtual void Seek(const FAnimSequenceDecompressionContext& DecompContext, float SampleAtTime) override;
+
+	template<int32 Format> inline FQuat GetLinearRotation(const FAnimSequenceDecompressionContext& DecompContext, int32 TrackIndex) const;
+	template<int32 Format> inline FVector GetLinearTranslation(const FAnimSequenceDecompressionContext& DecompContext, int32 TrackIndex) const;
+	template<int32 Format> inline FVector GetLinearScale(const FAnimSequenceDecompressionContext& DecompContext, int32 TrackIndex) const;
+};
+#endif
+
 /**
  * Base class for all Animation Encoding Formats using variably-spaced key interpolation.
  */
@@ -101,6 +162,11 @@ public:
 		FMemoryWriter& MemoryWriter,
 		uint8*& ScaleTrackData,
 		int32 NumKeysScale) override;
+
+#if USE_SEGMENTING_CONTEXT
+	virtual void CreateEncodingContext(FAnimSequenceDecompressionContext& DecompContext) override;
+	virtual void ReleaseEncodingContext(FAnimSequenceDecompressionContext& DecompContext) override;
+#endif
 };
 
 template<int32 FORMAT>
@@ -111,56 +177,28 @@ public:
 	 * Decompress the Rotation component of a BoneAtom
 	 *
 	 * @param	OutAtom			The FTransform to fill in.
-	 * @param	Stream			The compressed animation data.
-	 * @param	NumKeys			The number of keys present in Stream.
-	 * @param	Time			Current time to solve for.
-	 * @param	RelativePos		Current position within the animation to solve for in the range [0.0,1.0].
-	 * @return					None. 
+	 * @param	DecompContext	The decompression context to use.
+	 * @param	TrackIndex		The index of the track desired in the Animation Sequence.
 	 */
-	virtual void GetBoneAtomRotation(	
-		FTransform& OutAtom,
-		const UAnimSequence& Seq,
-		const uint8* RESTRICT Stream,
-		int32 NumKeys,
-		float Time,
-		float RelativePos) override;
+	virtual void GetBoneAtomRotation(FTransform& OutAtom, const FAnimSequenceDecompressionContext& DecompContext, int32 TrackIndex) override;
 
 	/**
 	 * Decompress the Translation component of a BoneAtom
 	 *
 	 * @param	OutAtom			The FTransform to fill in.
-	 * @param	Stream			The compressed animation data.
-	 * @param	NumKeys			The number of keys present in Stream.
-	 * @param	Time			Current time to solve for.
-	 * @param	RelativePos		Current position within the animation to solve for in the range [0.0,1.0].
-	 * @return					None. 
+	 * @param	DecompContext	The decompression context to use.
+	 * @param	TrackIndex		The index of the track desired in the Animation Sequence.
 	 */
-	virtual void GetBoneAtomTranslation(	
-		FTransform& OutAtom,
-		const UAnimSequence& Seq,
-		const uint8* RESTRICT Stream,
-		int32 NumKeys,
-		float Time,
-		float RelativePos) override;
-
-	
+	virtual void GetBoneAtomTranslation(FTransform& OutAtom, const FAnimSequenceDecompressionContext& DecompContext, int32 TrackIndex) override;
+
 	/**
 	 * Decompress the Scale component of a BoneAtom
 	 *
 	 * @param	OutAtom			The FTransform to fill in.
-	 * @param	Stream			The compressed animation data.
-	 * @param	NumKeys			The number of keys present in Stream.
-	 * @param	Time			Current time to solve for.
-	 * @param	RelativePos		Current position within the animation to solve for in the range [0.0,1.0].
-	 * @return					None. 
+	 * @param	DecompContext	The decompression context to use.
+	 * @param	TrackIndex		The index of the track desired in the Animation Sequence.
 	 */
-	virtual void GetBoneAtomScale(	
-		FTransform& OutAtom,
-		const UAnimSequence& Seq,
-		const uint8* RESTRICT Stream,
-		int32 NumKeys,
-		float Time,
-		float RelativePos) override;
+	virtual void GetBoneAtomScale(FTransform& OutAtom, const FAnimSequenceDecompressionContext& DecompContext, int32 TrackIndex) override;
 
 #if USE_ANIMATION_CODEC_BATCH_SOLVER
 
@@ -169,45 +207,36 @@ public:
 	 *
 	 * @param	Atoms			The FTransform array to fill in.
 	 * @param	DesiredPairs	Array of requested bone information
-	 * @param	Seq				The animation sequence to use.
-	 * @param	Time			Current time to solve for.
-	 * @return					None. 
+	 * @param	DecompContext	The decompression context to use.
 	 */
-	virtual void GetPoseRotations(	
-		FTransformArray& Atoms, 
+	virtual void GetPoseRotations(
+		FTransformArray& Atoms,
 		const BoneTrackArray& DesiredPairs,
-		const UAnimSequence& Seq,
-		float RelativePos) override;
+		FAnimSequenceDecompressionContext& DecompContext) override;
 
 	/**
 	 * Decompress all requested translation components from an Animation Sequence
 	 *
 	 * @param	Atoms			The FTransform array to fill in.
 	 * @param	DesiredPairs	Array of requested bone information
-	 * @param	Seq				The animation sequence to use.
-	 * @param	Time			Current time to solve for.
-	 * @return					None. 
+	 * @param	DecompContext	The decompression context to use.
 	 */
-	virtual void GetPoseTranslations(	
+	virtual void GetPoseTranslations(
 		FTransformArray& Atoms,
 		const BoneTrackArray& DesiredPairs,
-		const UAnimSequence& Seq,
-		float RelativePos) override;
+		FAnimSequenceDecompressionContext& DecompContext) override;
 
 	/**
 	 * Decompress all requested Scale components from an Animation Sequence
 	 *
 	 * @param	Atoms			The FTransform array to fill in.
 	 * @param	DesiredPairs	Array of requested bone information
-	 * @param	Seq				The animation sequence to use.
-	 * @param	Time			Current time to solve for.
-	 * @return					None. 
+	 * @param	DecompContext	The decompression context to use.
 	 */
-	virtual void GetPoseScales(	
+	virtual void GetPoseScales(
 		FTransformArray& Atoms,
 		const BoneTrackArray& DesiredPairs,
-		const UAnimSequence& Seq,
-		float RelativePos) override;
+		FAnimSequenceDecompressionContext& DecompContext) override;
 #endif
 
 };
@@ -217,21 +246,47 @@ public:
  * Decompress the Rotation component of a BoneAtom
  *
  * @param	OutAtom			The FTransform to fill in.
- * @param	Stream			The compressed animation data.
- * @param	NumKeys			The number of keys present in Stream.
- * @param	Time			Current time to solve for.
- * @param	RelativePos		Current position within the animation to solve for in the range [0.0,1.0].
- * @return					None. 
+ * @param	DecompContext	The decompression context to use.
+ * @param	TrackIndex		The index of the track desired in the Animation Sequence.
  */
 template<int32 FORMAT>
-FORCEINLINE_DEBUGGABLE void AEFVariableKeyLerp<FORMAT>::GetBoneAtomRotation(	
-	FTransform& OutAtom,
-	const UAnimSequence& Seq,
-	const uint8* RESTRICT RotStream,
-	int32 NumRotKeys,
-	float Time,
-	float RelativePos)
+FORCEINLINE_DEBUGGABLE void AEFVariableKeyLerp<FORMAT>::GetBoneAtomRotation(FTransform& OutAtom, const FAnimSequenceDecompressionContext& DecompContext, int32 TrackIndex)
 {
+#if USE_SEGMENTING_CONTEXT
+	if (DecompContext.AnimSeq->CompressedSegments.Num() != 0)
+	{
+		const FTrivialAnimKeyHandle TrivialKeyHandle = DecompContext.GetTrivialRotationKeyHandle(TrackIndex);
+		if (TrivialKeyHandle.IsValid())
+		{
+			DecompContext.GetTrivialRotation(OutAtom, TrivialKeyHandle);
+		}
+		else
+		{
+			FQuat Rotation;
+			if (DecompContext.bIsSorted)
+			{
+				const FAEVariableKeyLerpSortedContext& EncodingContext = *static_cast<const FAEVariableKeyLerpSortedContext*>(DecompContext.EncodingContext);
+				Rotation = EncodingContext.GetSortedRotation<FORMAT>(DecompContext, TrackIndex);
+			}
+			else
+			{
+				const FAEVariableKeyLerpLinearContext& EncodingContext = *static_cast<const FAEVariableKeyLerpLinearContext*>(DecompContext.EncodingContext);
+				Rotation = EncodingContext.GetLinearRotation<FORMAT>(DecompContext, TrackIndex);
+			}
+
+			OutAtom.SetRotation(Rotation);
+		}
+
+		return;
+	}
+
+#endif
+
+	const int32* RESTRICT TrackData = DecompContext.GetCompressedTrackOffsets() + (TrackIndex * 4);
+	int32 RotKeysOffset = TrackData[2];
+	int32 NumRotKeys = TrackData[3];
+	const uint8* RESTRICT RotStream = DecompContext.GetCompressedByteStream() + RotKeysOffset;
+
 	if (NumRotKeys == 1)
 	{
 		// For a rotation track of n=1 keys, the single key is packed as an FQuatFloat96NoW.
@@ -247,7 +302,7 @@ FORCEINLINE_DEBUGGABLE void AEFVariableKeyLerp<FORMAT>::GetBoneAtomRotation(
 
 		int32 Index0;
 		int32 Index1;
-		float Alpha = TimeToIndex(Seq,FrameTable,RelativePos,NumRotKeys,Index0,Index1);
+		float Alpha = TimeToIndex(*DecompContext.AnimSeq, FrameTable, DecompContext.RelativePos, NumRotKeys, Index0, Index1);
 
 
 		if (Index0 != Index1)
@@ -283,27 +338,53 @@ FORCEINLINE_DEBUGGABLE void AEFVariableKeyLerp<FORMAT>::GetBoneAtomRotation(
  * Decompress the Translation component of a BoneAtom
  *
  * @param	OutAtom			The FTransform to fill in.
- * @param	Stream			The compressed animation data.
- * @param	NumKeys			The number of keys present in Stream.
- * @param	Time			Current time to solve for.
- * @param	RelativePos		Current position within the animation to solve for in the range [0.0,1.0].
- * @return					None. 
+ * @param	DecompContext	The decompression context to use.
+ * @param	TrackIndex		The index of the track desired in the Animation Sequence.
  */
 template<int32 FORMAT>
-FORCEINLINE_DEBUGGABLE void AEFVariableKeyLerp<FORMAT>::GetBoneAtomTranslation(	
-	FTransform& OutAtom,
-	const UAnimSequence& Seq,
-	const uint8* RESTRICT TransStream,
-	int32 NumTransKeys,
-	float Time,
-	float RelativePos)
+FORCEINLINE_DEBUGGABLE void AEFVariableKeyLerp<FORMAT>::GetBoneAtomTranslation(FTransform& OutAtom, const FAnimSequenceDecompressionContext& DecompContext, int32 TrackIndex)
 {
+#if USE_SEGMENTING_CONTEXT
+	if (DecompContext.AnimSeq->CompressedSegments.Num() != 0)
+	{
+		const FTrivialAnimKeyHandle TrivialKeyHandle = DecompContext.GetTrivialTranslationKeyHandle(TrackIndex);
+		if (TrivialKeyHandle.IsValid())
+		{
+			DecompContext.GetTrivialTranslation(OutAtom, TrivialKeyHandle);
+		}
+		else
+		{
+			FVector Translation;
+			if (DecompContext.bIsSorted)
+			{
+				const FAEVariableKeyLerpSortedContext& EncodingContext = *static_cast<const FAEVariableKeyLerpSortedContext*>(DecompContext.EncodingContext);
+				Translation = EncodingContext.GetSortedTranslation<FORMAT>(DecompContext, TrackIndex);
+			}
+			else
+			{
+				const FAEVariableKeyLerpLinearContext& EncodingContext = *static_cast<const FAEVariableKeyLerpLinearContext*>(DecompContext.EncodingContext);
+				Translation = EncodingContext.GetLinearTranslation<FORMAT>(DecompContext, TrackIndex);
+			}
+
+			OutAtom.SetTranslation(Translation);
+		}
+
+		return;
+	}
+
+#endif
+
+	const int32* RESTRICT TrackData = DecompContext.GetCompressedTrackOffsets() + (TrackIndex * 4);
+	int32 TransKeysOffset = TrackData[0];
+	int32 NumTransKeys = TrackData[1];
+	const uint8* RESTRICT TransStream = DecompContext.GetCompressedByteStream() + TransKeysOffset;
+
 	const uint8* RESTRICT FrameTable= TransStream +(NumTransKeys*CompressedTranslationStrides[FORMAT]*CompressedTranslationNum[FORMAT]);
 	FrameTable= Align(FrameTable, 4);
 
 	int32 Index0;
 	int32 Index1;
-	float Alpha = TimeToIndex(Seq,FrameTable,RelativePos,NumTransKeys,Index0,Index1);
+	float Alpha = TimeToIndex(*DecompContext.AnimSeq, FrameTable, DecompContext.RelativePos, NumTransKeys, Index0, Index1);
 	const int32 TransStreamOffset = ((FORMAT == ACF_IntervalFixed32NoW) && NumTransKeys > 1) ? (sizeof(float)*6) : 0; // offset past Min and Range data
 
 	if (Index0 != Index1)
@@ -330,27 +411,51 @@ FORCEINLINE_DEBUGGABLE void AEFVariableKeyLerp<FORMAT>::GetBoneAtomTranslation(
  * Decompress the Scale component of a BoneAtom
  *
  * @param	OutAtom			The FTransform to fill in.
- * @param	Stream			The compressed animation data.
- * @param	NumKeys			The number of keys present in Stream.
- * @param	Time			Current time to solve for.
- * @param	RelativePos		Current position within the animation to solve for in the range [0.0,1.0].
- * @return					None. 
+ * @param	DecompContext	The decompression context to use.
+ * @param	TrackIndex		The index of the track desired in the Animation Sequence.
  */
 template<int32 FORMAT>
-FORCEINLINE_DEBUGGABLE void AEFVariableKeyLerp<FORMAT>::GetBoneAtomScale(	
-	FTransform& OutAtom,
-	const UAnimSequence& Seq,
-	const uint8* RESTRICT ScaleStream,
-	int32 NumScaleKeys,
-	float Time,
-	float RelativePos)
+FORCEINLINE_DEBUGGABLE void AEFVariableKeyLerp<FORMAT>::GetBoneAtomScale(FTransform& OutAtom, const FAnimSequenceDecompressionContext& DecompContext, int32 TrackIndex)
 {
+#if USE_SEGMENTING_CONTEXT
+	if (DecompContext.AnimSeq->CompressedSegments.Num() != 0)
+	{
+		const FTrivialAnimKeyHandle TrivialKeyHandle = DecompContext.GetTrivialScaleKeyHandle(TrackIndex);
+		if (TrivialKeyHandle.IsValid())
+		{
+			DecompContext.GetTrivialScale(OutAtom, TrivialKeyHandle);
+		}
+		else
+		{
+			FVector Scale;
+			if (DecompContext.bIsSorted)
+			{
+				const FAEVariableKeyLerpSortedContext& EncodingContext = *static_cast<const FAEVariableKeyLerpSortedContext*>(DecompContext.EncodingContext);
+				Scale = EncodingContext.GetSortedScale<FORMAT>(DecompContext, TrackIndex);
+			}
+			else
+			{
+				const FAEVariableKeyLerpLinearContext& EncodingContext = *static_cast<const FAEVariableKeyLerpLinearContext*>(DecompContext.EncodingContext);
+				Scale = EncodingContext.GetLinearScale<FORMAT>(DecompContext, TrackIndex);
+			}
+
+			OutAtom.SetScale3D(Scale);
+		}
+
+		return;
+	}
+#endif
+
+	const int32 ScaleKeysOffset = DecompContext.GetCompressedScaleOffsets()->GetOffsetData(TrackIndex, 0);
+	const int32 NumScaleKeys = DecompContext.GetCompressedScaleOffsets()->GetOffsetData(TrackIndex, 1);
+	const uint8* RESTRICT ScaleStream = DecompContext.GetCompressedByteStream() + ScaleKeysOffset;
+
 	const uint8* RESTRICT FrameTable= ScaleStream +(NumScaleKeys*CompressedScaleStrides[FORMAT]*CompressedScaleNum[FORMAT]);
 	FrameTable= Align(FrameTable, 4);
 
 	int32 Index0;
 	int32 Index1;
-	float Alpha = TimeToIndex(Seq,FrameTable,RelativePos,NumScaleKeys,Index0,Index1);
+	float Alpha = TimeToIndex(*DecompContext.AnimSeq, FrameTable, DecompContext.RelativePos, NumScaleKeys, Index0, Index1);
 	const int32 ScaleStreamOffset = ((FORMAT == ACF_IntervalFixed32NoW) && NumScaleKeys > 1) ? (sizeof(float)*6) : 0; // offset past Min and Range data
 
 	if (Index0 != Index1)
@@ -380,19 +485,15 @@ FORCEINLINE_DEBUGGABLE void AEFVariableKeyLerp<FORMAT>::GetBoneAtomScale(
  *
  * @param	Atoms			The FTransform array to fill in.
  * @param	DesiredPairs	Array of requested bone information
- * @param	Seq				The animation sequence to use.
- * @param	Time			Current time to solve for.
- * @return					None. 
+ * @param	DecompContext	The decompression context to use.
  */
 template<int32 FORMAT>
-void AEFVariableKeyLerp<FORMAT>::GetPoseRotations(	
-	FTransformArray& Atoms, 
+void AEFVariableKeyLerp<FORMAT>::GetPoseRotations(
+	FTransformArray& Atoms,
 	const BoneTrackArray& DesiredPairs,
-	const UAnimSequence& Seq,
-	float Time)
+	FAnimSequenceDecompressionContext& DecompContext)
 {
 	const int32 PairCount = DesiredPairs.Num();
-	const float RelativePos = Time / (float)Seq.SequenceLength;
 
 	for (int32 PairIndex=0; PairIndex<PairCount; ++PairIndex)
 	{
@@ -401,13 +502,8 @@ void AEFVariableKeyLerp<FORMAT>::GetPoseRotations(
 		const int32 AtomIndex = Pair.AtomIndex;
 		FTransform& BoneAtom = Atoms[AtomIndex];
 
-		const int32* RESTRICT TrackData = Seq.CompressedTrackOffsets.GetData() + (TrackIndex*4);
-		const int32 RotKeysOffset	= *(TrackData+2);
-		const int32 NumRotKeys	= *(TrackData+3);
-		const uint8* RESTRICT RotStream		= Seq.CompressedByteStream.GetData()+RotKeysOffset;
-
 		// call the decoder directly (not through the vtable)
-		AEFVariableKeyLerp<FORMAT>::GetBoneAtomRotation(BoneAtom, Seq, RotStream, NumRotKeys, Time, RelativePos);
+		AEFVariableKeyLerp<FORMAT>::GetBoneAtomRotation(BoneAtom, DecompContext, TrackIndex);
 	}
 }
 
@@ -416,19 +512,15 @@ void AEFVariableKeyLerp<FORMAT>::GetPoseRotations(
  *
  * @param	Atoms			The FTransform array to fill in.
  * @param	DesiredPairs	Array of requested bone information
- * @param	Seq				The animation sequence to use.
- * @param	Time			Current time to solve for.
- * @return					None. 
+ * @param	DecompContext	The decompression context to use.
  */
 template<int32 FORMAT>
-void AEFVariableKeyLerp<FORMAT>::GetPoseTranslations(	
-	FTransformArray& Atoms, 
+void AEFVariableKeyLerp<FORMAT>::GetPoseTranslations(
+	FTransformArray& Atoms,
 	const BoneTrackArray& DesiredPairs,
-	const UAnimSequence& Seq,
-	float Time)
+	FAnimSequenceDecompressionContext& DecompContext)
 {
 	const int32 PairCount= DesiredPairs.Num();
-	const float RelativePos = Time / (float)Seq.SequenceLength;
 
 	for (int32 PairIndex=0; PairIndex<PairCount; ++PairIndex)
 	{
@@ -437,13 +529,8 @@ void AEFVariableKeyLerp<FORMAT>::GetPoseTranslations(
 		const int32 AtomIndex = Pair.AtomIndex;
 		FTransform& BoneAtom = Atoms[AtomIndex];
 
-		const int32* RESTRICT TrackData = Seq.CompressedTrackOffsets.GetData() + (TrackIndex*4);
-		const int32 TransKeysOffset	= *(TrackData+0);
-		const int32 NumTransKeys		= *(TrackData+1);
-		const uint8* RESTRICT TransStream = Seq.CompressedByteStream.GetData()+TransKeysOffset;
-
 		// call the decoder directly (not through the vtable)
-		AEFVariableKeyLerp<FORMAT>::GetBoneAtomTranslation(BoneAtom, Seq, TransStream, NumTransKeys, Time, RelativePos);
+		AEFVariableKeyLerp<FORMAT>::GetBoneAtomTranslation(BoneAtom, DecompContext, TrackIndex);
 	}
 }
 
@@ -452,21 +539,17 @@ void AEFVariableKeyLerp<FORMAT>::GetPoseTranslations(
  *
  * @param	Atoms			The FTransform array to fill in.
  * @param	DesiredPairs	Array of requested bone information
- * @param	Seq				The animation sequence to use.
- * @param	Time			Current time to solve for.
- * @return					None. 
+ * @param	DecompContext	The decompression context to use.
  */
 template<int32 FORMAT>
-void AEFVariableKeyLerp<FORMAT>::GetPoseScales(	
-	FTransformArray& Atoms, 
+void AEFVariableKeyLerp<FORMAT>::GetPoseScales(
+	FTransformArray& Atoms,
 	const BoneTrackArray& DesiredPairs,
-	const UAnimSequence& Seq,
-	float Time)
+	FAnimSequenceDecompressionContext& DecompContext)
 {
-	check (Seq.CompressedScaleOffsets.IsValid());
+	checkSlow(DecompContext.bHasScale);
 
 	const int32 PairCount= DesiredPairs.Num();
-	const float RelativePos = Time / (float)Seq.SequenceLength;
 
 	for (int32 PairIndex=0; PairIndex<PairCount; ++PairIndex)
 	{
@@ -475,12 +558,281 @@ void AEFVariableKeyLerp<FORMAT>::GetPoseScales(
 		const int32 AtomIndex = Pair.AtomIndex;
 		FTransform& BoneAtom = Atoms[AtomIndex];
 
-		const int32 ScaleKeysOffset	= Seq.CompressedScaleOffsets.GetOffsetData(TrackIndex, 0);
-		const int32 NumScaleKeys	= Seq.CompressedScaleOffsets.GetOffsetData(TrackIndex, 1);
-		const uint8* RESTRICT ScaleStream = Seq.CompressedByteStream.GetData()+ScaleKeysOffset;
-
 		// call the decoder directly (not through the vtable)
-		AEFVariableKeyLerp<FORMAT>::GetBoneAtomScale(BoneAtom, Seq, ScaleStream, NumScaleKeys, Time, RelativePos);
+		AEFVariableKeyLerp<FORMAT>::GetBoneAtomScale(BoneAtom, DecompContext, TrackIndex);
 	}
 }
 #endif
+
+#if USE_SEGMENTING_CONTEXT
+template<int32 Format>
+FQuat FAEVariableKeyLerpSortedContext::GetSortedRotation(const FAnimSequenceDecompressionContext& DecompContext, int32 TrackIndex) const
+{
+	const FCachedKey& CachedKey = CachedKeys[TrackIndex];
+
+	float Alpha;
+	if (DecompContext.NeedsTwoSegments)
+	{
+		Alpha = DecompContext.KeyAlpha;
+	}
+	else
+	{
+		// compute the blend parameters for the keys we have found
+		const int32 FrameIndex0 = SegmentStartFrame0 + CachedKey.RotFrameIndices[0];
+		const int32 FrameIndex1 = SegmentStartFrame1 + CachedKey.RotFrameIndices[1];
+
+		const int32 Delta = FMath::Max(FrameIndex1 - FrameIndex0, 1);
+		const float Remainder = FramePos - float(FrameIndex0);
+		Alpha = DecompContext.AnimSeq->Interpolation == EAnimInterpolationType::Step ? 0.0f : (Remainder / float(Delta));
+	}
+
+	const uint8* KeyData0 = DecompContext.CompressedByteStream + CachedKey.RotOffsets[0];
+
+	FQuat Rotation0;
+	DecompressRotation<Format, false>(Rotation0, DecompContext.TrackRangeData[0], KeyData0);
+
+	const uint8* KeyData1 = DecompContext.CompressedByteStream + CachedKey.RotOffsets[1];
+
+	FQuat Rotation1;
+	DecompressRotation<Format, false>(Rotation1, DecompContext.TrackRangeData[1], KeyData1);
+
+	// Fast linear quaternion interpolation.
+	FQuat BlendedQuat = FQuat::FastLerp(Rotation0, Rotation1, Alpha);
+	BlendedQuat.Normalize();
+
+	return BlendedQuat;
+}
+
+template<int32 Format>
+FVector FAEVariableKeyLerpSortedContext::GetSortedTranslation(const FAnimSequenceDecompressionContext& DecompContext, int32 TrackIndex) const
+{
+	const FCachedKey& CachedKey = CachedKeys[TrackIndex];
+
+	float Alpha;
+	if (DecompContext.NeedsTwoSegments)
+	{
+		Alpha = DecompContext.KeyAlpha;
+	}
+	else
+	{
+		// compute the blend parameters for the keys we have found
+		const int32 FrameIndex0 = SegmentStartFrame0 + CachedKey.TransFrameIndices[0];
+		const int32 FrameIndex1 = SegmentStartFrame1 + CachedKey.TransFrameIndices[1];
+
+		const int32 Delta = FMath::Max(FrameIndex1 - FrameIndex0, 1);
+		const float Remainder = FramePos - float(FrameIndex0);
+		Alpha = DecompContext.AnimSeq->Interpolation == EAnimInterpolationType::Step ? 0.0f : (Remainder / float(Delta));
+	}
+
+	const uint8* KeyData0 = DecompContext.CompressedByteStream + CachedKey.TransOffsets[0];
+
+	FVector Translation0;
+	DecompressTranslation<Format, false>(Translation0, DecompContext.TrackRangeData[0], KeyData0);
+
+	const uint8* KeyData1 = DecompContext.CompressedByteStream + CachedKey.TransOffsets[1];
+
+	FVector Translation1;
+	DecompressTranslation<Format, false>(Translation1, DecompContext.TrackRangeData[1], KeyData1);
+
+	return FMath::Lerp(Translation0, Translation1, Alpha);
+}
+
+template<int32 Format>
+FVector FAEVariableKeyLerpSortedContext::GetSortedScale(const FAnimSequenceDecompressionContext& DecompContext, int32 TrackIndex) const
+{
+	const FCachedKey& CachedKey = CachedKeys[TrackIndex];
+
+	float Alpha;
+	if (DecompContext.NeedsTwoSegments)
+	{
+		Alpha = DecompContext.KeyAlpha;
+	}
+	else
+	{
+		// compute the blend parameters for the keys we have found
+		const int32 FrameIndex0 = SegmentStartFrame0 + CachedKey.ScaleFrameIndices[0];
+		const int32 FrameIndex1 = SegmentStartFrame1 + CachedKey.ScaleFrameIndices[1];
+
+		const int32 Delta = FMath::Max(FrameIndex1 - FrameIndex0, 1);
+		const float Remainder = FramePos - float(FrameIndex0);
+		Alpha = DecompContext.AnimSeq->Interpolation == EAnimInterpolationType::Step ? 0.0f : (Remainder / float(Delta));
+	}
+
+	const uint8* KeyData0 = DecompContext.CompressedByteStream + CachedKey.ScaleOffsets[0];
+
+	FVector Scale0;
+	DecompressScale<Format, false>(Scale0, DecompContext.TrackRangeData[0], KeyData0);
+
+	const uint8* KeyData1 = DecompContext.CompressedByteStream + CachedKey.ScaleOffsets[1];
+
+	FVector Scale1;
+	DecompressScale<Format, false>(Scale1, DecompContext.TrackRangeData[1], KeyData1);
+
+	return FMath::Lerp(Scale0, Scale1, Alpha);
+}
+
+template<int32 Format>
+FQuat FAEVariableKeyLerpLinearContext::GetLinearRotation(const FAnimSequenceDecompressionContext& DecompContext, int32 TrackIndex) const
+{
+	const int32 NumTrackStreams = NumAnimatedTrackStreams[DecompContext.GetRotationValueOffset(TrackIndex)];
+
+	const uint8* OffsetNumKeysPair0 = OffsetNumKeysPairs[0] + (OffsetNumKeysPairSize * NumTrackStreams);
+	const uint16 NumKeys0 = *reinterpret_cast<const uint16*>(OffsetNumKeysPair0 + sizeof(uint32));
+
+	const uint32 KeysOffset0 = AnimationCompressionUtils::UnalignedRead<uint32>(OffsetNumKeysPair0);
+	const int32 TimeMarkersOffset0 = DecompContext.Segment0->ByteStreamOffset + KeysOffset0;
+
+	int32 FrameIndex0;
+	int32 FrameIndex1;
+	float Alpha;
+	if (DecompContext.NeedsTwoSegments)
+	{
+		FrameIndex0 = NumKeys0 - 1;
+		FrameIndex1 = 0;
+		Alpha = DecompContext.KeyAlpha;
+	}
+	else
+	{
+		const uint8* TimeMarkers0 = DecompContext.CompressedByteStream + TimeMarkersOffset0;
+
+		Alpha = AnimEncoding::TimeToIndex(DecompContext, TimeMarkers0, NumKeys0, DecompContext.Segment0->NumFrames, TimeMarkerSize[0], SegmentRelativePos0, FrameIndex0, FrameIndex1);
+	}
+
+	const int32 TrackDataOffset0 = Align(TimeMarkersOffset0 + (NumKeys0 * TimeMarkerSize[0]), 4);
+	const int32 KeyDataOffset0 = TrackDataOffset0 + (FrameIndex0 * CompressedRotationStrides[Format] * CompressedRotationNum[Format]);
+	const uint8* KeyData0 = DecompContext.CompressedByteStream + KeyDataOffset0;
+
+	FQuat Rotation;
+	DecompressRotation<Format>(Rotation, DecompContext.TrackRangeData[0], KeyData0);
+
+	if (DecompContext.NeedsInterpolation)
+	{
+		const uint8* OffsetNumKeysPair1 = OffsetNumKeysPairs[1] + (OffsetNumKeysPairSize * NumTrackStreams);
+		const uint32 KeysOffset1 = AnimationCompressionUtils::UnalignedRead<uint32>(OffsetNumKeysPair1);
+		const uint16 NumKeys1 = *reinterpret_cast<const uint16*>(OffsetNumKeysPair1 + sizeof(uint32));
+
+		const int32 TimeMarkersOffset1 = DecompContext.Segment1->ByteStreamOffset + KeysOffset1;
+		const int32 TrackDataOffset1 = Align(TimeMarkersOffset1 + (NumKeys1 * TimeMarkerSize[1]), 4);
+		const int32 KeyDataOffset1 = TrackDataOffset1 + (FrameIndex1 * CompressedRotationStrides[Format] * CompressedRotationNum[Format]);
+		const uint8* KeyData1 = DecompContext.CompressedByteStream + KeyDataOffset1;
+
+		FQuat Rotation1;
+		DecompressRotation<Format>(Rotation1, DecompContext.TrackRangeData[1], KeyData1);
+
+		// Fast linear quaternion interpolation.
+		FQuat BlendedQuat = FQuat::FastLerp(Rotation, Rotation1, Alpha);
+		BlendedQuat.Normalize();
+		Rotation = BlendedQuat;
+	}
+
+	return Rotation;
+}
+
+template<int32 Format>
+FVector FAEVariableKeyLerpLinearContext::GetLinearTranslation(const FAnimSequenceDecompressionContext& DecompContext, int32 TrackIndex) const
+{
+	const int32 NumTrackStreams = NumAnimatedTrackStreams[DecompContext.GetTranslationValueOffset(TrackIndex)];
+
+	const uint8* OffsetNumKeysPair0 = OffsetNumKeysPairs[0] + (OffsetNumKeysPairSize * NumTrackStreams);
+	const uint16 NumKeys0 = *reinterpret_cast<const uint16*>(OffsetNumKeysPair0 + sizeof(uint32));
+
+	const uint32 KeysOffset0 = AnimationCompressionUtils::UnalignedRead<uint32>(OffsetNumKeysPair0);
+	const int32 TimeMarkersOffset0 = DecompContext.Segment0->ByteStreamOffset + KeysOffset0;
+
+	int32 FrameIndex0;
+	int32 FrameIndex1;
+	float Alpha;
+	if (DecompContext.NeedsTwoSegments)
+	{
+		FrameIndex0 = NumKeys0 - 1;
+		FrameIndex1 = 0;
+		Alpha = DecompContext.KeyAlpha;
+	}
+	else
+	{
+		const uint8* TimeMarkers0 = DecompContext.CompressedByteStream + TimeMarkersOffset0;
+
+		Alpha = AnimEncoding::TimeToIndex(DecompContext, TimeMarkers0, NumKeys0, DecompContext.Segment0->NumFrames, TimeMarkerSize[0], SegmentRelativePos0, FrameIndex0, FrameIndex1);
+	}
+
+	const int32 TrackDataOffset0 = Align(TimeMarkersOffset0 + (NumKeys0 * TimeMarkerSize[0]), 4);
+	const int32 KeyDataOffset0 = TrackDataOffset0 + (FrameIndex0 * CompressedTranslationStrides[Format] * CompressedTranslationNum[Format]);
+	const uint8* KeyData0 = DecompContext.CompressedByteStream + KeyDataOffset0;
+
+	FVector Translation;
+	DecompressTranslation<Format>(Translation, DecompContext.TrackRangeData[0], KeyData0);
+
+	if (DecompContext.NeedsInterpolation)
+	{
+		const uint8* OffsetNumKeysPair1 = OffsetNumKeysPairs[1] + (OffsetNumKeysPairSize * NumTrackStreams);
+		const uint32 KeysOffset1 = AnimationCompressionUtils::UnalignedRead<uint32>(OffsetNumKeysPair1);
+		const uint16 NumKeys1 = *reinterpret_cast<const uint16*>(OffsetNumKeysPair1 + sizeof(uint32));
+
+		const int32 TimeMarkersOffset1 = DecompContext.Segment1->ByteStreamOffset + KeysOffset1;
+		const int32 TrackDataOffset1 = Align(TimeMarkersOffset1 + (NumKeys1 * TimeMarkerSize[1]), 4);
+		const int32 KeyDataOffset1 = TrackDataOffset1 + (FrameIndex1 * CompressedTranslationStrides[Format] * CompressedTranslationNum[Format]);
+		const uint8* KeyData1 = DecompContext.CompressedByteStream + KeyDataOffset1;
+
+		FVector Translation1;
+		DecompressTranslation<Format>(Translation1, DecompContext.TrackRangeData[1], KeyData1);
+
+		Translation = FMath::Lerp(Translation, Translation1, Alpha);
+	}
+
+	return Translation;
+}
+
+template<int32 Format>
+FVector FAEVariableKeyLerpLinearContext::GetLinearScale(const FAnimSequenceDecompressionContext& DecompContext, int32 TrackIndex) const
+{
+	const int32 NumTrackStreams = NumAnimatedTrackStreams[DecompContext.GetScaleValueOffset(TrackIndex)];
+
+	const uint8* OffsetNumKeysPair0 = OffsetNumKeysPairs[0] + (OffsetNumKeysPairSize * NumTrackStreams);
+	const uint16 NumKeys0 = *reinterpret_cast<const uint16*>(OffsetNumKeysPair0 + sizeof(uint32));
+
+	const uint32 KeysOffset0 = AnimationCompressionUtils::UnalignedRead<uint32>(OffsetNumKeysPair0);
+	const int32 TimeMarkersOffset0 = DecompContext.Segment0->ByteStreamOffset + KeysOffset0;
+
+	int32 FrameIndex0;
+	int32 FrameIndex1;
+	float Alpha;
+	if (DecompContext.NeedsTwoSegments)
+	{
+		FrameIndex0 = NumKeys0 - 1;
+		FrameIndex1 = 0;
+		Alpha = DecompContext.KeyAlpha;
+	}
+	else
+	{
+		const uint8* TimeMarkers0 = DecompContext.CompressedByteStream + TimeMarkersOffset0;
+
+		Alpha = AnimEncoding::TimeToIndex(DecompContext, TimeMarkers0, NumKeys0, DecompContext.Segment0->NumFrames, TimeMarkerSize[0], SegmentRelativePos0, FrameIndex0, FrameIndex1);
+	}
+
+	const int32 TrackDataOffset0 = Align(TimeMarkersOffset0 + (NumKeys0 * TimeMarkerSize[0]), 4);
+	const int32 KeyDataOffset0 = TrackDataOffset0 + (FrameIndex0 * CompressedScaleStrides[Format] * CompressedScaleNum[Format]);
+	const uint8* KeyData0 = DecompContext.CompressedByteStream + KeyDataOffset0;
+
+	FVector Scale;
+	DecompressScale<Format>(Scale, DecompContext.TrackRangeData[0], KeyData0);
+
+	if (DecompContext.NeedsInterpolation)
+	{
+		const uint8* OffsetNumKeysPair1 = OffsetNumKeysPairs[1] + (OffsetNumKeysPairSize * NumTrackStreams);
+		const uint32 KeysOffset1 = AnimationCompressionUtils::UnalignedRead<uint32>(OffsetNumKeysPair1);
+		const uint16 NumKeys1 = *reinterpret_cast<const uint16*>(OffsetNumKeysPair1 + sizeof(uint32));
+
+		const int32 TimeMarkersOffset1 = DecompContext.Segment1->ByteStreamOffset + KeysOffset1;
+		const int32 TrackDataOffset1 = Align(TimeMarkersOffset1 + (NumKeys1 * TimeMarkerSize[1]), 4);
+		const int32 KeyDataOffset1 = TrackDataOffset1 + (FrameIndex1 * CompressedScaleStrides[Format] * CompressedScaleNum[Format]);
+		const uint8* KeyData1 = DecompContext.CompressedByteStream + KeyDataOffset1;
+
+		FVector Scale1;
+		DecompressScale<Format>(Scale1, DecompContext.TrackRangeData[1], KeyData1);
+
+		Scale = FMath::Lerp(Scale, Scale1, Alpha);
+	}
+
+	return Scale;
+}
+#endif
\ No newline at end of file
diff --git a/Engine/Source/Runtime/Engine/Public/Animation/AnimEncodingDecompressionContext.h b/Engine/Source/Runtime/Engine/Public/Animation/AnimEncodingDecompressionContext.h
new file mode 100644
--- /dev/null
+++ b/Engine/Source/Runtime/Engine/Public/Animation/AnimEncodingDecompressionContext.h
@@ -0,0 +1,14 @@
+// Copyright 2018 Epic Games, Inc. All Rights Reserved.
+
+#pragma once
+
+/**
+ * Class to expose a custom per encoding decompression context and a common interface.
+ */
+class FAnimEncodingDecompressionContext
+{
+public:
+	virtual ~FAnimEncodingDecompressionContext() {}
+
+	virtual void Seek(const FAnimSequenceDecompressionContext& DecompContext, float SampleAtTime) = 0;
+};
diff --git a/Engine/Source/Runtime/Engine/Public/Animation/AnimMTStats.h b/Engine/Source/Runtime/Engine/Public/Animation/AnimMTStats.h
--- a/Engine/Source/Runtime/Engine/Public/Animation/AnimMTStats.h
+++ b/Engine/Source/Runtime/Engine/Public/Animation/AnimMTStats.h
@@ -9,3 +9,5 @@ DO_ANIMSTAT_PROCESSING(BlendPosesInGraph);
 DO_ANIMSTAT_PROCESSING(EvaluateAnimStateMachine);
 DO_ANIMSTAT_PROCESSING(FillComponentSpaceTransforms);
 DO_ANIMSTAT_PROCESSING(ProxyUpdateAnimation);
+DO_ANIMSTAT_PROCESSING(PoseBlendNodeEvaluate);
+DO_ANIMSTAT_PROCESSING(PoseAssetGetAnimationPose);
\ No newline at end of file
diff --git a/Engine/Source/Runtime/Engine/Public/Animation/AnimSequenceDecompressionContext.h b/Engine/Source/Runtime/Engine/Public/Animation/AnimSequenceDecompressionContext.h
new file mode 100644
--- /dev/null
+++ b/Engine/Source/Runtime/Engine/Public/Animation/AnimSequenceDecompressionContext.h
@@ -0,0 +1,220 @@
+// Copyright 2018 Epic Games, Inc. All Rights Reserved.
+
+#pragma once
+
+#include "CoreMinimal.h"
+#include "Animation/AnimEncodingHeapAllocator.h"
+#include "AnimationCompression.h"
+#include "Containers/StaticArray.h"
+#include "Animation/AnimSequence.h"
+
+struct FCompressedOffsetData;
+struct FCompressedSegment;
+class FAnimEncodingDecompressionContext;
+class AnimEncoding;
+
+#define USE_SEGMENTING_CONTEXT 0 // Uses segmenting in anim compression + context in decompression
+
+/**
+ * Structure to wrap a trivial key handle. Used by the decompression context and encoders.
+ */
+struct FTrivialAnimKeyHandle
+{
+	constexpr FTrivialAnimKeyHandle() : Value(-1) {}
+	constexpr explicit FTrivialAnimKeyHandle(int32 InValue) : Value(InValue) {}
+
+	constexpr bool IsValid() const { return Value >= 0; }
+
+	int32 Value;
+};
+
+#if !USE_SEGMENTING_CONTEXT
+// Lightweight decompression context to match API of when USE_SEGMENT_CONTEXT is 1
+struct ENGINE_API FAnimSequenceDecompressionContext
+{
+public:
+	FAnimSequenceDecompressionContext() : AnimSeq(nullptr), Time(0.f) {}
+	FAnimSequenceDecompressionContext(const UAnimSequence* AnimSeq_) :AnimSeq(AnimSeq_), Time(0.f) {}
+
+	/**
+	* We don't disable copying but we also don't retain any state. Copying is required by the UE serialization
+	* but retaining state is not necessary and can cause issues.
+	*/
+	FAnimSequenceDecompressionContext(const FAnimSequenceDecompressionContext& Other) : AnimSeq(Other.AnimSeq), Time(Other.Time) {}
+	FAnimSequenceDecompressionContext& operator=(const FAnimSequenceDecompressionContext& Other) { AnimSeq = Other.AnimSeq; Time = Other.Time; return *this; }
+
+	const UAnimSequence* AnimSeq;
+	float Time;
+	float RelativePos;
+	bool bHasScale;
+
+	void Seek(float SampleAtTime)
+	{
+		Time = SampleAtTime;
+		RelativePos = SampleAtTime / AnimSeq->SequenceLength;
+		bHasScale = AnimSeq->CompressedScaleOffsets.IsValid();
+	}
+
+	// Following just expose the properties of the anim sequence, 
+	AnimEncoding* GetRotationCodec() const { return AnimSeq->RotationCodec; }
+	AnimEncoding* GetTranslationCodec() const { return AnimSeq->TranslationCodec; }
+	AnimEncoding* GetScaleCodec() const { return AnimSeq->ScaleCodec; }
+
+	const int32* GetCompressedTrackOffsets() const { return AnimSeq->CompressedTrackOffsets.GetData(); }
+	const uint8* GetCompressedByteStream() const { return AnimSeq->CompressedByteStream.GetData(); }
+	const FCompressedOffsetData* GetCompressedScaleOffsets() const { return &AnimSeq->CompressedScaleOffsets; }
+};
+
+#else
+
+/**
+ * Structure to hold the data required for decompression.
+ * It is created from an UAnimSequence and stores all the intermediary
+ * information required. A context is bound to a single sequence. Multiple
+ * context instances can be bound to the same sequence.
+ * In order to re-use a context with a new sequence, it needs to be bound.
+ */
+struct ENGINE_API FAnimSequenceDecompressionContext
+{
+	FAnimSequenceDecompressionContext();
+	FAnimSequenceDecompressionContext(const UAnimSequence* AnimSeq_);
+	~FAnimSequenceDecompressionContext();
+
+	/**
+	 * We don't disable copying but we also don't retain any state. Copying is required by the UE serialization
+	 * but retaining state is not necessary and can cause issues.
+	 */
+	FAnimSequenceDecompressionContext(const FAnimSequenceDecompressionContext& Other);
+	FAnimSequenceDecompressionContext& operator=(const FAnimSequenceDecompressionContext& Other);
+
+	/**
+	 * Seek into an animation sequence at a particular time.
+	 */
+	void Seek(float SampleAtTime);
+
+	/**
+	 * Returns whether this context is bound to the provided anim sequence.
+	 */
+	bool IsStale(const UAnimSequence* AnimSeq_) const;
+
+	/**
+	 * Binds the context to the provided anim sequence. The sequence can be NULL
+	 * in which case the context will be reset.
+	 */
+	void Bind(const UAnimSequence* AnimSeq_);
+
+	inline FTrivialAnimKeyHandle GetTrivialRotationKeyHandle(int32 TrackIndex) const;
+	inline FTrivialAnimKeyHandle GetTrivialTranslationKeyHandle(int32 TrackIndex) const;
+	inline FTrivialAnimKeyHandle GetTrivialScaleKeyHandle(int32 TrackIndex) const;
+	FORCEINLINE void GetTrivialRotation(FTransform& OutAtom, FTrivialAnimKeyHandle KeyHandle) const;
+	FORCEINLINE void GetTrivialTranslation(FTransform& OutAtom, FTrivialAnimKeyHandle KeyHandle) const;
+	FORCEINLINE void GetTrivialScale(FTransform& OutAtom, FTrivialAnimKeyHandle KeyHandle) const;
+
+	inline int32 GetRotationValueOffset(int32 TrackIndex) const { return (TrackIndex * NumStreamsPerTrack) + 1; }
+	inline int32 GetTranslationValueOffset(int32 TrackIndex) const { return (TrackIndex * NumStreamsPerTrack) + 0; }
+	inline int32 GetScaleValueOffset(int32 TrackIndex) const { return (TrackIndex * NumStreamsPerTrack) + 2; }
+
+	AnimEncoding* GetRotationCodec() const { return RotationCodec; }
+	AnimEncoding* GetTranslationCodec() const { return TranslationCodec; }
+	AnimEncoding* GetScaleCodec() const { return ScaleCodec; }
+
+	const int32* GetCompressedTrackOffsets() const { return CompressedTrackOffsets; }
+	const uint8* GetCompressedByteStream() const { return CompressedByteStream; }
+	const FCompressedOffsetData* GetCompressedScaleOffsets() const { return CompressedScaleOffsets; }
+
+	const UAnimSequence* AnimSeq;
+
+	const int32* CompressedTrackOffsets;
+	const uint8* CompressedByteStream;
+	const FCompressedOffsetData* CompressedScaleOffsets;
+
+	float Time;
+	float RelativePos;
+
+	bool bHasScale;
+	bool NeedsTwoSegments;
+	bool NeedsInterpolation;
+	bool HasSegments;
+	bool bIsSorted;
+
+	uint8 NumStreamsPerTrack;
+
+	int32 NumTracks;
+	uint16 SegmentIndex0;
+	uint16 SegmentIndex1;
+
+	int32 KeyIndex0;
+	int32 KeyIndex1;
+	int32 SegmentKeyIndex0;
+	int32 SegmentKeyIndex1;
+	float KeyAlpha;
+
+	int32 PackedTranslationSize0;
+	int32 PackedRotationSize0;
+	int32 PackedScaleSize0;
+
+	int32 RangeDataSize0;
+
+	const uint8* TrackFlags;
+	const uint8* TrivialTrackKeys;
+	const uint8* TrackRangeData[2];
+
+	const FCompressedSegment* Segment0;
+	const FCompressedSegment* Segment1;
+
+	FAnimEncodingDecompressionContext* EncodingContext;
+	TArray<int32, FAnimEncodingHeapAllocator> TrivialTrackStreamOffsets;
+	uint32 SequenceCRC;
+
+#if WITH_EDITOR
+	TStaticArray<double, 4> PreviousBindTimeStamps;
+#endif
+
+private:
+	AnimEncoding * RotationCodec;
+	AnimEncoding* TranslationCodec;
+	AnimEncoding* ScaleCodec;
+};
+
+FTrivialAnimKeyHandle FAnimSequenceDecompressionContext::GetTrivialRotationKeyHandle(int32 TrackIndex) const
+{
+	const int32 TrivialKeyOffset = TrivialTrackStreamOffsets[GetRotationValueOffset(TrackIndex)];
+	return FTrivialAnimKeyHandle(TrivialKeyOffset);
+}
+
+FTrivialAnimKeyHandle FAnimSequenceDecompressionContext::GetTrivialTranslationKeyHandle(int32 TrackIndex) const
+{
+	const int32 TrivialKeyOffset = TrivialTrackStreamOffsets[GetTranslationValueOffset(TrackIndex)];
+	return FTrivialAnimKeyHandle(TrivialKeyOffset);
+}
+
+FTrivialAnimKeyHandle FAnimSequenceDecompressionContext::GetTrivialScaleKeyHandle(int32 TrackIndex) const
+{
+	const int32 TrivialKeyOffset = TrivialTrackStreamOffsets[GetScaleValueOffset(TrackIndex)];
+	return FTrivialAnimKeyHandle(TrivialKeyOffset);
+}
+
+FORCEINLINE void FAnimSequenceDecompressionContext::GetTrivialRotation(FTransform& OutAtom, FTrivialAnimKeyHandle KeyHandle) const
+{
+	const uint8* TrivialKeyData = TrivialTrackKeys + KeyHandle.Value;
+	FQuat R0;
+	DecompressRotation<ACF_Float96NoW>(R0, TrivialKeyData, TrivialKeyData);
+	OutAtom.SetRotation(R0);
+}
+
+FORCEINLINE void FAnimSequenceDecompressionContext::GetTrivialTranslation(FTransform& OutAtom, FTrivialAnimKeyHandle KeyHandle) const
+{
+	const uint8* TrivialKeyData = TrivialTrackKeys + KeyHandle.Value;
+	FVector P0;
+	DecompressTranslation<ACF_None>(P0, TrivialKeyData, TrivialKeyData);
+	OutAtom.SetTranslation(P0);
+}
+
+FORCEINLINE void FAnimSequenceDecompressionContext::GetTrivialScale(FTransform& OutAtom, FTrivialAnimKeyHandle KeyHandle) const
+{
+	const uint8* TrivialKeyData = TrivialTrackKeys + KeyHandle.Value;
+	FVector P0;
+	DecompressScale<ACF_None>(P0, TrivialKeyData, TrivialKeyData);
+	OutAtom.SetScale3D(P0);
+}
+#endif
\ No newline at end of file
diff --git a/Engine/Source/Runtime/Engine/Public/Animation/AnimSingleNodeInstanceProxy.h b/Engine/Source/Runtime/Engine/Public/Animation/AnimSingleNodeInstanceProxy.h
--- a/Engine/Source/Runtime/Engine/Public/Animation/AnimSingleNodeInstanceProxy.h
+++ b/Engine/Source/Runtime/Engine/Public/Animation/AnimSingleNodeInstanceProxy.h
@@ -8,6 +8,7 @@
 #include "Animation/AnimationAsset.h"
 #include "Animation/AnimInstanceProxy.h"
 #include "Animation/AnimNodeBase.h"
+#include "Animation/AnimSequenceDecompressionContext.h"
 #include "AnimSingleNodeInstanceProxy.generated.h"
 
 struct FAnimSingleNodeInstanceProxy;
@@ -75,6 +76,8 @@ public:
 #endif
 	}
 
+	virtual ~FAnimSingleNodeInstanceProxy();
+
 	virtual void Initialize(UAnimInstance* InAnimInstance) override;
 	virtual bool Evaluate(FPoseContext& Output) override;
 	virtual void UpdateAnimationNode(float DeltaSeconds) override;
diff --git a/Engine/Source/Runtime/Engine/Public/AnimationCompression.h b/Engine/Source/Runtime/Engine/Public/AnimationCompression.h
--- a/Engine/Source/Runtime/Engine/Public/AnimationCompression.h
+++ b/Engine/Source/Runtime/Engine/Public/AnimationCompression.h
@@ -7,8 +7,9 @@
 #pragma once
 
 #include "CoreMinimal.h"
-#include "Animation/AnimSequence.h"
 #include "Math/FloatPacker.h"
+#include "Animation/AnimEnums.h"
+#include "EngineLogs.h"
 
 DECLARE_LOG_CATEGORY_EXTERN(LogAnimationCompression, Log, All);
 
@@ -34,6 +35,27 @@ DECLARE_LOG_CATEGORY_EXTERN(LogAnimationCompression, Log, All);
 #define Quant11BitFactor  (1023.f)
 #define Quant11BitOffs    (1023)
 
+//@TODO: Explore different scales
+static constexpr int32 LogScale = 7;
+
+namespace AnimationCompressionUtils
+{
+	template<typename ValueType>
+	inline ValueType UnalignedRead(const void* Ptr)
+	{
+#if PLATFORM_SUPPORTS_UNALIGNED_INT_LOADS
+		return *reinterpret_cast<const ValueType*>(Ptr);
+#else
+		// TODO: On ARM devices this will be slower than it needs to be.
+		// To make it fast, __packed keyword must be used with the reinterpret cast.
+		// Ideally this code would be moved into FPlatformMisc to handle this per platform properly.
+		ValueType Result;
+		memcpy(&Result, Ptr, sizeof(ValueType));
+		return Result;
+#endif
+	}
+}
+
 class FQuatFixed48NoW
 {
 public:
@@ -123,7 +145,8 @@ public:
 		Packed = (PackedX << XShift) | (PackedY << YShift) | (PackedZ);
 	}
 
-	void ToQuat(FQuat& Out) const
+	template<bool bIsDataAligned>
+	static FQuat ToQuat(const uint32* PackedValue)
 	{
 		const uint32 XShift = 21;
 		const uint32 YShift = 10;
@@ -131,6 +154,7 @@ public:
 		const uint32 YMask = 0x001ffc00;
 		const uint32 XMask = 0xffe00000;
 
+		const uint32 Packed = bIsDataAligned ? *PackedValue : AnimationCompressionUtils::UnalignedRead<uint32>(PackedValue);
 		const uint32 UnpackedX = Packed >> XShift;
 		const uint32 UnpackedY = (Packed & YMask) >> YShift;
 		const uint32 UnpackedZ = (Packed & ZMask);
@@ -140,10 +164,12 @@ public:
 		const float Z = ((int32)UnpackedZ - (int32)Quant10BitOffs) / Quant10BitDiv;
 		const float WSquared = 1.f - X*X - Y*Y - Z*Z;
 
-		Out.X = X;
-		Out.Y = Y;
-		Out.Z = Z;
-		Out.W = WSquared > 0.f ? FMath::Sqrt( WSquared ) : 0.f;
+		return FQuat(X, Y, Z, WSquared > 0.f ? FMath::Sqrt(WSquared) : 0.f);
+	}
+
+	void ToQuat(FQuat& Out) const
+	{
+		Out = ToQuat<true>(&Packed);
 	}
 
 	friend FArchive& operator<<(FArchive& Ar, FQuatFixed32NoW& Quat)
@@ -190,14 +216,20 @@ public:
 		Z = Temp.Z;
 	}
 
-	void ToQuat(FQuat& Out) const
+	template<bool bIsDataAligned = true>
+	static FQuat ToQuat(const float* Values)
 	{
+		const float X = bIsDataAligned ? *Values++ : AnimationCompressionUtils::UnalignedRead<float>(Values++);
+		const float Y = bIsDataAligned ? *Values++ : AnimationCompressionUtils::UnalignedRead<float>(Values++);
+		const float Z = bIsDataAligned ? *Values++ : AnimationCompressionUtils::UnalignedRead<float>(Values++);
 		const float WSquared = 1.f - X*X - Y*Y - Z*Z;
 
-		Out.X = X;
-		Out.Y = Y;
-		Out.Z = Z;
-		Out.W = WSquared > 0.f ? FMath::Sqrt( WSquared ) : 0.f;
+		return FQuat(X, Y, Z, WSquared > 0.f ? FMath::Sqrt(WSquared) : 0.f);
+	}
+
+	void ToQuat(FQuat& Out) const
+	{
+		Out = ToQuat<true>(&X);
 	}
 
 	friend FArchive& operator<<(FArchive& Ar, FQuatFloat96NoW& Quat)
@@ -286,7 +318,8 @@ public:
 		Packed = (PackedZ << ZShift) | (PackedY << YShift) | (PackedX);
 	}
 
-	void ToVector(FVector& Out, const float* Mins, const float *Ranges) const
+	template<bool bIsDataAligned = true>
+	static FVector ToVector(const float* Mins, const float *Ranges, const uint32* PackedValue)
 	{
 		const uint32 ZShift = 21;
 		const uint32 YShift = 10;
@@ -294,17 +327,21 @@ public:
 		const uint32 YMask = 0x001ffc00;
 		const uint32 ZMask = 0xffe00000;
 
+		const uint32 Packed = bIsDataAligned ? *PackedValue : AnimationCompressionUtils::UnalignedRead<uint32>(PackedValue);
 		const uint32 UnpackedZ = Packed >> ZShift;
 		const uint32 UnpackedY = (Packed & YMask) >> YShift;
 		const uint32 UnpackedX = (Packed & XMask);
 
-		const float X = ( (((int32)UnpackedX - (int32)Quant10BitOffs) / Quant10BitDiv) * Ranges[0] + Mins[0] );
-		const float Y = ( (((int32)UnpackedY - (int32)Quant11BitOffs) / Quant11BitDiv) * Ranges[1] + Mins[1] );
-		const float Z = ( (((int32)UnpackedZ - (int32)Quant11BitOffs) / Quant11BitDiv) * Ranges[2] + Mins[2] );
+		const float X = ((((int32)UnpackedX - (int32)Quant10BitOffs) / Quant10BitDiv) * Ranges[0] + Mins[0]);
+		const float Y = ((((int32)UnpackedY - (int32)Quant11BitOffs) / Quant11BitDiv) * Ranges[1] + Mins[1]);
+		const float Z = ((((int32)UnpackedZ - (int32)Quant11BitOffs) / Quant11BitDiv) * Ranges[2] + Mins[2]);
+
+		return FVector(X, Y, Z);
+	}
 
-		Out.X = X;
-		Out.Y = Y;
-		Out.Z = Z;
+	void ToVector(FVector& Out, const float* Mins, const float *Ranges) const
+	{
+		Out = ToVector<true>(Mins, Ranges, &Packed);
 	}
 
 	friend FArchive& operator<<(FArchive& Ar, FVectorIntervalFixed32NoW& Value)
@@ -354,7 +391,8 @@ public:
 		Packed = (PackedX << XShift) | (PackedY << YShift) | (PackedZ);
 	}
 
-	void ToQuat(FQuat& Out, const float* Mins, const float *Ranges) const
+	template<bool bIsDataAligned>
+	static FQuat ToQuat(const float* Mins, const float *Ranges, const uint32* PackedValue)
 	{
 		const uint32 XShift = 21;
 		const uint32 YShift = 10;
@@ -362,19 +400,22 @@ public:
 		const uint32 YMask = 0x001ffc00;
 		const uint32 XMask = 0xffe00000;
 
+		const uint32 Packed = bIsDataAligned ? *PackedValue : AnimationCompressionUtils::UnalignedRead<uint32>(PackedValue);
 		const uint32 UnpackedX = Packed >> XShift;
 		const uint32 UnpackedY = (Packed & YMask) >> YShift;
 		const uint32 UnpackedZ = (Packed & ZMask);
 
-		const float X = ( (((int32)UnpackedX - (int32)Quant11BitOffs) / Quant11BitDiv) * Ranges[0] + Mins[0] );
-		const float Y = ( (((int32)UnpackedY - (int32)Quant11BitOffs) / Quant11BitDiv) * Ranges[1] + Mins[1] );
-		const float Z = ( (((int32)UnpackedZ - (int32)Quant10BitOffs) / Quant10BitDiv) * Ranges[2] + Mins[2] );
+		const float X = ((((int32)UnpackedX - (int32)Quant11BitOffs) / Quant11BitDiv) * Ranges[0] + Mins[0]);
+		const float Y = ((((int32)UnpackedY - (int32)Quant11BitOffs) / Quant11BitDiv) * Ranges[1] + Mins[1]);
+		const float Z = ((((int32)UnpackedZ - (int32)Quant10BitOffs) / Quant10BitDiv) * Ranges[2] + Mins[2]);
 		const float WSquared = 1.f - X*X - Y*Y - Z*Z;
 
-		Out.X = X;
-		Out.Y = Y;
-		Out.Z = Z;
-		Out.W = WSquared > 0.f ? FMath::Sqrt( WSquared ) : 0.f;
+		return FQuat(X, Y, Z, WSquared > 0.f ? FMath::Sqrt(WSquared) : 0.f);
+	}
+
+	void ToQuat(FQuat& Out, const float* Mins, const float *Ranges) const
+	{
+		Out = ToQuat<true>(Mins, Ranges, &Packed);
 	}
 
 	friend FArchive& operator<<(FArchive& Ar, FQuatIntervalFixed32NoW& Quat)
@@ -422,7 +463,8 @@ public:
 		Packed = (PackedX << XShift) | (PackedY << YShift) | (PackedZ);
 	}
 
-	void ToQuat(FQuat& Out) const
+	template<bool bIsDataAligned>
+	static FQuat ToQuat(const uint32* PackedValue)
 	{
 		const uint32 XShift = 21;
 		const uint32 YShift = 10;
@@ -430,6 +472,7 @@ public:
 		const uint32 YMask = 0x001ffc00;
 		const uint32 XMask = 0xffe00000;
 
+		const uint32 Packed = bIsDataAligned ? *PackedValue : AnimationCompressionUtils::UnalignedRead<uint32>(PackedValue);
 		const uint32 UnpackedX = Packed >> XShift;
 		const uint32 UnpackedY = (Packed & YMask) >> YShift;
 		const uint32 UnpackedZ = (Packed & ZMask);
@@ -437,15 +480,17 @@ public:
 		TFloatPacker<3, 7, true> Packer7e3;
 		TFloatPacker<3, 6, true> Packer6e3;
 
-		const float X = Packer7e3.Decode( UnpackedX );
-		const float Y = Packer7e3.Decode( UnpackedY );
-		const float Z = Packer6e3.Decode( UnpackedZ );
+		const float X = Packer7e3.Decode(UnpackedX);
+		const float Y = Packer7e3.Decode(UnpackedY);
+		const float Z = Packer6e3.Decode(UnpackedZ);
 		const float WSquared = 1.f - X*X - Y*Y - Z*Z;
 
-		Out.X = X;
-		Out.Y = Y;
-		Out.Z = Z;
-		Out.W = WSquared > 0.f ? FMath::Sqrt( WSquared ) : 0.f;
+		return FQuat(X, Y, Z, WSquared > 0.f ? FMath::Sqrt(WSquared) : 0.f);
+	}
+
+	void ToQuat(FQuat& Out) const
+	{
+		Out = ToQuat<true>(&Packed);
 	}
 
 	friend FArchive& operator<<(FArchive& Ar, FQuatFloat32NoW& Quat)
@@ -469,7 +514,7 @@ public:
  * @param	KeyData			The compressed rotation data to decompress.
  * @return	None. 
  */
-template <int32 FORMAT>
+template <int32 FORMAT, bool bIsDataAligned = true>
 FORCEINLINE void DecompressRotation(FQuat& Out, const uint8* RESTRICT TopOfStream, const uint8* RESTRICT KeyData)
 {
 	// this if-else stack gets compiled away to a single result based on the template parameter
@@ -477,16 +522,20 @@ FORCEINLINE void DecompressRotation(FQuat& Out, const uint8* RESTRICT TopOfStrea
 	{
 		// due to alignment issue, this crahses accessing non aligned 
 		// we don't think this is common case, so this is slower. 
-		float* Keys = (float*)KeyData;
-		Out = FQuat(Keys[0], Keys[1], Keys[2], Keys[3]);
+		const float* Keys = (const float*)KeyData;
+		const float X = bIsDataAligned ? Keys[0] : AnimationCompressionUtils::UnalignedRead<float>(&Keys[0]);
+		const float Y = bIsDataAligned ? Keys[1] : AnimationCompressionUtils::UnalignedRead<float>(&Keys[1]);
+		const float Z = bIsDataAligned ? Keys[2] : AnimationCompressionUtils::UnalignedRead<float>(&Keys[2]);
+		const float W = bIsDataAligned ? Keys[3] : AnimationCompressionUtils::UnalignedRead<float>(&Keys[3]);
+		Out = FQuat(X, Y, Z, W);
 	}
 	else if ( FORMAT == ACF_Float96NoW )
 	{
-		((FQuatFloat96NoW*)KeyData)->ToQuat( Out );
+		Out = FQuatFloat96NoW::ToQuat<bIsDataAligned>((const float*)KeyData);
 	}
 	else if ( FORMAT == ACF_Fixed32NoW )
 	{
-		((FQuatFixed32NoW*)KeyData)->ToQuat( Out );
+		Out = FQuatFixed32NoW::ToQuat<bIsDataAligned>((const uint32*)KeyData);
 	}
 	else if ( FORMAT == ACF_Fixed48NoW )
 	{
@@ -496,11 +545,11 @@ FORCEINLINE void DecompressRotation(FQuat& Out, const uint8* RESTRICT TopOfStrea
 	{
 		const float* RESTRICT Mins = (float*)TopOfStream;
 		const float* RESTRICT Ranges = (float*)(TopOfStream+sizeof(float)*3);
-		((FQuatIntervalFixed32NoW*)KeyData)->ToQuat( Out, Mins, Ranges );
+		Out = FQuatIntervalFixed32NoW::ToQuat<bIsDataAligned>(Mins, Ranges, (const uint32*)KeyData);
 	}
 	else if ( FORMAT == ACF_Float32NoW )
 	{
-		((FQuatFloat32NoW*)KeyData)->ToQuat( Out );
+		Out = FQuatFloat32NoW::ToQuat<bIsDataAligned>((const uint32*)KeyData);
 	}
 	else if ( FORMAT == ACF_Identity )
 	{
@@ -521,18 +570,18 @@ FORCEINLINE void DecompressRotation(FQuat& Out, const uint8* RESTRICT TopOfStrea
  * @param	KeyData			The compressed translation data to decompress.
  * @return	None. 
  */
-template <int32 FORMAT>
+template <int32 FORMAT, bool bIsDataAligned = true>
 FORCEINLINE void DecompressTranslation(FVector& Out, const uint8* RESTRICT TopOfStream, const uint8* RESTRICT KeyData)
 {
 	if ( (FORMAT == ACF_None) || (FORMAT == ACF_Float96NoW) )
 	{
-		Out = *((FVector*)KeyData);
+		Out = bIsDataAligned ? *((FVector*)KeyData) : AnimationCompressionUtils::UnalignedRead<FVector>(KeyData);
 	}
 	else if ( FORMAT == ACF_IntervalFixed32NoW )
 	{
 		const float* RESTRICT Mins = (float*)TopOfStream;
 		const float* RESTRICT Ranges = (float*)(TopOfStream+sizeof(float)*3);
-		((FVectorIntervalFixed32NoW*)KeyData)->ToVector( Out, Mins, Ranges );
+		Out = FVectorIntervalFixed32NoW::ToVector<bIsDataAligned>(Mins, Ranges, (const uint32*)KeyData);
 	}
 	else if ( FORMAT == ACF_Identity )
 	{
@@ -558,18 +607,18 @@ FORCEINLINE void DecompressTranslation(FVector& Out, const uint8* RESTRICT TopOf
  * @param	KeyData			The compressed Scale data to decompress.
  * @return	None. 
  */
-template <int32 FORMAT>
+template <int32 FORMAT, bool bIsDataAligned = true>
 FORCEINLINE void DecompressScale(FVector& Out, const uint8* RESTRICT TopOfStream, const uint8* RESTRICT KeyData)
 {
 	if ( (FORMAT == ACF_None) || (FORMAT == ACF_Float96NoW) )
 	{
-		Out = *((FVector*)KeyData);
+		Out = bIsDataAligned ? *((FVector*)KeyData) : AnimationCompressionUtils::UnalignedRead<FVector>(KeyData);
 	}
 	else if ( FORMAT == ACF_IntervalFixed32NoW )
 	{
 		const float* RESTRICT Mins = (float*)TopOfStream;
 		const float* RESTRICT Ranges = (float*)(TopOfStream+sizeof(float)*3);
-		((FVectorIntervalFixed32NoW*)KeyData)->ToVector( Out, Mins, Ranges );
+		Out = FVectorIntervalFixed32NoW::ToVector<bIsDataAligned>(Mins, Ranges, (const uint32*)KeyData);
 	}
 	else if ( FORMAT == ACF_Identity )
 	{
@@ -738,6 +787,7 @@ public:
 	}
 
 	/** Decompress a single translation key from a single track that was compressed with the PerTrack codec (scalar) */
+	template<bool bIsDataAligned = true>
 	static FORCEINLINE_DEBUGGABLE void DecompressTranslation(int32 Format, int32 FormatFlags, FVector& Out, const uint8* RESTRICT TopOfStream, const uint8* RESTRICT KeyData)
 	{
 		if( Format == ACF_Float96NoW )
@@ -745,7 +795,7 @@ public:
 			// Legacy Format, all components stored
 			if( (FormatFlags & 7) == 0 )
 			{
-				Out = *((FVector*)KeyData);
+				Out = bIsDataAligned ? *((FVector*)KeyData) : AnimationCompressionUtils::UnalignedRead<FVector>(KeyData);
 			}
 			// Stored per components
 			else
@@ -753,7 +803,7 @@ public:
 				const float* RESTRICT TypedKeyData = (const float*)KeyData;
 				if (FormatFlags & 1)
 				{
-					Out.X = *TypedKeyData++;
+					Out.X = bIsDataAligned ? *TypedKeyData++ : AnimationCompressionUtils::UnalignedRead<float>(TypedKeyData++);
 				}
 				else
 				{
@@ -762,7 +812,7 @@ public:
 
 				if (FormatFlags & 2)
 				{
-					Out.Y = *TypedKeyData++;
+					Out.Y = bIsDataAligned ? *TypedKeyData++ : AnimationCompressionUtils::UnalignedRead<float>(TypedKeyData++);
 				}
 				else
 				{
@@ -771,7 +821,7 @@ public:
 
 				if (FormatFlags & 4)
 				{
-					Out.Z = *TypedKeyData++;
+					Out.Z = bIsDataAligned ? *TypedKeyData++ : AnimationCompressionUtils::UnalignedRead<float>(TypedKeyData++);
 				}
 				else
 				{
@@ -802,12 +852,10 @@ public:
 				Ranges[2] = *SourceBounds++;
 			}
 
-			((FVectorIntervalFixed32NoW*)KeyData)->ToVector(Out, Mins, Ranges);
+			Out = FVectorIntervalFixed32NoW::ToVector<bIsDataAligned>(Mins, Ranges, (const uint32*)KeyData);
 		}
 		else if (Format == ACF_Fixed48NoW)
 		{
-			//@TODO: Explore different scales
-			const int32 LogScale = 7;
 			const uint16* RESTRICT TypedKeyData = (const uint16*)KeyData;
 			if (FormatFlags & 1)
 			{
@@ -849,15 +897,16 @@ public:
 	}
 
 	/** Decompress a single rotation key from a single track that was compressed with the PerTrack codec (scalar) */
+	template<bool bIsDataAligned = true>
 	static FORCEINLINE_DEBUGGABLE void DecompressRotation(int32 Format, int32 FormatFlags, FQuat& Out, const uint8* RESTRICT TopOfStream, const uint8* RESTRICT KeyData)
 	{
 		if (Format == ACF_Fixed48NoW)
 		{
 			const uint16* RESTRICT TypedKeyData = (const uint16*)KeyData;
 
-			const float Xa = (FormatFlags & 1) ? (float)(*TypedKeyData++) : 32767.0f;
-			const float Ya = (FormatFlags & 2) ? (float)(*TypedKeyData++) : 32767.0f;
-			const float Za = (FormatFlags & 4) ? (float)(*TypedKeyData++) : 32767.0f;
+			const float Xa = (FormatFlags & 1) ? (float)(bIsDataAligned ? (*TypedKeyData++) : AnimationCompressionUtils::UnalignedRead<uint16>(TypedKeyData++)) : 32767.0f;
+			const float Ya = (FormatFlags & 2) ? (float)(bIsDataAligned ? (*TypedKeyData++) : AnimationCompressionUtils::UnalignedRead<uint16>(TypedKeyData++)) : 32767.0f;
+			const float Za = (FormatFlags & 4) ? (float)(bIsDataAligned ? (*TypedKeyData++) : AnimationCompressionUtils::UnalignedRead<uint16>(TypedKeyData++)) : 32767.0f;
 
 			const float X = (Xa - 32767.0f) * 3.0518509475997192297128208258309e-5f;
 			const float XX = X*X;
@@ -874,7 +923,7 @@ public:
 		}
 		else if (Format == ACF_Float96NoW)
 		{
-			((FQuatFloat96NoW*)KeyData)->ToQuat(Out);
+			Out = FQuatFloat96NoW::ToQuat<bIsDataAligned>((const float*)KeyData);
 		}
 		else if ( Format == ACF_IntervalFixed32NoW )
 		{
@@ -899,15 +948,15 @@ public:
 				Ranges[2] = *SourceBounds++;
 			}
 
-			((FQuatIntervalFixed32NoW*)KeyData)->ToQuat( Out, Mins, Ranges );
+			Out = FQuatIntervalFixed32NoW::ToQuat<bIsDataAligned>(Mins, Ranges, (const uint32*)KeyData);
 		}
 		else if ( Format == ACF_Float32NoW )
 		{
-			((FQuatFloat32NoW*)KeyData)->ToQuat( Out );
+			Out = FQuatFloat32NoW::ToQuat<bIsDataAligned>((const uint32*)KeyData);
 		}
 		else if (Format == ACF_Fixed32NoW)
 		{
-			((FQuatFixed32NoW*)KeyData)->ToQuat(Out);
+			Out = FQuatFixed32NoW::ToQuat<bIsDataAligned>((const uint32*)KeyData);
 		}
 		else if ( Format == ACF_Identity )
 		{
@@ -921,6 +970,7 @@ public:
 	}
 
 	/** Decompress a single Scale key from a single track that was compressed with the PerTrack codec (scalar) */
+	template<bool bIsDataAligned = true>
 	static FORCEINLINE_DEBUGGABLE void DecompressScale(int32 Format, int32 FormatFlags, FVector& Out, const uint8* RESTRICT TopOfStream, const uint8* RESTRICT KeyData)
 	{
 		if( Format == ACF_Float96NoW )
@@ -928,7 +978,7 @@ public:
 			// Legacy Format, all components stored
 			if( (FormatFlags & 7) == 0 )
 			{
-				Out = *((FVector*)KeyData);
+				Out = bIsDataAligned ? *((FVector*)KeyData) : AnimationCompressionUtils::UnalignedRead<FVector>(KeyData);
 			}
 			// Stored per components
 			else
@@ -936,7 +986,7 @@ public:
 				const float* RESTRICT TypedKeyData = (const float*)KeyData;
 				if (FormatFlags & 1)
 				{
-					Out.X = *TypedKeyData++;
+					Out.X = bIsDataAligned ? *TypedKeyData++ : AnimationCompressionUtils::UnalignedRead<float>(TypedKeyData++);
 				}
 				else
 				{
@@ -945,7 +995,7 @@ public:
 
 				if (FormatFlags & 2)
 				{
-					Out.Y = *TypedKeyData++;
+					Out.Y = bIsDataAligned ? *TypedKeyData++ : AnimationCompressionUtils::UnalignedRead<float>(TypedKeyData++);
 				}
 				else
 				{
@@ -954,7 +1004,7 @@ public:
 
 				if (FormatFlags & 4)
 				{
-					Out.Z = *TypedKeyData++;
+					Out.Z = bIsDataAligned ? *TypedKeyData++ : AnimationCompressionUtils::UnalignedRead<float>(TypedKeyData++);
 				}
 				else
 				{
@@ -985,12 +1035,10 @@ public:
 				Ranges[2] = *SourceBounds++;
 			}
 
-			((FVectorIntervalFixed32NoW*)KeyData)->ToVector(Out, Mins, Ranges);
+			Out = FVectorIntervalFixed32NoW::ToVector<bIsDataAligned>(Mins, Ranges, (const uint32*)KeyData);
 		}
 		else if (Format == ACF_Fixed48NoW)
 		{
-			//@TODO: Explore different scales
-			const int32 LogScale = 7;
 			const uint16* RESTRICT TypedKeyData = (const uint16*)KeyData;
 			if (FormatFlags & 1)
 			{
@@ -1030,7 +1078,6 @@ public:
 			Out = FVector::ZeroVector;
 		}
 	}
-
 };
 
 template <>
diff --git a/Engine/Source/Runtime/Engine/Public/AnimationRuntime.h b/Engine/Source/Runtime/Engine/Public/AnimationRuntime.h
--- a/Engine/Source/Runtime/Engine/Public/AnimationRuntime.h
+++ b/Engine/Source/Runtime/Engine/Public/AnimationRuntime.h
@@ -12,6 +12,7 @@
 #include "Animation/AnimTypes.h"
 #include "Animation/AnimationAsset.h"
 #include "Animation/AnimCurveTypes.h"
+#include "Animation/AnimSequenceBase.h"
 #include "Components/SkinnedMeshComponent.h"
 #include "BonePose.h"
 #include "Containers/ArrayView.h"
@@ -344,7 +345,7 @@ public:
 	 *
 	 * return ETypeAdvanceAnim type
 	 */
-	static enum ETypeAdvanceAnim AdvanceTime(const bool& bAllowLooping, const float& MoveDelta, float& InOutTime, const float& EndTime);
+	static ETypeAdvanceAnim AdvanceTime(const bool& bAllowLooping, const float& MoveDelta, float& InOutTime, const float& EndTime);
 
 	static void TickBlendWeight(float DeltaTime, float DesiredWeight, float& Weight, float& BlendTime);
 	/** 
diff --git a/Engine/Source/Runtime/Engine/Public/AnimationUtils.h b/Engine/Source/Runtime/Engine/Public/AnimationUtils.h
--- a/Engine/Source/Runtime/Engine/Public/AnimationUtils.h
+++ b/Engine/Source/Runtime/Engine/Public/AnimationUtils.h
@@ -167,10 +167,7 @@ public:
 	 * @param	bFirstRecompressUsingCurrentOrDefault	If true, then the animation will be first recompressed with it's current compressor if non-NULL, or with the global default compressor (specified in the engine ini)
 	 * @param	bForceBelowThreshold					If true and the existing compression error is greater than MasterTolerance, then any compression technique (even one that increases the size) with a lower error will be used until it falls below the threshold
 	 * @param	bRaiseMaxErrorToExisting				If true and the existing compression error is greater than MasterTolerance, then MasterTolerance will be effectively raised to the existing error level
-	 * @param	bTryFixedBitwiseCompression				If true, the uniform bitwise techniques will be tried
-	 * @param	bTryPerTrackBitwiseCompression			If true, the per-track compressor techniques will be tried
-	 * @param	bTryLinearKeyRemovalCompression			If true, the linear key removal techniques will be tried
-	 * @param	bTryIntervalKeyRemoval					If true, the resampling techniques will be tried
+	 * @param	bTryExhaustiveSearch					If true, then an exhaustive search is used otherwise only a short list of the best methods is tried
 	 *
 	 * @return	None.
 	 */
@@ -181,10 +178,11 @@ public:
 		bool bFirstRecompressUsingCurrentOrDefault,
 		bool bForceBelowThreshold,
 		bool bRaiseMaxErrorToExisting,
-		bool bTryFixedBitwiseCompression,
-		bool bTryPerTrackBitwiseCompression,
-		bool bTryLinearKeyRemovalCompression,
-		bool bTryIntervalKeyRemoval);
+		bool bTryExhaustiveSearch,
+		bool bEnableSegmenting,
+		int32 IdealNumFramesPerSegment,
+		int32 MaxNumFramesPerSegment,
+		const TArray<FBoneData>& BoneData);
 
 	/**
 	 * Tests for a missing or invalid mesh on the animation sequence, warning if one or more were found
diff --git a/Engine/Source/Runtime/Engine/Public/BonePose.h b/Engine/Source/Runtime/Engine/Public/BonePose.h
--- a/Engine/Source/Runtime/Engine/Public/BonePose.h
+++ b/Engine/Source/Runtime/Engine/Public/BonePose.h
@@ -7,9 +7,10 @@
 #include "BoneIndices.h"
 #include "Animation/AnimTypes.h"
 #include "CustomBoneIndexArray.h"
-#include "AnimEncoding.h"
 #include "Animation/AnimStats.h"
 #include "Misc/Base64.h"
+#include "Animation/Skeleton.h"
+#include "BoneContainer.h"
 
 struct FBoneTransform
 {
@@ -318,28 +319,6 @@ public:
 		return GetBoneContainer().GetRefPoseTransform(BoneIndex);
 	}
 
-	// Populates this pose from the supplied animation and track data
-	void PopulateFromAnimation(
-		const UAnimSequence& Seq,
-		const BoneTrackArray& RotationTracks,
-		const BoneTrackArray& TranslationTracks,
-		const BoneTrackArray& ScaleTracks,
-		float Time)
-	{
-		// @todo fixme 
-		FTransformArray LocalBones;
-		LocalBones = this->Bones;
-
-		AnimationFormat_GetAnimationPose(
-			LocalBones, //@TODO:@ANIMATION: Nasty hack
-			RotationTracks,
-			TranslationTracks,
-			ScaleTracks,
-			Seq,
-			Time);
-		this->Bones = LocalBones;
-	}
-
 protected:
 
 	//Reference to our BoneContainer
diff --git a/Engine/Source/Runtime/Engine/Public/EngineSharedPCH.h b/Engine/Source/Runtime/Engine/Public/EngineSharedPCH.h
--- a/Engine/Source/Runtime/Engine/Public/EngineSharedPCH.h
+++ b/Engine/Source/Runtime/Engine/Public/EngineSharedPCH.h
@@ -563,7 +563,6 @@
 #include "Animation/AnimSequence.h"
 #include "BonePose.h"
 #include "CustomBoneIndexArray.h"
-#include "AnimEncoding.h"
 #include "Animation/AnimStats.h"
 #include "GenericOctreePublic.h"
 #include "GameFramework/OnlineReplStructs.h"
