PR #4298: Avoid division by 0 in AnimLinkableElement (Contributed by DSDambuster)

diff --git a/Engine/Source/Runtime/Engine/Private/Animation/AnimLinkableElement.cpp b/Engine/Source/Runtime/Engine/Private/Animation/AnimLinkableElement.cpp
--- a/Engine/Source/Runtime/Engine/Private/Animation/AnimLinkableElement.cpp
+++ b/Engine/Source/Runtime/Engine/Private/Animation/AnimLinkableElement.cpp
@@ -289,12 +289,28 @@ void FAnimLinkableElement::SetTimeFromProportional(float NewTime, EAnimLinkMetho
 	{
 		case EAnimLinkMethod::Absolute:
 		{
-			LinkValue = (NewTime - SegmentBeginTime) / SegmentLength;
+			if(SegmentLength != 0.0f)
+			{
+				LinkValue = (NewTime - SegmentBeginTime) / SegmentLength;
+			}
+			else
+			{
+				// if segment length is 0, we can't set anywhere else but 0.f
+				LinkValue = 0.f;
+			}
 		}
 			break;
 		case EAnimLinkMethod::Relative:
 		{
-			LinkValue = NewTime / SegmentLength;
+			if(SegmentLength != 0.0f)
+			{
+				LinkValue = NewTime / SegmentLength;
+			}
+			else
+			{
+				// if segment length is 0, we can't set anywhere else but 0.f
+				LinkValue = 0.f;			
+			}
 		}
 			break;
 	}
