PR #4842: UE-52767: UAnimComposite doesn't inherit from UAnimSequence (Contributed by [REDACTED_PERSON])

Fixes issue with stepping not working with AnimComposite
- it also loops around

diff --git a/Engine/Source/Runtime/Engine/Private/Animation/AnimSingleNodeInstance.cpp b/Engine/Source/Runtime/Engine/Private/Animation/AnimSingleNodeInstance.cpp
--- a/Engine/Source/Runtime/Engine/Private/Animation/AnimSingleNodeInstance.cpp
+++ b/Engine/Source/Runtime/Engine/Private/Animation/AnimSingleNodeInstance.cpp
@@ -393,28 +393,45 @@ float UAnimSingleNodeInstance::GetLength()
 
 void UAnimSingleNodeInstance::StepForward()
 {
-	if (UAnimSequence* Sequence = Cast<UAnimSequence>(CurrentAsset))
+	if (UAnimSequenceBase* Sequence = Cast<UAnimSequenceBase>(CurrentAsset))
 	{
 		FAnimSingleNodeInstanceProxy& Proxy = GetProxyOnGameThread<FAnimSingleNodeInstanceProxy>();
-		const FAnimKeyHelper Helper(Sequence->SequenceLength, Sequence->NumFrames);
+		const FAnimKeyHelper Helper(Sequence->SequenceLength, Sequence->GetNumberOfFrames());
 		float KeyLength = Helper.TimePerKey() + SMALL_NUMBER;
 		float Fraction = (Proxy.GetCurrentTime()+KeyLength)/Sequence->SequenceLength;
-		int32 Frames = FMath::Clamp<int32>((float)(Helper.LastKey()*Fraction), 0, Helper.LastKey());
+		int32 Frames = (float)Helper.LastKey() * Fraction;
+		if (IsLooping())
+		{
+			int32 PreviousFrame = (float)Helper.LastKey() * (Proxy.GetCurrentTime() / Sequence->SequenceLength);
+			Frames = (PreviousFrame == Helper.LastKey()) ? 0 : Frames;
+		}
+		else
+		{
+			Frames = FMath::Clamp<int32>(Frames, 0, Helper.LastKey());
+		}
 		SetPosition(Frames*KeyLength);
 	}	
 }
 
 void UAnimSingleNodeInstance::StepBackward()
 {
-	if (UAnimSequence* Sequence = Cast<UAnimSequence>(CurrentAsset))
+	if (UAnimSequenceBase* Sequence = Cast<UAnimSequenceBase>(CurrentAsset))
 	{
 		FAnimSingleNodeInstanceProxy& Proxy = GetProxyOnGameThread<FAnimSingleNodeInstanceProxy>();
 
-		const FAnimKeyHelper Helper(Sequence->SequenceLength, Sequence->NumFrames);
+		const FAnimKeyHelper Helper(Sequence->SequenceLength, Sequence->GetNumberOfFrames());
 		float KeyLength = Helper.TimePerKey() + SMALL_NUMBER;
 		float Fraction = (Proxy.GetCurrentTime()-KeyLength)/Sequence->SequenceLength;
-		int32 Frames = FMath::Clamp<int32>((float)(Helper.LastKey()*Fraction), 0, Helper.LastKey());
-		SetPosition(Frames*KeyLength);
+		int32 Frames = (float)Helper.LastKey() * Fraction;
+		if (IsLooping())
+		{
+			Frames = (Frames < 0) ? Helper.LastKey() : Frames;
+		}
+		else
+		{
+			Frames = FMath::Clamp<int32>(Frames, 0, Helper.LastKey());
+		}
+		SetPosition(Frames * KeyLength);
 	}
 }
 
