Fixed mis-reported HLOD build indicator

diff --git a/Engine/Source/Runtime/Engine/Private/HLODProxy.cpp b/Engine/Source/Runtime/Engine/Private/HLODProxy.cpp
--- a/Engine/Source/Runtime/Engine/Private/HLODProxy.cpp
+++ b/Engine/Source/Runtime/Engine/Private/HLODProxy.cpp
@@ -367,7 +367,17 @@ FName UHLODProxy::GenerateKeyForActor(const ALODActor* LODActor)
 bool UHLODProxy::ContainsDataForActor(const ALODActor* InLODActor) const
 {
 #if WITH_EDITOR
-	FName Key = GenerateKeyForActor(InLODActor);
+	FName Key;
+
+	// Only re-generate the key in non-PIE worlds
+	if(InLODActor->GetOutermost()->HasAnyPackageFlags(PKG_PlayInEditor))
+	{
+		Key = InLODActor->GetKey();
+	}
+	else
+	{
+		Key = GenerateKeyForActor(InLODActor);
+	}
 #else
 	FName Key = InLODActor->GetKey();
 #endif
diff --git a/Engine/Source/Runtime/Engine/Private/LODActor.cpp b/Engine/Source/Runtime/Engine/Private/LODActor.cpp
--- a/Engine/Source/Runtime/Engine/Private/LODActor.cpp
+++ b/Engine/Source/Runtime/Engine/Private/LODActor.cpp
@@ -18,6 +18,7 @@
 #include "UObject/FrameworkObjectVersion.h"
 #include "UObject/FortniteMainBranchObjectVersion.h"
 #include "Engine/HLODProxy.h"
+#include "UObject/PropertyPortFlags.h"
 
 #if WITH_EDITOR
 #include "Editor.h"
@@ -974,6 +975,13 @@ void ALODActor::OnCVarsChanged()
 
 void ALODActor::Serialize(FArchive& Ar)
 {
+#if WITH_EDITOR
+	if (Ar.GetPortFlags() & PPF_DuplicateForPIE && Ar.IsSaving())
+	{
+		Key = UHLODProxy::GenerateKeyForActor(this);
+	}
+#endif
+
 	Super::Serialize(Ar);
 #if WITH_EDITOR
 	Ar.UsingCustomVersion(FFrameworkObjectVersion::GUID);
