Fix issue where preview asset is not saved and not loaded unless it's loaded manually

diff --git a/Engine/Source/Runtime/Engine/Private/AnimBlueprint.cpp b/Engine/Source/Runtime/Engine/Private/AnimBlueprint.cpp
--- a/Engine/Source/Runtime/Engine/Private/AnimBlueprint.cpp
+++ b/Engine/Source/Runtime/Engine/Private/AnimBlueprint.cpp
@@ -202,6 +202,10 @@ USkeletalMesh* UAnimBlueprint::GetPreviewMesh(bool bFindIfNotSet/*=false*/)
 USkeletalMesh* UAnimBlueprint::GetPreviewMesh() const
 {
 #if WITH_EDITORONLY_DATA
+	if (!PreviewSkeletalMesh.IsValid())
+	{
+		PreviewSkeletalMesh.LoadSynchronous();
+	}
 	return PreviewSkeletalMesh.Get();
 #else
 	return nullptr;
diff --git a/Engine/Source/Runtime/Engine/Private/Animation/Skeleton.cpp b/Engine/Source/Runtime/Engine/Private/Animation/Skeleton.cpp
--- a/Engine/Source/Runtime/Engine/Private/Animation/Skeleton.cpp
+++ b/Engine/Source/Runtime/Engine/Private/Animation/Skeleton.cpp
@@ -789,6 +789,10 @@ USkeletalMesh* USkeleton::GetPreviewMesh(bool bFindIfNotSet/*=false*/)
 USkeletalMesh* USkeleton::GetPreviewMesh() const
 {
 #if WITH_EDITORONLY_DATA
+	if (!PreviewSkeletalMesh.IsValid())
+	{
+		PreviewSkeletalMesh.LoadSynchronous();
+	}
 	return PreviewSkeletalMesh.Get();
 #else
 	return nullptr;
diff --git a/Engine/Source/Runtime/Engine/Private/AnimationAsset.cpp b/Engine/Source/Runtime/Engine/Private/AnimationAsset.cpp
--- a/Engine/Source/Runtime/Engine/Private/AnimationAsset.cpp
+++ b/Engine/Source/Runtime/Engine/Private/AnimationAsset.cpp
@@ -278,6 +278,10 @@ USkeletalMesh* UAnimationAsset::GetPreviewMesh(bool bFindIfNotSet)
 USkeletalMesh* UAnimationAsset::GetPreviewMesh() const
 {
 #if WITH_EDITORONLY_DATA
+	if (!PreviewSkeletalMesh.IsValid())
+	{
+		PreviewSkeletalMesh.LoadSynchronous();
+	}
 	return PreviewSkeletalMesh.Get();
 #else
 	return nullptr;
