Perform anim update even when we have zero bones, but still skip evaluation

This ensures that notifies still get triggered in all update-relevant tick modes, even on dedicated servers.

diff --git a/Engine/Source/Runtime/Engine/Private/Components/SkeletalMeshComponent.cpp b/Engine/Source/Runtime/Engine/Private/Components/SkeletalMeshComponent.cpp
--- a/Engine/Source/Runtime/Engine/Private/Components/SkeletalMeshComponent.cpp
+++ b/Engine/Source/Runtime/Engine/Private/Components/SkeletalMeshComponent.cpp
@@ -1753,8 +1753,7 @@ void USkeletalMeshComponent::PerformAnimationProcessing(const USkeletalMesh* InS
 	ANIM_MT_SCOPE_CYCLE_COUNTER(PerformAnimEvaluation, IsRunningParallelEvaluation());
 
 	// Can't do anything without a SkeletalMesh
-	// Do nothing more if no bones in skeleton.
-	if (!InSkeletalMesh || OutSpaceBases.Num() == 0)
+	if (!InSkeletalMesh)
 	{
 		return;
 	}
@@ -1771,7 +1770,8 @@ void USkeletalMeshComponent::PerformAnimationProcessing(const USkeletalMesh* InS
 		PostProcessAnimInstance->ParallelUpdateAnimation();
 	}
 
-	if(bInDoEvaluation)
+	// Do nothing more if no bones in skeleton.
+	if(bInDoEvaluation && OutSpaceBases.Num() > 0)
 	{
 		FMemMark Mark(FMemStack::Get());
 		FCompactPose EvaluatedPose;
