Rich Curve: Don't expose RCIM_None and RCTM_None as DisplayName="None". It doesn't make sense to have no interpolation

diff --git a/Engine/Source/Runtime/Engine/Classes/Curves/RichCurve.h b/Engine/Source/Runtime/Engine/Classes/Curves/RichCurve.h
--- a/Engine/Source/Runtime/Engine/Classes/Curves/RichCurve.h
+++ b/Engine/Source/Runtime/Engine/Classes/Curves/RichCurve.h
@@ -16,7 +16,7 @@ enum ERichCurveInterpMode
 	RCIM_Linear UMETA(DisplayName="Linear"),
 	RCIM_Constant UMETA(DisplayName="Constant"),
 	RCIM_Cubic UMETA(DisplayName="Cubic"),
-	RCIM_None UMETA(DisplayName="None")
+	RCIM_None UMETA(Hidden)
 };
 
 
@@ -27,7 +27,7 @@ enum ERichCurveTangentMode
 	RCTM_Auto UMETA(DisplayName="Auto"),
 	RCTM_User UMETA(DisplayName="User"),
 	RCTM_Break UMETA(DisplayName="Break"),
-	RCTM_None UMETA(DisplayName="None")
+	RCTM_None UMETA(Hidden)
 };
 
 
