Added scale parameter to Canvas::DrawText

diff --git a/Engine/Source/Runtime/Engine/Classes/Engine/Canvas.h b/Engine/Source/Runtime/Engine/Classes/Engine/Canvas.h
--- a/Engine/Source/Runtime/Engine/Classes/Engine/Canvas.h
+++ b/Engine/Source/Runtime/Engine/Classes/Engine/Canvas.h
@@ -533,7 +533,7 @@ public:
 	 * @param OutlineColor				Color to render the outline for the text.
 	 */	
 	UFUNCTION(BlueprintCallable, Category=Canvas, meta=(DisplayName="Draw Text", ScriptName="DrawText"))
-	void K2_DrawText(UFont* RenderFont, const FString& RenderText, FVector2D ScreenPosition, FLinearColor RenderColor=FLinearColor::White, float Kerning=0.0f, FLinearColor ShadowColor=FLinearColor::Black, FVector2D ShadowOffset=FVector2D::UnitVector, bool bCentreX=false, bool bCentreY=false, bool bOutlined=false, FLinearColor OutlineColor=FLinearColor::Black);
+	void K2_DrawText(UFont* RenderFont, const FString& RenderText, FVector2D ScreenPosition, FVector2D Scale=FVector2D(1.0f,1.0f), FLinearColor RenderColor=FLinearColor::White, float Kerning=0.0f, FLinearColor ShadowColor=FLinearColor::Black, FVector2D ShadowOffset=FVector2D::UnitVector, bool bCentreX=false, bool bCentreY=false, bool bOutlined=false, FLinearColor OutlineColor=FLinearColor::Black);
 	
 	/**
 	 * Draws a 3x3 grid border with tiled frame and tiled interior on the Canvas.
diff --git a/Engine/Source/Runtime/Engine/Private/UserInterface/Canvas.cpp b/Engine/Source/Runtime/Engine/Private/UserInterface/Canvas.cpp
--- a/Engine/Source/Runtime/Engine/Private/UserInterface/Canvas.cpp
+++ b/Engine/Source/Runtime/Engine/Private/UserInterface/Canvas.cpp
@@ -2012,7 +2012,8 @@ void UCanvas::K2_DrawMaterial(UMaterialInterface* RenderMaterial, FVector2D Scre
 	}
 }
 
-void UCanvas::K2_DrawText(UFont* RenderFont, const FString& RenderText, FVector2D ScreenPosition, FLinearColor RenderColor, float Kerning, FLinearColor ShadowColor, FVector2D ShadowOffset, bool bCentreX, bool bCentreY, bool bOutlined, FLinearColor OutlineColor)
+
+void UCanvas::K2_DrawText(UFont* RenderFont, const FString& RenderText, FVector2D ScreenPosition, FVector2D Scale, FLinearColor RenderColor, float Kerning, FLinearColor ShadowColor, FVector2D ShadowOffset, bool bCentreX, bool bCentreY, bool bOutlined, FLinearColor OutlineColor)
 {
 	if (!RenderText.IsEmpty() && Canvas)
 	{
@@ -2022,6 +2023,7 @@ void UCanvas::K2_DrawText(UFont* RenderFont, const FString& RenderText, FVector2
 		TextItem.ShadowOffset = ShadowOffset;
 		TextItem.bCentreX = bCentreX;
 		TextItem.bCentreY = bCentreY;
+		TextItem.Scale = Scale;
 		TextItem.bOutlined = bOutlined;
 		TextItem.OutlineColor = OutlineColor;
 		DrawItem(TextItem);
