Fixed touch events simulated through mouse not respecting high dpi

diff --git a/Engine/Source/Runtime/Engine/Private/Slate/SceneViewport.cpp b/Engine/Source/Runtime/Engine/Private/Slate/SceneViewport.cpp
--- a/Engine/Source/Runtime/Engine/Private/Slate/SceneViewport.cpp
+++ b/Engine/Source/Runtime/Engine/Private/Slate/SceneViewport.cpp
@@ -27,8 +27,8 @@ static EPixelFormat SceneTargetFormat = PF_A2B10G10R10;
 FSceneViewport::FSceneViewport( FViewportClient* InViewportClient, TSharedPtr<SViewport> InViewportWidget )
 	: FViewport( InViewportClient )
 	, CurrentReplyState( FReply::Unhandled() )
-	, CachedMousePos(-1, -1)
-	, PreCaptureMousePos(-1, -1)
+	, CachedCursorPos(-1, -1)
+	, PreCaptureCursorPos(-1, -1)
 	, SoftwareCursorPosition( 0, 0 )
 	, bIsSoftwareCursorVisible( false )	
 	, DebugCanvasDrawer( new FDebugCanvasDrawer )
@@ -130,7 +130,7 @@ void FSceneViewport::ShowCursor( bool bVisible )
 		else
 		{
 			// Restore the old mouse position when we show the cursor.
-			CurrentReplyState.SetMousePos( PreCaptureMousePos );
+			CurrentReplyState.SetMousePos( PreCaptureCursorPos );
 		}
 
 		SetPreCaptureMousePosFromSlateCursor();
@@ -172,23 +172,23 @@ void FSceneViewport::Destroy()
 
 int32 FSceneViewport::GetMouseX() const
 {
-	return CachedMousePos.X;
+	return CachedCursorPos.X;
 }
 
 int32 FSceneViewport::GetMouseY() const
 {
-	return CachedMousePos.Y;
+	return CachedCursorPos.Y;
 }
 
 void FSceneViewport::GetMousePos( FIntPoint& MousePosition, const bool bLocalPosition )
 {
 	if (bLocalPosition)
 	{
-		MousePosition = CachedMousePos;
+		MousePosition = CachedCursorPos;
 	}
 	else
 	{
-		const FVector2D AbsoluteMousePos = CachedGeometry.LocalToAbsolute(FVector2D(CachedMousePos.X / CachedGeometry.Scale, CachedMousePos.Y / CachedGeometry.Scale));
+		const FVector2D AbsoluteMousePos = CachedGeometry.LocalToAbsolute(FVector2D(CachedCursorPos.X / CachedGeometry.Scale, CachedCursorPos.Y / CachedGeometry.Scale));
 		MousePosition.X = AbsoluteMousePos.X;
 		MousePosition.Y = AbsoluteMousePos.Y;
 	}
@@ -198,7 +198,7 @@ void FSceneViewport::SetMouse( int32 X, int32 Y )
 {
 	FVector2D AbsolutePos = CachedGeometry.LocalToAbsolute(FVector2D(X, Y));
 	FSlateApplication::Get().SetCursorPos( AbsolutePos );
-	CachedMousePos = FIntPoint(X, Y);
+	CachedCursorPos = FIntPoint(X, Y);
 }
 
 void FSceneViewport::ProcessInput( float DeltaTime )
@@ -206,13 +206,13 @@ void FSceneViewport::ProcessInput( float DeltaTime )
 	// Required 
 }
 
-void FSceneViewport::UpdateCachedMousePos( const FGeometry& InGeometry, const FPointerEvent& InMouseEvent )
+void FSceneViewport::UpdateCachedCursorPos( const FGeometry& InGeometry, const FPointerEvent& InMouseEvent )
 {
 	FVector2D LocalPixelMousePos = InGeometry.AbsoluteToLocal(InMouseEvent.GetScreenSpacePosition());
 	LocalPixelMousePos.X *= CachedGeometry.Scale;
 	LocalPixelMousePos.Y *= CachedGeometry.Scale;
 
-	CachedMousePos = LocalPixelMousePos.IntPoint();
+	CachedCursorPos = LocalPixelMousePos.IntPoint();
 }
 
 void FSceneViewport::UpdateCachedGeometry( const FGeometry& InGeometry )
@@ -466,8 +466,9 @@ FReply FSceneViewport::OnMouseButtonDown( const FGeometry& InGeometry, const FPo
 
 	KeyStateMap.Add(InMouseEvent.GetEffectingButton(), true);
 	UpdateModifierKeys(InMouseEvent);
-	UpdateCachedMousePos(InGeometry, InMouseEvent);
+
 	UpdateCachedGeometry(InGeometry);
+	UpdateCachedCursorPos(InGeometry, InMouseEvent);
 
 	// Switch to the viewport clients world before processing input
 	FScopedConditionalWorldSwitcher WorldSwitcher(ViewportClient);
@@ -576,8 +577,10 @@ FReply FSceneViewport::OnMouseButtonUp( const FGeometry& InGeometry, const FPoin
 
 	KeyStateMap.Add( InMouseEvent.GetEffectingButton(), false );
 	UpdateModifierKeys( InMouseEvent );
-	UpdateCachedMousePos( InGeometry, InMouseEvent );
+
 	UpdateCachedGeometry(InGeometry);
+	UpdateCachedCursorPos( InGeometry, InMouseEvent );
+	
 
 	// Switch to the viewport clients world before processing input
 	FScopedConditionalWorldSwitcher WorldSwitcher( ViewportClient );
@@ -626,7 +629,7 @@ FReply FSceneViewport::OnMouseButtonUp( const FGeometry& InGeometry, const FPoin
 
 void FSceneViewport::OnMouseEnter( const FGeometry& MyGeometry, const FPointerEvent& MouseEvent )
 {
-	UpdateCachedMousePos( MyGeometry, MouseEvent );
+	UpdateCachedCursorPos( MyGeometry, MouseEvent );
 	ViewportClient->MouseEnter( this, GetMouseX(), GetMouseY() );
 }
 
@@ -638,7 +641,7 @@ void FSceneViewport::OnMouseLeave( const FPointerEvent& MouseEvent )
 	
 		if (IsCurrentlyGameViewport())
 		{
-			CachedMousePos = FIntPoint(-1, -1);
+			CachedCursorPos = FIntPoint(-1, -1);
 		}
 	}
 }
@@ -648,8 +651,8 @@ FReply FSceneViewport::OnMouseMove( const FGeometry& InGeometry, const FPointerE
 	// Start a new reply state
 	CurrentReplyState = FReply::Handled();
 
-	UpdateCachedMousePos(InGeometry, InMouseEvent);
 	UpdateCachedGeometry(InGeometry);
+	UpdateCachedCursorPos(InGeometry, InMouseEvent);
 
 	const bool bViewportHasCapture = ViewportWidget.IsValid() && ViewportWidget.Pin()->HasMouseCapture();
 	if ( ViewportClient && GetSizeXY() != FIntPoint::ZeroValue )
@@ -694,8 +697,8 @@ FReply FSceneViewport::OnMouseWheel( const FGeometry& InGeometry, const FPointer
 	// Start a new reply state
 	CurrentReplyState = FReply::Handled();
 
-	UpdateCachedMousePos( InGeometry, InMouseEvent );
 	UpdateCachedGeometry(InGeometry);
+	UpdateCachedCursorPos( InGeometry, InMouseEvent );
 
 	if( ViewportClient  && GetSizeXY() != FIntPoint::ZeroValue  )
 	{
@@ -724,8 +727,9 @@ FReply FSceneViewport::OnMouseButtonDoubleClick( const FGeometry& InGeometry, co
 	//	WM_*BUTTONDBLCLK	(Needs to set the KeyStates[*] to true)
 	//	WM_*BUTTONUP
 	KeyStateMap.Add( InMouseEvent.GetEffectingButton(), true );
-	UpdateCachedMousePos( InGeometry, InMouseEvent );
+
 	UpdateCachedGeometry(InGeometry);
+	UpdateCachedCursorPos(InGeometry, InMouseEvent);
 
 	if( ViewportClient && GetSizeXY() != FIntPoint::ZeroValue  )
 	{
@@ -746,17 +750,15 @@ FReply FSceneViewport::OnTouchStarted( const FGeometry& MyGeometry, const FPoint
 	CurrentReplyState = FReply::Handled().PreventThrottling(); 
 	++NumTouches;
 
-	UpdateCachedMousePos(MyGeometry, TouchEvent);
 	UpdateCachedGeometry(MyGeometry);
+	UpdateCachedCursorPos(MyGeometry, TouchEvent);
 	
 	if( ViewportClient )
 	{
 		// Switch to the viewport clients world before processing input
 		FScopedConditionalWorldSwitcher WorldSwitcher( ViewportClient );
 
-		const FVector2D TouchPosition = MyGeometry.AbsoluteToLocal(TouchEvent.GetScreenSpacePosition());
-
-		if( !ViewportClient->InputTouch( this, TouchEvent.GetUserIndex(), TouchEvent.GetPointerIndex(), ETouchType::Began, TouchPosition, FDateTime::Now(), TouchEvent.GetTouchpadIndex()) )
+		if( !ViewportClient->InputTouch( this, TouchEvent.GetUserIndex(), TouchEvent.GetPointerIndex(), ETouchType::Began, CachedCursorPos, FDateTime::Now(), TouchEvent.GetTouchpadIndex()) )
 		{
 			CurrentReplyState = FReply::Unhandled(); 
 		}
@@ -768,19 +770,18 @@ FReply FSceneViewport::OnTouchStarted( const FGeometry& MyGeometry, const FPoint
 FReply FSceneViewport::OnTouchMoved( const FGeometry& MyGeometry, const FPointerEvent& TouchEvent )
 {
 	// Start a new reply state
-	CurrentReplyState = FReply::Handled(); 
+	CurrentReplyState = FReply::Handled();
 
-	UpdateCachedMousePos(MyGeometry, TouchEvent);
 	UpdateCachedGeometry(MyGeometry);
+	UpdateCachedCursorPos(MyGeometry, TouchEvent);
 
 	if( ViewportClient )
 	{
 		// Switch to the viewport clients world before processing input
 		FScopedConditionalWorldSwitcher WorldSwitcher( ViewportClient );
 
-		const FVector2D TouchPosition = MyGeometry.AbsoluteToLocal(TouchEvent.GetScreenSpacePosition());
 
-		if( !ViewportClient->InputTouch( this, TouchEvent.GetUserIndex(), TouchEvent.GetPointerIndex(), ETouchType::Moved, TouchPosition, FDateTime::Now(), TouchEvent.GetTouchpadIndex()) )
+		if( !ViewportClient->InputTouch( this, TouchEvent.GetUserIndex(), TouchEvent.GetPointerIndex(), ETouchType::Moved, CachedCursorPos, FDateTime::Now(), TouchEvent.GetTouchpadIndex()) )
 		{
 			CurrentReplyState = FReply::Unhandled(); 
 		}
@@ -792,26 +793,25 @@ FReply FSceneViewport::OnTouchMoved( const FGeometry& MyGeometry, const FPointer
 FReply FSceneViewport::OnTouchEnded( const FGeometry& MyGeometry, const FPointerEvent& TouchEvent )
 {
 	// Start a new reply state
-	CurrentReplyState = FReply::Handled(); 
+	CurrentReplyState = FReply::Handled();
 
+	UpdateCachedGeometry(MyGeometry);
 	if (--NumTouches > 0)
 	{
-		UpdateCachedMousePos(MyGeometry, TouchEvent);
+		UpdateCachedCursorPos(MyGeometry, TouchEvent);
 	}
 	else
 	{
-		CachedMousePos = FIntPoint(-1, -1);
+		CachedCursorPos = FIntPoint(-1, -1);
 	}
-	UpdateCachedGeometry(MyGeometry);
+
 
 	if( ViewportClient )
 	{
 		// Switch to the viewport clients world before processing input
 		FScopedConditionalWorldSwitcher WorldSwitcher( ViewportClient );
 
-		const FVector2D TouchPosition = MyGeometry.AbsoluteToLocal(TouchEvent.GetScreenSpacePosition());
-
-		if( !ViewportClient->InputTouch( this, TouchEvent.GetUserIndex(), TouchEvent.GetPointerIndex(), ETouchType::Ended, TouchPosition, FDateTime::Now(), TouchEvent.GetTouchpadIndex()) )
+		if( !ViewportClient->InputTouch( this, TouchEvent.GetUserIndex(), TouchEvent.GetPointerIndex(), ETouchType::Ended, CachedCursorPos, FDateTime::Now(), TouchEvent.GetTouchpadIndex()) )
 		{
 			CurrentReplyState = FReply::Unhandled(); 
 		}
@@ -825,8 +825,8 @@ FReply FSceneViewport::OnTouchGesture( const FGeometry& MyGeometry, const FPoint
 	// Start a new reply state
 	CurrentReplyState = FReply::Handled();
 
-	UpdateCachedMousePos( MyGeometry, GestureEvent );
-	UpdateCachedGeometry( MyGeometry );
+	UpdateCachedGeometry(MyGeometry);
+	UpdateCachedCursorPos( MyGeometry, GestureEvent );
 
 	if( ViewportClient )
 	{
@@ -1869,5 +1869,5 @@ void FSceneViewport::ReleaseDynamicRHI()
 
 void FSceneViewport::SetPreCaptureMousePosFromSlateCursor()
 {
-	PreCaptureMousePos = FSlateApplication::Get().GetCursorPos().IntPoint();
+	PreCaptureCursorPos = FSlateApplication::Get().GetCursorPos().IntPoint();
 }
diff --git a/Engine/Source/Runtime/Engine/Public/Slate/SceneViewport.h b/Engine/Source/Runtime/Engine/Public/Slate/SceneViewport.h
--- a/Engine/Source/Runtime/Engine/Public/Slate/SceneViewport.h
+++ b/Engine/Source/Runtime/Engine/Public/Slate/SceneViewport.h
@@ -308,7 +308,7 @@ private:
 	 * @param InGeometry	The geometry of the viewport to convert to local space
 	 * @param InMouseEvent	The mouse event containing the position of the mouse in absolute space
 	 */
-	void UpdateCachedMousePos( const FGeometry& InGeometry, const FPointerEvent& InMouseEvent );
+	void UpdateCachedCursorPos( const FGeometry& InGeometry, const FPointerEvent& InMouseEvent );
 
 	/**
 	 * Updates the cached viewport geometry
@@ -371,11 +371,11 @@ private:
 	/** A mapping of key names to their pressed state */
 	TMap<FKey,bool> KeyStateMap;
 	/** The last known mouse position in local space, -1, -1 if unknown */
-	FIntPoint CachedMousePos;
+	FIntPoint CachedCursorPos;
 	/** The last known geometry info */
 	FGeometry CachedGeometry;
 	/** Mouse position before the latest capture */
-	FIntPoint PreCaptureMousePos;
+	FIntPoint PreCaptureCursorPos;
 	/**	The current position of the software cursor */
 	FVector2D SoftwareCursorPosition;
 	/**	Whether the software cursor should be drawn in the viewport */
