PR #4675: Properly set TextScale for OnScreenDebugMessages (Contributed by [REDACTED_PERSON])

diff --git a/Engine/Source/Runtime/Engine/Private/UnrealEngine.cpp b/Engine/Source/Runtime/Engine/Private/UnrealEngine.cpp
--- a/Engine/Source/Runtime/Engine/Private/UnrealEngine.cpp
+++ b/Engine/Source/Runtime/Engine/Private/UnrealEngine.cpp
@@ -8511,6 +8511,7 @@ void UEngine::AddOnScreenDebugMessage(uint64 Key, float TimeToDisplay, FColor Di
 				NewMessage->ScreenMessage = DebugMessage;
 				NewMessage->DisplayColor = DisplayColor;
 				NewMessage->TimeToDisplay = TimeToDisplay;
+				NewMessage->TextScale = TextScale;
 				NewMessage->CurrentTimeDisplayed = 0.0f;				
 			}
 			else
@@ -8520,7 +8521,8 @@ void UEngine::AddOnScreenDebugMessage(uint64 Key, float TimeToDisplay, FColor Di
 				NewMessage.Key = Key;
 				NewMessage.DisplayColor = DisplayColor;
 				NewMessage.TimeToDisplay = TimeToDisplay;
-				NewMessage.ScreenMessage = DebugMessage;				
+				NewMessage.ScreenMessage = DebugMessage;
+				NewMessage.TextScale = TextScale;
 				PriorityScreenMessages.Insert(NewMessage, 0);
 			}
 		}
@@ -8536,6 +8538,7 @@ void UEngine::AddOnScreenDebugMessage(uint64 Key, float TimeToDisplay, FColor Di
 				NewMessage.DisplayColor = DisplayColor;
 				NewMessage.TimeToDisplay = TimeToDisplay;
 				NewMessage.ScreenMessage = DebugMessage;				
+				NewMessage.TextScale = TextScale;
 				ScreenMessages.Add((int32)Key, NewMessage);
 			}
 			else
@@ -8546,6 +8549,7 @@ void UEngine::AddOnScreenDebugMessage(uint64 Key, float TimeToDisplay, FColor Di
 				Message->TextScale = TextScale;
 				Message->TimeToDisplay = TimeToDisplay;
 				Message->CurrentTimeDisplayed = 0.0f;				
+				Message->TextScale = TextScale;
 			}
 		}
 	}
@@ -8636,14 +8640,12 @@ bool UEngine::FErrorsAndWarningsCollector::Tick(float Seconds)
 /** Wrapper from int32 to uint64 */
 void UEngine::AddOnScreenDebugMessage(int32 Key, float TimeToDisplay, FColor DisplayColor, const FString& DebugMessage, bool bNewerOnTop, const FVector2D& TextScale)
 {
-	if (bEnableOnScreenDebugMessages == true)
-	{
-		AddOnScreenDebugMessage((uint64)Key, TimeToDisplay, DisplayColor, DebugMessage, bNewerOnTop, TextScale);
-	}
+	AddOnScreenDebugMessage((uint64)Key, TimeToDisplay, DisplayColor, DebugMessage, bNewerOnTop, TextScale);
 }
 
 bool UEngine::OnScreenDebugMessageExists(uint64 Key)
 {
+#if !UE_BUILD_SHIPPING
 	if (bEnableOnScreenDebugMessages == true)
 	{
 		if (Key == (uint64)-1)
@@ -8658,14 +8660,17 @@ bool UEngine::OnScreenDebugMessageExists(uint64 Key)
 			return true;
 		}
 	}
+#endif // !UE_BUILD_SHIPPING
 
 	return false;
 }
 
 void UEngine::ClearOnScreenDebugMessages()
 {
+#if !UE_BUILD_SHIPPING
 	ScreenMessages.Empty();
 	PriorityScreenMessages.Empty();
+#endif // !UE_BUILD_SHIPPING
 }
 
 #if !UE_BUILD_SHIPPING
@@ -9454,7 +9459,8 @@ float DrawOnscreenDebugMessages(UWorld* World, FViewport* Viewport, FCanvas* Can
 			if (YPos < MaxYPos)
 			{
 				MessageTextItem.Text = FText::FromString(Message.ScreenMessage);
-				MessageTextItem.SetColor(Message.DisplayColor);				
+				MessageTextItem.SetColor(Message.DisplayColor);
+				MessageTextItem.Scale = Message.TextScale;
 				Canvas->DrawItem(MessageTextItem, FVector2D(MessageX, YPos));
 				YPos += MessageTextItem.DrawnSize.Y * 1.15f;
 			}
