Optimize import time for morph targets

diff --git a/Engine/Source/Editor/UnrealEd/Private/Fbx/FbxSkeletalMeshImport.cpp b/Engine/Source/Editor/UnrealEd/Private/Fbx/FbxSkeletalMeshImport.cpp
--- a/Engine/Source/Editor/UnrealEd/Private/Fbx/FbxSkeletalMeshImport.cpp
+++ b/Engine/Source/Editor/UnrealEd/Private/Fbx/FbxSkeletalMeshImport.cpp
@@ -4105,6 +4105,7 @@ void UnFbx::FFbxImporter::ImportMorphTargetsInternal( TArray<FbxNode*>& SkelMesh
 		GWarn->StatusUpdate( NumCompleted, NumTasks, FText::Format( LOCTEXT("ImportingMorphTargetStatus", "Importing Morph Target: {NumCompleted} of {NumTasks}"), Args ) );
 	}
 
+	bool bNeedToInvalidateRegisteredMorph = false;
 	// Create morph streams for each morph target we are importing.
 	// This has to happen on a single thread since the skeletal meshes' bulk data is locked and cant be accessed by multiple threads simultaneously
 	for (int32 Index = 0; Index < MorphTargets.Num(); Index++)
@@ -4120,13 +4121,18 @@ void UnFbx::FFbxImporter::ImportMorphTargetsInternal( TArray<FbxNode*>& SkelMesh
 		// register does mark package as dirty
 		if (MorphTarget->HasValidData())
 		{
-			BaseSkelMesh->RegisterMorphTarget(MorphTarget);
+			bNeedToInvalidateRegisteredMorph |= BaseSkelMesh->RegisterMorphTarget(MorphTarget, false);
 		}
 
 		delete Results[Index];
 		Results[Index] = nullptr;
 	}
 
+	if (bNeedToInvalidateRegisteredMorph)
+	{
+		BaseSkelMesh->InitMorphTargetsAndRebuildRenderData();
+	}
+
 	GWarn->EndSlowTask();
 }
 
diff --git a/Engine/Source/Runtime/Engine/Classes/Engine/SkeletalMesh.h b/Engine/Source/Runtime/Engine/Classes/Engine/SkeletalMesh.h
--- a/Engine/Source/Runtime/Engine/Classes/Engine/SkeletalMesh.h
+++ b/Engine/Source/Runtime/Engine/Classes/Engine/SkeletalMesh.h
@@ -922,8 +922,11 @@ public:
 	UMorphTarget* FindMorphTarget(FName MorphTargetName) const;
 	UMorphTarget* FindMorphTargetAndIndex(FName MorphTargetName, int32& OutIndex) const;
 
+	/* Initialize morph targets and rebuild the render data */
+	void InitMorphTargetsAndRebuildRenderData();
+
 	/** if name conflicts, it will overwrite the reference */
-	void RegisterMorphTarget(UMorphTarget* MorphTarget);
+	bool RegisterMorphTarget(UMorphTarget* MorphTarget, bool bInvalidateRenderData = true);
 
 	void UnregisterMorphTarget(UMorphTarget* MorphTarget);
 
diff --git a/Engine/Source/Runtime/Engine/Private/SkeletalMesh.cpp b/Engine/Source/Runtime/Engine/Private/SkeletalMesh.cpp
--- a/Engine/Source/Runtime/Engine/Private/SkeletalMesh.cpp
+++ b/Engine/Source/Runtime/Engine/Private/SkeletalMesh.cpp
@@ -1731,7 +1731,16 @@ void USkeletalMesh::DebugVerifySkeletalMeshLOD()
 	}
 }
 
-void USkeletalMesh::RegisterMorphTarget(UMorphTarget* MorphTarget)
+void USkeletalMesh::InitMorphTargetsAndRebuildRenderData()
+{
+	MarkPackageDirty();
+	// need to refresh the map
+	InitMorphTargets();
+	// invalidate render data
+	InvalidateRenderData();
+}
+
+bool USkeletalMesh::RegisterMorphTarget(UMorphTarget* MorphTarget, bool bInvalidateRenderData)
 {
 	if ( MorphTarget )
 	{
@@ -1765,25 +1774,20 @@ void USkeletalMesh::RegisterMorphTarget(UMorphTarget* MorphTarget)
 			bRegistered = true;
 		}
 
-		if (bRegistered)
+		if (bRegistered && bInvalidateRenderData)
 		{
-			MarkPackageDirty();
-			// need to refresh the map
-			InitMorphTargets();
-			// invalidate render data
-			InvalidateRenderData();
+			InitMorphTargetsAndRebuildRenderData();
 		}
+		return bRegistered;
 	}
+	return false;
 }
 
+
 void USkeletalMesh::UnregisterAllMorphTarget()
 {
 	MorphTargets.Empty();
-	MarkPackageDirty();
-	// need to refresh the map
-	InitMorphTargets();
-	// invalidate render data
-	InvalidateRenderData();
+	InitMorphTargetsAndRebuildRenderData();
 }
 
 void USkeletalMesh::UnregisterMorphTarget(UMorphTarget* MorphTarget)
@@ -1798,15 +1802,10 @@ void USkeletalMesh::UnregisterMorphTarget(UMorphTarget* MorphTarget)
 			{
 				MorphTargets.RemoveAt(I);
 				--I;
-				MarkPackageDirty();
-				// need to refresh the map
-				InitMorphTargets();
-				// invalidate render data
-				InvalidateRenderData();
+				InitMorphTargetsAndRebuildRenderData();
 				return;
 			}
 		}
-
 		UE_LOG( LogSkeletalMesh, Log, TEXT("UnregisterMorphTarget: %s not found."), *MorphTarget->GetName() );
 	}
 }
