Bumped static mesh mesh data GUID.

diff --git a/Engine/Source/Runtime/Engine/Private/StaticMesh.cpp b/Engine/Source/Runtime/Engine/Private/StaticMesh.cpp
--- a/Engine/Source/Runtime/Engine/Private/StaticMesh.cpp
+++ b/Engine/Source/Runtime/Engine/Private/StaticMesh.cpp
@@ -2977,7 +2977,7 @@ void UStaticMesh::FixupMaterialSlotName()
 // differences, etc.) replace the version GUID below with a new one.
 // In case of merge conflicts with DDC versions, you MUST generate a new GUID
 // and set this new GUID as the version.                                       
-#define MESHDATAKEY_STATICMESH_DERIVEDDATA_VER TEXT("B6081AC5FA9949A7B87DC14E39BC25E6")
+#define MESHDATAKEY_STATICMESH_DERIVEDDATA_VER TEXT("ED22489A741846E9830C0AEFB207E591")
 
 static const FString& GetMeshDataKeyStaticMeshDerivedDataVersion()
 {
