Removed some test code which got checked in by mistake.

diff --git a/Engine/Source/Runtime/Engine/Private/StaticMesh.cpp b/Engine/Source/Runtime/Engine/Private/StaticMesh.cpp
--- a/Engine/Source/Runtime/Engine/Private/StaticMesh.cpp
+++ b/Engine/Source/Runtime/Engine/Private/StaticMesh.cpp
@@ -2614,27 +2614,6 @@ void UStaticMesh::RegisterMeshAttributes( FMeshDescription& MeshDescription )
 	MeshDescription.PolygonGroupAttributes().RegisterAttribute<FName>( MeshAttribute::PolygonGroup::ImportedMaterialSlotName ); //The unique key to match the mesh material slot
 	MeshDescription.PolygonGroupAttributes().RegisterAttribute<bool>( MeshAttribute::PolygonGroup::EnableCollision ); //Deprecated
 	MeshDescription.PolygonGroupAttributes().RegisterAttribute<bool>( MeshAttribute::PolygonGroup::CastShadow ); //Deprecated
-
-	{
-	TVertexAttributesView<FVector4> View = MeshDescription.VertexAttributes().GetAttributesView<FVector4>( MeshAttribute::Vertex::Position );
-	FVector4 Blah = View.Get(FVertexID(0));
-	(void)View.GetNumElements();
-	FVector Def = View.GetDefaultValue();
-	View.SetNumIndices(2);
-	View.Set(FVertexID(0), FVector4(1,1,1,1));
-	View.IsValid();
-	}
-
-	const FMeshDescription& MD2 = MeshDescription;
-	{
-	TVertexAttributesConstView<FVector4> View = MD2.VertexAttributes().GetAttributesView<FVector4>( MeshAttribute::Vertex::Position );
-	FVector4 Blah = View.Get(FVertexID(0));
-	(void)View.GetNumElements();
-	FVector Def = View.GetDefaultValue();
-//	View.SetNumIndices(2);
-//	View.Set(FVertexID(0), FVector4(1,1,1,1));
-	View.IsValid();
-	}
 }
 
 
