Made DataTableUtils::GetX functions take a const data pointer

diff --git a/Engine/Source/Runtime/Engine/Private/DataTableUtils.cpp b/Engine/Source/Runtime/Engine/Private/DataTableUtils.cpp
--- a/Engine/Source/Runtime/Engine/Private/DataTableUtils.cpp
+++ b/Engine/Source/Runtime/Engine/Private/DataTableUtils.cpp
@@ -299,7 +299,7 @@ FString DataTableUtils::AssignStringToProperty(const FString& InString, const UP
 	return Error;
 }
 
-FString DataTableUtils::GetPropertyValueAsStringDirect(const UProperty* InProp, uint8* InData, const EDataTableExportFlags InDTExportFlags)
+FString DataTableUtils::GetPropertyValueAsStringDirect(const UProperty* InProp, const uint8* InData, const EDataTableExportFlags InDTExportFlags)
 {
 	FString Result;
 
@@ -311,7 +311,7 @@ FString DataTableUtils::GetPropertyValueAsStringDirect(const UProperty* InProp,
 	return Result;
 }
 
-FString DataTableUtils::GetPropertyValueAsString(const UProperty* InProp, uint8* InData, const EDataTableExportFlags InDTExportFlags)
+FString DataTableUtils::GetPropertyValueAsString(const UProperty* InProp, const uint8* InData, const EDataTableExportFlags InDTExportFlags)
 {
 	FString Result;
 
@@ -342,7 +342,7 @@ FString DataTableUtils::GetPropertyValueAsString(const UProperty* InProp, uint8*
 	return Result;
 }
 
-FText DataTableUtils::GetPropertyValueAsTextDirect(const UProperty* InProp, uint8* InData)
+FText DataTableUtils::GetPropertyValueAsTextDirect(const UProperty* InProp, const uint8* InData)
 {
 	FText Result;
 
@@ -357,7 +357,7 @@ FText DataTableUtils::GetPropertyValueAsTextDirect(const UProperty* InProp, uint
 	return Result;
 }
 
-FText DataTableUtils::GetPropertyValueAsText(const UProperty* InProp, uint8* InData)
+FText DataTableUtils::GetPropertyValueAsText(const UProperty* InProp, const uint8* InData)
 {
 	FText Result;
 
diff --git a/Engine/Source/Runtime/Engine/Public/DataTableUtils.h b/Engine/Source/Runtime/Engine/Public/DataTableUtils.h
--- a/Engine/Source/Runtime/Engine/Public/DataTableUtils.h
+++ b/Engine/Source/Runtime/Engine/Public/DataTableUtils.h
@@ -43,25 +43,25 @@ namespace DataTableUtils
 	 * Util to get a property as a string.
 	 * This always gets a string for the given property without adjusting the address.
 	 */
-	ENGINE_API FString GetPropertyValueAsStringDirect(const UProperty* InProp, uint8* InData, const EDataTableExportFlags InDTExportFlags);
+	ENGINE_API FString GetPropertyValueAsStringDirect(const UProperty* InProp, const uint8* InData, const EDataTableExportFlags InDTExportFlags);
 
 	/** 
 	 * Util to get a property as a string.
 	 * When the property is a static sized array, this will return a string containing each element in the array.
 	 */
-	ENGINE_API FString GetPropertyValueAsString(const UProperty* InProp, uint8* InData, const EDataTableExportFlags InDTExportFlags);
+	ENGINE_API FString GetPropertyValueAsString(const UProperty* InProp, const uint8* InData, const EDataTableExportFlags InDTExportFlags);
 
 	/** 
 	 * Util to get a property as text (this will use the display name of the value where available - use GetPropertyValueAsString if you need an internal identifier).
 	 * This always gets a string for the given property without adjusting the address.
 	 */
-	ENGINE_API FText GetPropertyValueAsTextDirect(const UProperty* InProp, uint8* InData);
+	ENGINE_API FText GetPropertyValueAsTextDirect(const UProperty* InProp, const uint8* InData);
 
 	/** 
 	 * Util to get a property as text (this will use the display name of the value where available - use GetPropertyValueAsString if you need an internal identifier).
 	 * When the property is a static sized array, this will return a string containing each element in the array. 
 	 */
-	ENGINE_API FText GetPropertyValueAsText(const UProperty* InProp, uint8* InData);
+	ENGINE_API FText GetPropertyValueAsText(const UProperty* InProp, const uint8* InData);
 
 	/**
 	 * Util to get all property names from a struct.
