Fixed SceneCaptureComponent so ShowOnlyActors property is writeable in blueprints.

diff --git a/Engine/Source/Runtime/Engine/Classes/Components/SceneCaptureComponent.h b/Engine/Source/Runtime/Engine/Classes/Components/SceneCaptureComponent.h
--- a/Engine/Source/Runtime/Engine/Classes/Components/SceneCaptureComponent.h
+++ b/Engine/Source/Runtime/Engine/Classes/Components/SceneCaptureComponent.h
@@ -78,7 +78,7 @@ class ENGINE_API USceneCaptureComponent : public USceneComponent
  	TArray<TWeakObjectPtr<UPrimitiveComponent> > ShowOnlyComponents;
 
 	/** The only actors to be rendered by this scene capture, if PrimitiveRenderMode is set to UseShowOnlyList.*/
-	UPROPERTY(EditInstanceOnly, BlueprintReadOnly, Category=SceneCapture)
+	UPROPERTY(EditInstanceOnly, BlueprintReadWrite, Category=SceneCapture)
 	TArray<AActor*> ShowOnlyActors;
 
 	/** Whether to update the capture's contents every frame.  If disabled, the component will render once on load and then only when moved. */
