Added Message Log output for invalid software cursor as opposed to ensure/log.

diff --git a/Engine/Source/Runtime/Engine/Private/GameViewportClient.cpp b/Engine/Source/Runtime/Engine/Private/GameViewportClient.cpp
--- a/Engine/Source/Runtime/Engine/Private/GameViewportClient.cpp
+++ b/Engine/Source/Runtime/Engine/Private/GameViewportClient.cpp
@@ -819,7 +819,7 @@ void UGameViewportClient::SetVirtualCursorWidget(EMouseCursor::Type Cursor, UUse
 
 void UGameViewportClient::AddSoftwareCursor(EMouseCursor::Type Cursor, const FSoftClassPath& CursorClass)
 {
-	if (ensureMsgf(CursorClass.IsValid(), TEXT("UGameViewportClient::AddCusor: Cursor class is not valid!")))
+	if (CursorClass.IsValid())
 	{
 		UClass* Class = CursorClass.TryLoadClass<UUserWidget>();
 		if (Class)
@@ -829,9 +829,13 @@ void UGameViewportClient::AddSoftwareCursor(EMouseCursor::Type Cursor, const FSo
 		}
 		else
 		{
-			UE_LOG(LogPlayerManagement, Warning, TEXT("UGameViewportClient::AddCursor: Could not load cursor class %s."), *CursorClass.GetAssetName());
+			FMessageLog("PIE").Warning(FText::Format(LOCTEXT("AddCursor:LoadFailed", "UGameViewportClient::AddCursor: Could not load cursor class '{0}'."), FText::FromString(CursorClass.GetAssetName())));
 		}
 	}
+	else
+	{
+		FMessageLog("PIE").Warning(LOCTEXT("AddCursor:InvalidClass", "UGameViewportClient::AddCursor: Invalid class specified."));
+	}
 }
 
 void UGameViewportClient::AddCursorWidget(EMouseCursor::Type Cursor, class UUserWidget* CursorWidget)
