[NO NON-PII COMMIT MESSAGE]

diff --git a/Engine/Source/Runtime/Engine/Classes/Components/InstancedStaticMeshComponent.h b/Engine/Source/Runtime/Engine/Classes/Components/InstancedStaticMeshComponent.h
--- a/Engine/Source/Runtime/Engine/Classes/Components/InstancedStaticMeshComponent.h
+++ b/Engine/Source/Runtime/Engine/Classes/Components/InstancedStaticMeshComponent.h
@@ -37,6 +37,7 @@ struct FInstanceUpdateCmdBuffer
 		Update,
 		Hide,
 		EditorData,
+		LightmapData,
 	};
 	
 	struct FInstanceUpdateCommand
@@ -47,6 +48,9 @@ struct FInstanceUpdateCmdBuffer
 		
 		FColor HitProxyColor;
 		bool bSelected;
+
+		FVector2D LightmapUVBias;
+		FVector2D ShadowmapUVBias;
 	};
 	
 	FInstanceUpdateCmdBuffer();
@@ -56,6 +60,8 @@ struct FInstanceUpdateCmdBuffer
 	void AddInstance(const FMatrix& InTransform);
 	void UpdateInstance(int32 RenderIndex, const FMatrix& InTransform);
 	void SetEditorData(int32 RenderIndex, const FColor& Color, bool bSelected);
+	void SetLightMapData(int32 RenderIndex, const FVector2D& LightmapUVBias);
+	void SetShadowMapData(int32 RenderIndex, const FVector2D& ShadowmapUVBias);
 	void ResetInlineCommands();
 	int32 NumInlineCommands() const { return Cmds.Num(); }
 
diff --git a/Engine/Source/Runtime/Engine/Private/InstancedStaticMesh.cpp b/Engine/Source/Runtime/Engine/Private/InstancedStaticMesh.cpp
--- a/Engine/Source/Runtime/Engine/Private/InstancedStaticMesh.cpp
+++ b/Engine/Source/Runtime/Engine/Private/InstancedStaticMesh.cpp
@@ -104,6 +104,58 @@ void FInstanceUpdateCmdBuffer::SetEditorData(int32 RenderIndex, const FColor& Co
 	Edit();
 }
 
+void FInstanceUpdateCmdBuffer::SetLightMapData(int32 RenderIndex, const FVector2D& LightmapUVBias)
+{
+	// We only support 1 command to update lightmap/shadowmap
+	bool CommandExist = false;
+
+	for (FInstanceUpdateCommand& Cmd : Cmds)
+	{
+		if (Cmd.Type == FInstanceUpdateCmdBuffer::LightmapData && Cmd.InstanceIndex == RenderIndex)
+		{
+			CommandExist = true;
+			Cmd.LightmapUVBias = LightmapUVBias;
+			break;
+		}
+	}
+
+	if (!CommandExist)
+	{
+		FInstanceUpdateCommand& Cmd = Cmds.AddDefaulted_GetRef();
+		Cmd.InstanceIndex = RenderIndex;
+		Cmd.Type = FInstanceUpdateCmdBuffer::LightmapData;
+		Cmd.LightmapUVBias = LightmapUVBias;
+	}
+
+	Edit();
+}
+
+void FInstanceUpdateCmdBuffer::SetShadowMapData(int32 RenderIndex, const FVector2D& ShadowmapUVBias)
+{
+	// We only support 1 command to update lightmap/shadowmap
+	bool CommandExist = false;
+
+	for (FInstanceUpdateCommand& Cmd : Cmds)
+	{
+		if (Cmd.Type == FInstanceUpdateCmdBuffer::LightmapData && Cmd.InstanceIndex == RenderIndex)
+		{
+			CommandExist = true;
+			Cmd.ShadowmapUVBias = ShadowmapUVBias;
+			break;
+		}
+	}
+
+	if (!CommandExist)
+	{
+		FInstanceUpdateCommand& Cmd = Cmds.AddDefaulted_GetRef();
+		Cmd.InstanceIndex = RenderIndex;
+		Cmd.Type = FInstanceUpdateCmdBuffer::LightmapData;
+		Cmd.ShadowmapUVBias = ShadowmapUVBias;
+	}
+
+	Edit();
+}
+
 void FInstanceUpdateCmdBuffer::ResetInlineCommands()
 {
 	Cmds.Reset();
@@ -200,6 +252,9 @@ void FStaticMeshInstanceBuffer::UpdateFromCommandBuffer_RenderThread(FInstanceUp
 		case FInstanceUpdateCmdBuffer::EditorData:
 			InstanceData->SetInstanceEditorData(Cmd.InstanceIndex, Cmd.HitProxyColor, Cmd.bSelected);
 			break;
+		case FInstanceUpdateCmdBuffer::LightmapData:
+			InstanceData->SetInstanceLightMapData(Cmd.InstanceIndex, Cmd.LightmapUVBias, Cmd.ShadowmapUVBias);
+			break;
 		default:
 			check(false);
 		}
@@ -1028,8 +1083,6 @@ FPrimitiveSceneProxy* UInstancedStaticMeshComponent::CreateSceneProxy()
 		{
 			InstanceUpdateCmdBuffer.Reset();
 
-			//TODO: should we not use the command buffer here instead of updating all instances?!
-
 			FStaticMeshInstanceData RenderInstanceData = FStaticMeshInstanceData(GVertexElementTypeSupport.IsSupported(VET_Half2));
 			BuildRenderData(RenderInstanceData, PerInstanceRenderData->HitProxies);
 			PerInstanceRenderData->UpdateFromPreallocatedData(RenderInstanceData);
diff --git a/Engine/Source/Runtime/Engine/Private/LightMap.cpp b/Engine/Source/Runtime/Engine/Private/LightMap.cpp
--- a/Engine/Source/Runtime/Engine/Private/LightMap.cpp
+++ b/Engine/Source/Runtime/Engine/Private/LightMap.cpp
@@ -280,7 +280,7 @@ struct FLightMapAllocation
 				// Need to create per-LOD instance data to fix that
 				MeshBuildData->PerInstanceLightmapData[InstanceIndex].LightmapUVBias = LightMap->GetCoordinateBias();
 
-				Component->InstanceUpdateCmdBuffer.Edit();
+				Component->InstanceUpdateCmdBuffer.SetLightMapData(InstanceIndex, MeshBuildData->PerInstanceLightmapData[InstanceIndex].LightmapUVBias);
 				Component->MarkRenderStateDirty();
 			}
 		}
diff --git a/Engine/Source/Runtime/Engine/Private/ShadowMap.cpp b/Engine/Source/Runtime/Engine/Private/ShadowMap.cpp
--- a/Engine/Source/Runtime/Engine/Private/ShadowMap.cpp
+++ b/Engine/Source/Runtime/Engine/Private/ShadowMap.cpp
@@ -110,7 +110,7 @@ struct FShadowMapAllocation
 				// Need to create per-LOD instance data to fix that
 				MeshBuildData->PerInstanceLightmapData[InstanceIndex].ShadowmapUVBias = ShadowMap->GetCoordinateBias();
 
-				Component->InstanceUpdateCmdBuffer.Edit();
+				Component->InstanceUpdateCmdBuffer.SetShadowMapData(InstanceIndex, MeshBuildData->PerInstanceLightmapData[InstanceIndex].ShadowmapUVBias);
 				Component->MarkRenderStateDirty();
 			}
 		}
