Fix for material parameter node crashing

diff --git a/Engine/Source/Runtime/Engine/Private/Materials/Material.cpp b/Engine/Source/Runtime/Engine/Private/Materials/Material.cpp
--- a/Engine/Source/Runtime/Engine/Private/Materials/Material.cpp
+++ b/Engine/Source/Runtime/Engine/Private/Materials/Material.cpp
@@ -2318,7 +2318,7 @@ bool UMaterial::IsScalarParameterUsedAsAtlasPosition(const FMaterialParameterInf
 					OutValue = Parameter->IsUsedAsAtlasPosition();
 					if (OutValue)
 					{
-						UMaterialExpressionCurveAtlasRowParameter* AtlasParameter = Cast<UMaterialExpressionCurveAtlasRowParameter>(Expression);
+						UMaterialExpressionCurveAtlasRowParameter* AtlasParameter = Cast<UMaterialExpressionCurveAtlasRowParameter>(Parameter);
 						Curve = TSoftObjectPtr<UCurveLinearColor>(FSoftObjectPath(AtlasParameter->Curve->GetPathName()));
 						Atlas = TSoftObjectPtr<UCurveLinearColorAtlas>(FSoftObjectPath(AtlasParameter->Atlas->GetPathName()));
 					}
@@ -2338,7 +2338,7 @@ bool UMaterial::IsScalarParameterUsedAsAtlasPosition(const FMaterialParameterInf
 					OutValue = Parameter->IsUsedAsAtlasPosition();
 					if (OutValue)
 					{
-						UMaterialExpressionCurveAtlasRowParameter* AtlasParameter = Cast<UMaterialExpressionCurveAtlasRowParameter>(Expression);
+						UMaterialExpressionCurveAtlasRowParameter* AtlasParameter = Cast<UMaterialExpressionCurveAtlasRowParameter>(Parameter);
 						Curve = TSoftObjectPtr<UCurveLinearColor>(FSoftObjectPath(AtlasParameter->Curve->GetPathName()));
 						Atlas = TSoftObjectPtr<UCurveLinearColorAtlas>(FSoftObjectPath(AtlasParameter->Atlas->GetPathName()));
 					}
