PR #4422: Set default material parameter name (Contributed by [REDACTED_PERSON])

diff --git a/Engine/Source/Runtime/Engine/Classes/Materials/MaterialExpressionParameter.h b/Engine/Source/Runtime/Engine/Classes/Materials/MaterialExpressionParameter.h
--- a/Engine/Source/Runtime/Engine/Classes/Materials/MaterialExpressionParameter.h
+++ b/Engine/Source/Runtime/Engine/Classes/Materials/MaterialExpressionParameter.h
@@ -38,7 +38,7 @@ class UMaterialExpressionParameter : public UMaterialExpression
 
 	//~ Begin UMaterialExpression Interface
 	virtual bool MatchesSearchQuery( const TCHAR* SearchQuery ) override;
-
+	virtual void PostEditChangeProperty(FPropertyChangedEvent& PropertyChangedEvent) override;
 	virtual bool CanRenameNode() const override { return true; }
 	virtual FString GetEditableName() const override;
 	virtual void SetEditableName(const FString& NewName) override;
@@ -46,7 +46,7 @@ class UMaterialExpressionParameter : public UMaterialExpression
 	virtual bool HasAParameterName() const override { return true; }
 	virtual FName GetParameterName() const override { return ParameterName; }
 	virtual void SetParameterName(const FName& Name) override { ParameterName = Name; }
-
+	virtual void ValidateParameterName() override;
 #endif
 	//~ End UMaterialExpression Interface
 
@@ -63,7 +63,9 @@ class UMaterialExpressionParameter : public UMaterialExpression
 	 * Get list of parameter names for static parameter sets
 	 */
 	void GetAllParameterInfo(TArray<FMaterialParameterInfo> &OutParameterInfo, TArray<FGuid> &OutParameterIds, const FMaterialParameterInfo& InBaseParameterInfo) const;
-};
-
-
 
+private:
+#if WITH_EDITOR
+	static FName ParameterDefaultName;
+#endif
+};
diff --git a/Engine/Source/Runtime/Engine/Private/Materials/MaterialExpressions.cpp b/Engine/Source/Runtime/Engine/Private/Materials/MaterialExpressions.cpp
--- a/Engine/Source/Runtime/Engine/Private/Materials/MaterialExpressions.cpp
+++ b/Engine/Source/Runtime/Engine/Private/Materials/MaterialExpressions.cpp
@@ -5977,6 +5977,7 @@ int32 UMaterialExpressionDesaturation::Compile(class FMaterialCompiler* Compiler
 //
 //	UMaterialExpressionParameter
 //
+FName UMaterialExpressionParameter::ParameterDefaultName = TEXT("Param");
 
 UMaterialExpressionParameter::UMaterialExpressionParameter(const FObjectInitializer& ObjectInitializer)
 	: Super(ObjectInitializer)
@@ -5985,14 +5986,17 @@ UMaterialExpressionParameter::UMaterialExpressionParameter(const FObjectInitiali
 	struct FConstructorStatics
 	{
 		FText NAME_Parameters;
+		FName ParameterName;
 		FConstructorStatics()
 			: NAME_Parameters(LOCTEXT( "Parameters", "Parameters" ))
+			, ParameterName(UMaterialExpressionParameter::ParameterDefaultName)
 		{
 		}
 	};
 	static FConstructorStatics ConstructorStatics;
 
 	bIsParameterExpression = true;
+	ParameterName = ConstructorStatics.ParameterName;
 
 #if WITH_EDITORONLY_DATA
 	MenuCategories.Add(ConstructorStatics.NAME_Parameters);
@@ -6014,6 +6018,26 @@ bool UMaterialExpressionParameter::MatchesSearchQuery( const TCHAR* SearchQuery
 	return Super::MatchesSearchQuery(SearchQuery);
 }
 
+void UMaterialExpressionParameter::PostEditChangeProperty(FPropertyChangedEvent& PropertyChangedEvent)
+{
+	UProperty* PropertyThatChanged = PropertyChangedEvent.Property;
+	if (PropertyThatChanged != NULL)
+	{
+		// Update the preview for this node if we adjusted a property
+		bNeedToUpdatePreview = true;
+
+		const FName PropertyName = PropertyThatChanged->GetFName();
+		if (PropertyName == GET_MEMBER_NAME_CHECKED(UMaterialExpressionParameter, ParameterName))
+		{
+			ValidateParameterName();
+		}
+	}
+
+	Super::PostEditChangeProperty(PropertyChangedEvent);
+}
+
+
+
 FString UMaterialExpressionParameter::GetEditableName() const
 {
 	return ParameterName.ToString();
@@ -6037,6 +6061,53 @@ void UMaterialExpressionParameter::GetAllParameterInfo(TArray<FMaterialParameter
 	}
 }
 
+void UMaterialExpressionParameter::ValidateParameterName()
+{
+	if (Material != nullptr)
+	{
+		int32 NameIndex = 1;
+		bool FoundValidName = false;
+		FName PotentialName;
+
+		// Find an available unique name
+		while (!FoundValidName)
+		{
+			PotentialName = GetParameterName();
+
+			// Parameters cannot be named Name_None, use the default name instead
+			if (PotentialName == NAME_None)
+			{
+				PotentialName = UMaterialExpressionParameter::ParameterDefaultName;
+			}
+
+			if (NameIndex != 1)
+			{
+				PotentialName.SetNumber(NameIndex);
+			}
+
+			FoundValidName = true;
+
+			for (UMaterialExpression* Expression : Material->Expressions)
+			{
+				if (Expression != nullptr && Expression->HasAParameterName())
+				{
+					// Name are unique per class type
+					if (Expression != this && Expression->GetClass() == GetClass() && Expression->GetParameterName() == PotentialName)
+					{
+						FoundValidName = false;
+						break;
+					}
+				}
+			}
+
+			++NameIndex;
+		}
+
+		SetParameterName(PotentialName);
+	}
+}
+
+
 bool UMaterialExpressionParameter::NeedsLoadForClient() const
 {
 	// Keep named parameters
