Now all renamable nodes do name verification also

diff --git a/Engine/Source/Editor/UnrealEd/Private/MaterialGraphNode.cpp b/Engine/Source/Editor/UnrealEd/Private/MaterialGraphNode.cpp
--- a/Engine/Source/Editor/UnrealEd/Private/MaterialGraphNode.cpp
+++ b/Engine/Source/Editor/UnrealEd/Private/MaterialGraphNode.cpp
@@ -669,6 +669,25 @@ void UMaterialGraphNode::OnRenameNode(const FString& NewName)
 	MaterialExpression->Modify();
 	SetParameterName(NewName);
 	MaterialExpression->MarkPackageDirty();
+	MaterialExpression->ValidateParameterName();
+	UProperty* NameProperty = nullptr;
+	if (Cast<UMaterialExpressionParameter>(MaterialExpression))
+	{
+		NameProperty = FindFieldChecked<UProperty>(UMaterialExpressionParameter::StaticClass(), GET_MEMBER_NAME_CHECKED(UMaterialExpressionParameter, ParameterName));
+	}
+	else if (Cast<UMaterialExpressionFontSampleParameter>(MaterialExpression))
+	{
+		NameProperty = FindFieldChecked<UProperty>(UMaterialExpressionFontSampleParameter::StaticClass(), GET_MEMBER_NAME_CHECKED(UMaterialExpressionFontSampleParameter, ParameterName));
+	}
+	else if (Cast<UMaterialExpressionTextureSampleParameter>(MaterialExpression))
+	{
+		NameProperty = FindFieldChecked<UProperty>(UMaterialExpressionTextureSampleParameter::StaticClass(), GET_MEMBER_NAME_CHECKED(UMaterialExpressionTextureSampleParameter, ParameterName));
+	}
+	if(NameProperty)
+	{
+		FPropertyChangedEvent PropertyChangeEvent(NameProperty, EPropertyChangeType::ValueSet);
+		MaterialExpression->PostEditChangeProperty(PropertyChangeEvent);
+	}
 	MaterialDirtyDelegate.ExecuteIfBound();
 }
 
diff --git a/Engine/Source/Runtime/Engine/Classes/Materials/MaterialExpressionFontSampleParameter.h b/Engine/Source/Runtime/Engine/Classes/Materials/MaterialExpressionFontSampleParameter.h
--- a/Engine/Source/Runtime/Engine/Classes/Materials/MaterialExpressionFontSampleParameter.h
+++ b/Engine/Source/Runtime/Engine/Classes/Materials/MaterialExpressionFontSampleParameter.h
@@ -47,6 +47,7 @@ class UMaterialExpressionFontSampleParameter : public UMaterialExpressionFontSam
 	virtual bool HasAParameterName() const override { return true; }
 	virtual FName GetParameterName() const override { return ParameterName; }
 	virtual void SetParameterName(const FName& Name) override { ParameterName = Name; }
+	virtual void ValidateParameterName() override;
 #endif
 	//~ End UMaterialExpression Interface
 	
diff --git a/Engine/Source/Runtime/Engine/Classes/Materials/MaterialExpressionParameter.h b/Engine/Source/Runtime/Engine/Classes/Materials/MaterialExpressionParameter.h
--- a/Engine/Source/Runtime/Engine/Classes/Materials/MaterialExpressionParameter.h
+++ b/Engine/Source/Runtime/Engine/Classes/Materials/MaterialExpressionParameter.h
@@ -33,12 +33,13 @@ class UMaterialExpressionParameter : public UMaterialExpression
 	UPROPERTY(EditAnywhere, Category=MaterialExpressionParameter)
 	int32 SortPriority;
 #endif
+	static FName ParameterDefaultName;
+
 
 #if WITH_EDITOR
 
 	//~ Begin UMaterialExpression Interface
 	virtual bool MatchesSearchQuery( const TCHAR* SearchQuery ) override;
-	virtual void PostEditChangeProperty(FPropertyChangedEvent& PropertyChangedEvent) override;
 	virtual bool CanRenameNode() const override { return true; }
 	virtual FString GetEditableName() const override;
 	virtual void SetEditableName(const FString& NewName) override;
@@ -64,8 +65,4 @@ class UMaterialExpressionParameter : public UMaterialExpression
 	 */
 	void GetAllParameterInfo(TArray<FMaterialParameterInfo> &OutParameterInfo, TArray<FGuid> &OutParameterIds, const FMaterialParameterInfo& InBaseParameterInfo) const;
 
-private:
-#if WITH_EDITOR
-	static FName ParameterDefaultName;
-#endif
 };
diff --git a/Engine/Source/Runtime/Engine/Classes/Materials/MaterialExpressionTextureSampleParameter.h b/Engine/Source/Runtime/Engine/Classes/Materials/MaterialExpressionTextureSampleParameter.h
--- a/Engine/Source/Runtime/Engine/Classes/Materials/MaterialExpressionTextureSampleParameter.h
+++ b/Engine/Source/Runtime/Engine/Classes/Materials/MaterialExpressionTextureSampleParameter.h
@@ -46,7 +46,7 @@ class ENGINE_API UMaterialExpressionTextureSampleParameter : public UMaterialExp
 	virtual bool HasAParameterName() const override { return true; }
 	virtual FName GetParameterName() const override { return ParameterName; }
 	virtual void SetParameterName(const FName& Name) override { ParameterName = Name; }
-
+	virtual void ValidateParameterName() override;
 #endif
 	//~ End UMaterialExpression Interface
 
diff --git a/Engine/Source/Runtime/Engine/Private/Materials/MaterialExpressions.cpp b/Engine/Source/Runtime/Engine/Private/Materials/MaterialExpressions.cpp
--- a/Engine/Source/Runtime/Engine/Private/Materials/MaterialExpressions.cpp
+++ b/Engine/Source/Runtime/Engine/Private/Materials/MaterialExpressions.cpp
@@ -375,6 +375,53 @@ bool CanConnectMaterialValueTypes(uint32 InputType, uint32 OutputType)
 
 #if WITH_EDITOR
 
+
+void ValidateParameterNameInternal(class UMaterialExpression* ExpressionToValidate, class UMaterial* OwningMaterial)
+{
+	if (OwningMaterial != nullptr)
+	{
+		int32 NameIndex = 1;
+		bool FoundValidName = false;
+		FName PotentialName;
+
+		// Find an available unique name
+		while (!FoundValidName)
+		{
+			PotentialName = ExpressionToValidate->GetParameterName();
+
+			// Parameters cannot be named Name_None, use the default name instead
+			if (PotentialName == NAME_None)
+			{
+				PotentialName = UMaterialExpressionParameter::ParameterDefaultName;
+			}
+
+			if (NameIndex != 1)
+			{
+				PotentialName.SetNumber(NameIndex);
+			}
+
+			FoundValidName = true;
+
+			for (UMaterialExpression* Expression : OwningMaterial->Expressions)
+			{
+				if (Expression != nullptr && Expression->HasAParameterName())
+				{
+					// Name are unique per class type
+					if (Expression != ExpressionToValidate && Expression->GetClass() == ExpressionToValidate->GetClass() && Expression->GetParameterName() == PotentialName)
+					{
+						FoundValidName = false;
+						break;
+					}
+				}
+			}
+
+			++NameIndex;
+		}
+
+		ExpressionToValidate->SetParameterName(PotentialName);
+	}
+}
+
 /**
  * Helper function that wraps the supplied texture coordinates in the necessary math to transform them for external textures
  *
@@ -672,6 +719,13 @@ void UMaterialExpression::PostEditChangeProperty(FPropertyChangedEvent& Property
 		bNeedToUpdatePreview = true;
 
 		const FName PropertyName = PropertyThatChanged->GetFName();
+
+		const FName ParameterName = TEXT("ParameterName");
+		if (PropertyName == ParameterName)
+		{
+			ValidateParameterName();
+		}
+
 		if (PropertyName == GET_MEMBER_NAME_CHECKED(UMaterialExpression, Desc) && !IsA(UMaterialExpressionComment::StaticClass()))
 		{
 			if (GraphNode)
@@ -1155,42 +1209,7 @@ void UMaterialExpression::SetEditableName(const FString& NewName)
 
 void UMaterialExpression::ValidateParameterName()
 {
-	if (Material != nullptr)
-	{
-		int32 NameIndex = 1;
-		bool FoundValidName = false;
-		FName PotentialName;
-
-		// Find an available unique name
-		while (!FoundValidName)
-		{
-			PotentialName = GetParameterName();
-
-			if (NameIndex != 1)
-			{
-				PotentialName.SetNumber(NameIndex);
-			}
-
-			FoundValidName = true;
-
-			for (UMaterialExpression* Expression : Material->Expressions)
-			{
-				if (Expression != nullptr && Expression->HasAParameterName())
-				{
-					// Name are unique per class type
-					if (Expression != this && Expression->GetClass() == GetClass() && Expression->GetParameterName() == PotentialName)
-					{
-						FoundValidName = false;
-						break;
-					}
-				}
-			}
-
-			++NameIndex;
-		}
-
-		SetParameterName(PotentialName);
-	}
+	// Incrementing the name is now handled in UMaterialExpressionParameter::ValidateParameterName
 }
 
 #endif // WITH_EDITOR
@@ -1865,6 +1884,12 @@ void UMaterialExpressionTextureSampleParameter::GetCaption(TArray<FString>& OutC
 	OutCaptions.Add(TEXT("Texture Param")); 
 	OutCaptions.Add(FString::Printf(TEXT("'%s'"), *ParameterName.ToString()));
 }
+
+void UMaterialExpressionTextureSampleParameter::ValidateParameterName()
+{
+	ValidateParameterNameInternal(this, Material);
+}
+
 #endif // WITH_EDITOR
 
 bool UMaterialExpressionTextureSampleParameter::IsNamedParameter(const FMaterialParameterInfo& ParameterInfo, UTexture*& OutValue) const
@@ -6018,24 +6043,6 @@ bool UMaterialExpressionParameter::MatchesSearchQuery( const TCHAR* SearchQuery
 	return Super::MatchesSearchQuery(SearchQuery);
 }
 
-void UMaterialExpressionParameter::PostEditChangeProperty(FPropertyChangedEvent& PropertyChangedEvent)
-{
-	UProperty* PropertyThatChanged = PropertyChangedEvent.Property;
-	if (PropertyThatChanged != NULL)
-	{
-		// Update the preview for this node if we adjusted a property
-		bNeedToUpdatePreview = true;
-
-		const FName PropertyName = PropertyThatChanged->GetFName();
-		if (PropertyName == GET_MEMBER_NAME_CHECKED(UMaterialExpressionParameter, ParameterName))
-		{
-			ValidateParameterName();
-		}
-	}
-
-	Super::PostEditChangeProperty(PropertyChangedEvent);
-}
-
 
 
 FString UMaterialExpressionParameter::GetEditableName() const
@@ -6061,52 +6068,12 @@ void UMaterialExpressionParameter::GetAllParameterInfo(TArray<FMaterialParameter
 	}
 }
 
+#if WITH_EDITOR
 void UMaterialExpressionParameter::ValidateParameterName()
 {
-	if (Material != nullptr)
-	{
-		int32 NameIndex = 1;
-		bool FoundValidName = false;
-		FName PotentialName;
-
-		// Find an available unique name
-		while (!FoundValidName)
-		{
-			PotentialName = GetParameterName();
-
-			// Parameters cannot be named Name_None, use the default name instead
-			if (PotentialName == NAME_None)
-			{
-				PotentialName = UMaterialExpressionParameter::ParameterDefaultName;
-			}
-
-			if (NameIndex != 1)
-			{
-				PotentialName.SetNumber(NameIndex);
-			}
-
-			FoundValidName = true;
-
-			for (UMaterialExpression* Expression : Material->Expressions)
-			{
-				if (Expression != nullptr && Expression->HasAParameterName())
-				{
-					// Name are unique per class type
-					if (Expression != this && Expression->GetClass() == GetClass() && Expression->GetParameterName() == PotentialName)
-					{
-						FoundValidName = false;
-						break;
-					}
-				}
-			}
-
-			++NameIndex;
-		}
-
-		SetParameterName(PotentialName);
-	}
+	ValidateParameterNameInternal(this, Material);
 }
-
+#endif
 
 bool UMaterialExpressionParameter::NeedsLoadForClient() const
 {
@@ -8789,6 +8756,11 @@ void UMaterialExpressionFontSampleParameter::GetCaption(TArray<FString>& OutCapt
 	OutCaptions.Add(TEXT("Font Param")); 
 	OutCaptions.Add(FString::Printf(TEXT("'%s'"), *ParameterName.ToString()));
 }
+
+void UMaterialExpressionFontSampleParameter::ValidateParameterName()
+{
+	ValidateParameterNameInternal(this, Material);
+}
 #endif // WITH_EDITOR
 
 bool UMaterialExpressionFontSampleParameter::IsNamedParameter(const FMaterialParameterInfo& ParameterInfo, UFont*& OutFontValue, int32& OutFontPage) const
