Fixing input labels on material function inputs

diff --git a/Engine/Source/Runtime/Engine/Private/Materials/MaterialExpressions.cpp b/Engine/Source/Runtime/Engine/Private/Materials/MaterialExpressions.cpp
--- a/Engine/Source/Runtime/Engine/Private/Materials/MaterialExpressions.cpp
+++ b/Engine/Source/Runtime/Engine/Private/Materials/MaterialExpressions.cpp
@@ -11442,6 +11442,7 @@ void UMaterialExpressionFunctionInput::GetCaption(TArray<FString>& OutCaptions)
 		TEXT("Vector4"),
 		TEXT("Texture2D"),
 		TEXT("TextureCube"),
+		TEXT("VolumeTexture"),
 		TEXT("StaticBool"),
 		TEXT("MaterialAttributes"),
 		TEXT("External")
