Fixing issue with bIsImmortal and bWillBecomeZombie

Issue mentioned in this UDN.
https://udn.unrealengine.com/questions/446170/cascade-immortalzombie-bug.html

diff --git a/Engine/Source/Runtime/Engine/Private/Particles/ParticleComponents.cpp b/Engine/Source/Runtime/Engine/Private/Particles/ParticleComponents.cpp
--- a/Engine/Source/Runtime/Engine/Private/Particles/ParticleComponents.cpp
+++ b/Engine/Source/Runtime/Engine/Private/Particles/ParticleComponents.cpp
@@ -2170,6 +2170,8 @@ void UParticleSystem::PostEditChangeProperty(FPropertyChangedEvent& PropertyChan
 
 	// Recompute the looping flag
 	bAnyEmitterLoopsForever = false;
+	bIsImmortal = false;
+	bWillBecomeZombie = false;
 	HighestSignificance = EParticleSignificanceLevel::Low;
 	LowestSignificance = EParticleSignificanceLevel::Critical;
 	for (UParticleEmitter* Emitter : Emitters)
@@ -2201,7 +2203,6 @@ void UParticleSystem::PostEditChangeProperty(FPropertyChangedEvent& PropertyChan
 									bWillBecomeZombie = true;
 								}
 							}
-
 						}
 					}
 				}
@@ -2276,6 +2277,8 @@ void UParticleSystem::PostLoad()
 	// Run thru all of the emitters, load them up and compute some flags based on them
 	bHasPhysics = false;
 	bAnyEmitterLoopsForever = false;
+	bIsImmortal = false;
+	bWillBecomeZombie = false;
 	HighestSignificance = EParticleSignificanceLevel::Low;
 	LowestSignificance = EParticleSignificanceLevel::Critical;
 	for (int32 i = Emitters.Num() - 1; i >= 0; i--)
@@ -2314,6 +2317,7 @@ void UParticleSystem::PostLoad()
 			LODLevel->ConditionalPostLoad();
 				
 			//@todo. Move these flag calculations into the editor and serialize?
+			//Should mirror implementation in PostEditChangeProperty.
 			for (UParticleLODLevel* ParticleLODLevel : Emitter->LODLevels)
 			{
 				if (ParticleLODLevel)
@@ -2339,17 +2343,17 @@ void UParticleSystem::PostLoad()
 						if ((RequiredModule != nullptr) && (RequiredModule->EmitterLoops == 0))
 						{
 							bAnyEmitterLoopsForever = true;
-						}
 
-						UParticleModuleSpawn* SpawnModule = LODLevel->SpawnModule;
-						check(SpawnModule);
+							UParticleModuleSpawn* SpawnModule = LODLevel->SpawnModule;
+							check(SpawnModule);
 
-						if (RequiredModule->EmitterDuration == 0.0f)
-						{
-							bIsImmortal = true;
-							if (SpawnModule->GetMaximumSpawnRate() == 0.0f && !bAutoDeactivate)
+							if (RequiredModule->EmitterDuration == 0.0f)
 							{
-								bWillBecomeZombie = true;
+								bIsImmortal = true;
+								if (SpawnModule->GetMaximumSpawnRate() == 0.0f && !bAutoDeactivate)
+								{
+									bWillBecomeZombie = true;
+								}
 							}
 						}
 					}
