Interpreter Optimization 7 of N: CPF_ZeroConstructor assigned for struct properties that that are User Defined Structs

diff --git a/Engine/Source/Editor/UnrealEd/Private/UserDefinedStructEditorData.cpp b/Engine/Source/Editor/UnrealEd/Private/UserDefinedStructEditorData.cpp
--- a/Engine/Source/Editor/UnrealEd/Private/UserDefinedStructEditorData.cpp
+++ b/Engine/Source/Editor/UnrealEd/Private/UserDefinedStructEditorData.cpp
@@ -181,6 +181,9 @@ void UUserDefinedStructEditorData::RecreateDefaultInstance(FString* OutLog)
 			}
 		}
 	}
+
+	// Make sure StructFlags are in sync with the new default instance:
+	ScriptStruct->UpdateStructFlags();
 }
 
 void UUserDefinedStructEditorData::CleanDefaultInstance()
diff --git a/Engine/Source/Runtime/CoreUObject/Private/UObject/PropertyStruct.cpp b/Engine/Source/Runtime/CoreUObject/Private/UObject/PropertyStruct.cpp
--- a/Engine/Source/Runtime/CoreUObject/Private/UObject/PropertyStruct.cpp
+++ b/Engine/Source/Runtime/CoreUObject/Private/UObject/PropertyStruct.cpp
@@ -84,6 +84,11 @@ void UStructProperty::LinkInternal(FArchive& Ar)
 		// User Defined structs won't have UScriptStruct::ICppStructOps. Setting their flags here.
 		PropertyFlags |= CPF_HasGetValueTypeHash;
 	}
+
+	if (Struct->StructFlags & STRUCT_ZeroConstructor)
+	{
+		PropertyFlags |= CPF_ZeroConstructor;
+	}
 }
 
 bool UStructProperty::Identical( const void* A, const void* B, uint32 PortFlags ) const
diff --git a/Engine/Source/Runtime/Engine/Classes/Engine/UserDefinedStruct.h b/Engine/Source/Runtime/Engine/Classes/Engine/UserDefinedStruct.h
--- a/Engine/Source/Runtime/Engine/Classes/Engine/UserDefinedStruct.h
+++ b/Engine/Source/Runtime/Engine/Classes/Engine/UserDefinedStruct.h
@@ -124,4 +124,6 @@ public:
 	/** returns references from default instance */
 	static void AddReferencedObjects(UObject* InThis, FReferenceCollector& Collector);
 
+	/** Inspects properties and default values, setting appropriate StructFlags */
+	void UpdateStructFlags();
 };
diff --git a/Engine/Source/Runtime/Engine/Private/UserDefinedStruct.cpp b/Engine/Source/Runtime/Engine/Private/UserDefinedStruct.cpp
--- a/Engine/Source/Runtime/Engine/Private/UserDefinedStruct.cpp
+++ b/Engine/Source/Runtime/Engine/Private/UserDefinedStruct.cpp
@@ -126,6 +126,13 @@ void UUserDefinedStruct::Serialize(FStructuredArchive::FRecord Record)
 
 			FScopedPlaceholderRawContainerTracker TrackStruct(StructData);
 			SerializeItem(Record.EnterField(FIELD_NAME_TEXT("Data")), StructData, nullptr);
+
+			// Now that defaults have been loaded we can inspect our properties
+			// and default values and set the StructFlags accordingly:
+			if(UnderlyingArchive.IsLoading())
+			{
+				UpdateStructFlags();
+			}
 		}
 	}
 
@@ -434,4 +441,48 @@ void UUserDefinedStruct::AddReferencedObjects(UObject* InThis, FReferenceCollect
 	}
 
 	Super::AddReferencedObjects(This, Collector);
-}
\ No newline at end of file
+}
+
+void UUserDefinedStruct::UpdateStructFlags()
+{
+	// Adapted from PrepareCppStructOps, where we 'discover' zero constructability
+	// for native types:
+	bool bIsZeroConstruct = true;
+	
+	uint8* StructData = DefaultStructInstance.GetStructMemory();
+	if (StructData)
+	{
+		int32 Size = GetStructureSize();
+		for (int32 Index = 0; Index < Size; Index++)
+		{
+			if (StructData[Index])
+			{
+				bIsZeroConstruct = false;
+				break;
+			}
+		}
+	}
+
+	if(bIsZeroConstruct)
+	{	
+		for (TFieldIterator<UProperty> It(this); It; ++It)
+		{
+			UProperty* Property = *It;
+			if (Property && !Property->HasAnyPropertyFlags(CPF_ZeroConstructor))
+			{
+				bIsZeroConstruct = false;
+				break;
+			}
+		}
+	}
+
+	if(!bIsZeroConstruct)
+	{
+		StructFlags = EStructFlags(StructFlags & ~STRUCT_ZeroConstructor);
+	}
+	else
+	{
+		StructFlags = EStructFlags(StructFlags | STRUCT_ZeroConstructor);
+	}
+
+}
