Interpreter Optimization 8 of N: CPF_IsPlainOldData and CPF_NoDestructor assigned for struct properties that that are User Defined Structs

diff --git a/Engine/Source/Runtime/CoreUObject/Private/UObject/PropertyStruct.cpp b/Engine/Source/Runtime/CoreUObject/Private/UObject/PropertyStruct.cpp
--- a/Engine/Source/Runtime/CoreUObject/Private/UObject/PropertyStruct.cpp
+++ b/Engine/Source/Runtime/CoreUObject/Private/UObject/PropertyStruct.cpp
@@ -89,6 +89,14 @@ void UStructProperty::LinkInternal(FArchive& Ar)
 	{
 		PropertyFlags |= CPF_ZeroConstructor;
 	}
+	if (Struct->StructFlags & STRUCT_IsPlainOldData)
+	{
+		PropertyFlags |= CPF_IsPlainOldData;
+	}
+	if (Struct->StructFlags & STRUCT_NoDestructor)
+	{
+		PropertyFlags |= CPF_NoDestructor;
+	}
 }
 
 bool UStructProperty::Identical( const void* A, const void* B, uint32 PortFlags ) const
diff --git a/Engine/Source/Runtime/Engine/Private/UserDefinedStruct.cpp b/Engine/Source/Runtime/Engine/Private/UserDefinedStruct.cpp
--- a/Engine/Source/Runtime/Engine/Private/UserDefinedStruct.cpp
+++ b/Engine/Source/Runtime/Engine/Private/UserDefinedStruct.cpp
@@ -448,41 +448,80 @@ void UUserDefinedStruct::UpdateStructFlags()
 	// Adapted from PrepareCppStructOps, where we 'discover' zero constructability
 	// for native types:
 	bool bIsZeroConstruct = true;
-	
-	uint8* StructData = DefaultStructInstance.GetStructMemory();
-	if (StructData)
 	{
-		int32 Size = GetStructureSize();
-		for (int32 Index = 0; Index < Size; Index++)
+		uint8* StructData = DefaultStructInstance.GetStructMemory();
+		if (StructData)
 		{
-			if (StructData[Index])
+			int32 Size = GetStructureSize();
+			for (int32 Index = 0; Index < Size; Index++)
 			{
-				bIsZeroConstruct = false;
-				break;
+				if (StructData[Index])
+				{
+					bIsZeroConstruct = false;
+					break;
+				}
+			}
+		}
+
+		if(bIsZeroConstruct)
+		{	
+			for (TFieldIterator<UProperty> It(this); It; ++It)
+			{
+				UProperty* Property = *It;
+				if (Property && !Property->HasAnyPropertyFlags(CPF_ZeroConstructor))
+				{
+					bIsZeroConstruct = false;
+					break;
+				}
 			}
 		}
 	}
 
-	if(bIsZeroConstruct)
-	{	
+	// IsPOD/NoDtor could be derived earlier than bIsZeroConstruct because they do not depend on 
+	// the structs default values, but it is convenient to calculate them all in one place:
+	bool bIsPOD = true;
+	{
 		for (TFieldIterator<UProperty> It(this); It; ++It)
 		{
 			UProperty* Property = *It;
-			if (Property && !Property->HasAnyPropertyFlags(CPF_ZeroConstructor))
+			if (Property && !Property->HasAnyPropertyFlags(CPF_IsPlainOldData))
 			{
-				bIsZeroConstruct = false;
+				bIsPOD = false;
 				break;
 			}
 		}
 	}
 
+	bool bHasNoDtor = bIsPOD;
+	{
+		if(!bHasNoDtor)
+		{
+			// we're not POD, but we still may have no destructor, check properties:
+			bHasNoDtor = true;
+			for (TFieldIterator<UProperty> It(this); It; ++It)
+			{
+				UProperty* Property = *It;
+				if (Property && !Property->HasAnyPropertyFlags(CPF_NoDestructor))
+				{
+					bHasNoDtor = false;
+					break;
+				}
+			}
+		}
+	}
+	
+	StructFlags = EStructFlags(StructFlags | STRUCT_ZeroConstructor | STRUCT_IsPlainOldData | STRUCT_NoDestructor);
 	if(!bIsZeroConstruct)
 	{
 		StructFlags = EStructFlags(StructFlags & ~STRUCT_ZeroConstructor);
 	}
-	else
+	if(!bIsPOD)
+	{
+		StructFlags = EStructFlags(StructFlags & ~STRUCT_IsPlainOldData);
+	}
+	if(!bHasNoDtor)
 	{
-		StructFlags = EStructFlags(StructFlags | STRUCT_ZeroConstructor);
+		StructFlags = EStructFlags(StructFlags & ~STRUCT_NoDestructor);
 	}
 
 }
