[NO NON-PII COMMIT MESSAGE]

diff --git a/Engine/Source/Runtime/Engine/Private/LevelStreaming.cpp b/Engine/Source/Runtime/Engine/Private/LevelStreaming.cpp
--- a/Engine/Source/Runtime/Engine/Private/LevelStreaming.cpp
+++ b/Engine/Source/Runtime/Engine/Private/LevelStreaming.cpp
@@ -425,7 +425,11 @@ bool ULevelStreaming::DetermineTargetState()
 		{
 			TargetState = ETargetState::UnloadedAndRemoved;
 		}
-		else if (!GUseBackgroundLevelStreaming && !World->GetShouldForceUnloadStreamingLevels())
+		else if (World->GetShouldForceUnloadStreamingLevels())
+		{
+			bContinueToConsider = false;
+		}
+		else if (!GUseBackgroundLevelStreaming)
 		{
 			TargetState = ETargetState::LoadedNotVisible;
 		}
@@ -444,7 +448,7 @@ bool ULevelStreaming::DetermineTargetState()
 		break;
 
 	case ECurrentState::LoadedNotVisible:
-		if (bIsRequestingUnloadAndRemoval)
+		if (bIsRequestingUnloadAndRemoval || World->GetShouldForceUnloadStreamingLevels())
 		{
 			TargetState = ETargetState::Unloaded;
 		}
@@ -467,7 +471,7 @@ bool ULevelStreaming::DetermineTargetState()
 		break;
 
 	case ECurrentState::LoadedVisible:
-		if (bIsRequestingUnloadAndRemoval)
+		if (bIsRequestingUnloadAndRemoval || World->GetShouldForceUnloadStreamingLevels())
 		{
 			TargetState = ETargetState::LoadedNotVisible;
 		}
