Warn users when the outputs of a node that is pruned are used

diff --git a/Engine/Source/Editor/KismetCompiler/Private/KismetCompiler.cpp b/Engine/Source/Editor/KismetCompiler/Private/KismetCompiler.cpp
--- a/Engine/Source/Editor/KismetCompiler/Private/KismetCompiler.cpp
+++ b/Engine/Source/Editor/KismetCompiler/Private/KismetCompiler.cpp
@@ -1399,7 +1399,7 @@ void FKismetCompilerContext::PruneIsolatedNodes(const TArray<UEdGraphNode*>& Roo
 	}
 
 	const UEdGraphSchema* const K2Schema = UEdGraphSchema_K2::StaticClass()->GetDefaultObject<UEdGraphSchema_K2>();
-
+	TMap<UEdGraphNode*, TArray<UEdGraphNode*>> PrunedExecNodeNeighbors;
 	for (int32 NodeIndex = 0; NodeIndex < GraphNodes.Num(); ++NodeIndex)
 	{
 		UEdGraphNode* Node = GraphNodes[NodeIndex];
@@ -1434,6 +1434,9 @@ void FKismetCompilerContext::PruneIsolatedNodes(const TArray<UEdGraphNode*>& Roo
 			{
 				if (Node)
 				{
+					// Track nodes that are directly connected to the outputs of the node we are pruning so 
+					// that we can warn if one or more of those neighboring nodes are not also orphaned:
+					Node->ForEachNodeDirectlyConnectedToOutputs([&PrunedExecNodeNeighbors, Node](UEdGraphNode* NeighborNode) { PrunedExecNodeNeighbors.FindOrAdd(Node).Add(NeighborNode); });
 					Node->BreakAllNodeLinks();
 				}
 				GraphNodes.RemoveAtSwap(NodeIndex);
@@ -1476,6 +1479,24 @@ void FKismetCompilerContext::PruneIsolatedNodes(const TArray<UEdGraphNode*>& Roo
 			}
 		}
 	}
+
+	for(const TPair<UEdGraphNode*, TArray<UEdGraphNode*>>& PrunedExecNodeWithNeighbors : PrunedExecNodeNeighbors)
+	{
+		bool bNeighborsNotPruned = false;
+		for(UEdGraphNode* Neighbor : PrunedExecNodeWithNeighbors.Value)
+		{
+			if(GraphNodes.Contains(Neighbor))
+			{
+				bNeighborsNotPruned = true;
+			}
+		}
+
+		if(bNeighborsNotPruned)
+		{
+			// Warn the user if they are attempting to read an output value from a pruned exec node:
+			MessageLog.Warning(*LOCTEXT("PrunedExecNodeAttemptedUse", "@@ was pruned because its Exec pin is not connected, the connected value is not available and will instead be read as default").ToString(), PrunedExecNodeWithNeighbors.Key);
+		}
+	}
 }
 
 /**
diff --git a/Engine/Source/Runtime/Engine/Classes/EdGraph/EdGraphNode.h b/Engine/Source/Runtime/Engine/Classes/EdGraph/EdGraphNode.h
--- a/Engine/Source/Runtime/Engine/Classes/EdGraph/EdGraphNode.h
+++ b/Engine/Source/Runtime/Engine/Classes/EdGraph/EdGraphNode.h
@@ -788,6 +788,15 @@ public:
 	/** Sets a flag if the comment bubble needs to be made visible immediately */
 	void SetMakeCommentBubbleVisible(bool MakeVisible);
 
+	/** Execute a provided function once for each node that is directly connected to this node, will not include the node itself */
+	void ForEachNodeDirectlyConnected(TFunctionRef<void(UEdGraphNode*)> Func);
+
+	/** Execute a provided function once for each node that is directly connected to this node's input pins, will not include the node itself */
+	void ForEachNodeDirectlyConnectedToInputs(TFunctionRef<void(UEdGraphNode*)> Func);
+
+	/** Execute a provided function once for each node that is directly connected to this node's output pins, will not include the node itself */
+	void ForEachNodeDirectlyConnectedToOutputs(TFunctionRef<void(UEdGraphNode*)> Func);
+	
 protected:
 	/**
 	 * Finds the difference in properties of node instance
diff --git a/Engine/Source/Runtime/Engine/Private/EdGraph/EdGraphNode.cpp b/Engine/Source/Runtime/Engine/Private/EdGraph/EdGraphNode.cpp
--- a/Engine/Source/Runtime/Engine/Private/EdGraph/EdGraphNode.cpp
+++ b/Engine/Source/Runtime/Engine/Private/EdGraph/EdGraphNode.cpp
@@ -765,6 +765,77 @@ void UEdGraphNode::SetMakeCommentBubbleVisible(bool MakeVisible)
 	bCommentBubbleMakeVisible = MakeVisible;
 }
 
+void UEdGraphNode::ForEachNodeDirectlyConnected(TFunctionRef<void(UEdGraphNode*)> Func)
+{
+	TSet<UEdGraphNode*> DirectNeighbors;
+	for( UEdGraphPin* Pin : Pins )
+	{
+		if(Pin->LinkedTo.Num() > 0)
+		{
+			for(UEdGraphPin* Connection : Pin->LinkedTo)
+			{
+				// avoid including the current node in the case of a self connection:
+				if(Connection->GetOwningNode() != this)
+				{
+					DirectNeighbors.Add(Connection->GetOwningNode());
+				}
+			}
+		}
+	}
+
+	for(UEdGraphNode* Neighbor : DirectNeighbors)
+	{
+		Func(Neighbor);
+	}
+}
+void UEdGraphNode::ForEachNodeDirectlyConnectedToInputs(TFunctionRef<void(UEdGraphNode*)> Func)
+{
+	TSet<UEdGraphNode*> DirectNeighbors;
+	for( UEdGraphPin* Pin : Pins )
+	{
+		if(Pin->Direction == EGPD_Input && Pin->LinkedTo.Num() > 0)
+		{
+			for(UEdGraphPin* Connection : Pin->LinkedTo)
+			{
+				// avoid including the current node in the case of a self connection:
+				if(Connection->GetOwningNode() != this)
+				{
+					DirectNeighbors.Add(Connection->GetOwningNode());
+				}
+			}
+		}
+	}
+
+	for(UEdGraphNode* Neighbor : DirectNeighbors)
+	{
+		Func(Neighbor);
+	}
+}
+
+void UEdGraphNode::ForEachNodeDirectlyConnectedToOutputs(TFunctionRef<void(UEdGraphNode*)> Func)
+{
+	TSet<UEdGraphNode*> DirectNeighbors;
+	for( UEdGraphPin* Pin : Pins )
+	{
+		if(Pin->Direction == EGPD_Output && Pin->LinkedTo.Num() > 0)
+		{
+			for(UEdGraphPin* Connection : Pin->LinkedTo)
+			{
+				// avoid including the current node in the case of a self connection:
+				if(Connection->GetOwningNode() != this)
+				{
+					DirectNeighbors.Add(Connection->GetOwningNode());
+				}
+			}
+		}
+	}
+
+	for(UEdGraphNode* Neighbor : DirectNeighbors)
+	{
+		Func(Neighbor);
+	}
+}
+
 #endif	//#if WITH_EDITOR
 
 bool UEdGraphNode::IsInDevelopmentMode() const
