Don't warn when a non-pure node that is connected only via its exec is orphaned - fixes overzealous warnings EngineTest
#rnx

diff --git a/Engine/Source/Editor/KismetCompiler/Private/KismetCompiler.cpp b/Engine/Source/Editor/KismetCompiler/Private/KismetCompiler.cpp
--- a/Engine/Source/Editor/KismetCompiler/Private/KismetCompiler.cpp
+++ b/Engine/Source/Editor/KismetCompiler/Private/KismetCompiler.cpp
@@ -1436,7 +1436,17 @@ void FKismetCompilerContext::PruneIsolatedNodes(const TArray<UEdGraphNode*>& Roo
 				{
 					// Track nodes that are directly connected to the outputs of the node we are pruning so 
 					// that we can warn if one or more of those neighboring nodes are not also orphaned:
-					Node->ForEachNodeDirectlyConnectedToOutputs([&PrunedExecNodeNeighbors, Node](UEdGraphNode* NeighborNode) { PrunedExecNodeNeighbors.FindOrAdd(Node).Add(NeighborNode); });
+					Node->ForEachNodeDirectlyConnectedIf(
+						// Consider connections on output pins other than the exec pin:
+						[](const UEdGraphPin* Pin) {
+							if(Pin->Direction == EGPD_Output && Pin->PinType.PinCategory != UEdGraphSchema_K2::PC_Exec) 
+							{ 
+								return true; 
+							}
+							return false;
+						},
+						[&PrunedExecNodeNeighbors, Node](UEdGraphNode* NeighborNode) { PrunedExecNodeNeighbors.FindOrAdd(Node).Add(NeighborNode); }
+					);
 					Node->BreakAllNodeLinks();
 				}
 				GraphNodes.RemoveAtSwap(NodeIndex);
diff --git a/Engine/Source/Runtime/Engine/Classes/EdGraph/EdGraphNode.h b/Engine/Source/Runtime/Engine/Classes/EdGraph/EdGraphNode.h
--- a/Engine/Source/Runtime/Engine/Classes/EdGraph/EdGraphNode.h
+++ b/Engine/Source/Runtime/Engine/Classes/EdGraph/EdGraphNode.h
@@ -790,11 +790,24 @@ public:
 
 	/** Execute a provided function once for each node that is directly connected to this node, will not include the node itself */
 	void ForEachNodeDirectlyConnected(TFunctionRef<void(UEdGraphNode*)> Func);
+	
+	/** 
+	 * Often we are only interested in a subset of our connections (e.g. only output pins, or only output pins except our exec pin)
+	 * This function provides the ability to execute a provided function once for each node that is directly connect to this node, but
+	 * first filters out out which of this node's pins to consider:
+	 */
+	void ForEachNodeDirectlyConnectedIf(TFunctionRef<bool(const UEdGraphPin* Pin)> Filter, TFunctionRef<void(UEdGraphNode*)> Func);
 
-	/** Execute a provided function once for each node that is directly connected to this node's input pins, will not include the node itself */
+	/** 
+	 * Execute a provided function once for each node that is directly connected to this node's input pins, will not include the node itself
+	 * Implementation provides an example usage of ForEachNodeDirectlyConnectedIf.
+	 */
 	void ForEachNodeDirectlyConnectedToInputs(TFunctionRef<void(UEdGraphNode*)> Func);
 
-	/** Execute a provided function once for each node that is directly connected to this node's output pins, will not include the node itself */
+	/** 
+	 * Execute a provided function once for each node that is directly connected to this node's output pins, will not include the node itself
+	 * Implementation provides an example usage of ForEachNodeDirectlyConnectedIf.
+	 */
 	void ForEachNodeDirectlyConnectedToOutputs(TFunctionRef<void(UEdGraphNode*)> Func);
 	
 protected:
diff --git a/Engine/Source/Runtime/Engine/Private/EdGraph/EdGraphNode.cpp b/Engine/Source/Runtime/Engine/Private/EdGraph/EdGraphNode.cpp
--- a/Engine/Source/Runtime/Engine/Private/EdGraph/EdGraphNode.cpp
+++ b/Engine/Source/Runtime/Engine/Private/EdGraph/EdGraphNode.cpp
@@ -788,52 +788,59 @@ void UEdGraphNode::ForEachNodeDirectlyConnected(TFunctionRef<void(UEdGraphNode*)
 		Func(Neighbor);
 	}
 }
-void UEdGraphNode::ForEachNodeDirectlyConnectedToInputs(TFunctionRef<void(UEdGraphNode*)> Func)
+
+void UEdGraphNode::ForEachNodeDirectlyConnectedIf(TFunctionRef<bool(const UEdGraphPin* Pin)> Filter, TFunctionRef<void(UEdGraphNode*)> Func)
 {
-	TSet<UEdGraphNode*> DirectNeighbors;
+	TSet<UEdGraphNode*> NeighborsAcceptedForConsideration;
 	for( UEdGraphPin* Pin : Pins )
 	{
-		if(Pin->Direction == EGPD_Input && Pin->LinkedTo.Num() > 0)
+		if(Pin->LinkedTo.Num() > 0 && Filter(Pin))
 		{
 			for(UEdGraphPin* Connection : Pin->LinkedTo)
 			{
 				// avoid including the current node in the case of a self connection:
 				if(Connection->GetOwningNode() != this)
 				{
-					DirectNeighbors.Add(Connection->GetOwningNode());
+					NeighborsAcceptedForConsideration.Add(Connection->GetOwningNode());
 				}
 			}
 		}
 	}
 
-	for(UEdGraphNode* Neighbor : DirectNeighbors)
+	for(UEdGraphNode* Neighbor : NeighborsAcceptedForConsideration)
 	{
 		Func(Neighbor);
 	}
 }
 
-void UEdGraphNode::ForEachNodeDirectlyConnectedToOutputs(TFunctionRef<void(UEdGraphNode*)> Func)
+void UEdGraphNode::ForEachNodeDirectlyConnectedToInputs(TFunctionRef<void(UEdGraphNode*)> Func)
 {
-	TSet<UEdGraphNode*> DirectNeighbors;
-	for( UEdGraphPin* Pin : Pins )
-	{
-		if(Pin->Direction == EGPD_Output && Pin->LinkedTo.Num() > 0)
-		{
-			for(UEdGraphPin* Connection : Pin->LinkedTo)
+	ForEachNodeDirectlyConnectedIf(
+		[](const UEdGraphPin* Pin)
+		{ 
+			if(Pin->Direction == EGPD_Input)
 			{
-				// avoid including the current node in the case of a self connection:
-				if(Connection->GetOwningNode() != this)
-				{
-					DirectNeighbors.Add(Connection->GetOwningNode());
-				}
+				return true;
 			}
-		}
-	}
+			return false;
+		},
+		Func
+	);
+}
 
-	for(UEdGraphNode* Neighbor : DirectNeighbors)
-	{
-		Func(Neighbor);
-	}
+void UEdGraphNode::ForEachNodeDirectlyConnectedToOutputs(TFunctionRef<void(UEdGraphNode*)> Func)
+{
+	ForEachNodeDirectlyConnectedIf(
+		[](const UEdGraphPin* Pin)
+		{ 
+			if(Pin->Direction == EGPD_Output)
+			{
+				return true;
+			}
+			return false;
+		},
+		Func
+	);
 }
 
 #endif	//#if WITH_EDITOR
