Comment typo fixes
#rnx

diff --git a/Engine/Source/Runtime/Engine/Classes/EdGraph/EdGraphNode.h b/Engine/Source/Runtime/Engine/Classes/EdGraph/EdGraphNode.h
--- a/Engine/Source/Runtime/Engine/Classes/EdGraph/EdGraphNode.h
+++ b/Engine/Source/Runtime/Engine/Classes/EdGraph/EdGraphNode.h
@@ -793,8 +793,8 @@ public:
 	
 	/** 
 	 * Often we are only interested in a subset of our connections (e.g. only output pins, or only output pins except our exec pin)
-	 * This function provides the ability to execute a provided function once for each node that is directly connect to this node, but
-	 * first filters out out which of this node's pins to consider:
+	 * This function provides the ability to execute a provided function once for each node that is directly connected to this node, 
+	 * but first filters out which of this node's pins to consider:
 	 */
 	void ForEachNodeDirectlyConnectedIf(TFunctionRef<bool(const UEdGraphPin* Pin)> Filter, TFunctionRef<void(UEdGraphNode*)> Func);
 
