PR #5049: Make GetPlayerControllerID pure (Contributed by [REDACTED_PERSON])

diff --git a/Engine/Source/Runtime/Engine/Classes/Kismet/GameplayStatics.h b/Engine/Source/Runtime/Engine/Classes/Kismet/GameplayStatics.h
--- a/Engine/Source/Runtime/Engine/Classes/Kismet/GameplayStatics.h
+++ b/Engine/Source/Runtime/Engine/Classes/Kismet/GameplayStatics.h
@@ -141,7 +141,7 @@ class ENGINE_API UGameplayStatics : public UBlueprintFunctionLibrary
 	* @param Player		The player controller of the player to get the ID of
 	* @return			The ID of the passed in player. -1 if there is no controller for the passed in player
 	*/
-	UFUNCTION(BlueprintCallable, Category="Game", meta=(UnsafeDuringActorConstruction="true"))
+	UFUNCTION(BlueprintPure, Category="Game", meta=(UnsafeDuringActorConstruction="true"))
 	static int32 GetPlayerControllerID(APlayerController* Player);
 
 	/** 
