Don't check when fetching world settings from IsPaused as it is already set up to handle it coming back null
#rnx

diff --git a/Engine/Source/Runtime/Engine/Private/LevelTick.cpp b/Engine/Source/Runtime/Engine/Private/LevelTick.cpp
--- a/Engine/Source/Runtime/Engine/Private/LevelTick.cpp
+++ b/Engine/Source/Runtime/Engine/Private/LevelTick.cpp
@@ -448,7 +448,7 @@ void UWorld::TickNetClient( float DeltaSeconds )
 bool UWorld::IsPaused() const
 {
 	// pause if specifically set or if we're waiting for the end of the tick to perform streaming level loads (so actors don't fall through the world in the meantime, etc)
-	const AWorldSettings* Info = GetWorldSettings();
+	const AWorldSettings* Info = GetWorldSettings(/*bCheckStreamingPersistent=*/false, /*bChecked=*/false);
 	return ( (Info && Info->Pauser != nullptr && TimeSeconds >= PauseDelay) ||
 				(bRequestedBlockOnAsyncLoading && GetNetMode() == NM_Client) ||
 				(GEngine->ShouldCommitPendingMapChange(this)) ||
