ChildActor's need to be spawned from the ChildActorTemplate otherwise any non-replicated  customized properties will be incorrect on the Client.

diff --git a/Engine/Source/Runtime/Engine/Private/PackageMapClient.cpp b/Engine/Source/Runtime/Engine/Private/PackageMapClient.cpp
--- a/Engine/Source/Runtime/Engine/Private/PackageMapClient.cpp
+++ b/Engine/Source/Runtime/Engine/Private/PackageMapClient.cpp
@@ -24,6 +24,7 @@
 #include "Engine/DemoNetDriver.h"
 #include "Serialization/MemoryReader.h"
 #include "Serialization/MemoryWriter.h"
+#include "Components/ChildActorComponent.h"
 
 #if WITH_EDITOR
 #include "UObject/ObjectRedirector.h"
@@ -335,7 +336,16 @@ bool UPackageMapClient::SerializeNewActor(FArchive& Ar, class UActorChannel *Cha
 
 		if (Ar.IsSaving())
 		{
-			Archetype = Actor->GetArchetype();
+			// ChildActor's need to be spawned from the ChildActorTemplate otherwise any non-replicated 
+			// customized properties will be incorrect on the Client.
+			if (UChildActorComponent* CAC = Actor->GetParentComponent())
+			{
+				Archetype = CAC->GetChildActorTemplate();
+			}
+			if (Archetype == nullptr)
+			{
+				Archetype = Actor->GetArchetype();
+			}
 			ActorLevel = Actor->GetLevel();
 
 			check( Archetype != nullptr );
