Restore bunch bitfield optimization from CL 3683745.

diff --git a/Engine/Source/Runtime/Engine/Private/DataBunch.cpp b/Engine/Source/Runtime/Engine/Private/DataBunch.cpp
--- a/Engine/Source/Runtime/Engine/Private/DataBunch.cpp
+++ b/Engine/Source/Runtime/Engine/Private/DataBunch.cpp
@@ -49,10 +49,6 @@ FInBunch::FInBunch( UNetConnection* InConnection, uint8* Src, int64 CountBits )
 /** Copy constructor but with optional parameter to not copy buffer */
 FInBunch::FInBunch( FInBunch &InBunch, bool CopyBuffer )
 {
-	// Copy fields
-//	FMemory::Memcpy(&PacketId,&InBunch.PacketId,sizeof(FInBunch) - sizeof(FNetBitReader));
-
-	// This is portable
 	PacketId =	InBunch.PacketId;
 	Next =	InBunch.Next;
 	Connection = InBunch.Connection;
diff --git a/Engine/Source/Runtime/Engine/Public/Net/DataBunch.h b/Engine/Source/Runtime/Engine/Public/Net/DataBunch.h
--- a/Engine/Source/Runtime/Engine/Public/Net/DataBunch.h
+++ b/Engine/Source/Runtime/Engine/Public/Net/DataBunch.h
@@ -25,21 +25,21 @@ public:
 	FOutBunch *				Next;
 	UChannel *				Channel;
 	double					Time;
-	bool					ReceivedAck;
 	int32					ChIndex;
 	int32					ChType;
 	int32					ChSequence;
 	int32					PacketId;
-	uint8					bOpen;
-	uint8					bClose;
-	uint8					bDormant;
-	uint8					bIsReplicationPaused;   // Replication on this channel is being paused by the server
-	uint8					bReliable;
-	uint8					bPartial;				// Not a complete bunch
-	uint8					bPartialInitial;		// The first bunch of a partial bunch
-	uint8					bPartialFinal;			// The final bunch of a partial bunch
-	uint8					bHasPackageMapExports;	// This bunch has networkGUID name/id pairs
-	uint8					bHasMustBeMappedGUIDs;	// This bunch has guids that must be mapped before we can process this bunch
+	uint8					ReceivedAck:1;
+	uint8					bOpen:1;
+	uint8					bClose:1;
+	uint8					bDormant:1;
+	uint8					bIsReplicationPaused:1;   // Replication on this channel is being paused by the server
+	uint8					bReliable:1;
+	uint8					bPartial:1;				// Not a complete bunch
+	uint8					bPartialInitial:1;		// The first bunch of a partial bunch
+	uint8					bPartialFinal:1;			// The final bunch of a partial bunch
+	uint8					bHasPackageMapExports:1;	// This bunch has networkGUID name/id pairs
+	uint8					bHasMustBeMappedGUIDs:1;	// This bunch has guids that must be mapped before we can process this bunch
 
 	TArray< FNetworkGUID >	ExportNetGUIDs;			// List of GUIDs that went out on this bunch
 	TArray< uint64 >		NetFieldExports;
@@ -108,17 +108,17 @@ public:
 	int32				ChIndex;
 	int32				ChType;
 	int32				ChSequence;
-	uint8				bOpen;
-	uint8				bClose;
-	uint8				bDormant;				// Close, but go dormant
-	uint8				bIsReplicationPaused;	// Replication on this channel is being paused by the server
-	uint8				bReliable;
-	uint8				bPartial;				// Not a complete bunch
-	uint8				bPartialInitial;		// The first bunch of a partial bunch
-	uint8				bPartialFinal;			// The final bunch of a partial bunch
-	uint8				bHasPackageMapExports;	// This bunch has networkGUID name/id pairs
-	uint8				bHasMustBeMappedGUIDs;	// This bunch has guids that must be mapped before we can process this bunch
-	uint8				bIgnoreRPCs;
+	uint8				bOpen:1;
+	uint8				bClose:1;
+	uint8				bDormant:1;					// Close, but go dormant
+	uint8				bIsReplicationPaused:1;		// Replication on this channel is being paused by the server
+	uint8				bReliable:1;
+	uint8				bPartial:1;					// Not a complete bunch
+	uint8				bPartialInitial:1;			// The first bunch of a partial bunch
+	uint8				bPartialFinal:1;			// The final bunch of a partial bunch
+	uint8				bHasPackageMapExports:1;	// This bunch has networkGUID name/id pairs
+	uint8				bHasMustBeMappedGUIDs:1;	// This bunch has guids that must be mapped before we can process this bunch
+	uint8				bIgnoreRPCs:1;
 
 	FString	ToString()
 	{
