Fix initialization order & some whitespace

diff --git a/Engine/Source/Runtime/Engine/Private/DataBunch.cpp b/Engine/Source/Runtime/Engine/Private/DataBunch.cpp
--- a/Engine/Source/Runtime/Engine/Private/DataBunch.cpp
+++ b/Engine/Source/Runtime/Engine/Private/DataBunch.cpp
@@ -96,11 +96,11 @@ FOutBunch::FOutBunch( UChannel* InChannel, bool bInClose )
 ,	Next		( NULL )
 ,	Channel		( InChannel )
 ,	Time		( 0 )
-,	ReceivedAck	(false)
-,	ChIndex     ( InChannel->ChIndex )
-,	ChType      ( InChannel->ChType )
+,	ChIndex		( InChannel->ChIndex )
+,	ChType		( InChannel->ChType )
 ,	ChSequence	( 0 )
 ,	PacketId	( 0 )
+,	ReceivedAck	( 0 )
 ,	bOpen		( 0 )
 ,	bClose		( bInClose )
 ,	bDormant	( 0 )
@@ -129,12 +129,12 @@ FOutBunch::FOutBunch( UPackageMap *InPackageMap, int64 MaxBits )
 :	FNetBitWriter	( InPackageMap, MaxBits )
 ,	Next		( NULL )
 ,	Channel		( NULL )
-,   Time		( 0 )
-,	ReceivedAck	(false)
+,	Time		( 0 )
 ,	ChIndex     ( 0 )
 ,	ChType      ( 0 )
 ,	ChSequence	( 0 )
 ,	PacketId	( 0 )
+,	ReceivedAck	( 0 )
 ,	bOpen		( 0 )
 ,	bClose		( 0 )
 ,	bDormant	( 0 )
