Fix for shared rpc serialization with outparm properties

diff --git a/Engine/Source/Runtime/Engine/Private/RepLayout.cpp b/Engine/Source/Runtime/Engine/Private/RepLayout.cpp
--- a/Engine/Source/Runtime/Engine/Private/RepLayout.cpp
+++ b/Engine/Source/Runtime/Engine/Private/RepLayout.cpp
@@ -3739,6 +3739,11 @@ void FRepLayout::BuildSharedSerializationForRPC( void* Data )
 
 		for ( int32 i = 0; i < Parents.Num(); i++ )
 		{
+			if (Parents[i].Property->HasAnyPropertyFlags(CPF_OutParm))
+			{
+				continue;
+			}
+
 			bool bSend = true;
 
 			if ( !Cast<UBoolProperty>( Parents[i].Property ) )
@@ -3798,7 +3803,7 @@ void FRepLayout::SendPropertiesForRPC( UFunction * Function, UActorChannel * Cha
 		if ( !Cast<UBoolProperty>( Parents[i].Property ) )
 		{
 			// Used cached comparison result if possible
-			if (CVarNetShareSerializedData->GetInt() != 0 && SharedInfoRPC.IsValid())
+			if (CVarNetShareSerializedData->GetInt() != 0 && SharedInfoRPC.IsValid() && !Parents[i].Property->HasAnyPropertyFlags(CPF_OutParm))
 			{
 				Send = SharedInfoRPCParentsChanged[i];
 			}
