Add missing set of bClose to FInBunch copy constructor

diff --git a/Engine/Source/Runtime/Engine/Private/DataBunch.cpp b/Engine/Source/Runtime/Engine/Private/DataBunch.cpp
--- a/Engine/Source/Runtime/Engine/Private/DataBunch.cpp
+++ b/Engine/Source/Runtime/Engine/Private/DataBunch.cpp
@@ -56,6 +56,7 @@ FInBunch::FInBunch( FInBunch &InBunch, bool CopyBuffer )
 	ChType = InBunch.ChType;
 	ChSequence = InBunch.ChSequence;
 	bOpen =	InBunch.bOpen;
+	bClose = InBunch.bClose;
 	bDormant = InBunch.bDormant;
 	bIsReplicationPaused = InBunch.bIsReplicationPaused;
 	bReliable =	InBunch.bReliable;
