Fix replication for actor components created in blueprints - these components now properly return false from UActorComponent::IsNameStableForNetworking() by default.

diff --git a/Engine/Source/Runtime/Engine/Private/Components/ActorComponent.cpp b/Engine/Source/Runtime/Engine/Private/Components/ActorComponent.cpp
--- a/Engine/Source/Runtime/Engine/Private/Components/ActorComponent.cpp
+++ b/Engine/Source/Runtime/Engine/Private/Components/ActorComponent.cpp
@@ -1607,7 +1607,7 @@ bool UActorComponent::IsNameStableForNetworking() const
 	 *	-They were explicitly set to bNetAddressable (blueprint components created by SCS)
 	 */
 
-	return bNetAddressable || Super::IsNameStableForNetworking();
+	return bNetAddressable || (Super::IsNameStableForNetworking() && (CreationMethod != EComponentCreationMethod::UserConstructionScript));
 }
 
 bool UActorComponent::IsSupportedForNetworking() const
