Handle potential crash in TickWorldTravel

diff --git a/Engine/Source/Runtime/Engine/Private/UnrealEngine.cpp b/Engine/Source/Runtime/Engine/Private/UnrealEngine.cpp
--- a/Engine/Source/Runtime/Engine/Private/UnrealEngine.cpp
+++ b/Engine/Source/Runtime/Engine/Private/UnrealEngine.cpp
@@ -11672,10 +11672,18 @@ void UEngine::TickWorldTravel(FWorldContext& Context, float DeltaSeconds)
 
 				const bool bLoadedMapSuccessfully = LoadMap(Context, Context.PendingNetGame->URL, Context.PendingNetGame, Error);
 
-				Context.PendingNetGame->LoadMapCompleted(this, Context, bLoadedMapSuccessfully, Error);
+				if (Context.PendingNetGame != nullptr)
+				{
+					Context.PendingNetGame->LoadMapCompleted(this, Context, bLoadedMapSuccessfully, Error);
 
-				// Kill the pending level.
-				Context.PendingNetGame = NULL;
+					// Kill the pending level.
+					Context.PendingNetGame = nullptr;
+				}
+				else
+				{
+					BrowseToDefaultMap(Context);
+					BroadcastTravelFailure(Context.World(), ETravelFailure::TravelFailure, Error);
+				}
 			}
 		}
 	}
