PR #4965: Remove code that doesn't do anything (Contributed by [REDACTED_PERSON])

diff --git a/Engine/Source/Runtime/Engine/Private/SCS_Node.cpp b/Engine/Source/Runtime/Engine/Private/SCS_Node.cpp
--- a/Engine/Source/Runtime/Engine/Private/SCS_Node.cpp
+++ b/Engine/Source/Runtime/Engine/Private/SCS_Node.cpp
@@ -110,12 +110,6 @@ UActorComponent* USCS_Node::ExecuteNodeOnActor(AActor* Actor, USceneComponent* P
 			NewActorComp->OnComponentCreated();
 		}
 
-		if (NewActorComp->GetIsReplicated())
-		{
-			// Make sure this component is added to owning actor's replicated list.
-			NewActorComp->SetIsReplicated(true);
-		}
-
 		// Special handling for scene components
 		USceneComponent* NewSceneComp = Cast<USceneComponent>(NewActorComp);
 		if (NewSceneComp != nullptr)
