Fix crash if a socket error occurred when the PendingNetGame sent its initial join packet.

diff --git a/Engine/Source/Runtime/Engine/Private/UnrealEngine.cpp b/Engine/Source/Runtime/Engine/Private/UnrealEngine.cpp
--- a/Engine/Source/Runtime/Engine/Private/UnrealEngine.cpp
+++ b/Engine/Source/Runtime/Engine/Private/UnrealEngine.cpp
@@ -11463,6 +11463,15 @@ EBrowseReturnVal::Type UEngine::Browse( FWorldContext& WorldContext, FURL URL, F
 		WorldContext.PendingNetGame = NewObject<UPendingNetGame>();
 		WorldContext.PendingNetGame->Initialize(URL);
 		WorldContext.PendingNetGame->InitNetDriver();
+
+		if( !WorldContext.PendingNetGame )
+		{
+			// If the inital packet sent in InitNetDriver results in a socket error, HandleDisconnect() and CancelPending() may be called, which will null the PendingNetGame.
+			Error = NSLOCTEXT("Engine", "PendingNetGameInitFailure", "Error initializing the network driver.").ToString();
+			BroadcastTravelFailure(WorldContext.World(), ETravelFailure::PendingNetGameCreateFailure, Error);
+			return EBrowseReturnVal::Failure;
+		}
+
 		if( !WorldContext.PendingNetGame->NetDriver )
 		{
 			// UPendingNetGame will set the appropriate error code and connection lost type, so
