Suppress static analysis warning.

diff --git a/Engine/Source/Runtime/Engine/Private/UnrealEngine.cpp b/Engine/Source/Runtime/Engine/Private/UnrealEngine.cpp
--- a/Engine/Source/Runtime/Engine/Private/UnrealEngine.cpp
+++ b/Engine/Source/Runtime/Engine/Private/UnrealEngine.cpp
@@ -11461,8 +11461,8 @@ EBrowseReturnVal::Type UEngine::Browse( FWorldContext& WorldContext, FURL URL, F
 		}
 
 		WorldContext.PendingNetGame = NewObject<UPendingNetGame>();
-		WorldContext.PendingNetGame->Initialize(URL);
-		WorldContext.PendingNetGame->InitNetDriver();
+		WorldContext.PendingNetGame->Initialize(URL); //-V595
+		WorldContext.PendingNetGame->InitNetDriver(); //-V595
 
 		if( !WorldContext.PendingNetGame )
 		{
