Fix for level script actor gc references with replicated blueprint variables/functions

diff --git a/Engine/Source/Editor/UnrealEd/Private/EditorEngine.cpp b/Engine/Source/Editor/UnrealEd/Private/EditorEngine.cpp
--- a/Engine/Source/Editor/UnrealEd/Private/EditorEngine.cpp
+++ b/Engine/Source/Editor/UnrealEd/Private/EditorEngine.cpp
@@ -904,7 +904,7 @@ void UEditorEngine::Init(IEngineLoop* InEngineLoop)
 	{
 		if (GetPIEWorldContext() != nullptr && GetPIEWorldContext()->World() != nullptr)
 		{
-			GetPIEWorldContext()->World()->DestroyDemoNetDriver();
+			GEngine->ShutdownWorldNetDriver(GetPIEWorldContext()->World());
 		}
 	});
 
diff --git a/Engine/Source/Runtime/Engine/Classes/Engine/NetDriver.h b/Engine/Source/Runtime/Engine/Classes/Engine/NetDriver.h
--- a/Engine/Source/Runtime/Engine/Classes/Engine/NetDriver.h
+++ b/Engine/Source/Runtime/Engine/Classes/Engine/NetDriver.h
@@ -779,6 +779,10 @@ public:
 
 	ENGINE_API virtual void CleanPackageMaps();
 
+	void RemoveClassRepLayoutReferences(UClass* Class);
+
+	ENGINE_API void CleanupWorldForSeamlessTravel();
+
 	ENGINE_API void PreSeamlessTravelGarbageCollect();
 
 	ENGINE_API void PostSeamlessTravelGarbageCollect();
diff --git a/Engine/Source/Runtime/Engine/Private/NetworkDriver.cpp b/Engine/Source/Runtime/Engine/Private/NetworkDriver.cpp
--- a/Engine/Source/Runtime/Engine/Private/NetworkDriver.cpp
+++ b/Engine/Source/Runtime/Engine/Private/NetworkDriver.cpp
@@ -60,6 +60,7 @@
 #include "Stats/StatsMisc.h"
 #include "Engine/ReplicationDriver.h"
 #include "ProfilingDebugging/CsvProfiler.h"
+#include "Engine/LevelScriptActor.h"
 
 #if USE_SERVER_PERF_COUNTERS
 #include "PerfCountersModule.h"
@@ -1362,6 +1363,15 @@ void UNetDriver::Shutdown()
 	// Client closing connection to server
 	if (ServerConnection)
 	{
+		for (UChannel* Channel : ServerConnection->OpenChannels)
+		 {
+			 UActorChannel* ActorChannel = Cast<UActorChannel>(Channel);
+			 if (ActorChannel)
+			 {
+				 ActorChannel->CleanupReplicators();
+			 }
+		 }
+
 		// Calls Channel[0]->Close to send a close bunch to server
 		ServerConnection->Close();
 		ServerConnection->FlushNet();
@@ -4326,6 +4336,44 @@ void UNetDriver::CleanPackageMaps()
 	}
 }
 
+void UNetDriver::RemoveClassRepLayoutReferences(UClass* Class)
+{
+	RepLayoutMap.Remove(Class);
+
+	for (auto Func : TFieldRange<UFunction>(Class, EFieldIteratorFlags::ExcludeSuper))
+	{
+		if (Func && Func->HasAnyFunctionFlags(EFunctionFlags::FUNC_Net))
+		{
+			RepLayoutMap.Remove(Func);
+		}
+	}
+}
+
+void UNetDriver::CleanupWorldForSeamlessTravel()
+{
+	if (World != nullptr)
+	{
+		for (auto LevelIt(World->GetLevelIterator()); LevelIt; ++LevelIt)
+		{
+			if (const ULevel* Level = *LevelIt)
+			{
+				if (Level->LevelScriptActor)
+				{
+					// workaround for this not being called on clients
+					if (ServerConnection != nullptr)
+					{
+						NotifyActorLevelUnloaded(Level->LevelScriptActor);
+					}
+
+					RemoveClassRepLayoutReferences(Level->LevelScriptActor->GetClass());
+
+					ReplicationChangeListMap.Remove(Level->LevelScriptActor);
+				}
+			}
+		}
+	}
+}
+
 void UNetDriver::PreSeamlessTravelGarbageCollect()
 {
 	ResetGameWorldState();
@@ -4448,6 +4496,12 @@ void UNetDriver::OnLevelRemovedFromWorld(class ULevel* InLevel, class UWorld* In
 			}
 		}
 
+		if (InLevel && InLevel->LevelScriptActor)
+		{
+			RemoveClassRepLayoutReferences(InLevel->LevelScriptActor->GetClass());
+			ReplicationChangeListMap.Remove(InLevel->LevelScriptActor);
+		}
+
 		TArray<FNetworkGUID> RemovedGUIDs;
 		for (auto It = DestroyedStartupOrDormantActors.CreateIterator(); It; ++It)
 		{
diff --git a/Engine/Source/Runtime/Engine/Private/World.cpp b/Engine/Source/Runtime/Engine/Private/World.cpp
--- a/Engine/Source/Runtime/Engine/Private/World.cpp
+++ b/Engine/Source/Runtime/Engine/Private/World.cpp
@@ -5683,6 +5683,11 @@ UWorld* FSeamlessTravelHandler::Tick()
 				}
 			}
 
+			if (NetDriver)
+			{
+				NetDriver->CleanupWorldForSeamlessTravel();
+			}
+
 			bool bCreateNewGameMode = !bIsClient;
 			{
 				// scope because after GC the kept pointers will be bad
