SA fix

diff --git a/Engine/Source/Runtime/Engine/Private/NetworkDriver.cpp b/Engine/Source/Runtime/Engine/Private/NetworkDriver.cpp
--- a/Engine/Source/Runtime/Engine/Private/NetworkDriver.cpp
+++ b/Engine/Source/Runtime/Engine/Private/NetworkDriver.cpp
@@ -4487,19 +4487,22 @@ void UNetDriver::OnLevelRemovedFromWorld(class ULevel* InLevel, class UWorld* In
 {
 	if (InWorld == World)
 	{
-		for (AActor* Actor : InLevel->Actors)
+		if (InLevel)
 		{
-			if (Actor)
+			for (AActor* Actor : InLevel->Actors)
 			{
-				NotifyActorLevelUnloaded(Actor);
-				RemoveNetworkActor(Actor);
+				if (Actor)
+				{
+					NotifyActorLevelUnloaded(Actor);
+					RemoveNetworkActor(Actor);
+				}
 			}
-		}
 
-		if (InLevel && InLevel->LevelScriptActor)
-		{
-			RemoveClassRepLayoutReferences(InLevel->LevelScriptActor->GetClass());
-			ReplicationChangeListMap.Remove(InLevel->LevelScriptActor);
+			if (InLevel->LevelScriptActor)
+			{
+				RemoveClassRepLayoutReferences(InLevel->LevelScriptActor->GetClass());
+				ReplicationChangeListMap.Remove(InLevel->LevelScriptActor);
+			}
 		}
 
 		TArray<FNetworkGUID> RemovedGUIDs;
