Fix remote client character animation stuttering while stationary on listen servers by setting the default value of AGameNetworkManager::ClientNetSendMoveDeltaTimeStationary to 0.0166, the same value as ClientNetSendMoveDeltaTime, to match behavior of 4.19. Games may override this value to something lower as an optimization.

diff --git a/Engine/Config/BaseGame.ini b/Engine/Config/BaseGame.ini
--- a/Engine/Config/BaseGame.ini
+++ b/Engine/Config/BaseGame.ini
@@ -16,7 +16,7 @@ MAXCLIENTUPDATEINTERVAL=0.25f
 MaxMoveDeltaTime=0.125f
 ClientNetSendMoveDeltaTime=0.0166
 ClientNetSendMoveDeltaTimeThrottled=0.0222
-ClientNetSendMoveDeltaTimeStationary=0.0833
+ClientNetSendMoveDeltaTimeStationary=0.0166
 ClientNetSendMoveThrottleAtNetSpeed=10000
 ClientNetSendMoveThrottleOverPlayerCount=10
 ClientAuthorativePosition=false
diff --git a/Engine/Source/Runtime/Engine/Private/GameNetworkManager.cpp b/Engine/Source/Runtime/Engine/Private/GameNetworkManager.cpp
--- a/Engine/Source/Runtime/Engine/Private/GameNetworkManager.cpp
+++ b/Engine/Source/Runtime/Engine/Private/GameNetworkManager.cpp
@@ -25,7 +25,7 @@ AGameNetworkManager::AGameNetworkManager(const FObjectInitializer& ObjectInitial
 	MaxMoveDeltaTime = 0.125f;
 	ClientNetSendMoveDeltaTime = 0.0166f;
 	ClientNetSendMoveDeltaTimeThrottled = 0.0222f;
-	ClientNetSendMoveDeltaTimeStationary = 0.0833f;
+	ClientNetSendMoveDeltaTimeStationary = 0.0166f;
 	ClientNetSendMoveThrottleAtNetSpeed = 10000;
 	ClientNetSendMoveThrottleOverPlayerCount = 10;
 	ClientAuthorativePosition = false;
