Fix pawn relevancy position stuck at possession point.

diff --git a/Engine/Plugins/Runtime/PhysXVehicles/Source/PhysXVehicles/Private/WheeledVehicleMovementComponent.cpp b/Engine/Plugins/Runtime/PhysXVehicles/Source/PhysXVehicles/Private/WheeledVehicleMovementComponent.cpp
--- a/Engine/Plugins/Runtime/PhysXVehicles/Source/PhysXVehicles/Private/WheeledVehicleMovementComponent.cpp
+++ b/Engine/Plugins/Runtime/PhysXVehicles/Source/PhysXVehicles/Private/WheeledVehicleMovementComponent.cpp
@@ -1112,6 +1112,11 @@ void UWheeledVehicleMovementComponent::UpdateState( float DeltaTime )
 
 		// and send to server
 		ServerUpdateState(SteeringInput, ThrottleInput, BrakeInput, HandbrakeInput, GetCurrentGear());
+
+		if (PawnOwner && PawnOwner->IsNetMode(NM_Client))
+		{
+			MarkForClientCameraUpdate();
+		}
 	}
 	else
 	{
diff --git a/Engine/Source/Runtime/Engine/Classes/GameFramework/PawnMovementComponent.h b/Engine/Source/Runtime/Engine/Classes/GameFramework/PawnMovementComponent.h
--- a/Engine/Source/Runtime/Engine/Classes/GameFramework/PawnMovementComponent.h
+++ b/Engine/Source/Runtime/Engine/Classes/GameFramework/PawnMovementComponent.h
@@ -80,12 +80,20 @@ public:
 	// UNavMovementComponent override for input operations
 	virtual void RequestPathMove(const FVector& MoveInput) override;
 
+	virtual void OnTeleported() override;
+
 protected:
 
 	/** Pawn that owns this component. */
 	UPROPERTY(Transient, DuplicateTransient)
 	class APawn* PawnOwner;
 
+	/**
+	 * Attempts to mark the PlayerCameraManager as dirty.
+	 * This will have no effect if called from the server.
+	 */
+	void MarkForClientCameraUpdate();
+
 public:
 
 	virtual void Serialize(FArchive& Ar) override;
diff --git a/Engine/Source/Runtime/Engine/Private/Components/CharacterMovementComponent.cpp b/Engine/Source/Runtime/Engine/Private/Components/CharacterMovementComponent.cpp
--- a/Engine/Source/Runtime/Engine/Private/Components/CharacterMovementComponent.cpp
+++ b/Engine/Source/Runtime/Engine/Private/Components/CharacterMovementComponent.cpp
@@ -1167,12 +1167,7 @@ void UCharacterMovementComponent::TickComponent(float DeltaTime, enum ELevelTick
 		// Update camera to ensure client gets updates even when physics move him far away from point where simulation started
 		if (CharacterOwner->Role == ROLE_AutonomousProxy && IsNetMode(NM_Client))
 		{
-			APlayerController* PC = Cast<APlayerController>(CharacterOwner->GetController());
-			APlayerCameraManager* PlayerCameraManager = (PC ? PC->PlayerCameraManager : NULL);
-			if (PlayerCameraManager != NULL && PlayerCameraManager->bUseClientSideCameraUpdates)
-			{
-				PlayerCameraManager->bShouldSendClientSideCameraUpdate = true;
-			}
+			MarkForClientCameraUpdate();
 		}
 
 		ClearAccumulatedForces();
@@ -5311,6 +5306,8 @@ void UCharacterMovementComponent::OnTeleported()
 		return;
 	}
 
+	Super::OnTeleported();
+
 	bJustTeleported = true;
 
 	// Find floor at current location
@@ -7917,13 +7914,7 @@ void UCharacterMovementComponent::CallServerMove
 			);
 	}
 
-
-	APlayerController* PC = Cast<APlayerController>(CharacterOwner->GetController());
-	APlayerCameraManager* PlayerCameraManager = (PC ? PC->PlayerCameraManager : NULL);
-	if (PlayerCameraManager != NULL && PlayerCameraManager->bUseClientSideCameraUpdates)
-	{
-		PlayerCameraManager->bShouldSendClientSideCameraUpdate = true;
-	}
+	MarkForClientCameraUpdate();
 }
 
 
diff --git a/Engine/Source/Runtime/Engine/Private/Components/PawnMovementComponent.cpp b/Engine/Source/Runtime/Engine/Private/Components/PawnMovementComponent.cpp
--- a/Engine/Source/Runtime/Engine/Private/Components/PawnMovementComponent.cpp
+++ b/Engine/Source/Runtime/Engine/Private/Components/PawnMovementComponent.cpp
@@ -2,6 +2,8 @@
 
 #include "GameFramework/PawnMovementComponent.h"
 #include "GameFramework/Pawn.h"
+#include "GameFramework/PlayerController.h"
+#include "Camera/PlayerCameraManager.h"
 
 DEFINE_LOG_CATEGORY_STATIC(LogPawnMovementComponent, Log, All);
 
@@ -91,3 +93,26 @@ FVector UPawnMovementComponent::K2_GetInputVector() const
 {
 	return GetPendingInputVector();
 }
+
+void UPawnMovementComponent::OnTeleported()
+{
+	if (PawnOwner && PawnOwner->IsNetMode(NM_Client) && PawnOwner->IsLocallyControlled())
+	{
+		MarkForClientCameraUpdate();
+	}
+}
+
+void UPawnMovementComponent::MarkForClientCameraUpdate()
+{
+	if(PawnOwner)
+	{
+		if (APlayerController* PlayerController = Cast<APlayerController>(PawnOwner->GetController()))
+		{
+			APlayerCameraManager* PlayerCameraManager = PlayerController->PlayerCameraManager;
+			if (PlayerCameraManager != nullptr && PlayerCameraManager->bUseClientSideCameraUpdates)
+			{
+				PlayerCameraManager->bShouldSendClientSideCameraUpdate = true;
+			}
+		}
+	}	
+}
\ No newline at end of file
diff --git a/Engine/Source/Runtime/Engine/Private/PlayerController.cpp b/Engine/Source/Runtime/Engine/Private/PlayerController.cpp
--- a/Engine/Source/Runtime/Engine/Private/PlayerController.cpp
+++ b/Engine/Source/Runtime/Engine/Private/PlayerController.cpp
@@ -4418,7 +4418,10 @@ void APlayerController::TickActor( float DeltaSeconds, ELevelTick TickType, FAct
 
 	//root of tick hierarchy
 
-	if ((GetRemoteRole() == ROLE_AutonomousProxy) && !IsNetMode(NM_Client) && !IsLocalPlayerController())
+	const bool bIsClient = IsNetMode(NM_Client);
+	const bool bIsLocallyControlled = IsLocalPlayerController();
+
+	if ((GetRemoteRole() == ROLE_AutonomousProxy) && !bIsClient && !bIsLocallyControlled)
 	{
 		// force physics update for clients that aren't sending movement updates in a timely manner 
 		// this prevents cheats associated with artificially induced ping spikes
@@ -4462,11 +4465,11 @@ void APlayerController::TickActor( float DeltaSeconds, ELevelTick TickType, FAct
 		}
 
 		// update viewtarget replicated info
-		if (PlayerCameraManager != NULL)
+		if (PlayerCameraManager != nullptr)
 		{
 			APawn* TargetPawn = PlayerCameraManager->GetViewTargetPawn();
 			
-			if ((TargetPawn != GetPawn()) && (TargetPawn != NULL))
+			if ((TargetPawn != GetPawn()) && (TargetPawn != nullptr))
 			{
 				TargetViewRotation = TargetPawn->GetViewRotation();
 			}
@@ -4475,7 +4478,7 @@ void APlayerController::TickActor( float DeltaSeconds, ELevelTick TickType, FAct
 	else if (Role > ROLE_SimulatedProxy)
 	{
 		// Process PlayerTick with input.
-		if (!PlayerInput && (Player == NULL || Cast<ULocalPlayer>( Player ) != NULL))
+		if (!PlayerInput && (Player == nullptr || Cast<ULocalPlayer>( Player ) != nullptr))
 		{
 			InitInputSystem();
 		}
@@ -4491,13 +4494,24 @@ void APlayerController::TickActor( float DeltaSeconds, ELevelTick TickType, FAct
 		}
 
 		// update viewtarget replicated info
-		if (PlayerCameraManager != NULL)
+		if (PlayerCameraManager != nullptr)
 		{
 			APawn* TargetPawn = PlayerCameraManager->GetViewTargetPawn();
-			if ((TargetPawn != GetPawn()) && (TargetPawn != NULL))
+			if ((TargetPawn != GetPawn()) && (TargetPawn != nullptr))
 			{
 				SmoothTargetViewRotation(TargetPawn, DeltaSeconds);
 			}
+
+			// Send a camera update if necessary.
+			if (bIsClient && bIsLocallyControlled && GetPawn() && PlayerCameraManager->bUseClientSideCameraUpdates)
+			{
+				UPawnMovementComponent* PawnMovement = GetPawn()->GetMovementComponent();
+				if (!PawnMovement->IsMoveInputIgnored() &&
+					(PawnMovement->GetLastInputVector() != FVector::ZeroVector || PawnMovement->Velocity != FVector::ZeroVector))
+				{
+					PlayerCameraManager->bShouldSendClientSideCameraUpdate = true;
+				}
+			}
 		}
 	}
 
@@ -4510,7 +4524,7 @@ void APlayerController::TickActor( float DeltaSeconds, ELevelTick TickType, FAct
 	RotationInput = FRotator::ZeroRotator;
 
 #if !(UE_BUILD_SHIPPING || UE_BUILD_TEST)
-	if (CheatManager != NULL)
+	if (CheatManager != nullptr)
 	{
 		CheatManager->TickCollisionDebug();
 	}
