Prevent Streaming Level Visibility changes from closing channels for static actors on the server. This should prevent them from being destroyed on the client.

diff --git a/Engine/Source/Runtime/Engine/Private/NetworkDriver.cpp b/Engine/Source/Runtime/Engine/Private/NetworkDriver.cpp
--- a/Engine/Source/Runtime/Engine/Private/NetworkDriver.cpp
+++ b/Engine/Source/Runtime/Engine/Private/NetworkDriver.cpp
@@ -2456,8 +2456,42 @@ void UNetDriver::NotifyActorRenamed(AActor* ThisActor, FName PreviousName)
 	}
 }
 
+// This method will be called when a Streaming Level is about to be Garbage Collected.
 void UNetDriver::NotifyStreamingLevelUnload(ULevel* Level)
 {
+	if (IsServer())
+	{
+		for (AActor* Actor : Level->Actors)
+		{
+			if (Actor && Actor->IsNetStartupActor())
+			{
+				NotifyActorLevelUnloaded(Actor);
+			}
+		}
+		
+		TArray<FNetworkGUID> RemovedGUIDs;
+		for (auto It = DestroyedStartupOrDormantActors.CreateIterator(); It; ++It)
+		{
+			FActorDestructionInfo* DestructInfo = It->Value.Get();
+			if (DestructInfo->Level == Level && DestructInfo->NetGUID.IsStatic())
+			{
+				for (UNetConnection* Connection : ClientConnections)
+				{
+					Connection->RemoveDestructionInfo(DestructInfo);
+				}
+
+				RemovedGUIDs.Add(It->Key);
+				It.RemoveCurrent();
+			}
+		}
+	}	
+
+	if (Level->LevelScriptActor)
+	{
+		RemoveClassRepLayoutReferences(Level->LevelScriptActor->GetClass());
+		ReplicationChangeListMap.Remove(Level->LevelScriptActor);
+	}
+
 	if (ServerConnection && ServerConnection->PackageMap)
 	{
 		UE_LOG(LogNet, Log, TEXT("NotifyStreamingLevelUnload: %s"), *Level->GetFullName() );
@@ -4483,42 +4517,64 @@ TSharedPtr< FReplicationChangelistMgr > UNetDriver::GetReplicationChangeListMgr(
 	return *ReplicationChangeListMgrPtr;
 }
 
-void UNetDriver::OnLevelRemovedFromWorld(class ULevel* InLevel, class UWorld* InWorld)
+// This method will be called when Streaming Levels become Visible.
+void UNetDriver::OnLevelAddedToWorld(ULevel* Level, UWorld* InWorld)
+{
+	// Actors will be re-added to the network list when ULevel::SortActorList is called.
+}
+
+// This method will be called when Streaming Levels are hidden.
+void UNetDriver::OnLevelRemovedFromWorld(class ULevel* Level, class UWorld* InWorld)
 {
-	if (InWorld == World)
+	if (Level && InWorld == GetWorld())
 	{
-		if (InLevel)
+		if (!IsServer())
 		{
-			for (AActor* Actor : InLevel->Actors)
+			for (AActor* Actor : Level->Actors)
 			{
 				if (Actor)
 				{
+					// Always call this on clients.
+					// It won't actually destroy the Actor, but it will do some cleanup of the channel, etc.
 					NotifyActorLevelUnloaded(Actor);
-					RemoveNetworkActor(Actor);
 				}
 			}
-
-			if (InLevel->LevelScriptActor)
+		}	
+		else
+		{
+			for (AActor* Actor : Level->Actors)
 			{
-				RemoveClassRepLayoutReferences(InLevel->LevelScriptActor->GetClass());
-				ReplicationChangeListMap.Remove(InLevel->LevelScriptActor);
+				// Keep Startup actors alive, because they haven't been destroyed yet.
+				// Technically, Dynamic actors may not have been destroyed either, but this
+				// resembles relevancy.
+				if (!Actor->IsNetStartupActor())
+				{
+					NotifyActorLevelUnloaded(Actor);
+				}
+				else
+				{
+					// We still want to remove Startup actors from the Network list so they aren't processed anymore.
+					GetNetworkObjectList().Remove(Actor);
+				}
 			}
-		}
 
-		TArray<FNetworkGUID> RemovedGUIDs;
-		for (auto It = DestroyedStartupOrDormantActors.CreateIterator(); It; ++It)
-		{
-			if (It->Value->Level == InLevel)
+			TArray<FNetworkGUID> RemovedGUIDs;
+			for (auto It = DestroyedStartupOrDormantActors.CreateIterator(); It; ++It)
 			{
-				// Connections must be updated before we remove from our map
 				FActorDestructionInfo* DestructInfo = It->Value.Get();
-				for (UNetConnection* Connection : ClientConnections)
+				
+				// Until the level is actually unloaded / reloaded, any Static Actors that have been
+				// destroyed will not be recreated, so we still need to track them.
+				if (DestructInfo->Level == Level && !DestructInfo->NetGUID.IsStatic())
 				{
-					Connection->RemoveDestructionInfo(DestructInfo);
-				}
+					for (UNetConnection* Connection : ClientConnections)
+					{
+						Connection->RemoveDestructionInfo(DestructInfo);
+					}
 
-				RemovedGUIDs.Add(It->Key);
-				It.RemoveCurrent();
+					RemovedGUIDs.Add(It->Key);
+					It.RemoveCurrent();
+				}
 			}
 		}
 	}
@@ -4578,10 +4634,6 @@ FCreateReplicationDriver& UReplicationDriver::CreateReplicationDriverDelegate()
 	return Delegate;
 }
 
-void UNetDriver::OnLevelAddedToWorld(ULevel* InLevel, UWorld* InWorld)
-{
-}
-
 // --------------------------------------------------------------------------------------------------------------------------------------------
 // --------------------------------------------------------------------------------------------------------------------------------------------
 // --------------------------------------------------------------------------------------------------------------------------------------------
