Fix actor channel for the world settings of streaming levels not being marked broken on the client when seamless traveling

diff --git a/Engine/Source/Runtime/Engine/Private/NetworkDriver.cpp b/Engine/Source/Runtime/Engine/Private/NetworkDriver.cpp
--- a/Engine/Source/Runtime/Engine/Private/NetworkDriver.cpp
+++ b/Engine/Source/Runtime/Engine/Private/NetworkDriver.cpp
@@ -4403,6 +4403,16 @@ void UNetDriver::CleanupWorldForSeamlessTravel()
 
 					ReplicationChangeListMap.Remove(Level->LevelScriptActor);
 				}
+
+				// This is currently necessary because the actor iterator used in the seamless travel handler 
+				// skips over AWorldSettings actors for an unknown reason.
+				if (Level->GetWorldSettings())
+				{
+					if (ServerConnection != nullptr)
+					{
+						NotifyActorLevelUnloaded(Level->GetWorldSettings());
+					}
+				}
 			}
 		}
 	}
