Fix crash when trying to add a split player with an invalid net index

diff --git a/Engine/Source/Runtime/Engine/Private/NetConnection.cpp b/Engine/Source/Runtime/Engine/Private/NetConnection.cpp
--- a/Engine/Source/Runtime/Engine/Private/NetConnection.cpp
+++ b/Engine/Source/Runtime/Engine/Private/NetConnection.cpp
@@ -2428,10 +2428,11 @@ void UChildConnection::HandleClientPlayer(APlayerController* PC, UNetConnection*
 		}
 	}
 
-	if (NewPlayer == NULL)
+	if (!ensure(NewPlayer != NULL))
 	{
-#if !(UE_BUILD_SHIPPING || UE_BUILD_TEST)
 		UE_LOG(LogNet, Error, TEXT("Failed to find LocalPlayer for received PlayerController '%s' with index %d. PlayerControllers:"), *PC->GetName(), int32(PC->NetPlayerIndex));
+
+#if !(UE_BUILD_SHIPPING || UE_BUILD_TEST)
 		check( PC->GetWorld() );
 		for (TActorIterator<APlayerController> It(PC->GetWorld()); It; ++It)
 		{
@@ -2440,10 +2441,12 @@ void UChildConnection::HandleClientPlayer(APlayerController* PC, UNetConnection*
 				UE_LOG(LogNet, Log, TEXT(" - %s"), *It->GetFullName());
 			}
 		}
-		UE_LOG(LogNet, Fatal,TEXT("Failed to find LocalPlayer for received PlayerController"));
-#else
-		return; // avoid crash
 #endif
+		if (ensure(Parent != nullptr))
+		{
+			Parent->Close();
+		}
+		return; // avoid crash
 	}
 
 	// Detach old player.
