Convert some replay checks/ensures to log warnings
Serialize all static guids into replay checkpoints, unless deleted

diff --git a/Engine/Source/Runtime/Engine/Classes/Engine/DemoNetDriver.h b/Engine/Source/Runtime/Engine/Classes/Engine/DemoNetDriver.h
--- a/Engine/Source/Runtime/Engine/Classes/Engine/DemoNetDriver.h
+++ b/Engine/Source/Runtime/Engine/Classes/Engine/DemoNetDriver.h
@@ -340,7 +340,8 @@ class ENGINE_API UDemoNetDriver : public UNetDriver
 
 	/** Net startup actors that need to be destroyed after checkpoints are loaded */
 	TSet< FString >									DeletedNetStartupActors;
-
+	TSet< FNetworkGUID >							DeletedNetStartupActorGUIDs;
+	
 	/** 
 	 * Net startup actors that need to be rolled back during scrubbing by being destroyed and re-spawned 
 	 * NOTE - DeletedNetStartupActors will take precedence here, and destroy the actor instead
diff --git a/Engine/Source/Runtime/Engine/Private/DataChannel.cpp b/Engine/Source/Runtime/Engine/Private/DataChannel.cpp
--- a/Engine/Source/Runtime/Engine/Private/DataChannel.cpp
+++ b/Engine/Source/Runtime/Engine/Private/DataChannel.cpp
@@ -632,8 +632,12 @@ bool UChannel::ReceivedNextBunch( FInBunch & Bunch, bool & bOutSkipAck )
 				// If we opened the channel, we shouldn't be receiving bOpen commands from the other side
 				checkf(!OpenedLocally, TEXT("Received channel open command for channel that was already opened locally. %s"), *Describe());
 
-				check( OpenPacketId.First == INDEX_NONE );	// This should be the first and only assignment of the packet range (we should only receive one bOpen bunch)
-				check( OpenPacketId.Last == INDEX_NONE );	// This should be the first and only assignment of the packet range (we should only receive one bOpen bunch)
+				if (!ensure( OpenPacketId.First == INDEX_NONE && OpenPacketId.Last == INDEX_NONE ))
+				{
+					UE_LOG( LogNetTraffic, Error, TEXT("Received channel open command for channel that was already opened locally. %s"), *Describe() );
+					Bunch.SetError();
+					return false;
+				}
 			}
 
 			// Remember the range.
@@ -2179,7 +2183,12 @@ void UActorChannel::ReceivedBunch( FInBunch & Bunch )
 				}
 
 				// This GUID better have been exported before we get here, which means it must be registered by now
-				check( Connection->Driver->GuidCache->IsGUIDRegistered( NetGUID ) );
+				if (!Connection->Driver->GuidCache->IsGUIDRegistered(NetGUID))
+				{
+					UE_LOG(LogNet, Warning, TEXT("UActorChannel::ReceivedBunch: Received a MustBeMappedGUID that is not registered. ChIndex: %i NetGUID: %s.  %s" ), ChIndex, *NetGUID.ToString(), *Describe());
+					Bunch.SetError();
+					return;
+				}
 
 				if ( !Connection->Driver->GuidCache->IsGUIDLoaded( NetGUID ) )
 				{
@@ -3243,9 +3252,9 @@ bool UActorChannel::ReadFieldHeaderAndPayload( UObject* Object, const FClassNetC
 
 	if ( Connection->InternalAck )
 	{
-		if ( !ensure( NetFieldExportGroup != nullptr ) )
+		if ( NetFieldExportGroup == nullptr )
 		{
-			UE_LOG( LogNet, Error, TEXT( "ReadFieldHeaderAndPayload: NetFieldExportGroup was null. Object: %s" ), *Object->GetFullName() );
+			UE_LOG( LogNet, Warning, TEXT( "ReadFieldHeaderAndPayload: NetFieldExportGroup was null. Object: %s" ), *Object->GetFullName() );
 			Bunch.SetError();
 			return false;
 		}
diff --git a/Engine/Source/Runtime/Engine/Private/DemoNetDriver.cpp b/Engine/Source/Runtime/Engine/Private/DemoNetDriver.cpp
--- a/Engine/Source/Runtime/Engine/Private/DemoNetDriver.cpp
+++ b/Engine/Source/Runtime/Engine/Private/DemoNetDriver.cpp
@@ -1549,49 +1549,56 @@ bool UDemoNetDriver::DemoReplicateActor(AActor* Actor, UNetConnection* Connectio
 void UDemoNetDriver::SerializeGuidCache(TSharedPtr<FNetGUIDCache> InGuidCache, FArchive* CheckpointArchive)
 {
 	int32 NumValues = 0;
+	int32 UnloadedValues = 0;
 
-	for ( auto It = InGuidCache->ObjectLookup.CreateIterator(); It; ++It )
-	{
-		FNetGuidCacheObject& CacheObject = It.Value();
-
-		if ( CacheObject.Object == nullptr || !CacheObject.Object->IsNameStableForNetworking() )
-		{
-			continue;
-		}
-
-		NumValues++;
-	}
+	int64 CountPos = CheckpointArchive->Tell();
 
 	*CheckpointArchive << NumValues;
 
-	UE_LOG( LogDemo, Verbose, TEXT( "Checkpoint. SerializeGuidCache: %i" ), NumValues );
-
 	for ( auto It = InGuidCache->ObjectLookup.CreateIterator(); It; ++It )
 	{
 		FNetworkGUID& NetworkGUID = It.Key();
 		FNetGuidCacheObject& CacheObject = It.Value();
 
-		if ( CacheObject.Object == NULL || !CacheObject.Object->IsNameStableForNetworking() )
+		if (NetworkGUID.IsValid() && (NetworkGUID.IsStatic() || (CacheObject.Object.IsValid() && CacheObject.Object->IsNameStableForNetworking())))
 		{
-			continue;
-		}
+			// if we know the guid was specifically deleted, do not serialize it
+			if (DeletedNetStartupActorGUIDs.Contains(NetworkGUID))
+			{
+				continue;
+			}
 
-		FString PathName = CacheObject.Object->GetName();
+			FString PathName = CacheObject.Object.IsValid() ? CacheObject.Object->GetName() : CacheObject.PathName.ToString();
 
-		GEngine->NetworkRemapPath(this, PathName, false);
+			GEngine->NetworkRemapPath(this, PathName, false);
 
-		*CheckpointArchive << NetworkGUID;
-		*CheckpointArchive << CacheObject.OuterGUID;
-		*CheckpointArchive << PathName;
-		*CheckpointArchive << CacheObject.NetworkChecksum;
+			*CheckpointArchive << NetworkGUID;
+			*CheckpointArchive << CacheObject.OuterGUID;
+			*CheckpointArchive << PathName;
+			*CheckpointArchive << CacheObject.NetworkChecksum;
 
-		uint8 Flags = 0;
-		
-		Flags |= CacheObject.bNoLoad ? ( 1 << 0 ) : 0;
-		Flags |= CacheObject.bIgnoreWhenMissing ? ( 1 << 1 ) : 0;
+			uint8 Flags = 0;
+			Flags |= CacheObject.bNoLoad ? (1 << 0) : 0;
+			Flags |= CacheObject.bIgnoreWhenMissing ? (1 << 1) : 0;
+
+			*CheckpointArchive << Flags;
 
-		*CheckpointArchive << Flags;
+			++NumValues;
+
+			const bool bUnloaded = !CacheObject.Object.IsValid() || !CacheObject.Object->IsNameStableForNetworking();
+			if (bUnloaded)
+			{
+				++UnloadedValues;
+			}
+		}
 	}
+
+	int64 Pos = CheckpointArchive->Tell();
+	CheckpointArchive->Seek(CountPos);
+	*CheckpointArchive << NumValues;
+	CheckpointArchive->Seek(Pos);
+
+	UE_LOG( LogDemo, Verbose, TEXT( "Checkpoint. SerializeGuidCache: %i Unloaded: %i" ), NumValues, UnloadedValues );
 }
 
 float UDemoNetDriver::GetCheckpointSaveMaxMSPerFrame() const
@@ -4891,6 +4898,12 @@ void UDemoNetDriver::NotifyActorDestroyed( AActor* Actor, bool IsSeamlessTravel
 
 			UE_LOG(LogDemo, VeryVerbose, TEXT("NotifyActyorDestroyed: adding actor to deleted startup list: %s"), *Actor->GetFullName());
 			DeletedNetStartupActors.Add( Actor->GetFullName() );
+
+			FNetworkGUID NetGUID = GuidCache->NetGUIDLookup.FindRef(Actor);
+			if (NetGUID.IsValid())
+			{
+				DeletedNetStartupActorGUIDs.Add(NetGUID);
+			}
 		}
 	}
 
