PR #5023: Fix documentation comment to reflect reality. (Contributed by [REDACTED_PERSON])

diff --git a/Engine/Source/Runtime/Engine/Classes/Engine/World.h b/Engine/Source/Runtime/Engine/Classes/Engine/World.h
--- a/Engine/Source/Runtime/Engine/Classes/Engine/World.h
+++ b/Engine/Source/Runtime/Engine/Classes/Engine/World.h
@@ -2983,9 +2983,8 @@ public:
 	void WelcomePlayer(UNetConnection* Connection);
 
 	/**
-	 * Used to get a net driver object by name. Default name is the game net driver
-	 * @param NetDriverName the name of the net driver being asked for
-	 * @return a pointer to the net driver or NULL if the named driver is not found
+	 * Used to get a net driver object.
+	 * @return a pointer to the net driver or NULL if no driver is available.
 	 */
 	FORCEINLINE_DEBUGGABLE UNetDriver* GetNetDriver() const
 	{
