Add missing null check in OnLevelRemovedFromWorld

diff --git a/Engine/Source/Runtime/Engine/Private/NetworkDriver.cpp b/Engine/Source/Runtime/Engine/Private/NetworkDriver.cpp
--- a/Engine/Source/Runtime/Engine/Private/NetworkDriver.cpp
+++ b/Engine/Source/Runtime/Engine/Private/NetworkDriver.cpp
@@ -4554,17 +4554,20 @@ void UNetDriver::OnLevelRemovedFromWorld(class ULevel* Level, class UWorld* InWo
 		{
 			for (AActor* Actor : Level->Actors)
 			{
-				// Keep Startup actors alive, because they haven't been destroyed yet.
-				// Technically, Dynamic actors may not have been destroyed either, but this
-				// resembles relevancy.
-				if (!Actor->IsNetStartupActor())
-				{
-					NotifyActorLevelUnloaded(Actor);
-				}
-				else
+				if (Actor)
 				{
-					// We still want to remove Startup actors from the Network list so they aren't processed anymore.
-					GetNetworkObjectList().Remove(Actor);
+					// Keep Startup actors alive, because they haven't been destroyed yet.
+					// Technically, Dynamic actors may not have been destroyed either, but this
+					// resembles relevancy.
+					if (!Actor->IsNetStartupActor())
+					{
+						NotifyActorLevelUnloaded(Actor);
+					}
+					else
+					{
+						// We still want to remove Startup actors from the Network list so they aren't processed anymore.
+						GetNetworkObjectList().Remove(Actor);
+					}
 				}
 			}
 
