Fix issue with FUrl improperly parsing Ipv4 addresses that are really IPv4 - it tries to use ::ffff which is incorrect (from JakeL)

diff --git a/Engine/Source/Runtime/Engine/Private/URL.cpp b/Engine/Source/Runtime/Engine/Private/URL.cpp
--- a/Engine/Source/Runtime/Engine/Private/URL.cpp
+++ b/Engine/Source/Runtime/Engine/Private/URL.cpp
@@ -306,14 +306,14 @@ FURL::FURL( FURL* Base, const TCHAR* TextURL, ETravelType Type )
 					*(t++) = 0;
 					Port = FCString::Atoi( t );
 				}
-
+/*
 				if(SquareBracket && ClosingSquareBracket)
 				{
 					// Trim the brackets from the host address
 					*ClosingSquareBracket = 0;
 					Host = ss + 1;
 				}
-				else
+				else*/
 				{
 					// Normal IPv4 address
 					Host = ss;
