Fix SpeedTree node not working on ES3_1 (Contributed by [REDACTED_PERSON])
PR #4761

diff --git a/Engine/Source/Runtime/Engine/Private/Materials/HLSLMaterialTranslator.h b/Engine/Source/Runtime/Engine/Private/Materials/HLSLMaterialTranslator.h
--- a/Engine/Source/Runtime/Engine/Private/Materials/HLSLMaterialTranslator.h
+++ b/Engine/Source/Runtime/Engine/Private/Materials/HLSLMaterialTranslator.h
@@ -5472,7 +5472,7 @@ protected:
 	*/
 	virtual int32 SpeedTree(int32 GeometryArg, int32 WindArg, int32 LODArg, float BillboardThreshold, bool bAccurateWindVelocities, bool bExtraBend, int32 ExtraBendArg) override
 	{ 
-		if (ErrorUnlessFeatureLevelSupported(ERHIFeatureLevel::SM4) == INDEX_NONE)
+		if (ErrorUnlessFeatureLevelSupported(ERHIFeatureLevel::ES3_1) == INDEX_NONE)
 		{
 			return INDEX_NONE;
 		}
