We should not Tick USceneCaptureComponent on DS by default.

diff --git a/Engine/Source/Runtime/Engine/Private/Components/SceneCaptureComponent.cpp b/Engine/Source/Runtime/Engine/Private/Components/SceneCaptureComponent.cpp
--- a/Engine/Source/Runtime/Engine/Private/Components/SceneCaptureComponent.cpp
+++ b/Engine/Source/Runtime/Engine/Private/Components/SceneCaptureComponent.cpp
@@ -438,6 +438,7 @@ USceneCaptureComponent2D::USceneCaptureComponent2D(const FObjectInitializer& Obj
 	bAutoActivate = true;
 	PrimaryComponentTick.bCanEverTick = true;
 	PrimaryComponentTick.TickGroup = TG_DuringPhysics;
+	PrimaryComponentTick.bAllowTickOnDedicatedServer = false;
 	// Tick in the editor so that bCaptureEveryFrame preview works
 	bTickInEditor = true;
 	// previous behavior was to capture from raw scene color 
@@ -715,6 +716,7 @@ UPlanarReflectionComponent::UPlanarReflectionComponent(const FObjectInitializer&
 	bCaptureEveryFrame = true;
 	PrimaryComponentTick.bCanEverTick = true;
 	PrimaryComponentTick.TickGroup = TG_DuringPhysics;
+	PrimaryComponentTick.bAllowTickOnDedicatedServer = false;
 	// Tick in the editor so that bCaptureEveryFrame preview works
 	bTickInEditor = true;
 	RenderTarget = NULL;
@@ -858,6 +860,7 @@ USceneCaptureComponentCube::USceneCaptureComponentCube(const FObjectInitializer&
 	bAutoActivate = true;
 	PrimaryComponentTick.bCanEverTick = true;
 	PrimaryComponentTick.TickGroup = TG_DuringPhysics;
+	PrimaryComponentTick.bAllowTickOnDedicatedServer = false;
 	bTickInEditor = true;
 	IPD = 6.2f;
 }
