Merging //UE4/Dev-Main to Dev-Mobile (//UE4/Dev-Mobile)

diff --git a/Engine/Source/Runtime/Engine/Private/URL.cpp b/Engine/Source/Runtime/Engine/Private/URL.cpp
--- a/Engine/Source/Runtime/Engine/Private/URL.cpp
+++ b/Engine/Source/Runtime/Engine/Private/URL.cpp
@@ -306,14 +306,14 @@ FURL::FURL( FURL* Base, const TCHAR* TextURL, ETravelType Type )
 					*(t++) = 0;
 					Port = FCString::Atoi( t );
 				}
-
+/*
 				if(SquareBracket && ClosingSquareBracket)
 				{
 					// Trim the brackets from the host address
 					*ClosingSquareBracket = 0;
 					Host = ss + 1;
 				}
-				else
+				else*/
 				{
 					// Normal IPv4 address
 					Host = ss;
