MN - Reflection captures can miss end of frame updates when loading a level from the (sub)level view in editor. In this case the HISM was just added into the level and it's instance was buffer changed, this caused a dangeling pointer within the vertexfactory during reflection capture rendering because the renderstate had not been recreated although it was Marked as such.

diff --git a/Engine/Source/Runtime/Engine/Private/Components/ReflectionCaptureComponent.cpp b/Engine/Source/Runtime/Engine/Private/Components/ReflectionCaptureComponent.cpp
--- a/Engine/Source/Runtime/Engine/Private/Components/ReflectionCaptureComponent.cpp
+++ b/Engine/Source/Runtime/Engine/Private/Components/ReflectionCaptureComponent.cpp
@@ -1039,6 +1039,9 @@ void UReflectionCaptureComponent::UpdateReflectionCaptureContents(UWorld* WorldT
 		// Note: this will also prevent uploads of cubemaps from DDC, which is unintentional
 		&& (GShaderCompilingManager == NULL || !GShaderCompilingManager->IsCompiling()))
 	{
+		//guarantee that all render proxies are up to date before kicking off this render
+		WorldToUpdate->SendAllEndOfFrameUpdates();
+
 		TArray<UReflectionCaptureComponent*> WorldCombinedCaptures;
 
 		for (int32 CaptureIndex = ReflectionCapturesToUpdate.Num() - 1; CaptureIndex >= 0; CaptureIndex--)
