[UE-Main] Merge:

CL 4395628

- Fix for crash when flushing/transitioning levels.
- Perform an audio render thread command flush
- Removed sound wave ref counting code

diff --git a/Engine/Source/Runtime/AudioMixer/Private/AudioMixerDevice.cpp b/Engine/Source/Runtime/AudioMixer/Private/AudioMixerDevice.cpp
--- a/Engine/Source/Runtime/AudioMixer/Private/AudioMixerDevice.cpp
+++ b/Engine/Source/Runtime/AudioMixer/Private/AudioMixerDevice.cpp
@@ -909,6 +909,10 @@ namespace Audio
 		}
 	}
 
+	void FMixerDevice::FlushAudioRenderingCommands()
+	{
+		SourceManager.FlushCommandQueue();
+	}
 
 	bool FMixerDevice::IsMasterSubmixType(USoundSubmix* InSubmix) const
 	{
diff --git a/Engine/Source/Runtime/AudioMixer/Private/AudioMixerSourceBuffer.cpp b/Engine/Source/Runtime/AudioMixer/Private/AudioMixerSourceBuffer.cpp
--- a/Engine/Source/Runtime/AudioMixer/Private/AudioMixerSourceBuffer.cpp
+++ b/Engine/Source/Runtime/AudioMixer/Private/AudioMixerSourceBuffer.cpp
@@ -105,9 +105,11 @@ namespace Audio
 	bool FMixerSourceBuffer::Init()
 	{
 		check(SoundWave);
-
-		// We flag that this sound wave is active for the lifetime of this object since we use it for decoding, etc.
-		SoundWave->IncrementNumSounds();
+		if (SoundWave->bProcedural && SoundWave->IsGenerating())
+		{
+			UE_LOG(LogAudioMixer, Warning, TEXT("Procedural sound wave is reinitializing even though it is currently actively generating audio. Please stop sound before trying to play it again."));
+			return false;
+		}
 
 		// We have successfully initialized which means our SoundWave has been flagged as bIsActive
 		// GC can run between PreInit and Init so when cleaning up FMixerSourceBuffer, we don't want to touch SoundWave unless bInitailized is true.
@@ -399,9 +401,9 @@ namespace Audio
 		{
 			if (SoundWave->bProcedural)
 			{
+				SoundWave->SetGenerating(true);
 				SoundWave->OnBeginGenerate();
 			}
-
 		}
 	}
 
@@ -416,9 +418,9 @@ namespace Audio
 			if (SoundWave->bProcedural)
 			{
 				SoundWave->OnEndGenerate();
+				SoundWave->SetGenerating(false);
 			}
 
-			SoundWave->DecrementNumSounds();
 			SoundWave = nullptr;
 			bInitialized = false;
 		}
@@ -430,7 +432,6 @@ namespace Audio
 			{
 				delete MixerBuffer;
 			}
-
 			MixerBuffer = nullptr;
 		}
 	}
diff --git a/Engine/Source/Runtime/AudioMixer/Private/AudioMixerSourceManager.cpp b/Engine/Source/Runtime/AudioMixer/Private/AudioMixerSourceManager.cpp
--- a/Engine/Source/Runtime/AudioMixer/Private/AudioMixerSourceManager.cpp
+++ b/Engine/Source/Runtime/AudioMixer/Private/AudioMixerSourceManager.cpp
@@ -2137,6 +2137,16 @@ namespace Audio
 		RenderThreadCommandBufferIndex.Set(!CurrentRenderThreadIndex);
 	}
 
+	void FMixerSourceManager::FlushCommandQueue()
+	{
+		bPumpQueue = true;
+		while (bPumpQueue)
+		{
+			FPlatformProcess::Sleep(0);
+		}
+		bPumpQueue = true;
+	}
+
 	void FMixerSourceManager::UpdatePendingReleaseData(bool bForceWait)
 	{
 		// Check for any pending delete procedural sound waves
diff --git a/Engine/Source/Runtime/AudioMixer/Private/AudioMixerSourceManager.h b/Engine/Source/Runtime/AudioMixer/Private/AudioMixerSourceManager.h
--- a/Engine/Source/Runtime/AudioMixer/Private/AudioMixerSourceManager.h
+++ b/Engine/Source/Runtime/AudioMixer/Private/AudioMixerSourceManager.h
@@ -269,7 +269,7 @@ namespace Audio
 		bool IsBus(const int32 SourceId) const;
 		void PumpCommandQueue();
 		void UpdatePendingReleaseData(bool bForceWait = false);
-
+		void FlushCommandQueue();
 	private:
 
 		void ReleaseSource(const int32 SourceId);
diff --git a/Engine/Source/Runtime/AudioMixer/Private/Components/SynthComponent.cpp b/Engine/Source/Runtime/AudioMixer/Private/Components/SynthComponent.cpp
--- a/Engine/Source/Runtime/AudioMixer/Private/Components/SynthComponent.cpp
+++ b/Engine/Source/Runtime/AudioMixer/Private/Components/SynthComponent.cpp
@@ -280,7 +280,7 @@ void USynthComponent::OnUnregister()
 bool USynthComponent::IsReadyForOwnerToAutoDestroy() const
 {
 	const bool bIsAudioComponentReadyForDestroy = !AudioComponent || (AudioComponent && !AudioComponent->IsPlaying());
-	const bool bIsSynthSoundReadyForDestroy = !Synth || !Synth->GetNumSoundsActive();
+	const bool bIsSynthSoundReadyForDestroy = !Synth || !Synth->IsGenerating();
 	return bIsAudioComponentReadyForDestroy && bIsSynthSoundReadyForDestroy;
 }
 
diff --git a/Engine/Source/Runtime/AudioMixer/Public/AudioMixerDevice.h b/Engine/Source/Runtime/AudioMixer/Public/AudioMixerDevice.h
--- a/Engine/Source/Runtime/AudioMixer/Public/AudioMixerDevice.h
+++ b/Engine/Source/Runtime/AudioMixer/Public/AudioMixerDevice.h
@@ -123,6 +123,8 @@ namespace Audio
 		// Submix buffer listener callbacks
 		virtual void RegisterSubmixBufferListener(ISubmixBufferListener* InSubmixBufferListener, USoundSubmix* InSubmix = nullptr) override;
 		virtual void UnregisterSubmixBufferListener(ISubmixBufferListener* InSubmixBufferListener, USoundSubmix* InSubmix = nullptr) override;
+
+		virtual void FlushAudioRenderingCommands() override;
 		//~ End FAudioDevice
 
 		//~ Begin IAudioMixer
diff --git a/Engine/Source/Runtime/Engine/Classes/Sound/SoundWave.h b/Engine/Source/Runtime/Engine/Classes/Sound/SoundWave.h
--- a/Engine/Source/Runtime/Engine/Classes/Sound/SoundWave.h
+++ b/Engine/Source/Runtime/Engine/Classes/Sound/SoundWave.h
@@ -220,10 +220,7 @@ private:
 	/** What state the precache decompressor is in. */
 	FThreadSafeCounter PrecacheState;
 
-	/** Number of sounds actively using this sound wave by the audio renderer (in audio mixer). Prevents GC issues with rendering realtime audio. */
-	FThreadSafeCounter NumSoundsActive;
-	int32 GCTimeOutCount;
-
+	FThreadSafeBool bGenerating;
 #if !WITH_EDITOR
 	// This is the sample rate gotten from platform settings.
 	float CachedSampleRateOverride;
@@ -391,13 +388,8 @@ public:
 	// Called when the procedural sound wave is done generating on the render thread. Only used in the audio mixer and when bProcedural is true..
 	virtual void OnEndGenerate() {};
 
-	// Returns number of sounds using this sound wave (audio mixer only)
-	int32 GetNumSoundsActive();
-
-	// Increment and decrement num sounds (used in audio mixer)
-	void IncrementNumSounds();
-	void DecrementNumSounds();
-
+	bool IsGenerating() const { return bGenerating; }
+	void SetGenerating(bool bInGenerating) { bGenerating = bInGenerating; }
 	/**
 	* Overwrite sample rate. Used for procedural soundwaves, as well as sound waves that are resampled on compress/decompress.
 	*/
diff --git a/Engine/Source/Runtime/Engine/Private/AudioDevice.cpp b/Engine/Source/Runtime/Engine/Private/AudioDevice.cpp
--- a/Engine/Source/Runtime/Engine/Private/AudioDevice.cpp
+++ b/Engine/Source/Runtime/Engine/Private/AudioDevice.cpp
@@ -3833,7 +3833,7 @@ void FAudioDevice::AddNewActiveSound(const FActiveSound& NewActiveSound)
 	}
 
 	USoundWave* SoundWave = Cast<USoundWave>(Sound);
-	if (SoundWave && SoundWave->bProcedural && SoundWave->GetNumSoundsActive() > 0)
+	if (SoundWave && SoundWave->bProcedural && SoundWave->IsGenerating())
 	{
 		FString SoundWaveName;
 		SoundWave->GetName(SoundWaveName);
@@ -4668,6 +4668,9 @@ void FAudioDevice::Flush(UWorld* WorldToFlush, bool bClearActivatedReverb)
 	// Immediately stop all pending active sounds
 	ProcessingPendingActiveSoundStops(WorldToFlush == nullptr || WorldToFlush->bIsTearingDown);
 
+	// Make sure any in-flight audio rendering commands get executed.
+	FlushAudioRenderingCommands();
+
 	// Anytime we flush, make sure to clear all the listeners.  We'll get the right ones soon enough.
 	Listeners.Reset();
 	Listeners.Add(FListener(this));
diff --git a/Engine/Source/Runtime/Engine/Private/SoundWave.cpp b/Engine/Source/Runtime/Engine/Private/SoundWave.cpp
--- a/Engine/Source/Runtime/Engine/Private/SoundWave.cpp
+++ b/Engine/Source/Runtime/Engine/Private/SoundWave.cpp
@@ -639,9 +639,6 @@ void USoundWave::BeginDestroy()
 	// the audio render thread stops the sound before GC hits.
 	bIsBeginDestroy = true;
 	
-	// Arbitrary count to wait before forcing IsFinishDestroy to return true. Temporary fallback to hunt a deadlock.
-	GCTimeOutCount = 10;
-
 #if WITH_EDITOR
 	// Flush any async results so we dont leak them in the DDC
 	if (GetDerivedDataCache() && AsyncLoadingDataFormats.Num() > 0)
@@ -859,22 +856,6 @@ FWaveInstance* USoundWave::HandleStart( FActiveSound& ActiveSound, const UPTRINT
 	return WaveInstance;
 }
 
-int32 USoundWave::GetNumSoundsActive()
-{
-	return NumSoundsActive.GetValue();
-}
-
-void USoundWave::IncrementNumSounds()
-{
-	NumSoundsActive.Increment();
-}
-
-void USoundWave::DecrementNumSounds()
-{
-	int32 NewValue = NumSoundsActive.Decrement();
-	check(NewValue >= 0);
-}
-
 bool USoundWave::IsReadyForFinishDestroy()
 {
 	const bool bIsStreamingInProgress = IStreamingManager::Get().GetAudioStreamingManager().IsStreamingInProgress(this);
@@ -894,23 +875,7 @@ bool USoundWave::IsReadyForFinishDestroy()
 			}, GET_STATID(STAT_AudioFreeResources));
 		}
 	
-	bool AllowFinishDestroy = (ResourceState == ESoundWaveResourceState::Freed);
-
-	if (AllowFinishDestroy)
-	{
-		AllowFinishDestroy = (NumSoundsActive.GetValue() == 0);
-		if (!AllowFinishDestroy && --GCTimeOutCount <= 0)
-		{
-			AllowFinishDestroy = true;
-
-			if (AllowFinishDestroy && NumSoundsActive.GetValue() > 0)
-			{
-				UE_LOG(LogAudio, Warning, TEXT("Sound wave '%s' destroyed while num sounds active was '%d'."), *GetName(), NumSoundsActive.GetValue());
-			}
-		}
-	}
-
-	return AllowFinishDestroy;
+	return !bGenerating && ResourceState == ESoundWaveResourceState::Freed;
 }
 
 
diff --git a/Engine/Source/Runtime/Engine/Public/Audio.h b/Engine/Source/Runtime/Engine/Public/Audio.h
--- a/Engine/Source/Runtime/Engine/Public/Audio.h
+++ b/Engine/Source/Runtime/Engine/Public/Audio.h
@@ -717,7 +717,7 @@ public:
 	ENGINE_API FSpatializationParams GetSpatializationParams();
 
 	/** Returns the contained sound buffer object. */
-	const FSoundBuffer* GetBuffer() const { return Buffer; }
+	virtual const FSoundBuffer* GetBuffer() const { return Buffer; }
 
 	/** Initializes any source effects for this sound source. */
 	virtual void InitializeSourceEffects(uint32 InEffectVoiceId)
diff --git a/Engine/Source/Runtime/Engine/Public/AudioDevice.h b/Engine/Source/Runtime/Engine/Public/AudioDevice.h
--- a/Engine/Source/Runtime/Engine/Public/AudioDevice.h
+++ b/Engine/Source/Runtime/Engine/Public/AudioDevice.h
@@ -582,6 +582,9 @@ public:
 	 */
 	void Flush(UWorld* WorldToFlush, bool bClearActivatedReverb = true);
 
+	/** Allows audio rendering command queue to flush during audio device flush. */
+	virtual void FlushAudioRenderingCommands() {}
+
 	/**
 	 * Stop any playing sounds that are using a particular SoundWave
 	 *
