Fix CIS error in UE4Game Mac NonUnity.

diff --git a/Engine/Source/Runtime/Engine/Private/PhysicsEngine/PhysicsReplication.cpp b/Engine/Source/Runtime/Engine/Private/PhysicsEngine/PhysicsReplication.cpp
--- a/Engine/Source/Runtime/Engine/Private/PhysicsEngine/PhysicsReplication.cpp
+++ b/Engine/Source/Runtime/Engine/Private/PhysicsEngine/PhysicsReplication.cpp
@@ -16,6 +16,7 @@
 #include "Engine/Classes/Engine/Player.h"
 #include "Physics/PhysicsInterfaceCore.h"
 #include "Physics/PhysScene_PhysX.h"
+#include "Components/SkeletalMeshComponent.h"
 
 #if WITH_PHYSX
 #include "PhysXPublic.h"
