Remove mismatched PRAGMA_ENABLE_OPTIMIZATION macros, causing CIS warning.

diff --git a/Engine/Source/Runtime/Engine/Private/Animation/PoseAsset.cpp b/Engine/Source/Runtime/Engine/Private/Animation/PoseAsset.cpp
--- a/Engine/Source/Runtime/Engine/Private/Animation/PoseAsset.cpp
+++ b/Engine/Source/Runtime/Engine/Private/Animation/PoseAsset.cpp
@@ -1454,5 +1454,4 @@ bool UPoseAsset::GetBasePoseTransform(TArray<FTransform>& OutBasePose, TArray<fl
 }
 #endif // WITH_EDITOR
 
-PRAGMA_ENABLE_OPTIMIZATION
 #undef LOCTEXT_NAMESPACE 
diff --git a/Engine/Source/Runtime/Engine/Private/StreamableManager.cpp b/Engine/Source/Runtime/Engine/Private/StreamableManager.cpp
--- a/Engine/Source/Runtime/Engine/Private/StreamableManager.cpp
+++ b/Engine/Source/Runtime/Engine/Private/StreamableManager.cpp
@@ -1353,4 +1353,3 @@ FSoftObjectPath FStreamableManager::HandleLoadedRedirector(UObjectRedirector* Lo
 	return NewPath;
 }
 
-PRAGMA_ENABLE_OPTIMIZATION
