Fix Main CIS NoPhysXEditor Win64 step due to missing function when running without PhysX interface

diff --git a/Engine/Source/Runtime/Engine/Private/PhysicsEngine/Experimental/PhysInterface_Apeiron.cpp b/Engine/Source/Runtime/Engine/Private/PhysicsEngine/Experimental/PhysInterface_Apeiron.cpp
--- a/Engine/Source/Runtime/Engine/Private/PhysicsEngine/Experimental/PhysInterface_Apeiron.cpp
+++ b/Engine/Source/Runtime/Engine/Private/PhysicsEngine/Experimental/PhysInterface_Apeiron.cpp
@@ -778,6 +778,14 @@ FTransform FPhysInterface_Apeiron::GetComTransform_AssumesLocked(const FPhysicsA
     return GlobalTransform * ComTransform;
 }
 
+FTransform FPhysInterface_Apeiron::GetComTransformLocal_AssumesLocked(const FPhysicsActorReference_Apeiron& InActorReference)
+{
+	uint32 Index = UINT32_MAX;
+	const auto& LocalParticles = GetParticlesAndIndex(InActorReference, Index);
+	Apeiron::TRigidTransform<float, 3> ComTransform(LocalParticles.Geometry(Index) ? LocalParticles.Geometry(Index)->BoundingBox().Center() : Apeiron::TVector<float, 3>(0), Apeiron::TRotation<float, 3>(FQuat(0, 0, 0, 1)));
+	return ComTransform;
+}
+
 FVector FPhysInterface_Apeiron::GetLocalInertiaTensor_AssumesLocked(const FPhysicsActorReference_Apeiron& InActorReference)
 {
     // @todo(mlentine): Just return directly once we implement DiagonalMatrix
diff --git a/Engine/Source/Runtime/Engine/Public/Physics/Experimental/PhysInterface_Apeiron.h b/Engine/Source/Runtime/Engine/Public/Physics/Experimental/PhysInterface_Apeiron.h
--- a/Engine/Source/Runtime/Engine/Public/Physics/Experimental/PhysInterface_Apeiron.h
+++ b/Engine/Source/Runtime/Engine/Public/Physics/Experimental/PhysInterface_Apeiron.h
@@ -313,6 +313,7 @@ public:
 	static ENGINE_API FVector GetWorldVelocityAtPoint_AssumesLocked(const FPhysicsActorReference_Apeiron& InActorReference, const FVector& InPoint);
 
 	static ENGINE_API FTransform GetComTransform_AssumesLocked(const FPhysicsActorReference_Apeiron& InActorReference);
+	static ENGINE_API FTransform GetComTransformLocal_AssumesLocked(const FPhysicsActorReference_Apeiron& InActorReference);
 
 	static ENGINE_API FVector GetLocalInertiaTensor_AssumesLocked(const FPhysicsActorReference_Apeiron& InActorReference);
 	static ENGINE_API FBox GetBounds_AssumesLocked(const FPhysicsActorReference_Apeiron& InActorReference);
