MN - Add extra debug info for missing cloth buffer

diff --git a/Engine/Source/Runtime/Engine/Public/GPUSkinVertexFactory.h b/Engine/Source/Runtime/Engine/Public/GPUSkinVertexFactory.h
--- a/Engine/Source/Runtime/Engine/Public/GPUSkinVertexFactory.h
+++ b/Engine/Source/Runtime/Engine/Public/GPUSkinVertexFactory.h
@@ -721,7 +721,7 @@ public:
 				Index = 1 - Index;
 			}
 
-			check(ClothSimulPositionNormalBuffer[Index].VertexBufferRHI.IsValid());
+			checkf(ClothSimulPositionNormalBuffer[Index].VertexBufferRHI.IsValid(), TEXT("Index: %i Buffer0: %s Frame0: %i Buffer1: %s Frame1: %i"), Index,  ClothSimulPositionNormalBuffer[0].VertexBufferRHI.IsValid() ? TEXT("true") : TEXT("false"), BufferFrameNumber[0], ClothSimulPositionNormalBuffer[1].VertexBufferRHI.IsValid() ? TEXT("true") : TEXT("false"), BufferFrameNumber[1]);
 			return ClothSimulPositionNormalBuffer[Index];
 		}
 		
