Merging //UE4/Dev-Main to Dev-Build (//UE4/Dev-Build)

diff --git a/Engine/Build/Build.version b/Engine/Build/Build.version
--- a/Engine/Build/Build.version
+++ b/Engine/Build/Build.version
@@ -1,6 +1,6 @@
 {
 	"MajorVersion": 4,
-	"MinorVersion": 21,
+	"MinorVersion": 22,
 	"PatchVersion": 0,
 	"Changelist": 0,
 	"CompatibleChangelist": 0,
diff --git a/Engine/Plugins/Runtime/GoogleARCore/Source/GoogleARCoreBase/Public/GoogleARCoreAugmentedImageDatabase.h b/Engine/Plugins/Runtime/GoogleARCore/Source/GoogleARCoreBase/Public/GoogleARCoreAugmentedImageDatabase.h
--- a/Engine/Plugins/Runtime/GoogleARCore/Source/GoogleARCoreBase/Public/GoogleARCoreAugmentedImageDatabase.h
+++ b/Engine/Plugins/Runtime/GoogleARCore/Source/GoogleARCoreBase/Public/GoogleARCoreAugmentedImageDatabase.h
@@ -34,7 +34,9 @@ struct FGoogleARCoreAugmentedImageDatabaseEntry
 	float Width;
 
 	FGoogleARCoreAugmentedImageDatabaseEntry()
-		: Width(0.0f) { }
+		: ImageAsset(nullptr)
+	    , Width(0.0f) 
+	{ }
 };
 
 UCLASS(BlueprintType)
diff --git a/Engine/Plugins/Runtime/OptionalMobileFeaturesBPLibrary/Source/OptionalMobileFeaturesBPLibrary/Private/OptionalMobileFeaturesBPLibrary.cpp b/Engine/Plugins/Runtime/OptionalMobileFeaturesBPLibrary/Source/OptionalMobileFeaturesBPLibrary/Private/OptionalMobileFeaturesBPLibrary.cpp
--- a/Engine/Plugins/Runtime/OptionalMobileFeaturesBPLibrary/Source/OptionalMobileFeaturesBPLibrary/Private/OptionalMobileFeaturesBPLibrary.cpp
+++ b/Engine/Plugins/Runtime/OptionalMobileFeaturesBPLibrary/Source/OptionalMobileFeaturesBPLibrary/Private/OptionalMobileFeaturesBPLibrary.cpp
@@ -18,7 +18,7 @@ int UOptionalMobileFeaturesBPLibrary::GetVolumeState()
 	int scaledVolume = (baseVolume * 100) / 15;
 	return scaledVolume;
 #elif PLATFORM_IOS
-	return FIOSPlatformMisc::GetAudioVolume();
+	return FPlatformMisc::GetDeviceVolume();
 #else
 	return 0;
 #endif
diff --git a/Engine/Source/Runtime/Engine/Private/Animation/PoseAsset.cpp b/Engine/Source/Runtime/Engine/Private/Animation/PoseAsset.cpp
--- a/Engine/Source/Runtime/Engine/Private/Animation/PoseAsset.cpp
+++ b/Engine/Source/Runtime/Engine/Private/Animation/PoseAsset.cpp
@@ -1454,5 +1454,4 @@ bool UPoseAsset::GetBasePoseTransform(TArray<FTransform>& OutBasePose, TArray<fl
 }
 #endif // WITH_EDITOR
 
-PRAGMA_ENABLE_OPTIMIZATION
 #undef LOCTEXT_NAMESPACE 
diff --git a/Engine/Source/Runtime/Engine/Private/PhysicsEngine/Experimental/PhysInterface_Apeiron.cpp b/Engine/Source/Runtime/Engine/Private/PhysicsEngine/Experimental/PhysInterface_Apeiron.cpp
--- a/Engine/Source/Runtime/Engine/Private/PhysicsEngine/Experimental/PhysInterface_Apeiron.cpp
+++ b/Engine/Source/Runtime/Engine/Private/PhysicsEngine/Experimental/PhysInterface_Apeiron.cpp
@@ -778,6 +778,14 @@ FTransform FPhysInterface_Apeiron::GetComTransform_AssumesLocked(const FPhysicsA
     return GlobalTransform * ComTransform;
 }
 
+FTransform FPhysInterface_Apeiron::GetComTransformLocal_AssumesLocked(const FPhysicsActorReference_Apeiron& InActorReference)
+{
+	uint32 Index = UINT32_MAX;
+	const auto& LocalParticles = GetParticlesAndIndex(InActorReference, Index);
+	Apeiron::TRigidTransform<float, 3> ComTransform(LocalParticles.Geometry(Index) ? LocalParticles.Geometry(Index)->BoundingBox().Center() : Apeiron::TVector<float, 3>(0), Apeiron::TRotation<float, 3>(FQuat(0, 0, 0, 1)));
+	return ComTransform;
+}
+
 FVector FPhysInterface_Apeiron::GetLocalInertiaTensor_AssumesLocked(const FPhysicsActorReference_Apeiron& InActorReference)
 {
     // @todo(mlentine): Just return directly once we implement DiagonalMatrix
diff --git a/Engine/Source/Runtime/Engine/Private/PhysicsEngine/PhysicsReplication.cpp b/Engine/Source/Runtime/Engine/Private/PhysicsEngine/PhysicsReplication.cpp
--- a/Engine/Source/Runtime/Engine/Private/PhysicsEngine/PhysicsReplication.cpp
+++ b/Engine/Source/Runtime/Engine/Private/PhysicsEngine/PhysicsReplication.cpp
@@ -16,6 +16,7 @@
 #include "Engine/Classes/Engine/Player.h"
 #include "Physics/PhysicsInterfaceCore.h"
 #include "Physics/PhysScene_PhysX.h"
+#include "Components/SkeletalMeshComponent.h"
 
 #if WITH_PHYSX
 #include "PhysXPublic.h"
diff --git a/Engine/Source/Runtime/Engine/Private/StreamableManager.cpp b/Engine/Source/Runtime/Engine/Private/StreamableManager.cpp
--- a/Engine/Source/Runtime/Engine/Private/StreamableManager.cpp
+++ b/Engine/Source/Runtime/Engine/Private/StreamableManager.cpp
@@ -1353,4 +1353,3 @@ FSoftObjectPath FStreamableManager::HandleLoadedRedirector(UObjectRedirector* Lo
 	return NewPath;
 }
 
-PRAGMA_ENABLE_OPTIMIZATION
diff --git a/Engine/Source/Runtime/Engine/Public/GPUSkinVertexFactory.h b/Engine/Source/Runtime/Engine/Public/GPUSkinVertexFactory.h
--- a/Engine/Source/Runtime/Engine/Public/GPUSkinVertexFactory.h
+++ b/Engine/Source/Runtime/Engine/Public/GPUSkinVertexFactory.h
@@ -721,7 +721,7 @@ public:
 				Index = 1 - Index;
 			}
 
-			check(ClothSimulPositionNormalBuffer[Index].VertexBufferRHI.IsValid());
+			checkf(ClothSimulPositionNormalBuffer[Index].VertexBufferRHI.IsValid(), TEXT("Index: %i Buffer0: %s Frame0: %i Buffer1: %s Frame1: %i"), Index,  ClothSimulPositionNormalBuffer[0].VertexBufferRHI.IsValid() ? TEXT("true") : TEXT("false"), BufferFrameNumber[0], ClothSimulPositionNormalBuffer[1].VertexBufferRHI.IsValid() ? TEXT("true") : TEXT("false"), BufferFrameNumber[1]);
 			return ClothSimulPositionNormalBuffer[Index];
 		}
 		
diff --git a/Engine/Source/Runtime/Engine/Public/Physics/Experimental/PhysInterface_Apeiron.h b/Engine/Source/Runtime/Engine/Public/Physics/Experimental/PhysInterface_Apeiron.h
--- a/Engine/Source/Runtime/Engine/Public/Physics/Experimental/PhysInterface_Apeiron.h
+++ b/Engine/Source/Runtime/Engine/Public/Physics/Experimental/PhysInterface_Apeiron.h
@@ -313,6 +313,7 @@ public:
 	static ENGINE_API FVector GetWorldVelocityAtPoint_AssumesLocked(const FPhysicsActorReference_Apeiron& InActorReference, const FVector& InPoint);
 
 	static ENGINE_API FTransform GetComTransform_AssumesLocked(const FPhysicsActorReference_Apeiron& InActorReference);
+	static ENGINE_API FTransform GetComTransformLocal_AssumesLocked(const FPhysicsActorReference_Apeiron& InActorReference);
 
 	static ENGINE_API FVector GetLocalInertiaTensor_AssumesLocked(const FPhysicsActorReference_Apeiron& InActorReference);
 	static ENGINE_API FBox GetBounds_AssumesLocked(const FPhysicsActorReference_Apeiron& InActorReference);
diff --git a/Engine/Source/Runtime/Launch/Resources/Version.h b/Engine/Source/Runtime/Launch/Resources/Version.h
--- a/Engine/Source/Runtime/Launch/Resources/Version.h
+++ b/Engine/Source/Runtime/Launch/Resources/Version.h
@@ -55,7 +55,7 @@
 // These numbers define the banner UE4 version, and are the most significant numbers when ordering two engine versions (that is, a 4.12.* version is always 
 // newer than a 4.11.* version, regardless of the changelist that it was built with)
 #define ENGINE_MAJOR_VERSION	4
-#define ENGINE_MINOR_VERSION	21
+#define ENGINE_MINOR_VERSION	22
 #define ENGINE_PATCH_VERSION	0
 
 // Macros for encoding strings
diff --git a/Engine/Source/Runtime/LiveLinkInterface/LiveLinkInterface.Build.cs b/Engine/Source/Runtime/LiveLinkInterface/LiveLinkInterface.Build.cs
--- a/Engine/Source/Runtime/LiveLinkInterface/LiveLinkInterface.Build.cs
+++ b/Engine/Source/Runtime/LiveLinkInterface/LiveLinkInterface.Build.cs
@@ -9,8 +9,7 @@ public class LiveLinkInterface : ModuleRules
         PublicDependencyModuleNames.AddRange(
             new string[] {
                 "Core",
-                "CoreUObject",
-				"TimeManagement"
+                "CoreUObject"
             }
         );
 	}
