Fixes for compiling the editor with -StressTestUnity enabled, caused by symbols being declared in multiple files, etc...

diff --git a/Engine/Source/Editor/DetailCustomizations/Private/LinuxTargetSettingsDetails.cpp b/Engine/Source/Editor/DetailCustomizations/Private/LinuxTargetSettingsDetails.cpp
--- a/Engine/Source/Editor/DetailCustomizations/Private/LinuxTargetSettingsDetails.cpp
+++ b/Engine/Source/Editor/DetailCustomizations/Private/LinuxTargetSettingsDetails.cpp
@@ -46,7 +46,7 @@ namespace LinuxTargetSettingsDetailsConstants
 	const FText DisabledTip = LOCTEXT("GitHubSourceRequiredToolTip", "This requires GitHub source.");
 }
 
-static FText GetFriendlyNameFromRHIName(const FString& InRHIName)
+static FText GetFriendlyNameFromLinuxRHIName(const FString& InRHIName)
 {
 	FText FriendlyRHIName = LOCTEXT("UnknownRHI", "UnknownRHI");
 	if (InRHIName == TEXT("GLSL_150"))
@@ -87,7 +87,7 @@ TSharedRef<IDetailCustomization> FLinuxTargetSettingsDetails::MakeInstance()
 	return MakeShareable(new FLinuxTargetSettingsDetails);
 }
 
-namespace EImageScope
+namespace ELinuxImageScope
 {
 	enum Type
 	{
@@ -97,11 +97,11 @@ namespace EImageScope
 }
 
 /* Helper function used to generate filenames for splash screens */
-static FString GetSplashFilename(EImageScope::Type Scope, bool bIsEditorSplash)
+static FString GetLinuxSplashFilename(ELinuxImageScope::Type Scope, bool bIsEditorSplash)
 {
 	FString Filename;
 
-	if (Scope == EImageScope::Engine)
+	if (Scope == ELinuxImageScope::Engine)
 	{
 		Filename = FPaths::EngineContentDir();
 	}
@@ -125,11 +125,11 @@ static FString GetSplashFilename(EImageScope::Type Scope, bool bIsEditorSplash)
 }
 
 /* Helper function used to generate filenames for icons */
-static FString GetIconFilename(EImageScope::Type Scope)
+static FString GetLinuxIconFilename(ELinuxImageScope::Type Scope)
 {
 	const FString& PlatformName = FModuleManager::GetModuleChecked<ITargetPlatformModule>("LinuxTargetPlatform").GetTargetPlatforms()[0]->PlatformName();
 
-	if (Scope == EImageScope::Engine)
+	if (Scope == ELinuxImageScope::Engine)
 	{
 		FString Filename = FPaths::EngineDir() / FString(TEXT("Source/Runtime/Launch/Resources")) / PlatformName / FString("UE4.png");
 		return FPaths::ConvertRelativePathToFull(Filename);
@@ -160,8 +160,8 @@ void FLinuxTargetSettingsDetails::CustomizeDetails( IDetailLayoutBuilder& Detail
 	IDetailCategoryBuilder& SplashCategoryBuilder = DetailBuilder.EditCategory(TEXT("Splash"));
 	FDetailWidgetRow& EditorSplashWidgetRow = SplashCategoryBuilder.AddCustomRow(EditorSplashDesc);
 
-	const FString EditorSplash_TargetImagePath = GetSplashFilename(EImageScope::GameOverride, true);
-	const FString EditorSplash_DefaultImagePath = GetSplashFilename(EImageScope::Engine, true);
+	const FString EditorSplash_TargetImagePath = GetLinuxSplashFilename(ELinuxImageScope::GameOverride, true);
+	const FString EditorSplash_DefaultImagePath = GetLinuxSplashFilename(ELinuxImageScope::Engine, true);
 
 	TArray<FString> ImageExtensions;
 	ImageExtensions.Add(TEXT("png"));
@@ -202,8 +202,8 @@ void FLinuxTargetSettingsDetails::CustomizeDetails( IDetailLayoutBuilder& Detail
 
 	const FText GameSplashDesc(LOCTEXT("GameSplashLabel", "Game Splash"));
 	FDetailWidgetRow& GameSplashWidgetRow = SplashCategoryBuilder.AddCustomRow(GameSplashDesc);
-	const FString GameSplash_TargetImagePath = GetSplashFilename(EImageScope::GameOverride, false);
-	const FString GameSplash_DefaultImagePath = GetSplashFilename(EImageScope::Engine, false);
+	const FString GameSplash_TargetImagePath = GetLinuxSplashFilename(ELinuxImageScope::GameOverride, false);
+	const FString GameSplash_DefaultImagePath = GetLinuxSplashFilename(ELinuxImageScope::Engine, false);
 
 	GameSplashWidgetRow
 	.NameContent()
@@ -261,7 +261,7 @@ void FLinuxTargetSettingsDetails::CustomizeDetails( IDetailLayoutBuilder& Detail
 		.FillWidth(1.0f)
 		.VAlign(VAlign_Center)
 		[
-			SNew(SExternalImageReference, GetIconFilename(EImageScope::Engine), GetIconFilename(EImageScope::GameOverride))
+			SNew(SExternalImageReference, GetLinuxIconFilename(ELinuxImageScope::Engine), GetLinuxIconFilename(ELinuxImageScope::GameOverride))
 			.FileDescription(GameSplashDesc)
 			.OnPreExternalImageCopy(FOnPreExternalImageCopy::CreateSP(this, &FLinuxTargetSettingsDetails::HandlePreExternalIconCopy))
 			.OnGetPickerPath(FOnGetPickerPath::CreateSP(this, &FLinuxTargetSettingsDetails::GetPickerPath))
@@ -422,7 +422,7 @@ void FLinuxTargetShaderFormatsPropertyDetails::CreateTargetShaderFormatsProperty
 
 	for (const FName& ShaderFormat : ShaderFormats)
 	{
-		FText FriendlyShaderFormatName = GetFriendlyNameFromRHIName(ShaderFormat.ToString());
+		FText FriendlyShaderFormatName = GetFriendlyNameFromLinuxRHIName(ShaderFormat.ToString());
 
 		FDetailWidgetRow& TargetedRHIWidgetRow = TargetedRHICategoryBuilder.AddCustomRow(FriendlyShaderFormatName);
 
diff --git a/Engine/Source/Editor/DetailCustomizations/Private/WindowsTargetSettingsDetails.cpp b/Engine/Source/Editor/DetailCustomizations/Private/WindowsTargetSettingsDetails.cpp
--- a/Engine/Source/Editor/DetailCustomizations/Private/WindowsTargetSettingsDetails.cpp
+++ b/Engine/Source/Editor/DetailCustomizations/Private/WindowsTargetSettingsDetails.cpp
@@ -47,7 +47,7 @@ namespace WindowsTargetSettingsDetailsConstants
 	const FText DisabledTip = LOCTEXT("GitHubSourceRequiredToolTip", "This requires GitHub source.");
 }
 
-static FText GetFriendlyNameFromRHIName(const FString& InRHIName)
+static FText GetFriendlyNameFromWindowsRHIName(const FString& InRHIName)
 {
 	FText FriendlyRHIName;
 	if (InRHIName == TEXT("PCD3D_SM5"))
@@ -96,7 +96,7 @@ TSharedRef<IDetailCustomization> FWindowsTargetSettingsDetails::MakeInstance()
 	return MakeShareable(new FWindowsTargetSettingsDetails);
 }
 
-namespace EImageScope
+namespace EWindowsImageScope
 {
 	enum Type
 	{
@@ -106,11 +106,11 @@ namespace EImageScope
 }
 
 /* Helper function used to generate filenames for splash screens */
-static FString GetSplashFilename(EImageScope::Type Scope, bool bIsEditorSplash)
+static FString GetWindowsSplashFilename(EWindowsImageScope::Type Scope, bool bIsEditorSplash)
 {
 	FString Filename;
 
-	if (Scope == EImageScope::Engine)
+	if (Scope == EWindowsImageScope::Engine)
 	{
 		Filename = FPaths::EngineContentDir();
 	}
@@ -134,11 +134,11 @@ static FString GetSplashFilename(EImageScope::Type Scope, bool bIsEditorSplash)
 }
 
 /* Helper function used to generate filenames for icons */
-static FString GetIconFilename(EImageScope::Type Scope)
+static FString GetWindowsIconFilename(EWindowsImageScope::Type Scope)
 {
 	const FString& PlatformName = FModuleManager::GetModuleChecked<ITargetPlatformModule>("WindowsTargetPlatform").GetTargetPlatforms()[0]->PlatformName();
 
-	if (Scope == EImageScope::Engine)
+	if (Scope == EWindowsImageScope::Engine)
 	{
 		FString Filename = FPaths::EngineDir() / FString(TEXT("Source/Runtime/Launch/Resources")) / PlatformName / FString("UE4.ico");
 		return FPaths::ConvertRelativePathToFull(Filename);
@@ -178,8 +178,8 @@ void FWindowsTargetSettingsDetails::CustomizeDetails( IDetailLayoutBuilder& Deta
 	IDetailCategoryBuilder& SplashCategoryBuilder = DetailBuilder.EditCategory(TEXT("Splash"));
 	FDetailWidgetRow& EditorSplashWidgetRow = SplashCategoryBuilder.AddCustomRow(EditorSplashDesc);
 
-	const FString EditorSplash_TargetImagePath = GetSplashFilename(EImageScope::GameOverride, true);
-	const FString EditorSplash_DefaultImagePath = GetSplashFilename(EImageScope::Engine, true);
+	const FString EditorSplash_TargetImagePath = GetWindowsSplashFilename(EWindowsImageScope::GameOverride, true);
+	const FString EditorSplash_DefaultImagePath = GetWindowsSplashFilename(EWindowsImageScope::Engine, true);
 
 	TArray<FString> ImageExtensions;
 	ImageExtensions.Add(TEXT("png"));
@@ -220,8 +220,8 @@ void FWindowsTargetSettingsDetails::CustomizeDetails( IDetailLayoutBuilder& Deta
 
 	const FText GameSplashDesc(LOCTEXT("GameSplashLabel", "Game Splash"));
 	FDetailWidgetRow& GameSplashWidgetRow = SplashCategoryBuilder.AddCustomRow(GameSplashDesc);
-	const FString GameSplash_TargetImagePath = GetSplashFilename(EImageScope::GameOverride, false);
-	const FString GameSplash_DefaultImagePath = GetSplashFilename(EImageScope::Engine, false);
+	const FString GameSplash_TargetImagePath = GetWindowsSplashFilename(EWindowsImageScope::GameOverride, false);
+	const FString GameSplash_DefaultImagePath = GetWindowsSplashFilename(EWindowsImageScope::Engine, false);
 
 	GameSplashWidgetRow
 	.NameContent()
@@ -279,7 +279,7 @@ void FWindowsTargetSettingsDetails::CustomizeDetails( IDetailLayoutBuilder& Deta
 		.FillWidth(1.0f)
 		.VAlign(VAlign_Center)
 		[
-			SNew(SExternalImageReference, GetIconFilename(EImageScope::Engine), GetIconFilename(EImageScope::GameOverride))
+			SNew(SExternalImageReference, GetWindowsIconFilename(EWindowsImageScope::Engine), GetWindowsIconFilename(EWindowsImageScope::GameOverride))
 			.FileDescription(GameSplashDesc)
 			.OnPreExternalImageCopy(FOnPreExternalImageCopy::CreateSP(this, &FWindowsTargetSettingsDetails::HandlePreExternalIconCopy))
 			.OnGetPickerPath(FOnGetPickerPath::CreateSP(this, &FWindowsTargetSettingsDetails::GetPickerPath))
@@ -483,7 +483,7 @@ void FTargetShaderFormatsPropertyDetails::CreateTargetShaderFormatsPropertyView(
 
 	for (const FName& ShaderFormat : ShaderFormats)
 	{
-		FText FriendlyShaderFormatName = GetFriendlyNameFromRHIName(ShaderFormat.ToString());
+		FText FriendlyShaderFormatName = GetFriendlyNameFromWindowsRHIName(ShaderFormat.ToString());
 		// Skip explicitly removed entries
 		if (!FriendlyShaderFormatName.IsEmpty())
 		{
diff --git a/Engine/Source/Editor/Persona/Private/AnimationEditorViewportClient.cpp b/Engine/Source/Editor/Persona/Private/AnimationEditorViewportClient.cpp
--- a/Engine/Source/Editor/Persona/Private/AnimationEditorViewportClient.cpp
+++ b/Engine/Source/Editor/Persona/Private/AnimationEditorViewportClient.cpp
@@ -680,15 +680,7 @@ bool FAnimationViewportClient::ShouldDisplayAdditiveScaleErrorMessage() const
 	return false;
 }
 
-static FText ConcatenateLine(const FText& InText, const FText& InNewLine)
-{
-	if(InText.IsEmpty())
-	{
-		return InNewLine;
-	}
-
-	return FText::Format(LOCTEXT("ViewportTextNewlineFormatter", "{0}\n{1}"), InText, InNewLine);
-}
+extern FText ConcatenateLine(const FText& InText, const FText& InNewLine);
 
 FText FAnimationViewportClient::GetDisplayInfo(bool bDisplayAllInfo) const
 {
diff --git a/Engine/Source/Editor/Persona/Private/SAnimCurveViewer.cpp b/Engine/Source/Editor/Persona/Private/SAnimCurveViewer.cpp
--- a/Engine/Source/Editor/Persona/Private/SAnimCurveViewer.cpp
+++ b/Engine/Source/Editor/Persona/Private/SAnimCurveViewer.cpp
@@ -31,7 +31,7 @@ static const FName ColumnID_AnimCurveWeightLabel( "Weight" );
 static const FName ColumnID_AnimCurveEditLabel( "Edit" );
 static const FName ColumnID_AnimCurveNumBoneLabel("Num Bones");
 
-const float MaxMorphWeight = 5.f;
+extern const float MaxMorphWeight = 5.f;
 //////////////////////////////////////////////////////////////////////////
 // SAnimCurveListRow
 
diff --git a/Engine/Source/Editor/Persona/Private/SAnimationEditorViewport.cpp b/Engine/Source/Editor/Persona/Private/SAnimationEditorViewport.cpp
--- a/Engine/Source/Editor/Persona/Private/SAnimationEditorViewport.cpp
+++ b/Engine/Source/Editor/Persona/Private/SAnimationEditorViewport.cpp
@@ -165,7 +165,7 @@ bool SAnimationEditorViewportTabBody::CanUseGizmos() const
 	return false;
 }
 
-static FText ConcatenateLine(const FText& InText, const FText& InNewLine)
+FText ConcatenateLine(const FText& InText, const FText& InNewLine)
 {
 	if(InText.IsEmpty())
 	{
diff --git a/Engine/Source/Editor/Persona/Private/SMorphTargetViewer.cpp b/Engine/Source/Editor/Persona/Private/SMorphTargetViewer.cpp
--- a/Engine/Source/Editor/Persona/Private/SMorphTargetViewer.cpp
+++ b/Engine/Source/Editor/Persona/Private/SMorphTargetViewer.cpp
@@ -23,7 +23,7 @@ static const FName ColumnID_MorphTargetWeightLabel( "Weight" );
 static const FName ColumnID_MorphTargetEditLabel( "Edit" );
 static const FName ColumnID_MorphTargetVertCountLabel( "NumberOfVerts" );
 
-const float MaxMorphWeight = 5.f;
+extern const float MaxMorphWeight;
 
 //////////////////////////////////////////////////////////////////////////
 // SMorphTargetListRow
diff --git a/Engine/Source/Editor/UnrealEd/Classes/Commandlets/AudioTestCommandlet.h b/Engine/Source/Editor/UnrealEd/Classes/Commandlets/AudioTestCommandlet.h
--- a/Engine/Source/Editor/UnrealEd/Classes/Commandlets/AudioTestCommandlet.h
+++ b/Engine/Source/Editor/UnrealEd/Classes/Commandlets/AudioTestCommandlet.h
@@ -15,4 +15,7 @@ class UAudioTestCommandlet : public UCommandlet
 	//~ Begin UCommandlet Interface
 	virtual int32 Main(const FString& Params) override;
 	//~ End UCommandlet Interface
+
+private:
+	void PrintUsage();
 };
diff --git a/Engine/Source/Editor/UnrealEd/Private/Commandlets/AudioTestCommandlet.cpp b/Engine/Source/Editor/UnrealEd/Private/Commandlets/AudioTestCommandlet.cpp
--- a/Engine/Source/Editor/UnrealEd/Private/Commandlets/AudioTestCommandlet.cpp
+++ b/Engine/Source/Editor/UnrealEd/Private/Commandlets/AudioTestCommandlet.cpp
@@ -174,7 +174,7 @@ static AudioTestInfo AudioTestInfoList[] =
 };
 static_assert(ARRAY_COUNT(AudioTestInfoList) == AUDIO_TESTS, "Mismatched info list and test enumeration");
 
-static void PrintUsage()
+void UAudioTestCommandlet::PrintUsage()
 {
 	UE_LOG(AudioTestCommandlet, Display, TEXT("AudioTestCommandlet Usage: {Editor}.exe UnrealEd.AudioTestCommandlet {testcategory} {test} {arglist}"));
 	UE_LOG(AudioTestCommandlet, Display, TEXT("Possible Tests:"));
diff --git a/Engine/Source/Editor/UnrealEd/Private/Commandlets/ImportAssetsCommandlet.cpp b/Engine/Source/Editor/UnrealEd/Private/Commandlets/ImportAssetsCommandlet.cpp
--- a/Engine/Source/Editor/UnrealEd/Private/Commandlets/ImportAssetsCommandlet.cpp
+++ b/Engine/Source/Editor/UnrealEd/Private/Commandlets/ImportAssetsCommandlet.cpp
@@ -20,7 +20,7 @@
 #include "HAL/PlatformFilemanager.h"
 #include "GameFramework/WorldSettings.h"
 
-static void PrintUsage()
+void UImportAssetsCommandlet::PrintUsage()
 {
 	UE_LOG(LogAutomatedImport, Display, TEXT("LogAutomatedImport Usage: LogAutomatedImport {arglist}"));
 	UE_LOG(LogAutomatedImport, Display, TEXT("Arglist:"));
diff --git a/Engine/Source/Editor/UnrealEd/Private/EditorViewportClient.cpp b/Engine/Source/Editor/UnrealEd/Private/EditorViewportClient.cpp
--- a/Engine/Source/Editor/UnrealEd/Private/EditorViewportClient.cpp
+++ b/Engine/Source/Editor/UnrealEd/Private/EditorViewportClient.cpp
@@ -106,10 +106,10 @@ void PixelInspectorRealtimeManagement(FEditorViewportClient *CurrentViewport, bo
 	}
 }
 
-namespace {
+namespace EditorViewportClient
+{
 	static const float GridSize = 2048.0f;
 	static const int32 CellSize = 16;
-	static const float AutoViewportOrbitCameraTranslate = 256.0f;
 	static const float LightRotSpeed = 0.22f;
 }
 
@@ -410,8 +410,8 @@ FEditorViewportClient::FEditorViewportClient(FEditorModeTools* InModeTools, FPre
 	DrawHelper.GridColorAxis = FColor(160, 160, 160);
 	DrawHelper.GridColorMajor = FColor(144, 144, 144);
 	DrawHelper.GridColorMinor = FColor(128, 128, 128);
-	DrawHelper.PerspectiveGridSize = GridSize;
-	DrawHelper.NumCells = DrawHelper.PerspectiveGridSize / ( CellSize * 2 );
+	DrawHelper.PerspectiveGridSize = EditorViewportClient::GridSize;
+	DrawHelper.NumCells = DrawHelper.PerspectiveGridSize / ( EditorViewportClient::CellSize * 2 );
 
 	// Most editor viewports do not want motion blur.
 	EngineShowFlags.MotionBlur = 0;
@@ -4004,8 +4004,8 @@ bool FEditorViewportClient::InputAxis(FViewport* InViewport, int32 ControllerId,
 		// Adjust the preview light direction
 		FRotator LightDir = PreviewScene->GetLightDirection();
 
-		LightDir.Yaw += -DragX * LightRotSpeed;
-		LightDir.Pitch += -DragY * LightRotSpeed;
+		LightDir.Yaw += -DragX * EditorViewportClient::LightRotSpeed;
+		LightDir.Pitch += -DragY * EditorViewportClient::LightRotSpeed;
 
 		PreviewScene->SetLightDirection( LightDir );
 		
diff --git a/Engine/Source/Editor/UnrealEd/Public/Commandlets/ImportAssetsCommandlet.h b/Engine/Source/Editor/UnrealEd/Public/Commandlets/ImportAssetsCommandlet.h
--- a/Engine/Source/Editor/UnrealEd/Public/Commandlets/ImportAssetsCommandlet.h
+++ b/Engine/Source/Editor/UnrealEd/Public/Commandlets/ImportAssetsCommandlet.h
@@ -42,6 +42,9 @@ private:
 	/** Loads a level to be used for spawning actors from import factories */
 	bool LoadLevel(const FString& LevelToLoad);
 
+	/** Prints command line arguments */
+	void PrintUsage();
+
 	/**
 	 * Clears dirty flag from all packages
 	 */
diff --git a/Engine/Source/Runtime/AnimGraphRuntime/Private/BoneControllers/AnimNode_RigidBody.cpp b/Engine/Source/Runtime/AnimGraphRuntime/Private/BoneControllers/AnimNode_RigidBody.cpp
--- a/Engine/Source/Runtime/AnimGraphRuntime/Private/BoneControllers/AnimNode_RigidBody.cpp
+++ b/Engine/Source/Runtime/AnimGraphRuntime/Private/BoneControllers/AnimNode_RigidBody.cpp
@@ -14,8 +14,6 @@
 #include "Logging/MessageLog.h"
 
 
-using namespace ImmediatePhysics;
-
 //#pragma optimize("", off)
 
 /////////////////////////////////////////////////////
@@ -532,7 +530,7 @@ void FAnimNode_RigidBody::InitPhysics(const UAnimInstance* InAnimInstance)
 	if(UsePhysicsAsset)
 	{
 		delete PhysicsSimulation;
-		PhysicsSimulation = new FSimulation();
+		PhysicsSimulation = new ImmediatePhysics::FSimulation();
 		const int32 NumBodies = UsePhysicsAsset->SkeletalBodySetups.Num();
 		Bodies.Empty(NumBodies);
 		ComponentsInSim.Reset();
@@ -544,15 +542,15 @@ void FAnimNode_RigidBody::InitPhysics(const UAnimInstance* InAnimInstance)
 		TArray<FConstraintInstance*> HighLevelConstraintInstances;
 		SkeletalMeshComp->InstantiatePhysicsAssetRefPose(*UsePhysicsAsset, SimulationSpace == ESimulationSpace::WorldSpace ? SkeletalMeshComp->GetComponentToWorld().GetScale3D() : FVector(1.f), HighLevelBodyInstances, HighLevelConstraintInstances);
 
-		TMap<FName, FActorHandle*> NamesToHandles;
-		TArray<FActorHandle*> IgnoreCollisionActors;
+		TMap<FName, ImmediatePhysics::FActorHandle*> NamesToHandles;
+		TArray<ImmediatePhysics::FActorHandle*> IgnoreCollisionActors;
 
 		TArray<FBoneIndexType> InsertionOrder;
 		ComputeBodyInsertionOrder(InsertionOrder, *SkeletalMeshComp);
 
 		const int32 NumBonesLOD0 = InsertionOrder.Num();
 
-		TArray<FActorHandle*> BodyIndexToActorHandle;
+		TArray<ImmediatePhysics::FActorHandle*> BodyIndexToActorHandle;
 		BodyIndexToActorHandle.AddZeroed(NumBonesLOD0);
 
 		TArray<FBodyInstance*> BodiesSorted;
@@ -576,7 +574,7 @@ void FAnimNode_RigidBody::InitPhysics(const UAnimInstance* InAnimInstance)
 					const bool bKinematic = BodySetup->PhysicsType != EPhysicsType::PhysType_Simulated;
 					const FTransform& LastTransform = SkeletalMeshComp->GetBoneTransform(InsertBone);	//This is out of date, but will still give our bodies an initial setup that matches the constraints (TODO: use refpose)
 
-					FActorHandle* NewBodyHandle = nullptr;
+					ImmediatePhysics::FActorHandle* NewBodyHandle = nullptr;
 					if (bSimulatedBodies && !bKinematic)
 					{
 #if WITH_PHYSX && PHYSICS_INTERFACE_PHYSX
@@ -660,8 +658,8 @@ void FAnimNode_RigidBody::InitPhysics(const UAnimInstance* InAnimInstance)
 			for(int32 ConstraintIdx = 0; ConstraintIdx < HighLevelConstraintInstances.Num(); ++ConstraintIdx)
 			{
 				FConstraintInstance* CI = HighLevelConstraintInstances[ConstraintIdx];
-				FActorHandle* Body1Handle = NamesToHandles.FindRef(CI->ConstraintBone1);
-				FActorHandle* Body2Handle = NamesToHandles.FindRef(CI->ConstraintBone2);
+				ImmediatePhysics::FActorHandle* Body1Handle = NamesToHandles.FindRef(CI->ConstraintBone1);
+				ImmediatePhysics::FActorHandle* Body2Handle = NamesToHandles.FindRef(CI->ConstraintBone2);
 
 				if(Body1Handle && Body2Handle)
 				{
@@ -699,11 +697,11 @@ void FAnimNode_RigidBody::InitPhysics(const UAnimInstance* InAnimInstance)
 		HighLevelBodyInstances.Empty();
 		BodiesSorted.Empty();
 
-		TArray<FSimulation::FIgnorePair> IgnorePairs;
+		TArray<ImmediatePhysics::FSimulation::FIgnorePair> IgnorePairs;
 		const TMap<FRigidBodyIndexPair, bool>& DisableTable = UsePhysicsAsset->CollisionDisableTable;
 		for(auto ConstItr = DisableTable.CreateConstIterator(); ConstItr; ++ConstItr)
 		{
-			FSimulation::FIgnorePair Pair;
+			ImmediatePhysics::FSimulation::FIgnorePair Pair;
 			Pair.A = BodyIndexToActorHandle[ConstItr.Key().Indices[0]];
 			Pair.B = BodyIndexToActorHandle[ConstItr.Key().Indices[1]];
 			IgnorePairs.Add(Pair);
@@ -759,20 +757,20 @@ void FAnimNode_RigidBody::UpdateWorldForces(const FTransform& ComponentToWorld,
 		for (const USkeletalMeshComponent::FPendingRadialForces& PendingRadialForce : PendingRadialForces)
 		{
 			const FVector RadialForceOrigin = WorldPositionToSpace(SimulationSpace, PendingRadialForce.Origin, ComponentToWorld, BaseBoneTM);
-			for(FActorHandle* Body : Bodies)
+			for(ImmediatePhysics::FActorHandle* Body : Bodies)
 			{
 				const float InvMass = Body->GetInverseMass();
 				if(InvMass > 0.f)
 				{
 					const float StrengthPerBody = PendingRadialForce.bIgnoreMass ? PendingRadialForce.Strength : PendingRadialForce.Strength / (TotalMass * InvMass);
-					FSimulation::EForceType ForceType;
+					ImmediatePhysics::FSimulation::EForceType ForceType;
 					if (PendingRadialForce.Type == USkeletalMeshComponent::FPendingRadialForces::AddImpulse)
 					{
-						ForceType = PendingRadialForce.bIgnoreMass ? FSimulation::EForceType::AddVelocity : FSimulation::EForceType::AddImpulse;
+						ForceType = PendingRadialForce.bIgnoreMass ? ImmediatePhysics::FSimulation::EForceType::AddVelocity : ImmediatePhysics::FSimulation::EForceType::AddImpulse;
 					}
 					else
 					{
-						ForceType = PendingRadialForce.bIgnoreMass ? FSimulation::EForceType::AddAcceleration : FSimulation::EForceType::AddForce;
+						ForceType = PendingRadialForce.bIgnoreMass ? ImmediatePhysics::FSimulation::EForceType::AddAcceleration : ImmediatePhysics::FSimulation::EForceType::AddForce;
 					}
 					
 					Body->AddRadialForce(RadialForceOrigin, StrengthPerBody, PendingRadialForce.Radius, PendingRadialForce.Falloff, ForceType);
@@ -783,7 +781,7 @@ void FAnimNode_RigidBody::UpdateWorldForces(const FTransform& ComponentToWorld,
 		if(!ExternalForce.IsNearlyZero())
 		{
 			const FVector ExternalForceInSimSpace = WorldVectorToSpaceNoScale(SimulationSpace, ExternalForce, ComponentToWorld, BaseBoneTM);
-			for (FActorHandle* Body : Bodies)
+			for (ImmediatePhysics::FActorHandle* Body : Bodies)
 			{
 				const float InvMass = Body->GetInverseMass();
 				if (InvMass > 0.f)
diff --git a/Engine/Source/Runtime/Core/Private/Stats/StatsCommand.cpp b/Engine/Source/Runtime/Core/Private/Stats/StatsCommand.cpp
--- a/Engine/Source/Runtime/Core/Private/Stats/StatsCommand.cpp
+++ b/Engine/Source/Runtime/Core/Private/Stats/StatsCommand.cpp
@@ -1619,19 +1619,19 @@ bool FGroupFilter::IsRoot(const FName& MessageName) const
 	Dump...
 -----------------------------------------------------------------------------*/
 
-static int32 MaxDepth = MAX_int32;
-static FString NameFilter;
-static FString LeafFilter;
-static FDelegateHandle DumpFrameDelegateHandle;
-static FDelegateHandle DumpCPUDelegateHandle;
+static int32 GMaxDepth = MAX_int32;
+static FString GNameFilter;
+static FString GLeafFilter;
+static FDelegateHandle GDumpFrameDelegateHandle;
+static FDelegateHandle GDumpCPUDelegateHandle;
 
 static void DumpFrame(int64 Frame)
 {
 	FStatsThreadState& Stats = FStatsThreadState::GetLocalState();
 	int64 Latest = Stats.GetLatestValidFrame();
 	check(Latest > 0);
-	DumpHistoryFrame(Stats, Latest, DumpCull, MaxDepth, *NameFilter);
-	Stats.NewFrameDelegate.Remove(DumpFrameDelegateHandle);
+	DumpHistoryFrame(Stats, Latest, DumpCull, GMaxDepth, *GNameFilter);
+	Stats.NewFrameDelegate.Remove(GDumpFrameDelegateHandle);
 	StatsMasterEnableSubtract();
 }
 
@@ -1641,7 +1641,7 @@ static void DumpCPU(int64 Frame)
 	int64 Latest = Stats.GetLatestValidFrame();
 	check(Latest > 0);
 	DumpCPUSummary(Stats, Latest);
-	Stats.NewFrameDelegate.Remove(DumpCPUDelegateHandle);
+	Stats.NewFrameDelegate.Remove(GDumpCPUDelegateHandle);
 	StatsMasterEnableSubtract();
 }
 
@@ -1695,17 +1695,17 @@ struct FDumpMultiple
 			{
 				Stack->CullByCycles( int64( DumpCull / FPlatformTime::ToMilliseconds( 1 ) ) );
 			}
-			if (!NameFilter.IsEmpty() && !LeafFilter.IsEmpty())
+			if (!GNameFilter.IsEmpty() && !GLeafFilter.IsEmpty())
 			{
 				UE_LOG(LogStats, Log, TEXT("You can't have both a root and a leaf filter (though this wouldn't be hard to add)."));
 			}
-			else if (!LeafFilter.IsEmpty())
+			else if (!GLeafFilter.IsEmpty())
 			{
-				Stack->DebugPrintLeafFilter(*LeafFilter);
+				Stack->DebugPrintLeafFilter(*GLeafFilter);
 			}
 			else
 			{
-				Stack->DebugPrint(*NameFilter, MaxDepth);
+				Stack->DebugPrint(*GNameFilter, GMaxDepth);
 			}
 			delete Stack;
 			Stack = NULL;
@@ -1860,18 +1860,18 @@ static void StatCmd(FString InCmd, bool bStatCommand, FOutputDevice* Ar /*= null
 #if STATS
 		FStatsThreadState& Stats = FStatsThreadState::GetLocalState();
 		DumpCull = 1.0f;
-		MaxDepth = MAX_int32;
-		NameFilter.Empty();
-		LeafFilter.Empty();
+		GMaxDepth = MAX_int32;
+		GNameFilter.Empty();
+		GLeafFilter.Empty();
 
-		FParse::Value(Cmd, TEXT("ROOT="), NameFilter);
-		FParse::Value(Cmd, TEXT("LEAF="), LeafFilter);
+		FParse::Value(Cmd, TEXT("ROOT="), GNameFilter);
+		FParse::Value(Cmd, TEXT("LEAF="), GLeafFilter);
 		FParse::Value(Cmd, TEXT("MS="), DumpCull);
-		FParse::Value(Cmd, TEXT("DEPTH="), MaxDepth);
+		FParse::Value(Cmd, TEXT("DEPTH="), GMaxDepth);
 		if (FParse::Command(&Cmd, TEXT("DUMPFRAME")))
 		{
 			StatsMasterEnableAdd();
-			DumpFrameDelegateHandle = Stats.NewFrameDelegate.AddStatic(&DumpFrame);
+			GDumpFrameDelegateHandle = Stats.NewFrameDelegate.AddStatic(&DumpFrame);
 		}
 		else if (FParse::Command(&Cmd, TEXT("DUMPNONFRAME")))
 		{
@@ -1883,7 +1883,7 @@ static void StatCmd(FString InCmd, bool bStatCommand, FOutputDevice* Ar /*= null
 		else if (FParse::Command(&Cmd, TEXT("DUMPCPU")))
 		{
 			StatsMasterEnableAdd();
-			DumpCPUDelegateHandle = Stats.NewFrameDelegate.AddStatic(&DumpCPU);
+			GDumpCPUDelegateHandle = Stats.NewFrameDelegate.AddStatic(&DumpCPU);
 		}
 		else if (FParse::Command(&Cmd, TEXT("STOP")))
 		{
diff --git a/Engine/Source/Runtime/Core/Private/Windows/WindowsEventLogOutputDevice.cpp b/Engine/Source/Runtime/Core/Private/Windows/WindowsEventLogOutputDevice.cpp
--- a/Engine/Source/Runtime/Core/Private/Windows/WindowsEventLogOutputDevice.cpp
+++ b/Engine/Source/Runtime/Core/Private/Windows/WindowsEventLogOutputDevice.cpp
@@ -46,7 +46,7 @@ void FWindowsEventLogOutputDevice::Serialize(const TCHAR* Buffer, ELogVerbosity:
 		{
 		case ELogVerbosity::Error:
 			{
-				ReportEvent(EventLog,
+				ReportEventW(EventLog,
 					EVENTLOG_ERROR_TYPE,
 					NULL,
 					0xC0000001L,
@@ -59,7 +59,7 @@ void FWindowsEventLogOutputDevice::Serialize(const TCHAR* Buffer, ELogVerbosity:
 			}
 		case ELogVerbosity::Warning:
 			{
-				ReportEvent(EventLog,
+				ReportEventW(EventLog,
 					EVENTLOG_WARNING_TYPE,
 					NULL,
 					0x80000002L,
diff --git a/Engine/Source/Runtime/Core/Public/Windows/PostWindowsApi.h b/Engine/Source/Runtime/Core/Public/Windows/PostWindowsApi.h
--- a/Engine/Source/Runtime/Core/Public/Windows/PostWindowsApi.h
+++ b/Engine/Source/Runtime/Core/Public/Windows/PostWindowsApi.h
@@ -33,6 +33,8 @@ THIRD_PARTY_INCLUDES_END
 #undef GetProp
 #undef SetProp
 #undef GetFileAttributes
+#undef ReportEvent
+#undef GetClassName
 
 // Undefine all the atomics. AllowWindowsPlatformAtomics/HideWindowsPlatformAtomics temporarily defining these macros.
 #if !PLATFORM_XBOXONE
diff --git a/Engine/Source/Runtime/Engine/Private/ActorReplication.cpp b/Engine/Source/Runtime/Engine/Private/ActorReplication.cpp
--- a/Engine/Source/Runtime/Engine/Private/ActorReplication.cpp
+++ b/Engine/Source/Runtime/Engine/Private/ActorReplication.cpp
@@ -26,10 +26,13 @@
 //
 // Static variables for networking.
 //
-static bool		SavedbHidden;
-static AActor*	SavedOwner;
-static bool		SavedbRepPhysics;
-static ENetRole SavedRole;
+namespace ActorReplication
+{
+	static bool		SavedbHidden;
+	static AActor*	SavedOwner;
+	static bool		SavedbRepPhysics;
+	static ENetRole SavedRole;
+}
 
 float AActor::GetNetPriority(const FVector& ViewPos, const FVector& ViewDir, AActor* Viewer, AActor* ViewTarget, UActorChannel* InChannel, float Time, bool bLowBandwidth)
 {
@@ -125,10 +128,10 @@ bool AActor::GetNetDormancy(const FVector& ViewPos, const FVector& ViewDir, AAct
 
 void AActor::PreNetReceive()
 {
-	SavedbHidden = bHidden;
-	SavedOwner = Owner;
-	SavedbRepPhysics = ReplicatedMovement.bRepPhysics;
-	SavedRole = Role;
+	ActorReplication::SavedbHidden = bHidden;
+	ActorReplication::SavedOwner = Owner;
+	ActorReplication::SavedbRepPhysics = ReplicatedMovement.bRepPhysics;
+	ActorReplication::SavedRole = Role;
 }
 
 void AActor::PostNetReceive()
@@ -137,23 +140,23 @@ void AActor::PostNetReceive()
 	{
 		// Initially we need to sync the state regardless of whether bRepPhysics has "changed" since it may not currently match IsSimulatingPhysics().
 		SyncReplicatedPhysicsSimulation();
-		SavedbRepPhysics = ReplicatedMovement.bRepPhysics;
+		ActorReplication::SavedbRepPhysics = ReplicatedMovement.bRepPhysics;
 		bNetCheckedInitialPhysicsState = true;
 	}
 
-	ExchangeB( bHidden, SavedbHidden );
-	Exchange ( Owner, SavedOwner );
+	ExchangeB( bHidden, ActorReplication::SavedbHidden );
+	Exchange ( Owner, ActorReplication::SavedOwner );
 
-	if (bHidden != SavedbHidden)
+	if (bHidden != ActorReplication::SavedbHidden)
 	{
-		SetActorHiddenInGame(SavedbHidden);
+		SetActorHiddenInGame(ActorReplication::SavedbHidden);
 	}
-	if (Owner != SavedOwner)
+	if (Owner != ActorReplication::SavedOwner)
 	{
-		SetOwner(SavedOwner);
+		SetOwner(ActorReplication::SavedOwner);
 	}
 
-	if (Role != SavedRole)
+	if (Role != ActorReplication::SavedRole)
 	{
 		PostNetReceiveRole();
 	}
@@ -184,7 +187,7 @@ void AActor::OnRep_ReplicatedMovement()
 
 	if (RootComponent)
 	{
-		if (SavedbRepPhysics != ReplicatedMovement.bRepPhysics)
+		if (ActorReplication::SavedbRepPhysics != ReplicatedMovement.bRepPhysics)
 		{
 			// Turn on/off physics sim to match server.
 			SyncReplicatedPhysicsSimulation();
diff --git a/Engine/Source/Runtime/Engine/Private/AudioDecompress.cpp b/Engine/Source/Runtime/Engine/Private/AudioDecompress.cpp
--- a/Engine/Source/Runtime/Engine/Private/AudioDecompress.cpp
+++ b/Engine/Source/Runtime/Engine/Private/AudioDecompress.cpp
@@ -6,6 +6,7 @@
 #include "Interfaces/IAudioFormat.h"
 #include "ContentStreaming.h"
 #include "HAL/LowLevelMemTracker.h"
+#include "ADPCMAudioInfo.h"
 
 IStreamedCompressedInfo::IStreamedCompressedInfo()
 	: SrcBufferData(nullptr)
@@ -453,23 +454,6 @@ namespace ADPCM
 		FMemory::Memcpy(&OutAdaptationTable, AdaptationTable, sizeof(AdaptationTable));
 	}
 
-	template <typename T>
-	static void GetAdaptationCoefficients(T(&OutAdaptationCoefficient1)[NUM_ADAPTATION_COEFF], T(&OutAdaptationCoefficient2)[NUM_ADAPTATION_COEFF])
-	{
-		// Magic values as specified by standard
-		static T AdaptationCoefficient1[] =
-		{
-			256, 512, 0, 192, 240, 460, 392
-		};
-		static T AdaptationCoefficient2[] =
-		{
-			0, -256, 0, 64, 0, -208, -232
-		};
-
-		FMemory::Memcpy(&OutAdaptationCoefficient1, AdaptationCoefficient1, sizeof(AdaptationCoefficient1));
-		FMemory::Memcpy(&OutAdaptationCoefficient2, AdaptationCoefficient2, sizeof(AdaptationCoefficient2));
-	}
-
 	struct FAdaptationContext
 	{
 	public:
diff --git a/Engine/Source/Runtime/Engine/Private/AudioThread.cpp b/Engine/Source/Runtime/Engine/Private/AudioThread.cpp
--- a/Engine/Source/Runtime/Engine/Private/AudioThread.cpp
+++ b/Engine/Source/Runtime/Engine/Private/AudioThread.cpp
@@ -39,7 +39,7 @@ TAutoConsoleVariable<int32> CVarEnableBatchProcessing(
 	ECVF_Default); 
 
 static int32 GBatchAudioAsyncBatchSize = 128;
-static FAutoConsoleVariableRef CVarBatchParticleAsyncBatchSize(
+static FAutoConsoleVariableRef CVarBatchAudioAsyncBatchSize(
 	TEXT("AudioThread.BatchAsyncBatchSize"),
 	GBatchAudioAsyncBatchSize,
 	TEXT("When AudioThread.EnableBatchProcessing = 1, controls the number of audio commands grouped together for threading.")
diff --git a/Engine/Source/Runtime/Engine/Private/Collision/CollisionDebugDrawing.cpp b/Engine/Source/Runtime/Engine/Private/Collision/CollisionDebugDrawing.cpp
--- a/Engine/Source/Runtime/Engine/Private/Collision/CollisionDebugDrawing.cpp
+++ b/Engine/Source/Runtime/Engine/Private/Collision/CollisionDebugDrawing.cpp
@@ -9,16 +9,19 @@
 
 #if ENABLE_DRAW_DEBUG
 
-static FColor TraceColor(255,255,255);
-static FColor HitColor(156, 179, 209);
-static FColor PenetratingColor(64, 64, 255);
-static FColor BlockColor(255,64,64);
-static FColor TouchColor(64,255,64);
-static float NormalLength = 20.f;
+namespace CollisionDebugDrawing
+{
+	static FColor TraceColor(255,255,255);
+	static FColor HitColor(156, 179, 209);
+	static FColor PenetratingColor(64, 64, 255);
+	static FColor BlockColor(255,64,64);
+	static FColor TouchColor(64,255,64);
+	static float NormalLength = 20.f;
+}
 
 ENGINE_API void DrawLineTraces(const UWorld* InWorld, const FVector& Start, const FVector& End, const TArray<FHitResult>& Hits, float Lifetime)
 {
-	FColor Color = (Hits.Num() > 0)? HitColor: TraceColor;
+	FColor Color = (Hits.Num() > 0)? CollisionDebugDrawing::HitColor: CollisionDebugDrawing::TraceColor;
 	DrawDebugLine(InWorld, Start, End, Color, false, Lifetime);
 
 	for(int32 HitIdx=0; HitIdx<Hits.Num(); HitIdx++)
@@ -26,11 +29,11 @@ ENGINE_API void DrawLineTraces(const UWorld* InWorld, const FVector& Start, cons
 		const FHitResult& Hit = Hits[HitIdx];
 
 		FVector NormalStart = Hit.Location;
-		FVector NormalEnd = NormalStart + (Hit.Normal * NormalLength);
-		const FColor NormalColor = Hit.bBlockingHit ? BlockColor : TouchColor;
+		FVector NormalEnd = NormalStart + (Hit.Normal * CollisionDebugDrawing::NormalLength);
+		const FColor NormalColor = Hit.bBlockingHit ? CollisionDebugDrawing::BlockColor : CollisionDebugDrawing::TouchColor;
 		DrawDebugDirectionalArrow(InWorld, NormalStart, NormalEnd, 5.f, NormalColor, false, Lifetime);
 		NormalStart = Hit.ImpactPoint;
-		NormalEnd = NormalStart + (Hit.ImpactNormal * NormalLength);
+		NormalEnd = NormalStart + (Hit.ImpactNormal * CollisionDebugDrawing::NormalLength);
 		DrawDebugDirectionalArrow(InWorld, NormalStart, NormalEnd, 5.f, NormalColor, false, Lifetime);
 
 		UE_LOG(LogCollision, Log, TEXT("  %d: T=%f C='%s' BLOCK=%d"), HitIdx, Hit.Time, *GetPathNameSafe(Hit.Component.Get()), Hit.bBlockingHit);
@@ -39,7 +42,7 @@ ENGINE_API void DrawLineTraces(const UWorld* InWorld, const FVector& Start, cons
 
 ENGINE_API void DrawSphereSweeps(const UWorld* InWorld, const FVector& Start, const FVector& End, const float Radius, const TArray<FHitResult>& Hits, float Lifetime)
 {
-	FColor Color = (Hits.Num() > 0)? HitColor: TraceColor;
+	FColor Color = (Hits.Num() > 0)? CollisionDebugDrawing::HitColor: CollisionDebugDrawing::TraceColor;
 	DrawDebugSphere(InWorld, Start, Radius, FMath::Max(Radius/4.f, 2.f), Color, false, Lifetime);
 	DrawDebugSphere(InWorld, End, Radius, FMath::Max(Radius/4.f, 2.f), Color, false, Lifetime);
 	DrawDebugLine(InWorld, Start+FVector(0, 0, Radius), End+FVector(0, 0, Radius), Color, false, Lifetime);
@@ -49,11 +52,11 @@ ENGINE_API void DrawSphereSweeps(const UWorld* InWorld, const FVector& Start, co
 		const FHitResult& Hit = Hits[HitIdx];
 
 		FVector NormalStart = Hit.Location;
-		FVector NormalEnd = NormalStart + (Hit.Normal * NormalLength);
-		const FColor NormalColor = Hit.bBlockingHit ? BlockColor : TouchColor;
+		FVector NormalEnd = NormalStart + (Hit.Normal * CollisionDebugDrawing::NormalLength);
+		const FColor NormalColor = Hit.bBlockingHit ? CollisionDebugDrawing::BlockColor : CollisionDebugDrawing::TouchColor;
 		DrawDebugDirectionalArrow(InWorld, NormalStart, NormalEnd, 5.f, NormalColor, false, Lifetime);
 		NormalStart = Hit.ImpactPoint;
-		NormalEnd = NormalStart + (Hit.ImpactNormal * NormalLength);
+		NormalEnd = NormalStart + (Hit.ImpactNormal * CollisionDebugDrawing::NormalLength);
 		DrawDebugDirectionalArrow(InWorld, NormalStart, NormalEnd, 5.f, NormalColor, false, Lifetime);
 
 		//UE_LOG(LogCollision, Log, TEXT("  %d: T=%f C='%s' BLOCK=%d"), HitIdx, Hit.Time, Hit.Component ? *Hit.Component->GetPathName() : TEXT("NOCOMP"), Hit.bBlockingHit);
@@ -65,7 +68,7 @@ ENGINE_API void DrawBoxSweeps(const UWorld* InWorld, const FVector& Start, const
 	FBox StartBox(Start-Extent, Start+Extent);
 	FBox EndBox(End-Extent, End+Extent);
 
-	FColor Color = (Hits.Num() > 0)? HitColor: TraceColor;
+	FColor Color = (Hits.Num() > 0)? CollisionDebugDrawing::HitColor: CollisionDebugDrawing::TraceColor;
 
 	DrawDebugLine(InWorld, Start, End, Color, false, Lifetime);
 	DrawDebugBox(InWorld, StartBox.GetCenter(), StartBox.GetExtent(), Rot, Color, false, Lifetime);
@@ -84,11 +87,11 @@ ENGINE_API void DrawBoxSweeps(const UWorld* InWorld, const FVector& Start, const
 		const FHitResult& Hit = Hits[HitIdx];
 
 		FVector NormalStart = Hit.Location;
-		FVector NormalEnd = NormalStart + (Hit.Normal * NormalLength);
-		const FColor NormalColor = Hit.bBlockingHit ? BlockColor : TouchColor;
+		FVector NormalEnd = NormalStart + (Hit.Normal * CollisionDebugDrawing::NormalLength);
+		const FColor NormalColor = Hit.bBlockingHit ? CollisionDebugDrawing::BlockColor : CollisionDebugDrawing::TouchColor;
 		DrawDebugDirectionalArrow(InWorld, NormalStart, NormalEnd, 5.f, NormalColor, false, Lifetime);
 		NormalStart = Hit.ImpactPoint;
-		NormalEnd = NormalStart + (Hit.ImpactNormal * NormalLength);
+		NormalEnd = NormalStart + (Hit.ImpactNormal * CollisionDebugDrawing::NormalLength);
 		DrawDebugDirectionalArrow(InWorld, NormalStart, NormalEnd, 5.f, NormalColor, false, Lifetime);
 
 		//UE_LOG(LogCollision, Log, TEXT("  %d: T=%f C='%s' BLOCK=%d"), HitIdx, Hit.Time, Hit.Component ? *Hit.Component->GetPathName() : TEXT("NOCOMP"), Hit.bBlockingHit);
@@ -97,10 +100,10 @@ ENGINE_API void DrawBoxSweeps(const UWorld* InWorld, const FVector& Start, const
 
 ENGINE_API void DrawCapsuleSweeps(const UWorld* InWorld, const FVector& Start, const FVector& End, float HalfHeight, float Radius, const FQuat& Rotation, const TArray<FHitResult>& Hits, float Lifetime)
 {
-	FColor Color = (Hits.Num() > 0)? HitColor: TraceColor;
+	FColor Color = (Hits.Num() > 0)? CollisionDebugDrawing::HitColor: CollisionDebugDrawing::TraceColor;
 	if (Hits.Num() > 0 && Hits[0].bBlockingHit)
 	{
-		Color = PenetratingColor;
+		Color = CollisionDebugDrawing::PenetratingColor;
 		//		LifeTime=0.f;
 	}
 	DrawDebugLine(InWorld, Start, End, Color, false, Lifetime);
@@ -118,19 +121,19 @@ ENGINE_API void DrawCapsuleSweeps(const UWorld* InWorld, const FVector& Start, c
 	FVector Right = Dir ^ Up; 
 	Right *= Radius;
 
-	DrawDebugLine(InWorld, Start - Right, End - Right, TraceColor, false, Lifetime);
-	DrawDebugLine(InWorld, Start + Right, End + Right, TraceColor, false, Lifetime);
+	DrawDebugLine(InWorld, Start - Right, End - Right, CollisionDebugDrawing::TraceColor, false, Lifetime);
+	DrawDebugLine(InWorld, Start + Right, End + Right, CollisionDebugDrawing::TraceColor, false, Lifetime);
 
 	for(int32 HitIdx=0; HitIdx<Hits.Num(); HitIdx++)
 	{
 		const FHitResult& Hit = Hits[HitIdx];
 
 		FVector NormalStart = Hit.Location;
-		FVector NormalEnd = NormalStart + (Hit.Normal * NormalLength);
-		const FColor NormalColor = Hit.bBlockingHit ? BlockColor : TouchColor;
+		FVector NormalEnd = NormalStart + (Hit.Normal * CollisionDebugDrawing::NormalLength);
+		const FColor NormalColor = Hit.bBlockingHit ? CollisionDebugDrawing::BlockColor : CollisionDebugDrawing::TouchColor;
 		DrawDebugDirectionalArrow(InWorld, NormalStart, NormalEnd, 5.f, NormalColor, false, Lifetime);
 		NormalStart = Hit.ImpactPoint;
-		NormalEnd = NormalStart + (Hit.ImpactNormal * NormalLength);
+		NormalEnd = NormalStart + (Hit.ImpactNormal * CollisionDebugDrawing::NormalLength);
 		DrawDebugDirectionalArrow(InWorld, NormalStart, NormalEnd, 5.f, FColor(255, 255, 0), false, Lifetime);
 
 		//UE_LOG(LogCollision, Log, TEXT("  %d: T=%f C='%s' BLOCK=%d"), HitIdx, Hit.Time, Hit.Component ? *Hit.Component->GetPathName() : TEXT("NOCOMP"), Hit.bBlockingHit);
@@ -139,7 +142,7 @@ ENGINE_API void DrawCapsuleSweeps(const UWorld* InWorld, const FVector& Start, c
 
 ENGINE_API void DrawBoxOverlap(const UWorld* InWorld, const FVector& Pos, const FVector& Extent, const FQuat& Rot, TArray<struct FOverlapResult>& Overlaps, float Lifetime)
 {
-	FColor Color = (Overlaps.Num() > 0)? HitColor: TraceColor;
+	FColor Color = (Overlaps.Num() > 0)? CollisionDebugDrawing::HitColor: CollisionDebugDrawing::TraceColor;
 	DrawDebugBox(InWorld, Pos, Extent, Rot, Color, false, Lifetime);
 
 	for(int32 OverlapIdx=0; OverlapIdx<Overlaps.Num(); OverlapIdx++)
@@ -148,7 +151,7 @@ ENGINE_API void DrawBoxOverlap(const UWorld* InWorld, const FVector& Pos, const
 
 		if (Overlap.Component.Get())
 		{
-			const FColor OverlapColor = Overlap.bBlockingHit ? BlockColor : TouchColor;
+			const FColor OverlapColor = Overlap.bBlockingHit ? CollisionDebugDrawing::BlockColor : CollisionDebugDrawing::TouchColor;
 			DrawDebugDirectionalArrow(InWorld, Pos, Overlap.Component->GetComponentLocation(), 5.f, OverlapColor, false, Lifetime);
 		}
 
@@ -158,7 +161,7 @@ ENGINE_API void DrawBoxOverlap(const UWorld* InWorld, const FVector& Pos, const
 
 ENGINE_API void DrawSphereOverlap(const UWorld* InWorld, const FVector& Pos, const float Radius, TArray<struct FOverlapResult>& Overlaps, float Lifetime)
 {
-	FColor Color = (Overlaps.Num() > 0)? HitColor: TraceColor;
+	FColor Color = (Overlaps.Num() > 0)? CollisionDebugDrawing::HitColor: CollisionDebugDrawing::TraceColor;
 	DrawDebugSphere(InWorld, Pos, Radius, FMath::Max(Radius/4.f, 2.f), Color, false, Lifetime);
 
 	for(int32 OverlapIdx=0; OverlapIdx<Overlaps.Num(); OverlapIdx++)
@@ -167,7 +170,7 @@ ENGINE_API void DrawSphereOverlap(const UWorld* InWorld, const FVector& Pos, con
 
 		if (Overlap.Component.Get())
 		{
-			const FColor OverlapColor = Overlap.bBlockingHit ? BlockColor : TouchColor;
+			const FColor OverlapColor = Overlap.bBlockingHit ? CollisionDebugDrawing::BlockColor : CollisionDebugDrawing::TouchColor;
 			DrawDebugDirectionalArrow(InWorld, Pos, Overlap.Component->GetComponentLocation(), 5.f, OverlapColor, false, Lifetime);
 		}
 
@@ -177,7 +180,7 @@ ENGINE_API void DrawSphereOverlap(const UWorld* InWorld, const FVector& Pos, con
 
 ENGINE_API void DrawCapsuleOverlap(const UWorld* InWorld,const FVector& Pos, const float HalfHeight, const float Radius, const FQuat& Rot, TArray<struct FOverlapResult>& Overlaps, float Lifetime)
 {
-	FColor Color = (Overlaps.Num() > 0)? HitColor: TraceColor;
+	FColor Color = (Overlaps.Num() > 0)? CollisionDebugDrawing::HitColor: CollisionDebugDrawing::TraceColor;
 	DrawDebugCapsule(InWorld, Pos, HalfHeight, Radius, Rot, Color, false, Lifetime);
 
 	for(int32 OverlapIdx=0; OverlapIdx<Overlaps.Num(); OverlapIdx++)
@@ -186,7 +189,7 @@ ENGINE_API void DrawCapsuleOverlap(const UWorld* InWorld,const FVector& Pos, con
 
 		if (Overlap.Component.Get())
 		{
-			const FColor OverlapColor = Overlap.bBlockingHit ? BlockColor : TouchColor;
+			const FColor OverlapColor = Overlap.bBlockingHit ? CollisionDebugDrawing::BlockColor : CollisionDebugDrawing::TouchColor;
 			DrawDebugDirectionalArrow(InWorld, Pos, Overlap.Component->GetComponentLocation(), 5.f, OverlapColor, false, Lifetime);
 		}
 
diff --git a/Engine/Source/Runtime/Engine/Private/Level.cpp b/Engine/Source/Runtime/Engine/Private/Level.cpp
--- a/Engine/Source/Runtime/Engine/Private/Level.cpp
+++ b/Engine/Source/Runtime/Engine/Private/Level.cpp
@@ -62,7 +62,7 @@ Level.cpp: Level-related functions
 DEFINE_LOG_CATEGORY(LogLevel);
 
 int32 GActorClusteringEnabled = 1;
-static FAutoConsoleVariableRef CVarUseBackgroundLevelStreaming(
+static FAutoConsoleVariableRef CVarActorClusteringEnabled(
 	TEXT("gc.ActorClusteringEnabled"),
 	GActorClusteringEnabled,
 	TEXT("Whether to allow levels to create actor clusters for GC."),
diff --git a/Engine/Source/Runtime/Engine/Private/Particles/FXSystem.cpp b/Engine/Source/Runtime/Engine/Private/Particles/FXSystem.cpp
--- a/Engine/Source/Runtime/Engine/Private/Particles/FXSystem.cpp
+++ b/Engine/Source/Runtime/Engine/Private/Particles/FXSystem.cpp
@@ -13,6 +13,7 @@
 #include "VectorField/VectorField.h"
 #include "Components/VectorFieldComponent.h"
 #include "SceneUtils.h"
+#include "Renderer/Private/SceneRendering.h" // needed for STATGROUP_CommandListMarkers
 
 /*-----------------------------------------------------------------------------
 	External FX system interface.
@@ -330,8 +331,6 @@ bool FFXSystem::UsesGlobalDistanceField() const
 	return false;
 }
 
-DECLARE_STATS_GROUP(TEXT("Command List Markers"), STATGROUP_CommandListMarkers, STATCAT_Advanced);
-
 DECLARE_CYCLE_STAT(TEXT("FXPreRender_Prepare"), STAT_CLM_FXPreRender_Prepare, STATGROUP_CommandListMarkers);
 DECLARE_CYCLE_STAT(TEXT("FXPreRender_Simulate"), STAT_CLM_FXPreRender_Simulate, STATGROUP_CommandListMarkers);
 DECLARE_CYCLE_STAT(TEXT("FXPreRender_Finalize"), STAT_CLM_FXPreRender_Finalize, STATGROUP_CommandListMarkers);
diff --git a/Engine/Source/Runtime/Engine/Private/Particles/SubUVAnimation.cpp b/Engine/Source/Runtime/Engine/Private/Particles/SubUVAnimation.cpp
--- a/Engine/Source/Runtime/Engine/Private/Particles/SubUVAnimation.cpp
+++ b/Engine/Source/Runtime/Engine/Private/Particles/SubUVAnimation.cpp
@@ -438,7 +438,7 @@ void FindOptimalPolygon(int32 TargetVertexCount, const TArray<int32>& ConvexHull
 	}
 }
 
-FIntPoint Neighbors[] = 
+FIntPoint GNeighbors[] = 
 {
 	FIntPoint(-1, -1),
 	FIntPoint(0, -1),
@@ -483,15 +483,15 @@ int32 ComputeNeighborCount(int32 X, int32 Y, int32 GlobalX, int32 GlobalY, int32
 {
 	int32 NeighborCount = 0;
 
-	for (int32 NeighborIndex = 0; NeighborIndex < ARRAY_COUNT(Neighbors); NeighborIndex++)
+	for (int32 NeighborIndex = 0; NeighborIndex < ARRAY_COUNT(GNeighbors); NeighborIndex++)
 	{
-		int32 NeighborX = X + Neighbors[NeighborIndex].X;
-		int32 NeighborY = Y + Neighbors[NeighborIndex].Y;
+		int32 NeighborX = X + GNeighbors[NeighborIndex].X;
+		int32 NeighborY = Y + GNeighbors[NeighborIndex].Y;
 
 		if (NeighborX >= 0 && NeighborX < SubImageSizeX 
 			&& NeighborY >= 0 && NeighborY < SubImageSizeY)
 		{
-			uint8 NeighborAlphaValue = ComputeOpacityValue(&MipData[((GlobalY + Neighbors[NeighborIndex].Y) * TextureSizeX + GlobalX + Neighbors[NeighborIndex].X) * 4], OpacitySourceMode);
+			uint8 NeighborAlphaValue = ComputeOpacityValue(&MipData[((GlobalY + GNeighbors[NeighborIndex].Y) * TextureSizeX + GlobalX + GNeighbors[NeighborIndex].X) * 4], OpacitySourceMode);
 
 			if (NeighborAlphaValue > AlphaThresholdByte)
 			{
diff --git a/Engine/Source/Runtime/Engine/Private/PhysicsEngine/Experimental/PhysScene_LLImmediate.cpp b/Engine/Source/Runtime/Engine/Private/PhysicsEngine/Experimental/PhysScene_LLImmediate.cpp
--- a/Engine/Source/Runtime/Engine/Private/PhysicsEngine/Experimental/PhysScene_LLImmediate.cpp
+++ b/Engine/Source/Runtime/Engine/Private/PhysicsEngine/Experimental/PhysScene_LLImmediate.cpp
@@ -4,8 +4,6 @@
 #include "Async/ParallelFor.h"
 
 
-using namespace ImmediatePhysics;
-
 FPhysScene_LLImmediate::FPhysScene_LLImmediate()
 	: Simulation(nullptr)
 	, SimulationTime(0)
@@ -30,7 +28,7 @@ void FPhysScene_LLImmediate::Init()
 		Simulation = nullptr;
 	}
 	
-	Simulation = new FSimulation();
+	Simulation = new ImmediatePhysics::FSimulation();
 }
 
 void FPhysScene_LLImmediate::Tick(float InDeltaSeconds)
diff --git a/Engine/Source/Runtime/Engine/Private/SkeletalMeshComponentPhysics.cpp b/Engine/Source/Runtime/Engine/Private/SkeletalMeshComponentPhysics.cpp
--- a/Engine/Source/Runtime/Engine/Private/SkeletalMeshComponentPhysics.cpp
+++ b/Engine/Source/Runtime/Engine/Private/SkeletalMeshComponentPhysics.cpp
@@ -522,7 +522,7 @@ int32 USkeletalMeshComponent::FindRootBodyIndex() const
 }
 
 static int32 bAllowNotForDedServerPhysicsAssets = 1;
-static FAutoConsoleVariableRef CVarAllowCachedOverlaps(
+static FAutoConsoleVariableRef CVarAllowNotForDedServerPhysicsAssets(
 	TEXT("p.AllowNotForDedServerPhysicsAssets"),
 	bAllowNotForDedServerPhysicsAssets,
 	TEXT("Allow 'Not For Dedicated Server' flag on PhysicsAssets\n")
diff --git a/Engine/Source/Runtime/Engine/Private/UnrealClient.cpp b/Engine/Source/Runtime/Engine/Private/UnrealClient.cpp
--- a/Engine/Source/Runtime/Engine/Private/UnrealClient.cpp
+++ b/Engine/Source/Runtime/Engine/Private/UnrealClient.cpp
@@ -1193,24 +1193,24 @@ void FViewport::HighResScreenshot()
 	ViewportClient->GetEngineShowFlags()->SetMotionBlur(false);
 
 	// Forcing 128-bit rendering pipeline
-	static auto CVarSceneColorFormat = IConsoleManager::Get().FindConsoleVariable(TEXT("r.SceneColorFormat"));
-	static auto CVarPostColorFormat = IConsoleManager::Get().FindConsoleVariable(TEXT("r.PostProcessingColorFormat"));
-	static auto CVarForceLOD = IConsoleManager::Get().FindConsoleVariable(TEXT("r.ForceLOD"));
-
-	check(CVarSceneColorFormat && CVarPostColorFormat);
-	const int32 OldSceneColorFormat = CVarSceneColorFormat->GetInt();
-	const int32 OldPostColorFormat = CVarPostColorFormat->GetInt();
-	const int32 OldForceLOD = CVarForceLOD ? CVarForceLOD->GetInt() : -1;
+	static IConsoleVariable* SceneColorFormatVar = IConsoleManager::Get().FindConsoleVariable(TEXT("r.SceneColorFormat"));
+	static IConsoleVariable* PostColorFormatVar = IConsoleManager::Get().FindConsoleVariable(TEXT("r.PostProcessingColorFormat"));
+	static IConsoleVariable* ForceLODVar = IConsoleManager::Get().FindConsoleVariable(TEXT("r.ForceLOD"));
+
+	check(SceneColorFormatVar && PostColorFormatVar);
+	const int32 OldSceneColorFormat = SceneColorFormatVar->GetInt();
+	const int32 OldPostColorFormat = PostColorFormatVar->GetInt();
+	const int32 OldForceLOD = ForceLODVar ? ForceLODVar->GetInt() : -1;
 	if (GetHighResScreenshotConfig().bForce128BitRendering)
 	{
-		CVarSceneColorFormat->Set(5, ECVF_SetByCode);
-		CVarPostColorFormat->Set(1, ECVF_SetByCode);
+		SceneColorFormatVar->Set(5, ECVF_SetByCode);
+		PostColorFormatVar->Set(1, ECVF_SetByCode);
 	}
 
-	if (CVarForceLOD)
+	if (ForceLODVar)
 	{
 		// Force highest LOD
-		CVarForceLOD->Set(0, ECVF_SetByCode);
+		ForceLODVar->Set(0, ECVF_SetByCode);
 	}
 
 	// Render the requested number of frames (at least once)
@@ -1239,11 +1239,11 @@ void FViewport::HighResScreenshot()
 	ViewportClient->GetEngineShowFlags()->MotionBlur = MotionBlurShowFlagBackup;
 	ViewportClient->ProcessScreenShots(DummyViewport);
 
-	CVarSceneColorFormat->Set(OldSceneColorFormat, ECVF_SetByCode);
-	CVarPostColorFormat->Set(OldPostColorFormat, ECVF_SetByCode);
-	if (CVarForceLOD)
+	SceneColorFormatVar->Set(OldSceneColorFormat, ECVF_SetByCode);
+	PostColorFormatVar->Set(OldPostColorFormat, ECVF_SetByCode);
+	if (ForceLODVar)
 	{
-		CVarForceLOD->Set(OldForceLOD, ECVF_SetByCode);
+		ForceLODVar->Set(OldForceLOD, ECVF_SetByCode);
 	}
 
 	ENQUEUE_UNIQUE_RENDER_COMMAND_TWOPARAMETER(
diff --git a/Engine/Source/Runtime/Engine/Private/UnrealEngine.cpp b/Engine/Source/Runtime/Engine/Private/UnrealEngine.cpp
--- a/Engine/Source/Runtime/Engine/Private/UnrealEngine.cpp
+++ b/Engine/Source/Runtime/Engine/Private/UnrealEngine.cpp
@@ -240,8 +240,8 @@ void FEngineModule::StartupModule()
 	extern const FString GetMapNameStatic();
 	GGetMapNameDelegate.BindStatic(&GetMapNameStatic);
 
-	static auto CVarCacheWPOPrimitives = IConsoleManager::Get().FindConsoleVariable(TEXT("r.Shadow.CacheWPOPrimitives"));
-	CVarCacheWPOPrimitives->SetOnChangedCallback(FConsoleVariableDelegate::CreateStatic(&OnChangeEngineCVarRequiringRecreateRenderState));
+	static IConsoleVariable* CacheWPOPrimitivesVar = IConsoleManager::Get().FindConsoleVariable(TEXT("r.Shadow.CacheWPOPrimitives"));
+	CacheWPOPrimitivesVar->SetOnChangedCallback(FConsoleVariableDelegate::CreateStatic(&OnChangeEngineCVarRequiringRecreateRenderState));
 
 	SuspendTextureStreamingRenderTasks = &SuspendTextureStreamingRenderTasksInternal;
 	ResumeTextureStreamingRenderTasks = &ResumeTextureStreamingRenderTasksInternal;
diff --git a/Engine/Source/Runtime/Engine/Public/Physics/PhysicsInterfacePhysX.h b/Engine/Source/Runtime/Engine/Public/Physics/PhysicsInterfacePhysX.h
--- a/Engine/Source/Runtime/Engine/Public/Physics/PhysicsInterfacePhysX.h
+++ b/Engine/Source/Runtime/Engine/Public/Physics/PhysicsInterfacePhysX.h
@@ -248,6 +248,13 @@ struct ENGINE_API FPhysicsInterface_PhysX : public FGenericPhysicsInterface
 
 	template<typename AllocatorType>
 	static int32 GetAllShapes_AssumedLocked(const FPhysicsActorHandle_PhysX& InHandle, TArray<FPhysicsShapeHandle_PhysX, AllocatorType>& OutShapes, EPhysicsSceneType InSceneType = PST_MAX);
+
+	template <>
+	static int32 GetAllShapes_AssumedLocked(const FPhysicsActorHandle_PhysX& InActorHandle, TArray<FPhysicsShapeHandle_PhysX, FDefaultAllocator>& OutShapes, EPhysicsSceneType InSceneType);
+
+	template <>
+	static int32 GetAllShapes_AssumedLocked(const FPhysicsActorHandle_PhysX& InActorHandle, PhysicsInterfaceTypes::FInlineShapeArray& OutShapes, EPhysicsSceneType InSceneType);
+
 	static void GetNumShapes(const FPhysicsActorHandle_PhysX& InHandle, int32& OutNumSyncShapes, int32& OutNumAsyncShapes);
 
 	static void ReleaseShape(const FPhysicsShapeHandle_PhysX& InShape);
diff --git a/Engine/Source/Runtime/Renderer/Private/CompositionLighting/PostProcessDeferredDecals.cpp b/Engine/Source/Runtime/Renderer/Private/CompositionLighting/PostProcessDeferredDecals.cpp
--- a/Engine/Source/Runtime/Renderer/Private/CompositionLighting/PostProcessDeferredDecals.cpp
+++ b/Engine/Source/Runtime/Renderer/Private/CompositionLighting/PostProcessDeferredDecals.cpp
@@ -568,7 +568,7 @@ static FDepthStencilStateRHIParamRef GetDecalDepthState(uint32& StencilRef, FDec
 	}
 }
 
-static FRasterizerStateRHIParamRef GetDecalRasterizerState(EDecalRasterizerState DecalRasterizerState)
+FRasterizerStateRHIParamRef GetDecalRasterizerState(EDecalRasterizerState DecalRasterizerState)
 {
 	switch (DecalRasterizerState)
 	{
diff --git a/Engine/Source/Runtime/Renderer/Private/MobileDecalRendering.cpp b/Engine/Source/Runtime/Renderer/Private/MobileDecalRendering.cpp
--- a/Engine/Source/Runtime/Renderer/Private/MobileDecalRendering.cpp
+++ b/Engine/Source/Runtime/Renderer/Private/MobileDecalRendering.cpp
@@ -16,15 +16,7 @@
 #include "DecalRenderingShared.h"
 
 
-static FRasterizerStateRHIParamRef GetDecalRasterizerState(EDecalRasterizerState DecalRasterizerState)
-{
-	switch (DecalRasterizerState)
-	{
-	case DRS_CW: return TStaticRasterizerState<FM_Solid, CM_CW>::GetRHI();
-	case DRS_CCW: return TStaticRasterizerState<FM_Solid, CM_CCW>::GetRHI();
-	default: check(0); return nullptr;
-	}
-}
+extern FRasterizerStateRHIParamRef GetDecalRasterizerState(EDecalRasterizerState DecalRasterizerState);
 
 void FMobileSceneRenderer::RenderDecals(FRHICommandListImmediate& RHICmdList)
 {
diff --git a/Engine/Source/Runtime/Renderer/Private/PostProcess/PostProcessCircleDOF.cpp b/Engine/Source/Runtime/Renderer/Private/PostProcess/PostProcessCircleDOF.cpp
--- a/Engine/Source/Runtime/Renderer/Private/PostProcess/PostProcessCircleDOF.cpp
+++ b/Engine/Source/Runtime/Renderer/Private/PostProcess/PostProcessCircleDOF.cpp
@@ -521,19 +521,7 @@ FPooledRenderTargetDesc FRCPassPostProcessCircleDOFDilate::ComputeOutputDesc(EPa
 
 /** Encapsulates the Circle DOF pixel shader. */
 
-static float TemporalHalton2( int32 Index, int32 Base )
-{
-	float Result = 0.0f;
-	float InvBase = 1.0f / Base;
-	float Fraction = InvBase;
-	while( Index > 0 )
-	{
-		Result += ( Index % Base ) * Fraction;
-		Index /= Base;
-		Fraction *= InvBase;
-	}
-	return Result;
-}
+extern float TemporalHalton2( int32 Index, int32 Base );
 
 static void TemporalRandom2(FVector2D* RESTRICT const Constant, uint32 FrameNumber)
 {
diff --git a/Engine/Source/Runtime/Renderer/Private/ScenePrivate.h b/Engine/Source/Runtime/Renderer/Private/ScenePrivate.h
--- a/Engine/Source/Runtime/Renderer/Private/ScenePrivate.h
+++ b/Engine/Source/Runtime/Renderer/Private/ScenePrivate.h
@@ -2116,6 +2116,13 @@ public:
 	template<typename LightMapPolicyType>
 	TStaticMeshDrawList<TMobileBasePassDrawingPolicy<LightMapPolicyType> >& GetMobileBasePassCSMDrawList(EBasePassDrawListType DrawType);
 
+	/** Prototype explicit specializations to prevent instantiation */
+	template<>
+	TStaticMeshDrawList<TMobileBasePassDrawingPolicy<FUniformLightMapPolicy>>& GetMobileBasePassDrawList<FUniformLightMapPolicy>(EBasePassDrawListType DrawType);
+
+	template<>
+	TStaticMeshDrawList<TMobileBasePassDrawingPolicy<FUniformLightMapPolicy>>& GetMobileBasePassCSMDrawList<FUniformLightMapPolicy>(EBasePassDrawListType DrawType);
+
 #if WITH_EDITOR
 	/** Draw list to use for selected static meshes in the editor only */
 	TStaticMeshDrawList<FEditorSelectionDrawingPolicy> EditorSelectionDrawList;
diff --git a/Engine/Source/Runtime/Slate/Private/Framework/Application/MenuStack.cpp b/Engine/Source/Runtime/Slate/Private/Framework/Application/MenuStack.cpp
--- a/Engine/Source/Runtime/Slate/Private/Framework/Application/MenuStack.cpp
+++ b/Engine/Source/Runtime/Slate/Private/Framework/Application/MenuStack.cpp
@@ -67,7 +67,7 @@ public:
 };
 
 
-namespace
+namespace MenuStackInternal
 {
 	/** Widget that wraps any menu created in FMenuStack to provide default key handling, focus tracking and helps us spot menus in widget paths */
 	DECLARE_DELEGATE_RetVal_OneParam(FReply, FOnKeyDown, FKey)
@@ -656,8 +656,8 @@ TSharedRef<SWidget> FMenuStack::WrapContent(TSharedRef<SWidget> InContent, FOpti
 {
 	// Wrap menu content in a box that limits its maximum height
 	// and in a SMenuContentWrapper that handles key presses and focus changes.
-	return SNew(SMenuContentWrapper)
-		.OnKeyDown_Static(&OnMenuKeyDown)
+	return SNew(MenuStackInternal::SMenuContentWrapper)
+		.OnKeyDown_Static(&MenuStackInternal::OnMenuKeyDown)
 		.OnMenuLostFocus_Raw(this, &FMenuStack::OnMenuContentLostFocus)
 		.OptionalMinMenuWidth(OptionalMinWidth)
 		.OptionalMinMenuHeight(OptionalMinHeight)
