Convert FPhysicsInterface_PhysX::GetAllShaped_AssumedLocked to an overloaded function rather than a template specialization, to fix CIS issues.

diff --git a/Engine/Source/Runtime/Engine/Private/PhysicsEngine/PhysicsInterfacePhysX.cpp b/Engine/Source/Runtime/Engine/Private/PhysicsEngine/PhysicsInterfacePhysX.cpp
--- a/Engine/Source/Runtime/Engine/Private/PhysicsEngine/PhysicsInterfacePhysX.cpp
+++ b/Engine/Source/Runtime/Engine/Private/PhysicsEngine/PhysicsInterfacePhysX.cpp
@@ -33,11 +33,6 @@ using namespace PhysicsInterfaceTypes;
 
 #define LOCTEXT_NAMESPACE "PhysicsInterface_PhysX"
 
-template <>
-int32 ENGINE_API FPhysicsInterface_PhysX::GetAllShapes_AssumedLocked(const FPhysicsActorHandle_PhysX& InActorHandle, TArray<FPhysicsShapeHandle_PhysX, FDefaultAllocator>& OutShapes, EPhysicsSceneType InSceneType);
-template <>
-int32 ENGINE_API FPhysicsInterface_PhysX::GetAllShapes_AssumedLocked(const FPhysicsActorHandle_PhysX& InActorHandle, PhysicsInterfaceTypes::FInlineShapeArray& OutShapes, EPhysicsSceneType InSceneType);
-
 extern TAutoConsoleVariable<float> CVarConstraintLinearDampingScale;
 extern TAutoConsoleVariable<float> CVarConstraintLinearStiffnessScale;
 extern TAutoConsoleVariable<float> CVarConstraintAngularDampingScale;
@@ -792,13 +787,11 @@ int32 GetAllShapesInternal_AssumedLocked(const FPhysicsActorHandle_PhysX& InActo
 	return NumSyncShapes;
 }
 
-template <>
 int32 FPhysicsInterface_PhysX::GetAllShapes_AssumedLocked(const FPhysicsActorHandle_PhysX& InActorHandle, TArray<FPhysicsShapeHandle_PhysX, FDefaultAllocator>& OutShapes, EPhysicsSceneType InSceneType)
 {
 	return GetAllShapesInternal_AssumedLocked(InActorHandle, OutShapes, InSceneType);
 }
 
-template <>
 int32 FPhysicsInterface_PhysX::GetAllShapes_AssumedLocked(const FPhysicsActorHandle_PhysX& InActorHandle, PhysicsInterfaceTypes::FInlineShapeArray& OutShapes, EPhysicsSceneType InSceneType)
 {
 	return GetAllShapesInternal_AssumedLocked(InActorHandle, OutShapes, InSceneType);
diff --git a/Engine/Source/Runtime/Engine/Public/Physics/PhysicsInterfacePhysX.h b/Engine/Source/Runtime/Engine/Public/Physics/PhysicsInterfacePhysX.h
--- a/Engine/Source/Runtime/Engine/Public/Physics/PhysicsInterfacePhysX.h
+++ b/Engine/Source/Runtime/Engine/Public/Physics/PhysicsInterfacePhysX.h
@@ -246,14 +246,8 @@ struct ENGINE_API FPhysicsInterface_PhysX : public FGenericPhysicsInterface
 	static void ReleaseActor(FPhysicsActorHandle_PhysX& InHandle, FPhysScene* InScene = nullptr, bool bNeverDeferRelease = false);
 	//////////////////////////////////////////////////////////////////////////
 
-	template<typename AllocatorType>
-	static int32 GetAllShapes_AssumedLocked(const FPhysicsActorHandle_PhysX& InHandle, TArray<FPhysicsShapeHandle_PhysX, AllocatorType>& OutShapes, EPhysicsSceneType InSceneType = PST_MAX);
-
-	template <>
-	static int32 GetAllShapes_AssumedLocked(const FPhysicsActorHandle_PhysX& InActorHandle, TArray<FPhysicsShapeHandle_PhysX, FDefaultAllocator>& OutShapes, EPhysicsSceneType InSceneType);
-
-	template <>
-	static int32 GetAllShapes_AssumedLocked(const FPhysicsActorHandle_PhysX& InActorHandle, PhysicsInterfaceTypes::FInlineShapeArray& OutShapes, EPhysicsSceneType InSceneType);
+	static int32 GetAllShapes_AssumedLocked(const FPhysicsActorHandle_PhysX& InActorHandle, TArray<FPhysicsShapeHandle_PhysX, FDefaultAllocator>& OutShapes, EPhysicsSceneType InSceneType = PST_MAX);
+	static int32 GetAllShapes_AssumedLocked(const FPhysicsActorHandle_PhysX& InActorHandle, PhysicsInterfaceTypes::FInlineShapeArray& OutShapes, EPhysicsSceneType InSceneType = PST_MAX);
 
 	static void GetNumShapes(const FPhysicsActorHandle_PhysX& InHandle, int32& OutNumSyncShapes, int32& OutNumAsyncShapes);
 
