Clean up fixes for stress-test unity compilation.

diff --git a/Engine/Source/Runtime/Engine/Private/PhysicsEngine/PhysicsInterfacePhysX.cpp b/Engine/Source/Runtime/Engine/Private/PhysicsEngine/PhysicsInterfacePhysX.cpp
--- a/Engine/Source/Runtime/Engine/Private/PhysicsEngine/PhysicsInterfacePhysX.cpp
+++ b/Engine/Source/Runtime/Engine/Private/PhysicsEngine/PhysicsInterfacePhysX.cpp
@@ -787,11 +787,13 @@ int32 GetAllShapesInternal_AssumedLocked(const FPhysicsActorHandle_PhysX& InActo
 	return NumSyncShapes;
 }
 
+template<>
 int32 FPhysicsInterface_PhysX::GetAllShapes_AssumedLocked(const FPhysicsActorHandle_PhysX& InActorHandle, TArray<FPhysicsShapeHandle_PhysX, FDefaultAllocator>& OutShapes, EPhysicsSceneType InSceneType)
 {
 	return GetAllShapesInternal_AssumedLocked(InActorHandle, OutShapes, InSceneType);
 }
 
+template<>
 int32 FPhysicsInterface_PhysX::GetAllShapes_AssumedLocked(const FPhysicsActorHandle_PhysX& InActorHandle, PhysicsInterfaceTypes::FInlineShapeArray& OutShapes, EPhysicsSceneType InSceneType)
 {
 	return GetAllShapesInternal_AssumedLocked(InActorHandle, OutShapes, InSceneType);
diff --git a/Engine/Source/Runtime/Engine/Public/Physics/PhysicsInterfacePhysX.h b/Engine/Source/Runtime/Engine/Public/Physics/PhysicsInterfacePhysX.h
--- a/Engine/Source/Runtime/Engine/Public/Physics/PhysicsInterfacePhysX.h
+++ b/Engine/Source/Runtime/Engine/Public/Physics/PhysicsInterfacePhysX.h
@@ -246,9 +246,8 @@ struct ENGINE_API FPhysicsInterface_PhysX : public FGenericPhysicsInterface
 	static void ReleaseActor(FPhysicsActorHandle_PhysX& InHandle, FPhysScene* InScene = nullptr, bool bNeverDeferRelease = false);
 	//////////////////////////////////////////////////////////////////////////
 
-	static int32 GetAllShapes_AssumedLocked(const FPhysicsActorHandle_PhysX& InActorHandle, TArray<FPhysicsShapeHandle_PhysX, FDefaultAllocator>& OutShapes, EPhysicsSceneType InSceneType = PST_MAX);
-	static int32 GetAllShapes_AssumedLocked(const FPhysicsActorHandle_PhysX& InActorHandle, PhysicsInterfaceTypes::FInlineShapeArray& OutShapes, EPhysicsSceneType InSceneType = PST_MAX);
-
+	template<typename AllocatorType>
+	static int32 GetAllShapes_AssumedLocked(const FPhysicsActorHandle_PhysX& InHandle, TArray<FPhysicsShapeHandle_PhysX, AllocatorType>& OutShapes, EPhysicsSceneType InSceneType = PST_MAX);
 	static void GetNumShapes(const FPhysicsActorHandle_PhysX& InHandle, int32& OutNumSyncShapes, int32& OutNumAsyncShapes);
 
 	static void ReleaseShape(const FPhysicsShapeHandle_PhysX& InShape);
@@ -454,6 +453,11 @@ struct ENGINE_API FPhysicsInterface_PhysX : public FGenericPhysicsInterface
 
 };
 
+template <>
+int32 FPhysicsInterface_PhysX::GetAllShapes_AssumedLocked(const FPhysicsActorHandle_PhysX& InActorHandle, TArray<FPhysicsShapeHandle_PhysX, FDefaultAllocator>& OutShapes, EPhysicsSceneType InSceneType);
+template <>
+int32 FPhysicsInterface_PhysX::GetAllShapes_AssumedLocked(const FPhysicsActorHandle_PhysX& InActorHandle, PhysicsInterfaceTypes::FInlineShapeArray& OutShapes, EPhysicsSceneType InSceneType);
+
 template<>
 bool FPhysicsInterface_PhysX::GeomSweepMulti(const UWorld* World, const FPhysicsGeometryCollection& InGeom, const FQuat& InGeomRot, TArray<FHitResult>& OutHits, FVector Start, FVector End, ECollisionChannel TraceChannel, const FCollisionQueryParams& Params, const FCollisionResponseParams& ResponseParams, const FCollisionObjectQueryParams& ObjectParams /*= FCollisionObjectQueryParams::DefaultObjectQueryParam*/);
 template<>
diff --git a/Engine/Source/Runtime/Renderer/Private/RendererScene.cpp b/Engine/Source/Runtime/Renderer/Private/RendererScene.cpp
--- a/Engine/Source/Runtime/Renderer/Private/RendererScene.cpp
+++ b/Engine/Source/Runtime/Renderer/Private/RendererScene.cpp
@@ -3351,6 +3351,18 @@ TStaticMeshDrawList<TBasePassDrawingPolicy<FUniformLightMapPolicy> >& FScene::Ge
 	return BasePassUniformLightMapPolicyDrawList[DrawType];
 }
 
+template<>
+TStaticMeshDrawList<TMobileBasePassDrawingPolicy<FUniformLightMapPolicy>>& FScene::GetMobileBasePassDrawList<FUniformLightMapPolicy>(EBasePassDrawListType DrawType)
+{
+	return MobileBasePassUniformLightMapPolicyDrawList[DrawType];
+}
+
+template<>
+TStaticMeshDrawList<TMobileBasePassDrawingPolicy<FUniformLightMapPolicy>>& FScene::GetMobileBasePassCSMDrawList<FUniformLightMapPolicy>(EBasePassDrawListType DrawType)
+{
+	return MobileBasePassUniformLightMapPolicyDrawListWithCSM[DrawType];
+}
+
 /*-----------------------------------------------------------------------------
 	MotionBlurInfoData
 -----------------------------------------------------------------------------*/
diff --git a/Engine/Source/Runtime/Renderer/Private/ScenePrivate.h b/Engine/Source/Runtime/Renderer/Private/ScenePrivate.h
--- a/Engine/Source/Runtime/Renderer/Private/ScenePrivate.h
+++ b/Engine/Source/Runtime/Renderer/Private/ScenePrivate.h
@@ -2644,16 +2644,10 @@ inline bool ShouldIncludeDomainInMeshPass(EMaterialDomain Domain)
 }
 
 template<>
-inline TStaticMeshDrawList<TMobileBasePassDrawingPolicy<FUniformLightMapPolicy>>& FScene::GetMobileBasePassDrawList<FUniformLightMapPolicy>(EBasePassDrawListType DrawType)
-{
-	return MobileBasePassUniformLightMapPolicyDrawList[DrawType];
-}
+TStaticMeshDrawList<TMobileBasePassDrawingPolicy<FUniformLightMapPolicy>>& FScene::GetMobileBasePassDrawList<FUniformLightMapPolicy>(EBasePassDrawListType DrawType);
 
 template<>
-inline TStaticMeshDrawList<TMobileBasePassDrawingPolicy<FUniformLightMapPolicy>>& FScene::GetMobileBasePassCSMDrawList<FUniformLightMapPolicy>(EBasePassDrawListType DrawType)
-{
-	return MobileBasePassUniformLightMapPolicyDrawListWithCSM[DrawType];
-}
+TStaticMeshDrawList<TMobileBasePassDrawingPolicy<FUniformLightMapPolicy>>& FScene::GetMobileBasePassCSMDrawList<FUniformLightMapPolicy>(EBasePassDrawListType DrawType);
 
 #include "BasePassRendering.inl"
 
