Set linkage on specialized function.

#rnx

diff --git a/Engine/Source/Runtime/Engine/Public/Physics/PhysicsInterfacePhysX.h b/Engine/Source/Runtime/Engine/Public/Physics/PhysicsInterfacePhysX.h
--- a/Engine/Source/Runtime/Engine/Public/Physics/PhysicsInterfacePhysX.h
+++ b/Engine/Source/Runtime/Engine/Public/Physics/PhysicsInterfacePhysX.h
@@ -454,9 +454,9 @@ struct ENGINE_API FPhysicsInterface_PhysX : public FGenericPhysicsInterface
 };
 
 template <>
-int32 FPhysicsInterface_PhysX::GetAllShapes_AssumedLocked(const FPhysicsActorHandle_PhysX& InActorHandle, TArray<FPhysicsShapeHandle_PhysX, FDefaultAllocator>& OutShapes, EPhysicsSceneType InSceneType);
+ENGINE_API int32 FPhysicsInterface_PhysX::GetAllShapes_AssumedLocked(const FPhysicsActorHandle_PhysX& InActorHandle, TArray<FPhysicsShapeHandle_PhysX, FDefaultAllocator>& OutShapes, EPhysicsSceneType InSceneType);
 template <>
-int32 FPhysicsInterface_PhysX::GetAllShapes_AssumedLocked(const FPhysicsActorHandle_PhysX& InActorHandle, PhysicsInterfaceTypes::FInlineShapeArray& OutShapes, EPhysicsSceneType InSceneType);
+ENGINE_API int32 FPhysicsInterface_PhysX::GetAllShapes_AssumedLocked(const FPhysicsActorHandle_PhysX& InActorHandle, PhysicsInterfaceTypes::FInlineShapeArray& OutShapes, EPhysicsSceneType InSceneType);
 
 template<>
 bool FPhysicsInterface_PhysX::GeomSweepMulti(const UWorld* World, const FPhysicsGeometryCollection& InGeom, const FQuat& InGeomRot, TArray<FHitResult>& OutHits, FVector Start, FVector End, ECollisionChannel TraceChannel, const FCollisionQueryParams& Params, const FCollisionResponseParams& ResponseParams, const FCollisionObjectQueryParams& ObjectParams /*= FCollisionObjectQueryParams::DefaultObjectQueryParam*/);
