Fix StressTestUnity build failures.

#rnx

diff --git a/Engine/Source/Runtime/Engine/Private/Animation/AnimSequence.cpp b/Engine/Source/Runtime/Engine/Private/Animation/AnimSequence.cpp
--- a/Engine/Source/Runtime/Engine/Private/Animation/AnimSequence.cpp
+++ b/Engine/Source/Runtime/Engine/Private/Animation/AnimSequence.cpp
@@ -48,7 +48,7 @@ int32 GPerformFrameStripping = 0;
 
 static const TCHAR* StripFrameCVarName = TEXT("a.StripFramesOnCompression");
 
-static FAutoConsoleVariableRef CVarMaximumAllowedHLODLevel(
+static FAutoConsoleVariableRef CVarAnimStripFramesOnCompression(
 	StripFrameCVarName,
 	GPerformFrameStripping,
 	TEXT("1 = Strip every other frame on animations that have an even number of frames. 0 = off"));
