Redo changes to fix Physics force-unity compilation.

#rnx

diff --git a/Engine/Source/Runtime/Engine/Private/PhysicsEngine/PhysicsInterfacePhysX.cpp b/Engine/Source/Runtime/Engine/Private/PhysicsEngine/PhysicsInterfacePhysX.cpp
--- a/Engine/Source/Runtime/Engine/Private/PhysicsEngine/PhysicsInterfacePhysX.cpp
+++ b/Engine/Source/Runtime/Engine/Private/PhysicsEngine/PhysicsInterfacePhysX.cpp
@@ -34,11 +34,6 @@ using namespace PhysicsInterfaceTypes;
 
 #define LOCTEXT_NAMESPACE "PhysicsInterface_PhysX"
 
-template <>
-int32 ENGINE_API FPhysicsInterface_PhysX::GetAllShapes_AssumedLocked(const FPhysicsActorHandle_PhysX& InActorHandle, TArray<FPhysicsShapeHandle_PhysX, FDefaultAllocator>& OutShapes);
-template <>
-int32 ENGINE_API FPhysicsInterface_PhysX::GetAllShapes_AssumedLocked(const FPhysicsActorHandle_PhysX& InActorHandle, PhysicsInterfaceTypes::FInlineShapeArray& OutShapes);
-
 extern TAutoConsoleVariable<float> CVarConstraintLinearDampingScale;
 extern TAutoConsoleVariable<float> CVarConstraintLinearStiffnessScale;
 extern TAutoConsoleVariable<float> CVarConstraintAngularDampingScale;
diff --git a/Engine/Source/Runtime/Engine/Public/Physics/PhysicsInterfacePhysX.h b/Engine/Source/Runtime/Engine/Public/Physics/PhysicsInterfacePhysX.h
--- a/Engine/Source/Runtime/Engine/Public/Physics/PhysicsInterfacePhysX.h
+++ b/Engine/Source/Runtime/Engine/Public/Physics/PhysicsInterfacePhysX.h
@@ -406,9 +406,9 @@ struct ENGINE_API FPhysicsInterface_PhysX : public FGenericPhysicsInterface
 };
 
 template <>
-ENGINE_API int32 FPhysicsInterface_PhysX::GetAllShapes_AssumedLocked(const FPhysicsActorHandle_PhysX& InActorHandle, TArray<FPhysicsShapeHandle_PhysX, FDefaultAllocator>& OutShapes, EPhysicsSceneType InSceneType);
+ENGINE_API int32 FPhysicsInterface_PhysX::GetAllShapes_AssumedLocked(const FPhysicsActorHandle_PhysX& InActorHandle, TArray<FPhysicsShapeHandle_PhysX, FDefaultAllocator>& OutShapes);
 template <>
-ENGINE_API int32 FPhysicsInterface_PhysX::GetAllShapes_AssumedLocked(const FPhysicsActorHandle_PhysX& InActorHandle, PhysicsInterfaceTypes::FInlineShapeArray& OutShapes, EPhysicsSceneType InSceneType);
+ENGINE_API int32 FPhysicsInterface_PhysX::GetAllShapes_AssumedLocked(const FPhysicsActorHandle_PhysX& InActorHandle, PhysicsInterfaceTypes::FInlineShapeArray& OutShapes);
 
 template<>
 bool FGenericPhysicsInterface::GeomSweepMulti(const UWorld* World, const FPhysicsGeometryCollection& InGeom, const FQuat& InGeomRot, TArray<FHitResult>& OutHits, FVector Start, FVector End, ECollisionChannel TraceChannel, const FCollisionQueryParams& Params, const FCollisionResponseParams& ResponseParams, const FCollisionObjectQueryParams& ObjectParams /*= FCollisionObjectQueryParams::DefaultObjectQueryParam*/);
