Fix StressTestUnity compile errors.

#rnx

diff --git a/Engine/Source/Runtime/Engine/Private/Collision/SceneQueryPhysXImp.h b/Engine/Source/Runtime/Engine/Private/Collision/SceneQueryPhysXImp.h
--- a/Engine/Source/Runtime/Engine/Private/Collision/SceneQueryPhysXImp.h
+++ b/Engine/Source/Runtime/Engine/Private/Collision/SceneQueryPhysXImp.h
@@ -1,5 +1,7 @@
 // Copyright 1998-2019 Epic Games, Inc. All Rights Reserved.
 
+#pragma once
+
 #if WITH_PHYSX && PHYSICS_INTERFACE_PHYSX
 
 #include "PhysXPublic.h"
diff --git a/Engine/Source/Runtime/Engine/Private/PhysicsEngine/Experimental/PhysScene_LLImmediate.cpp b/Engine/Source/Runtime/Engine/Private/PhysicsEngine/Experimental/PhysScene_LLImmediate.cpp
--- a/Engine/Source/Runtime/Engine/Private/PhysicsEngine/Experimental/PhysScene_LLImmediate.cpp
+++ b/Engine/Source/Runtime/Engine/Private/PhysicsEngine/Experimental/PhysScene_LLImmediate.cpp
@@ -3,9 +3,6 @@
 #include "Physics/Experimental/PhysScene_LLImmediate.h"
 #include "Async/ParallelFor.h"
 
-
-using namespace ImmediatePhysics;
-
 FPhysScene_LLImmediate::FPhysScene_LLImmediate()
 	: Simulation(nullptr)
 	, SimulationTime(0)
@@ -31,7 +28,7 @@ void FPhysScene_LLImmediate::Init()
 	}
 
 	CurrentFrame = 0;
-	Simulation = new FSimulation();
+	Simulation = new ImmediatePhysics::FSimulation();
 
 	// #PHYS2 move to configuration somewhere
 	Simulation->SetPositionIterationCount(16);
