Merging //UE4/Dev-Main to Dev-Build (//UE4/Dev-Build)

#rnx

diff --git a/Engine/Plugins/Experimental/PixelStreaming/Source/PixelStreaming/Private/PixelStreamingInputComponent.cpp b/Engine/Plugins/Experimental/PixelStreaming/Source/PixelStreaming/Private/PixelStreamingInputComponent.cpp
--- a/Engine/Plugins/Experimental/PixelStreaming/Source/PixelStreaming/Private/PixelStreamingInputComponent.cpp
+++ b/Engine/Plugins/Experimental/PixelStreaming/Source/PixelStreaming/Private/PixelStreamingInputComponent.cpp
@@ -98,7 +98,7 @@ void UPixelStreamingInputComponent::AddJsonStringValue(const FString& Descriptor
 	if (!Descriptor.IsEmpty())
 	{
 		TSharedRef<TJsonReader<>> JsonReader = TJsonReaderFactory<>::Create(Descriptor);
-		if (!FJsonSerializer::Deserialize(JsonReader, JsonObject) && JsonObject.IsValid())
+		if (!FJsonSerializer::Deserialize(JsonReader, JsonObject) || !JsonObject.IsValid())
 		{
 			Success = false;
 			return;
diff --git a/Engine/Source/Editor/LevelEditor/Private/LevelEditorActions.cpp b/Engine/Source/Editor/LevelEditor/Private/LevelEditorActions.cpp
--- a/Engine/Source/Editor/LevelEditor/Private/LevelEditorActions.cpp
+++ b/Engine/Source/Editor/LevelEditor/Private/LevelEditorActions.cpp
@@ -586,6 +586,16 @@ bool FLevelEditorActionCallbacks::IsFeatureLevelPreviewActive()
 	return GEditor->IsFeatureLevelPreviewEnabled() && GEditor->IsFeatureLevelPreviewActive();
 }
 
+bool FLevelEditorActionCallbacks::IsFeatureLevelPreviewDropdownEnabled()
+{
+	if (GEditor->PlayWorld || GUnrealEd->bIsSimulatingInEditor)
+	{
+		return false;
+	}
+
+	return true;
+}
+
 void FLevelEditorActionCallbacks::SetPreviewPlatform(FName MaterialQualityPlatform, ERHIFeatureLevel::Type PreviewFeatureLevel)
 {
 	GEditor->SetPreviewPlatform(MaterialQualityPlatform, PreviewFeatureLevel);
diff --git a/Engine/Source/Editor/LevelEditor/Private/LevelEditorToolBar.cpp b/Engine/Source/Editor/LevelEditor/Private/LevelEditorToolBar.cpp
--- a/Engine/Source/Editor/LevelEditor/Private/LevelEditorToolBar.cpp
+++ b/Engine/Source/Editor/LevelEditor/Private/LevelEditorToolBar.cpp
@@ -1380,7 +1380,7 @@ TSharedRef< SWidget > FLevelEditorToolBar::MakeLevelEditorToolBar( const TShared
 			);
 
 		FUIAction PreviewModeMenuCanExecute;
-		PreviewModeMenuCanExecute.CanExecuteAction = FIsActionChecked::CreateStatic(&FLevelEditorActionCallbacks::IsFeatureLevelPreviewEnabled);
+		PreviewModeMenuCanExecute.CanExecuteAction = FIsActionChecked::CreateStatic(&FLevelEditorActionCallbacks::IsFeatureLevelPreviewDropdownEnabled);
 
 		ToolbarBuilder.AddComboButton(
 			PreviewModeMenuCanExecute,
diff --git a/Engine/Source/Editor/LevelEditor/Public/LevelEditorActions.h b/Engine/Source/Editor/LevelEditor/Public/LevelEditorActions.h
--- a/Engine/Source/Editor/LevelEditor/Public/LevelEditorActions.h
+++ b/Engine/Source/Editor/LevelEditor/Public/LevelEditorActions.h
@@ -808,6 +808,7 @@ public:
 	static void ToggleFeatureLevelPreview();
 	static bool IsFeatureLevelPreviewEnabled();
 	static bool IsFeatureLevelPreviewActive();
+	static bool IsFeatureLevelPreviewDropdownEnabled();
 	static void SetPreviewPlatform(FName MaterialQualityPlatform,ERHIFeatureLevel::Type PreviewFeatureLevel);
 	static bool IsPreviewPlatformChecked(FName MaterialQualityPlatform, ERHIFeatureLevel::Type PreviewFeatureLevel);
 	static void SetFeatureLevelPreview(ERHIFeatureLevel::Type InFeatureLevel);
diff --git a/Engine/Source/Runtime/Core/Private/HAL/MallocBinned3.cpp b/Engine/Source/Runtime/Core/Private/HAL/MallocBinned3.cpp
--- a/Engine/Source/Runtime/Core/Private/HAL/MallocBinned3.cpp
+++ b/Engine/Source/Runtime/Core/Private/HAL/MallocBinned3.cpp
@@ -257,6 +257,7 @@ MS_ALIGN(PLATFORM_CACHE_LINE_SIZE) static uint8 Binned3UnusedAlignPadding[PLATFO
 uint16 FMallocBinned3::SmallBlockSizesReversedShifted[BINNED3_SMALL_POOL_COUNT] = { 0 };
 uint32 FMallocBinned3::Binned3TlsSlot = 0;
 uint32 FMallocBinned3::OsAllocationGranularity = 0;
+uint32 FMallocBinned3::MaxAlignmentForMemoryRangeReserve = 0;
 
 #if !BINNED3_USE_SEPARATE_VM_PER_POOL
 	uint8* FMallocBinned3::Binned3BaseVMPtr = nullptr;
@@ -1258,6 +1259,7 @@ FMallocBinned3::FMallocBinned3()
 
 	FGenericPlatformMemoryConstants Constants = FPlatformMemory::GetConstants();
 	OsAllocationGranularity = Constants.BinnedAllocationGranularity ? Constants.BinnedAllocationGranularity : Constants.PageSize;
+	MaxAlignmentForMemoryRangeReserve = Constants.OsAllocationGranularity;
 	NumLargePoolsPerPage = OsAllocationGranularity / sizeof(FPoolInfoLarge);
 	check(OsAllocationGranularity % sizeof(FPoolInfoLarge) == 0);  // these need to divide evenly!
 	PtrToPoolMapping.Init(OsAllocationGranularity, NumLargePoolsPerPage, Constants.AddressLimit);
@@ -1480,12 +1482,12 @@ void* FMallocBinned3::MallocExternal(SIZE_T Size, uint32 Alignment)
 	Size = Align(FMath::Max((SIZE_T)1, Size), Alignment);
 
 	check(FMath::IsPowerOfTwo(Alignment));
-	check(Alignment <= OsAllocationGranularity);
+	UE_CLOG(Alignment > MaxAlignmentForMemoryRangeReserve ,LogMemory, Fatal, TEXT("Requested alignment was too large for OS. Alignment=%dKB MaxAlignmentForMemoryRangeReserve=%dKB"), Alignment/1024, MaxAlignmentForMemoryRangeReserve/1024);
 
-	FScopeLock Lock(&Mutex);
+	FScopeLock Lock(&Mutex); 
 
 	// Use OS for non-pooled allocations.
-	UPTRINT AlignedSize = Align(Size, OsAllocationGranularity);
+	UPTRINT AlignedSize = Align(Size, FMath::Max(OsAllocationGranularity,Alignment));
 
 #if BINNED3_TIME_LARGE_BLOCKS
 	double StartTime = FPlatformTime::Seconds();
diff --git a/Engine/Source/Runtime/Core/Private/Misc/NetworkVersion.cpp b/Engine/Source/Runtime/Core/Private/Misc/NetworkVersion.cpp
--- a/Engine/Source/Runtime/Core/Private/Misc/NetworkVersion.cpp
+++ b/Engine/Source/Runtime/Core/Private/Misc/NetworkVersion.cpp
@@ -18,7 +18,7 @@ FString FNetworkVersion::ProjectVersion = TEXT("1.0.0");
 bool FNetworkVersion::bHasCachedNetworkChecksum			= false;
 uint32 FNetworkVersion::CachedNetworkChecksum			= 0;
 
-uint32 FNetworkVersion::EngineNetworkProtocolVersion	= HISTORY_CHANNEL_CLOSE_REASON;
+uint32 FNetworkVersion::EngineNetworkProtocolVersion	= HISTORY_ACKS_INCLUDED_IN_HEADER;
 uint32 FNetworkVersion::GameNetworkProtocolVersion		= 0;
 
 uint32 FNetworkVersion::EngineCompatibleNetworkProtocolVersion		= HISTORY_REPLAY_BACKWARDS_COMPAT;
diff --git a/Engine/Source/Runtime/Core/Public/HAL/MallocBinned3.h b/Engine/Source/Runtime/Core/Public/HAL/MallocBinned3.h
--- a/Engine/Source/Runtime/Core/Public/HAL/MallocBinned3.h
+++ b/Engine/Source/Runtime/Core/Public/HAL/MallocBinned3.h
@@ -700,6 +700,7 @@ public:
 	static FMallocBinned3* MallocBinned3;
 	static uint32 Binned3TlsSlot;
 	static uint32 OsAllocationGranularity;
+	static uint32 MaxAlignmentForMemoryRangeReserve;
 #if !BINNED3_USE_SEPARATE_VM_PER_POOL
 	static uint8* Binned3BaseVMPtr;
 #else
diff --git a/Engine/Source/Runtime/Core/Public/Misc/NetworkVersion.h b/Engine/Source/Runtime/Core/Public/Misc/NetworkVersion.h
--- a/Engine/Source/Runtime/Core/Public/Misc/NetworkVersion.h
+++ b/Engine/Source/Runtime/Core/Public/Misc/NetworkVersion.h
@@ -39,6 +39,7 @@ enum EEngineNetworkVersionHistory
 	HISTORY_NEW_ACTOR_OVERRIDE_LEVEL = 5,			// Bump version since a level reference was added to the new actor information
 	HISTORY_CHANNEL_NAMES = 6,						// Bump version since channel type is now an fname
 	HISTORY_CHANNEL_CLOSE_REASON = 7,				// Bump version to serialize a channel close reason in bunches instead of bDormant
+	HISTORY_ACKS_INCLUDED_IN_HEADER = 8,			// Bump version since acks are now sent as part of the header
 };
 
 struct CORE_API FNetworkVersion
diff --git a/Engine/Source/Runtime/Engine/Classes/Engine/ChildConnection.h b/Engine/Source/Runtime/Engine/Classes/Engine/ChildConnection.h
--- a/Engine/Source/Runtime/Engine/Classes/Engine/ChildConnection.h
+++ b/Engine/Source/Runtime/Engine/Classes/Engine/ChildConnection.h
@@ -52,10 +52,6 @@ public:
 		Parent->AssertValid();
 	}
 
-	virtual void SendAck(int32 PacketId, bool FirstTime=1) override
-	{
-	}
-
 	virtual void FlushNet(bool bIgnoreSimulation = false) override
 	{
 		Parent->FlushNet(bIgnoreSimulation);
diff --git a/Engine/Source/Runtime/Engine/Classes/Engine/NetConnection.h b/Engine/Source/Runtime/Engine/Classes/Engine/NetConnection.h
--- a/Engine/Source/Runtime/Engine/Classes/Engine/NetConnection.h
+++ b/Engine/Source/Runtime/Engine/Classes/Engine/NetConnection.h
@@ -14,6 +14,7 @@
 #include "GameFramework/OnlineReplStructs.h"
 #include "Engine/NetDriver.h"
 #include "Net/DataBunch.h"
+#include "Net/NetPacketNotify.h"
 #include "Engine/Player.h"
 #include "Engine/Channel.h"
 #include "ProfilingDebugging/Histogram.h"
@@ -38,10 +39,11 @@ typedef TMap<TWeakObjectPtr<AActor>, UActorChannel*, FDefaultSetAllocator, TWeak
 	Types.
 -----------------------------------------------------------------------------*/
 enum { RELIABLE_BUFFER = 256 }; // Power of 2 >= 1.
-enum { MAX_PACKETID = 16384 };  // Power of 2 >= 1, covering guaranteed loss/misorder time.
+enum { MAX_PACKETID = FNetPacketNotify::SequenceNumberT::SeqNumberCount };  // Power of 2 >= 1, covering guaranteed loss/misorder time.
 enum { MAX_CHSEQUENCE = 1024 }; // Power of 2 >RELIABLE_BUFFER, covering loss/misorder time.
 enum { MAX_BUNCH_HEADER_BITS = 256 };
-enum { MAX_PACKET_HEADER_BITS = 15 }; // = FMath::CeilLogTwo(MAX_PACKETID) + 1 (IsAck)
+enum { UE_PACKET_MAX_SEQUENCE_HISTORY_BITS = FNetPacketNotify::MaxSequenceHistoryLength };
+enum { MAX_PACKET_HEADER_BITS = 49 + UE_PACKET_MAX_SEQUENCE_HISTORY_BITS }; // = [(Header)32 + (SeqHistory)32-MaxSequenceHistoryLength|(bool)1|?(ServerTime)8|(ReceviedRateByte)8]
 enum { MAX_PACKET_TRAILER_BITS = 1 };
 
 // 
@@ -389,7 +391,7 @@ public:
 	FBitWriter		SendBuffer;						// Queued up bits waiting to send
 	double			OutLagTime[256];				// For lag measuring.
 	int32			OutLagPacketId[256];			// For lag measuring.
-	int32			OutBytesPerSecondHistory[256];	// For saturation measuring.
+	uint8			OutBytesPerSecondHistory[256];	// For saturation measuring.
 	float			RemoteSaturation;
 	int32			InPacketId;						// Full incoming packet index.
 	int32			OutPacketId;					// Most recently sent packet.
@@ -405,8 +407,6 @@ public:
 	TArray<int32>		OutReliable;
 	TArray<int32>		InReliable;
 	TArray<int32>		PendingOutRec;	// Outgoing reliable unacked data from previous (now destroyed) channel in this slot.  This contains the first chsequence not acked
-	TArray<int32> QueuedAcks, ResendAcks;
-
 	int32				InitOutReliable;
 	int32				InitInReliable;
 
@@ -693,6 +693,7 @@ public:
 	ENGINE_API virtual void AssertValid();
 
 	/** Send an acknowledgment. */
+	UE_DEPRECATED(4.22, "This method will be removed")
 	ENGINE_API virtual void SendAck( int32 PacketId, bool FirstTime=1);
 
 	/**
@@ -857,6 +858,7 @@ public:
 	// Functions.
 
 	/** Resend any pending acks. */
+	UE_DEPRECATED(4.22, "This method will be removed.")
 	void PurgeAcks();
 
 	/** Send package map to the remote. */
@@ -1054,6 +1056,18 @@ private:
 	/** Returns true if an outgoing packet should be dropped due to packet simulation settings, including loss burst simulation. */
 	bool ShouldDropOutgoingPacketForLossSimulation(int64 NumBits) const;
 
+	/** Write packetHeader */
+	void WritePacketHeader(FBitWriter& Writer);
+
+	/** Write extended packet header information (ServerFrameTime, SaturationData) */
+	void WritePacketInfo(FBitWriter& Writer) const;
+	
+	/** Read extended packet header information (ServerFrameTime, SaturationData) */
+	bool ReadPacketInfo(FBitReader& Reader);
+
+	/** Packet was acknowledged as delivered */
+	void ReceivedAck(int32 AckPacketId);
+
 	/**
 	 * on the server, the world the client has told us it has loaded
 	 * used to make sure the client has traveled correctly, prevent replicating actors before level transitions are done, etc
@@ -1072,6 +1086,15 @@ private:
 
 	/** Bookkeeping for connections using internal ack to avoid acking all open channels */
 	TArray<int32> ChannelsWaitingForInternalAck;
+
+	/** Sequence data used to implement reliability */
+	FNetPacketNotify PacketNotify;
+
+	/** Full PacketId  of last sent packet that we have received notification for (i.e. we know if it was delivered or not). Related to OutAckPacketId which is tha last successfully delivered PacketId */
+	int32 LastNotifiedPacketId;
+
+	/** Keep old behavior where we send a packet with only acks even if we have no other outgoing data if we got incoming data */
+	uint32 HasDirtyAcks;
 };
 
 
diff --git a/Engine/Source/Runtime/Engine/Private/Net/NetPacketNotify.cpp b/Engine/Source/Runtime/Engine/Private/Net/NetPacketNotify.cpp
new file mode 100644
--- /dev/null
+++ b/Engine/Source/Runtime/Engine/Private/Net/NetPacketNotify.cpp
@@ -0,0 +1,139 @@
+// Copyright 1998-2019 Epic Games, Inc. All Rights Reserved.
+
+#include "Net/NetPacketNotify.h"
+#include "Serialization/BitWriter.h"
+#include "Serialization/BitReader.h"
+
+SIZE_T FNetPacketNotify::GetCurrentSequenceHistoryLength() const
+{
+	if (InAckSeq >= InAckSeqAck)
+	{
+		return (SIZE_T)SequenceNumberT::Diff(InAckSeq, InAckSeqAck);
+	}
+	else
+	{
+		// Worst case send full history
+		return SequenceHistoryT::Size;
+	}
+}
+
+void FNetPacketNotify::Init(SequenceNumberT InitialInSeq, SequenceNumberT InitialOutSeq)
+{
+	InSeqHistory.Reset();
+	InSeq = InitialInSeq;
+	InAckSeq = InitialInSeq;
+	InAckSeqAck = InitialInSeq;
+	OutSeq = InitialOutSeq;
+	OutAckSeq = SequenceNumberT(InitialOutSeq.Get() - 1);
+}
+
+void FNetPacketNotify::AckSeq(SequenceNumberT AckedSeq, bool IsAck)
+{
+	check( AckedSeq == InSeq);
+
+	while (AckedSeq > InAckSeq)
+	{
+		++InAckSeq;
+
+		const bool bReportAcked = InAckSeq == AckedSeq ? IsAck : false;
+
+		UE_LOG_PACKET_NOTIFY(TEXT("FNetPacketNotify::AckSeq - AckedSeq: %u, IsAck %u"), InAckSeq.Get(), bReportAcked ? 1u : 0u);
+
+		InSeqHistory.AddDeliveryStatus(bReportAcked);		
+	}
+}
+
+namespace 
+{
+	struct FPackedHeader
+	{
+		using SequenceNumberT = FNetPacketNotify::SequenceNumberT;
+
+		static_assert(FNetPacketNotify::SequenceNumberBits <= 14, "SequenceNumbers must be smaller than 14 bits to fit history word count");
+
+		enum { HistoryWordCountBits = 4 };
+		enum { SeqMask				= (1 << FNetPacketNotify::SequenceNumberBits) - 1 };
+		enum { HistoryWordCountMask	= (1 << HistoryWordCountBits) - 1 };
+		enum { AckSeqShift			= HistoryWordCountBits };
+		enum { SeqShift				= AckSeqShift + FNetPacketNotify::SequenceNumberBits };
+		
+		static uint32 Pack(SequenceNumberT Seq, SequenceNumberT AckedSeq, SIZE_T HistoryWordCount)
+		{
+			uint32 Packed = 0u;
+
+			Packed |= Seq.Get() << SeqShift;
+			Packed |= AckedSeq.Get() << AckSeqShift;
+			Packed |= HistoryWordCount & HistoryWordCountMask;
+
+			return Packed;
+		}
+
+		static SequenceNumberT GetSeq(uint32 Packed) { return SequenceNumberT(Packed >> SeqShift & SeqMask); }
+		static SequenceNumberT GetAckedSeq(uint32 Packed) { return SequenceNumberT(Packed >> AckSeqShift & SeqMask); }
+		static SIZE_T GetHistoryWordCount(uint32 Packed) { return (Packed & HistoryWordCountMask); }
+	};
+}
+
+// These methods must always write and read the exact same number of bits, that is the reason for not using WriteInt/WrittedWrappedInt
+void FNetPacketNotify::WriteHeader(FBitWriter& Writer, bool bRefresh)
+{
+	// we always write at least 1 word
+	SIZE_T CurrentHistoryWordCount = FMath::Clamp<SIZE_T>((GetCurrentSequenceHistoryLength() + SequenceHistoryT::BitsPerWord - 1u) / SequenceHistoryT::BitsPerWord, 1u, SequenceHistoryT::WordCount);
+
+	// We can only do a refresh if we do not need more space for the history
+	if (bRefresh && (CurrentHistoryWordCount > WrittenHistoryWordCount))
+	{
+		return;
+	}
+
+	// How many words of ack data should we write? If this is a refresh we must write the same size as the original header
+	WrittenHistoryWordCount = bRefresh ? WrittenHistoryWordCount : CurrentHistoryWordCount;
+	// This is the last InAck we have acknowledged at this time
+	WrittenInAckSeq = InAckSeq;
+
+	SequenceNumberT::SequenceT Seq = OutSeq.Get();
+	SequenceNumberT::SequenceT AckedSeq = InAckSeq.Get();
+
+	// Pack data into a uint
+	uint32 PackedHeader = FPackedHeader::Pack(Seq, AckedSeq, WrittenHistoryWordCount - 1);
+
+	// Write packed header
+	Writer << PackedHeader;
+
+	// Write ack history
+	InSeqHistory.Write(Writer, WrittenHistoryWordCount);
+
+	UE_LOG_PACKET_NOTIFY(TEXT("FNetPacketNotify::WriteHeader - Seq %u, AckedSeq %u bReFresh %u HistorySizeInWords %u"), Seq, AckedSeq, bRefresh ? 1u : 0u, WrittenHistoryWordCount);
+}
+
+bool FNetPacketNotify::ReadHeader(FNotificationHeader& Data, FBitReader& Reader) const
+{
+	// Read packed header
+	uint32 PackedHeader = 0;	
+	Reader << PackedHeader;
+
+	// unpack
+	Data.Seq = FPackedHeader::GetSeq(PackedHeader);
+	Data.AckedSeq = FPackedHeader::GetAckedSeq(PackedHeader);
+	Data.HistoryWordCount = FPackedHeader::GetHistoryWordCount(PackedHeader) + 1;
+
+	// Read ack history
+	Data.History.Read(Reader, Data.HistoryWordCount);
+
+	UE_LOG_PACKET_NOTIFY(TEXT("FNetPacketNotify::ReadHeader - Seq %u, AckedSeq %u HistorySizeInWords %u"), Data.Seq.Get(), Data.AckedSeq.Get(), Data.HistoryWordCount);
+
+	return Reader.IsError() == false;
+}
+
+FNetPacketNotify::SequenceNumberT FNetPacketNotify::CommitAndIncrementOutSeq()
+{
+	// we have not written a header...this is a fail.
+	check(WrittenHistoryWordCount != 0);
+
+	// Add entry to the ack-record so that we can update the InAckSeqAck when we received the ack for this OutSeq.
+	AckRecord.Enqueue( {OutSeq, WrittenInAckSeq} );
+	WrittenHistoryWordCount = 0u;
+	
+	return ++OutSeq;
+}
+
diff --git a/Engine/Source/Runtime/Engine/Private/Net/Tests/NetPacketNotifyTests.cpp b/Engine/Source/Runtime/Engine/Private/Net/Tests/NetPacketNotifyTests.cpp
new file mode 100644
--- /dev/null
+++ b/Engine/Source/Runtime/Engine/Private/Net/Tests/NetPacketNotifyTests.cpp
@@ -0,0 +1,303 @@
+// Copyright 1998-2019 Epic Games, Inc. All Rights Reserved.
+
+#include "Misc/AutomationTest.h"
+#include "Net/NetPacketNotify.h"
+
+#if WITH_DEV_AUTOMATION_TESTS
+
+IMPLEMENT_SIMPLE_AUTOMATION_TEST(FNetPacketNotifyTest, "Network.PacketNotifyTest", EAutomationTestFlags::ApplicationContextMask | EAutomationTestFlags::EngineFilter)
+
+struct FNetPacketNotifyTestUtil
+{
+	FNetPacketNotify DefaultNotify;
+	FNetPacketNotifyTestUtil()
+	{
+		DefaultNotify.Init(FNetPacketNotify::SequenceNumberT(-1),  FNetPacketNotify::SequenceNumberT(0));
+	}
+
+	// Helper to fill in SequenceHistory with expected result
+	template <typename T>
+	static void InitHistory(FNetPacketNotify::SequenceHistoryT& History, const T& DataToSet)
+	{
+		const SIZE_T Count = sizeof(T) / sizeof(DataToSet[0]);
+		static_assert(Count < FNetPacketNotify::SequenceHistoryT::WordCount, "DataToSet must be smaller than HistoryBuffer");
+
+		for (SIZE_T It=0; It < Count; ++It)
+		{
+			History.Data()[It] = DataToSet[It];
+		}
+	}
+
+	// Pretend to receive and acknowledge incoming packet to generate ackdata
+	static int32 PretendReceiveSeq(FNetPacketNotify& PacketNotify, FNetPacketNotify::SequenceNumberT Seq, bool Ack = true)
+	{
+		FNetPacketNotify::FNotificationHeader Data;
+		Data.Seq = Seq;
+		Data.AckedSeq = PacketNotify.GetOutAckSeq();
+		Data.History = FNetPacketNotify::SequenceHistoryT(0);
+		Data.HistoryWordCount = 1;
+		
+		FNetPacketNotify::SequenceNumberT::DifferenceT SeqDelta = PacketNotify.Update(Data, [](FNetPacketNotify::SequenceNumberT AckedSequence, bool delivered) {});
+		if (SeqDelta > 0)
+		{
+			if (Ack)
+				PacketNotify.AckSeq(Seq);
+		}
+
+		return SeqDelta;
+	}
+
+	// Pretend to send packet
+	static void PretendSendSeq(FNetPacketNotify& PacketNotify, FNetPacketNotify::SequenceNumberT LastAckSeq)
+	{
+		// set last InAcqSeq that we know that the remote end knows that we know (AckAck)
+		PacketNotify.WrittenHistoryWordCount = 1;
+		PacketNotify.WrittenInAckSeq = LastAckSeq;
+
+		// Store data
+		PacketNotify.CommitAndIncrementOutSeq();
+	}
+
+	// pretend to ack array of sequence numbers
+	template<typename T>
+	static void PretendAckSequenceNumbers(FNetPacketNotify& PacketNotify, const T& InSequenceNumbers)
+	{
+		SIZE_T SequenceNumberCount = sizeof(InSequenceNumbers) / sizeof(InSequenceNumbers[0]);
+
+		for (SIZE_T I=0; I<SequenceNumberCount; ++I)
+		{
+			FNetPacketNotifyTestUtil::PretendReceiveSeq(PacketNotify, InSequenceNumbers[I]);
+		}
+	}
+	
+	// Pretend that we received a packet
+	template<typename T>
+	static SIZE_T PretendReceivedPacket(FNetPacketNotify& PacketNotify, const FNetPacketNotify::FNotificationHeader Data, T& OutSequenceNumbers)
+	{
+		SIZE_T NotificationCount = 0;
+
+		auto HandleAck = [&OutSequenceNumbers, &NotificationCount](FNetPacketNotify::SequenceNumberT Seq, bool delivered)
+		{
+			const SIZE_T MaxSequenceNumberCount = sizeof(OutSequenceNumbers) / sizeof(OutSequenceNumbers[0]);
+
+			if (delivered)
+			{
+				if (NotificationCount < MaxSequenceNumberCount)
+				{
+					OutSequenceNumbers[NotificationCount] = Seq;
+				}
+				++NotificationCount;		
+			}
+		};
+		return PacketNotify.Update(Data, HandleAck);
+	}
+
+	// Test to fake sending and receiving an array of sequence numbers and test if we get the expected notifications back
+	template<typename T>
+	static bool TestNotificationSequence(const T& InSequenceNumbers, FNetPacketNotify::SequenceNumberT FirstSequence = 0)
+	{
+		T NotifiedSequenceNumbers = { 0 };
+
+		// Sender, which we will also be the receiver of the acks
+		FNetPacketNotify Sender;
+		Sender.Init(FNetPacketNotify::SequenceNumberT(FirstSequence.Get() - 1), FirstSequence);
+	
+		// pretend that we have sent the InSequenceNumbers
+		SIZE_T Count = sizeof(T) / sizeof(InSequenceNumbers[0]);
+		for (SIZE_T It=0; It < Count; ++It)
+		{			
+			do
+			{
+				FNetPacketNotifyTestUtil::PretendSendSeq(Sender, 0);
+			}
+			while (InSequenceNumbers[It] >= Sender.GetOutSeq());
+		}
+
+		// Receiver which we fake have received the packets sent from sender
+		FNetPacketNotify Receiver;
+		Receiver.Init(FNetPacketNotify::SequenceNumberT(FirstSequence.Get() - 1), FirstSequence);
+		PretendAckSequenceNumbers(Receiver, InSequenceNumbers);
+
+		// Fake Header with acks sent from receiver back to sender
+		FNetPacketNotify::FNotificationHeader Data;
+		Data.Seq = Receiver.GetOutSeq();
+		Data.AckedSeq = Receiver.GetInAckSeq();
+		Data.HistoryWordCount = FNetPacketNotify::SequenceHistoryT::WordCount;
+		Data.History = Receiver.GetInSeqHistory();
+
+		// Process the received ack information
+		PretendReceivedPacket(Sender, Data, NotifiedSequenceNumbers);
+
+		// Check that it matches the expected result
+		return FPlatformMemory::Memcmp(&InSequenceNumbers[0], &NotifiedSequenceNumbers[0], sizeof(InSequenceNumbers)) == 0;
+	}
+};
+
+bool FNetPacketNotifyTest::RunTest(const FString& Parameters)
+{
+	FNetPacketNotifyTestUtil Util;
+	// Test fill
+	{
+		FNetPacketNotify::SequenceNumberT ExpectedInSeq(31);
+		FNetPacketNotify::SequenceHistoryT ExpectedInSeqHistory(0xffffffffu, 1);
+
+		FNetPacketNotify Acks = Util.DefaultNotify;
+				
+		for (int32 I=0; I<32; ++I)
+		{
+			FNetPacketNotifyTestUtil::PretendReceiveSeq(Acks, I, true);
+		}
+
+		TestEqual(TEXT("Test fill - InSeq"), Acks.GetInSeq(), ExpectedInSeq);
+		TestEqual(TEXT("Test fill - History"), Acks.GetInSeqHistory(), ExpectedInSeqHistory);
+	}
+
+	// Test drop every other
+	{
+		FNetPacketNotify::SequenceNumberT ExpectedInSeq(30);
+		FNetPacketNotify::SequenceHistoryT ExpectedInSeqHistory(0x55555555u, 1);
+
+		FNetPacketNotify Acks = Util.DefaultNotify;
+				
+		for (int32 I=0; I<16; ++I)
+		{
+			FNetPacketNotifyTestUtil::PretendReceiveSeq(Acks, I*2);
+		}
+		
+		TestEqual(TEXT("Test drop every other - InSeq"), Acks.GetInSeq(), ExpectedInSeq);
+		TestEqual(TEXT("Test drop every other - History"), Acks.GetInSeqHistory(), ExpectedInSeqHistory);
+	}
+
+	// Test burst drop
+	{
+		FNetPacketNotify::SequenceNumberT ExpectedInSeq(128);
+		FNetPacketNotify::SequenceHistoryT ExpectedInSeqHistory;
+		uint32 ExpectedArray[] = {0x1, 0, 0, 0x20000000 };
+		FNetPacketNotifyTestUtil::InitHistory(ExpectedInSeqHistory, ExpectedArray );
+
+		FNetPacketNotify Acks = Util.DefaultNotify;
+
+		// Drop early
+		FNetPacketNotifyTestUtil::PretendReceiveSeq(Acks, 3);
+
+		// Large gap until next seq
+		FNetPacketNotifyTestUtil::PretendReceiveSeq(Acks, 128);
+
+		TestEqual(TEXT("Test burst drop - InSeq"), Acks.GetInSeq(), ExpectedInSeq);
+		TestEqual(TEXT("Test burst drop - History"), Acks.GetInSeqHistory(), ExpectedInSeqHistory);
+	}
+
+	// Test window overflow
+	{
+		FNetPacketNotify Acks = Util.DefaultNotify;
+
+		const FNetPacketNotify::SequenceNumberT ExpectedInSeq(0);
+		const FNetPacketNotify::SequenceNumberT MaxWindowSeq(FNetPacketNotify::SequenceNumberT::SeqNumberHalf);
+
+		FNetPacketNotifyTestUtil::PretendReceiveSeq(Acks, ExpectedInSeq, true);
+		TestEqual(TEXT("Test window overflow - Expect InSeq 0"), Acks.GetInSeq(), ExpectedInSeq);
+		TestEqual(TEXT("Test window overflow - Expect InAckSeq 0"), Acks.GetInAckSeq(), ExpectedInSeq);
+
+		FNetPacketNotifyTestUtil::PretendReceiveSeq(Acks, MaxWindowSeq, true);
+		TestEqual(TEXT("Test window overflow - Expect Seq reject"), Acks.GetInSeq(), ExpectedInSeq);
+		TestEqual(TEXT("Test window overflow - Expect Seq reject"), Acks.GetInAckSeq(), ExpectedInSeq);
+
+		const FNetPacketNotify::SequenceNumberT NextExpectedInSeq(1);
+		FNetPacketNotifyTestUtil::PretendReceiveSeq(Acks, NextExpectedInSeq, true);
+		TestEqual(TEXT("Test window overflow - Expect InSeq 1"), Acks.GetInSeq(), NextExpectedInSeq);
+		TestEqual(TEXT("Test window overflow - Expect InAckSeq 1"), Acks.GetInAckSeq(), NextExpectedInSeq);
+
+		FNetPacketNotifyTestUtil::PretendReceiveSeq(Acks, MaxWindowSeq, true);
+		TestEqual(TEXT("Test window overflow - Expect InSeq MaxWindowSeq"), Acks.GetInSeq(), MaxWindowSeq);
+		TestEqual(TEXT("Test window overflow - Expect InSeq MaxWindowSeq"), Acks.GetInAckSeq(), MaxWindowSeq);
+	}
+
+	// create history
+	{
+		const FNetPacketNotify::SequenceNumberT ExpectedInSeq(18);
+		const FNetPacketNotify::SequenceHistoryT ExpectedInSeqHistory(0x8853u);
+
+		const FNetPacketNotify::SequenceNumberT AckdPacketIds[] = {3, 7, 12, 14, 17, 18};
+		const SIZE_T ExpectedCount = sizeof(AckdPacketIds)/sizeof((AckdPacketIds)[0]);		
+
+		FNetPacketNotify Acks = Util.DefaultNotify;
+		FNetPacketNotifyTestUtil::PretendAckSequenceNumbers(Acks, AckdPacketIds);
+ 
+		TestEqual(TEXT("Create history - InSeq"), Acks.GetInSeq(), ExpectedInSeq);
+		TestEqual(TEXT("Create history - InSeqHistory"), Acks.GetInSeqHistory(), ExpectedInSeqHistory);
+	}
+	
+	// test notifications
+	{
+		static const FNetPacketNotify::SequenceNumberT ExpectedAckdPacketIds[] = {3, 7, 12, 14, 17, 18};
+		static const SIZE_T ExpectedCount = sizeof(ExpectedAckdPacketIds)/sizeof((ExpectedAckdPacketIds)[0]);		
+
+		FNetPacketNotify::SequenceNumberT RcvdAcks[ExpectedCount] = { 0 };
+		SIZE_T RcvdCount = 0;
+
+		// Create src data
+		FNetPacketNotify Acks = Util.DefaultNotify;
+
+		// Fill in some data
+		FNetPacketNotify::FNotificationHeader Data;
+		Data.Seq = FNetPacketNotify::SequenceNumberT(0);
+		Data.AckedSeq = FNetPacketNotify::SequenceNumberT(18);
+		Data.History = FNetPacketNotify::SequenceHistoryT(0x8853u);
+		Data.HistoryWordCount = 1;
+
+		// Need to fake ack record as well.
+		for (SIZE_T It=0; It <= 18; ++It)
+		{
+			FNetPacketNotifyTestUtil::PretendSendSeq(Acks, 0);
+		}
+	
+		SIZE_T DeltaSeq = FNetPacketNotifyTestUtil::PretendReceivedPacket(Acks, Data, RcvdAcks);
+
+		TestEqual(TEXT("Notifications - Create sequence delta"), DeltaSeq, 1);
+		TestEqual(TEXT("Notifications - Create sequence"), FPlatformMemory::Memcmp(ExpectedAckdPacketIds, RcvdAcks, sizeof(ExpectedAckdPacketIds)), 0u);
+	}
+
+	// test various sequence
+	{
+		static const FNetPacketNotify::SequenceNumberT TestSeqs[] = {3, 7, 12, 14, 17, 18};	
+		TestTrue(TEXT("Test Seq {3, 7, 12, 14, 17, 18}"), FNetPacketNotifyTestUtil::TestNotificationSequence(TestSeqs));
+	}
+	{
+		static const FNetPacketNotify::SequenceNumberT TestSeqs[] = {0, 1, 2, 3, 4, 5, 6, 7, 8, 9, 10, 11, 12, 13, 14, 15, 16, 17, 18, 19, 20, 21, 22, 23, 24, 25, 26, 27, 28, 29, 30, 31};	
+		TestTrue(TEXT("Test Seq {0, 1, 2, 3, 4, 5, 6, 7, 8, 9, 10, 11, 12, 13, 14, 15, 16, 17, 18, 19, 20, 21, 22, 23, 24, 25, 26, 27, 28, 29, 30, 31};"), FNetPacketNotifyTestUtil::TestNotificationSequence(TestSeqs));
+	}
+	{
+		static const FNetPacketNotify::SequenceNumberT TestSeqs[] = {2, 3, 4, 5, 6, 7, 8, 9, 10, 12, 13, 14, 16, 17, 18, 19, 20, 21, 22, 23, 24, 25, 26, 27, 28, 29, 30, 31};	
+		TestTrue(TEXT("{2, 3, 4, 5, 6, 7, 8, 9, 10, 12, 13, 14, 16, 17, 18, 19, 20, 21, 22, 23, 24, 25, 26, 27, 28, 29, 30, 31}"), FNetPacketNotifyTestUtil::TestNotificationSequence(TestSeqs));
+	}
+	{
+		static const FNetPacketNotify::SequenceNumberT TestSeqs[] = {0, FNetPacketNotify::MaxSequenceHistoryLength - 1};	
+		TestTrue(TEXT("Test Seq {0, FNetPacketNotify::MaxSequenceHistoryLength - 1}"), FNetPacketNotifyTestUtil::TestNotificationSequence(TestSeqs));
+	}
+	{
+		static const FNetPacketNotify::SequenceNumberT TestSeqs[] = {0, FNetPacketNotify::MaxSequenceHistoryLength};	
+		TestFalse(TEXT("Test Seq {0, FNetPacketNotify::MaxSequenceHistoryLength}"), FNetPacketNotifyTestUtil::TestNotificationSequence(TestSeqs));
+	}
+	{
+		static const FNetPacketNotify::SequenceNumberT TestSeqs[] = {0, FNetPacketNotify::SequenceNumberT::SeqNumberHalf };
+		TestFalse(TEXT("Test Seq {FNetPacketNotify::SequenceNumberT::SeqNumberMax, 0};"), FNetPacketNotifyTestUtil::TestNotificationSequence(TestSeqs));
+	}
+	{
+		static const FNetPacketNotify::SequenceNumberT TestSeqs[] = {0, FNetPacketNotify::SequenceNumberT::SeqNumberHalf - 1};
+		TestFalse(TEXT("Test Seq {FNetPacketNotify::SequenceNumberT::SeqNumberMax, 0};"), FNetPacketNotifyTestUtil::TestNotificationSequence(TestSeqs));
+	}
+	{
+		static const FNetPacketNotify::SequenceNumberT TestSeqs[] = {FNetPacketNotify::SequenceNumberT::SeqNumberMax, 0};
+		TestFalse(TEXT("Test Seq {FNetPacketNotify::SequenceNumberT::SeqNumberMax, 0} From 0;"), FNetPacketNotifyTestUtil::TestNotificationSequence(TestSeqs));
+	}
+	{
+		static const FNetPacketNotify::SequenceNumberT TestSeqs[] = {FNetPacketNotify::SequenceNumberT::SeqNumberMax, 0};
+		TestTrue(TEXT("Test Seq {FNetPacketNotify::SequenceNumberT::SeqNumberMax, 0} From SeqNumberHalf + 1;"), FNetPacketNotifyTestUtil::TestNotificationSequence(TestSeqs, FNetPacketNotify::SequenceNumberT::SeqNumberHalf + 2));
+	}
+
+	check(!HasAnyErrors());
+
+	return true;
+}
+
+#endif // #if WITH_DEV_AUTOMATION_TESTS
\ No newline at end of file
diff --git a/Engine/Source/Runtime/Engine/Private/Net/Tests/Util/ResizableCircularQueueTests.cpp b/Engine/Source/Runtime/Engine/Private/Net/Tests/Util/ResizableCircularQueueTests.cpp
new file mode 100644
--- /dev/null
+++ b/Engine/Source/Runtime/Engine/Private/Net/Tests/Util/ResizableCircularQueueTests.cpp
@@ -0,0 +1,218 @@
+// Copyright 1998-2019 Epic Games, Inc. All Rights Reserved.
+
+#include "Misc/AutomationTest.h"
+#include "Net/Util/ResizableCircularQueue.h"
+
+#if WITH_DEV_AUTOMATION_TESTS
+
+IMPLEMENT_SIMPLE_AUTOMATION_TEST(FResizableCircularQueueTest, "Network.ResizableCircularQueueTest", EAutomationTestFlags::ApplicationContextMask | EAutomationTestFlags::EngineFilter)
+
+struct FResizableCircularQueueTestUtil
+{
+	typedef TResizableCircularQueue<uint32> QueueT;
+
+	static bool VerifyQueueIntegrity(const QueueT& Queue, uint32 ExpectedValueAtFront, uint32 Increment)
+	{
+		bool bSuccess = true;
+		SIZE_T Offset = Queue.Count() - 1;
+	
+		uint32 ExpectedValue = ExpectedValueAtFront;
+
+		// Peek elements in queue at given offset, peek from back to front
+		for (SIZE_T It = 0; It < Queue.Count(); ++It)
+		{
+			bSuccess = bSuccess && (ExpectedValue == Queue.PeekAtOffset(It));
+			ExpectedValue += Increment;
+		}
+
+		return bSuccess;
+	}
+
+	static void OverideHeadAndTail(QueueT& Queue, uint32 Head, uint32 Tail)
+	{
+		Queue.Head = Head;
+		Queue.Tail = Tail;
+	}
+};
+
+bool FResizableCircularQueueTest::RunTest(const FString& Parameters)
+{
+	// Test empty
+	{
+		FResizableCircularQueueTestUtil::QueueT Q(0);
+
+		TestEqual(TEXT("Test empty - Size"), Q.Count(), 0);
+		TestTrue(TEXT("Test empty - IsEmpty"), Q.IsEmpty());
+		TestEqual(TEXT("Test empty - Capacity"), Q.AllocatedCapacity(), 0);
+	}
+
+	// Test Push to Capacity
+	{
+		const SIZE_T ElementsToPush = 8;
+
+		FResizableCircularQueueTestUtil::QueueT Q(ElementsToPush);
+
+		for (SIZE_T It=0; It < ElementsToPush; ++It)
+		{
+			Q.Enqueue(It);
+		}
+		
+		TestEqual(TEXT("Test Push to Capacity - Size"), Q.Count(), ElementsToPush);
+		TestEqual(TEXT("Test Push to Capacity - Capacity"), Q.AllocatedCapacity(), ElementsToPush);
+ 		TestTrue(TEXT("Test Push to Capacity - Expected"), FResizableCircularQueueTestUtil::VerifyQueueIntegrity(Q, 0, 1));
+	}
+
+	// Test Push over Capacity
+	{
+		const SIZE_T ElementsToPush = 32;
+
+		FResizableCircularQueueTestUtil::QueueT Q(ElementsToPush);
+
+		for (SIZE_T It=0; It < ElementsToPush; ++It)
+		{
+			Q.Enqueue(It);
+		}
+		
+		TestEqual(TEXT("Test Push over Capacity - Size"), Q.Count(), ElementsToPush);
+		TestEqual(TEXT("Test Push over Capacity - Capacity"), Q.AllocatedCapacity(), ElementsToPush);
+ 		TestTrue(TEXT("Test Push over Capacity - Expected"), FResizableCircularQueueTestUtil::VerifyQueueIntegrity(Q, 0, 1));
+	}
+
+	// Test Push and Pop
+	{
+		const SIZE_T ElementsToPush = 256;
+		const SIZE_T ElementPopMod = 16;
+		const SIZE_T ExpectedSize = 256 - ElementPopMod;
+		const SIZE_T ExpectedCapacity = 256;
+
+		FResizableCircularQueueTestUtil::QueueT Q(4);
+
+		uint32 ExpectedPoppedValue = 0;
+		for (uint32 It = 0; It < 256; ++It)
+		{
+			Q.Enqueue(It);
+			TestEqual(TEXT("Test Push and pop - Push"), It, Q.PeekAtOffset(Q.Count() - 1));
+
+			if (It % ElementPopMod == 0)
+			{
+				const uint32 PoppedValue = Q.PeekAtOffset(0);
+				TestEqual(TEXT("Test Push and pop - Pop"), ExpectedPoppedValue, PoppedValue);
+				++ExpectedPoppedValue;
+				Q.Pop();
+			}
+		}
+
+		TestEqual(TEXT("Test Push and pop - Size"), Q.Count(), ExpectedSize);
+		TestEqual(TEXT("Test Push and pop - Capacity"), Q.AllocatedCapacity(), ExpectedCapacity);
+ 		TestTrue (TEXT("Test Push and pop - Expected"), FResizableCircularQueueTestUtil::VerifyQueueIntegrity(Q, ElementPopMod, 1));	
+	}
+
+	// Test Push and pop all
+	{
+		const SIZE_T ElementsToPush = 256;
+
+		FResizableCircularQueueTestUtil::QueueT Q(ElementsToPush);
+
+		TestTrue(TEXT ("Test Push and pop all - IsEmpty before"), Q.IsEmpty());
+		TestEqual(TEXT("Test Push and pop all - Size before"), Q.Count(), 0);
+
+		for (SIZE_T It=0; It < ElementsToPush; ++It)
+		{
+			Q.Enqueue(It);
+		}
+		
+		TestEqual(TEXT("Test Push and pop all - Size"), Q.Count(), ElementsToPush);
+		TestEqual(TEXT("Test Push and pop all - Capacity"), Q.AllocatedCapacity(), ElementsToPush);
+ 		TestTrue (TEXT("Test Push and pop all - Expected"), FResizableCircularQueueTestUtil::VerifyQueueIntegrity(Q, 0, 1));
+
+		for (SIZE_T It=0; It < ElementsToPush; ++It)
+		{
+			Q.Pop();
+		}
+		
+		TestTrue( TEXT("Test Push and pop all - IsEmpty after"), Q.IsEmpty());
+		TestEqual(TEXT("Test Push and pop all - Size after"), Q.Count(), 0);
+		TestEqual(TEXT("Test Push and pop all - Capacity after"), Q.AllocatedCapacity(), ElementsToPush);
+	}
+
+	// Test index wrap
+	{
+		const SIZE_T ElementsToPush = 256;
+		const SIZE_T ElementPopMod = 16;
+		const SIZE_T ExpectedSize = 256 - ElementPopMod;
+		const SIZE_T ExpectedCapacity = 256;
+
+		FResizableCircularQueueTestUtil::QueueT Q(4);
+
+		// Set head and tail at the end of the space to test index wraparound
+		FResizableCircularQueueTestUtil::OverideHeadAndTail(Q, uint32(-2), uint32(-2));
+
+		TestTrue(TEXT ("Test index wrap - IsEmpty before"), Q.IsEmpty());
+		TestEqual(TEXT("Test index wrap - Size before"), Q.Count(), 0);
+
+		for (SIZE_T It=0; It < ElementsToPush; ++It)
+		{
+			Q.Enqueue(It);
+		}
+		
+		TestEqual(TEXT("Test index wrap - Size"), Q.Count(), ElementsToPush);
+		TestEqual(TEXT("Test index wrap - Capacity"), Q.AllocatedCapacity(), ElementsToPush);
+ 		TestTrue (TEXT("Test index wrap - Expected"), FResizableCircularQueueTestUtil::VerifyQueueIntegrity(Q, 0, 1));
+
+		for (SIZE_T It=0; It < ElementsToPush; ++It)
+		{
+			Q.Pop();
+		}
+		
+		TestTrue( TEXT("Test index wrap - IsEmpty after"), Q.IsEmpty());
+		TestEqual(TEXT("Test index wrap - Size after"), Q.Count(), 0);
+		TestEqual(TEXT("Test index wrap - Capacity after"), Q.AllocatedCapacity(), ElementsToPush);
+	}
+
+	// Test Trim
+	{
+		const SIZE_T ElementsToPush = 9;
+		const SIZE_T ElementsToPop = 5;
+		const SIZE_T ExpectedCapacity = 16;
+		const SIZE_T ExpectedCapacityAfterTrim = 4;
+
+		FResizableCircularQueueTestUtil::QueueT Q(0);
+
+		for (SIZE_T It=0; It < ElementsToPush; ++It)
+		{
+			Q.Enqueue(It);
+		}
+		
+		TestEqual(TEXT("Test Trim - Size"), Q.Count(), ElementsToPush);
+		TestEqual(TEXT("Test Trim - Capacity"), Q.AllocatedCapacity(), ExpectedCapacity);
+ 		TestTrue(TEXT("Test Trim - Expected"), FResizableCircularQueueTestUtil::VerifyQueueIntegrity(Q, 0, 1));
+
+		for (SIZE_T It=0; It < ElementsToPop; ++It)
+		{
+			Q.Pop();
+		}
+
+		Q.Trim();
+
+		TestEqual(TEXT("Test Trim - Size"), Q.Count(), ElementsToPush - ElementsToPop);
+		TestEqual(TEXT("Test Trim - Capacity"), Q.AllocatedCapacity(), ExpectedCapacityAfterTrim);
+ 		TestTrue(TEXT("Test Trim - Expected"), FResizableCircularQueueTestUtil::VerifyQueueIntegrity(Q, ElementsToPop, 1));
+	}
+
+	// Test Trim empty
+	{
+		FResizableCircularQueueTestUtil::QueueT Q(0);
+
+		Q.Trim();
+
+		TestEqual(TEXT("Test trim empty - Size"), Q.Count(), 0);
+		TestEqual(TEXT("Test trim empty - Capacity"), Q.AllocatedCapacity(), 0);
+	}
+
+
+	check(!HasAnyErrors());
+
+	return true;
+}
+
+#endif // #if WITH_DEV_AUTOMATION_TESTS
\ No newline at end of file
diff --git a/Engine/Source/Runtime/Engine/Private/Net/Tests/Util/SequenceNumberTests.cpp b/Engine/Source/Runtime/Engine/Private/Net/Tests/Util/SequenceNumberTests.cpp
new file mode 100644
--- /dev/null
+++ b/Engine/Source/Runtime/Engine/Private/Net/Tests/Util/SequenceNumberTests.cpp
@@ -0,0 +1,111 @@
+// Copyright 1998-2019 Epic Games, Inc. All Rights Reserved.
+
+#include "Misc/AutomationTest.h"
+#include "Net/Util/SequenceNumber.h"
+
+#if WITH_DEV_AUTOMATION_TESTS
+
+IMPLEMENT_SIMPLE_AUTOMATION_TEST(FNetSequenceNumberTest, "Network.SequenceNumberTest", EAutomationTestFlags::ApplicationContextMask | EAutomationTestFlags::EngineFilter)
+
+bool FNetSequenceNumberTest::RunTest(const FString& Parameters)
+{
+	typedef TSequenceNumber<3, uint16> FSequence3;
+
+	// Test construct
+	{
+		FSequence3 Seq;
+		TestEqual(TEXT("SequenceNumbers - Seq() == 0"), 0u, Seq.Get());		
+	}
+
+	// Test construct with value
+	{
+		FSequence3 Seq(8);
+		TestEqual(TEXT("SequenceNumbers - Seq(8) == 0"), 0u, Seq.Get());		
+	}
+
+	// Test postIncrement
+	{
+		FSequence3 Seq;
+		Seq++;
+		TestEqual("SequenceNumbers - Seq()++ == 1", 1, Seq.Get());
+	}
+
+	// Test preIncrement
+	{
+		FSequence3 Seq;
+		++Seq;
+		TestEqual("SequenceNumbers - ++Seq == 1", 1, Seq.Get());
+	}
+
+	// Test wraparound
+	{
+		FSequence3 Seq(7);
+		Seq++;
+		TestEqual("SequenceNumbers - Seq(7) + 1 == 0", 0, Seq.Get());
+	}
+
+	// Test equal and not equal
+	{
+		FSequence3 A(2);
+		FSequence3 B(2);
+		FSequence3 C(1);
+
+		TestTrue("SequenceNumbers - operator ==", A == B);
+		TestTrue("SequenceNumbers - operator ==", B == A);
+		TestTrue("SequenceNumbers - operator !=", A != C);
+		TestTrue("SequenceNumbers - operator !=", C != A);
+		TestFalse("SequenceNumbers - operator !=", A != B);
+	}
+
+	// Test operator >
+	{
+		FSequence3 Zero(0);
+		FSequence3 HalfMinusOne(FSequence3::SeqNumberHalf-1);
+		FSequence3 Half(FSequence3::SeqNumberHalf);
+		FSequence3 MaxSeq(FSequence3::SeqNumberCount - 1);
+
+		TestTrue("SequenceNumbers - operator> HalfMinusOne > Zero", HalfMinusOne > Zero);
+		TestTrue("SequenceNumbers - operator> Half > HalfMinusOne", Half > HalfMinusOne);
+		TestFalse("SequenceNumbers - operator> Half > Zero", Half > Zero);
+
+		TestFalse("SequenceNumbers - operator> MaxSeq > Zero", MaxSeq > Zero);
+		TestTrue("SequenceNumbers - operator> Zero > MaxSeq", Zero > MaxSeq);
+	}
+	
+	// Test diff sequence numbers
+	{		
+		// Valid sequence = 0-7, max distance between sequence number in order to determine order is half the sequence space (0-3)
+		for (int32 I = 0; I < FSequence3::SeqNumberCount; ++I)
+		{
+			FSequence3 Seq(I);
+			FSequence3 Ref(I);
+			
+			for (int32 U = 0; U < FSequence3::SeqNumberCount; ++U)
+			{	
+				FSequence3::DifferenceT Diff = FSequence3::Diff(Seq, Ref);
+
+				if (U < FSequence3::SeqNumberHalf)
+				{
+					if (!(Diff == U))
+					{
+						TestTrue(TEXT("SequenceNumbers - Expected Diff "), Diff == U);
+					}
+				}
+				else
+				{
+					if (!(Diff == (U - FSequence3::SeqNumberCount)))
+					{
+						TestTrue(TEXT("SequenceNumbers - Expected Diff"), Diff == (U - FSequence3::SeqNumberCount));
+					}
+				}
+				++Seq;
+			}
+		}
+
+		check(!HasAnyErrors());
+	}
+
+	return true;
+}
+
+#endif // #if WITH_DEV_AUTOMATION_TESTS
\ No newline at end of file
diff --git a/Engine/Source/Runtime/Engine/Private/NetConnection.cpp b/Engine/Source/Runtime/Engine/Private/NetConnection.cpp
--- a/Engine/Source/Runtime/Engine/Private/NetConnection.cpp
+++ b/Engine/Source/Runtime/Engine/Private/NetConnection.cpp
@@ -19,6 +19,7 @@
 #include "Net/UnrealNetwork.h"
 #include "Net/NetworkProfiler.h"
 #include "Net/DataReplication.h"
+#include "Net/NetPacketNotify.h"
 #include "Engine/ActorChannel.h"
 #include "Engine/ChildConnection.h"
 #include "Engine/VoiceChannel.h"
@@ -26,7 +27,6 @@
 #include "Engine/PackageMapClient.h"
 #include "Engine/NetworkObjectList.h"
 #include "EncryptionComponent.h"
-
 #include "Net/PerfCountersHelpers.h"
 #include "GameDelegates.h"
 #include "Misc/PackageName.h"
@@ -133,6 +133,8 @@ UNetConnection::UNetConnection(const FObjectInitializer& ObjectInitializer)
 #endif
 ,	PlayerOnlinePlatformName( NAME_None )
 ,	ClientWorldPackageName( NAME_None )
+,	LastNotifiedPacketId( -1 )
+,	HasDirtyAcks(0u)
 {
 	MaxChannelSize = CVarMaxChannelSize.GetValueOnAnyThread();
 	if (MaxChannelSize <= 0)
@@ -145,6 +147,8 @@ UNetConnection::UNetConnection(const FObjectInitializer& ObjectInitializer)
 	OutReliable.AddDefaulted(MaxChannelSize);
 	InReliable.AddDefaulted(MaxChannelSize);
 	PendingOutRec.AddDefaulted(MaxChannelSize);
+
+	PacketNotify.Init(InPacketId, OutPacketId);
 }
 
 /**
@@ -338,6 +342,7 @@ void UNetConnection::InitSequence(int32 IncomingSequence, int32 OutgoingSequence
 		InPacketId = IncomingSequence - 1;
 		OutPacketId = OutgoingSequence;
 		OutAckPacketId = OutgoingSequence - 1;
+		LastNotifiedPacketId = OutAckPacketId;
 
 		// Initialize the reliable packet sequence (more useful/effective at preventing attacks)
 		InitInReliable = IncomingSequence & (MAX_CHSEQUENCE - 1);
@@ -346,6 +351,8 @@ void UNetConnection::InitSequence(int32 IncomingSequence, int32 OutgoingSequence
 		InReliable.Init(InitInReliable, InReliable.Num());
 		OutReliable.Init(InitOutReliable, OutReliable.Num());
 
+		PacketNotify.Init(InPacketId, OutPacketId);
+
 		UE_LOG(LogNet, Verbose, TEXT("InitSequence: IncomingSequence: %i, OutgoingSequence: %i, InitInReliable: %i, InitOutReliable: %i"), IncomingSequence, OutgoingSequence, InitInReliable, InitOutReliable);
 	}
 }
@@ -456,7 +463,6 @@ bool UNetConnection::IsEncryptionEnabled() const
 void UNetConnection::Serialize( FArchive& Ar )
 {
 	UObject::Serialize( Ar );
-	
 	Ar << PackageMap;
 	for (UChannel* Channel : Channels)
 	{
@@ -485,8 +491,6 @@ void UNetConnection::Serialize( FArchive& Ar )
 		OutReliable.CountBytes(Ar);
 		InReliable.CountBytes(Ar);
 		PendingOutRec.CountBytes(Ar);
-		QueuedAcks.CountBytes(Ar);
-		ResendAcks.CountBytes(Ar);
 		ActorChannels.CountBytes(Ar);
 		DestroyedStartupOrDormantActorGUIDs.CountBytes(Ar);
 		KeepProcessingActorChannelBunchesMap.CountBytes(Ar);
@@ -1095,7 +1099,7 @@ void UNetConnection::FlushNet(bool bIgnoreSimulation)
 	TimeSensitive = 0;
 
 	// If there is any pending data to send, send it.
-	if (SendBuffer.GetNumBits() || ( Driver->Time-LastSendTime > Driver->KeepAliveTime && !InternalAck && State != USOCK_Closed))
+	if (SendBuffer.GetNumBits() || HasDirtyAcks || ( Driver->Time-LastSendTime > Driver->KeepAliveTime && !InternalAck && State != USOCK_Closed))
 	{
 		// Due to the PacketHandler handshake code, servers must never send the client data,
 		// before first receiving a client control packet (which is taken as an indication of a complete handshake).
@@ -1110,10 +1114,10 @@ void UNetConnection::FlushNet(bool bIgnoreSimulation)
 
 		FOutPacketTraits Traits;
 
-		// If sending keepalive packet, still write the packet id
+		// If sending keepalive packet or just acks, still write the packet header
 		if (SendBuffer.GetNumBits() == 0)
 		{
-			WriteBitsToSendBuffer( NULL, 0 );		// This will force the packet id to be written
+			WriteBitsToSendBuffer( NULL, 0 );		// This will force the packet header to be written
 
 			Traits.bIsKeepAlive = true;
 			AnalyticsVars.OutKeepAliveCount++;
@@ -1130,6 +1134,14 @@ void UNetConnection::FlushNet(bool bIgnoreSimulation)
 		// Write the UNetConnection-level termination bit
 		SendBuffer.WriteBit(1);
 
+		// Refresh outgoing header with latest data
+		if ( !InternalAck )
+		{
+			// if we update ack, we also update received ack associated with outgoing seq
+			// so we know how many ack bits we need to write (which is updated in received packet)
+			WritePacketHeader(SendBuffer);
+		}
+
 		ValidateSendBuffer();
 
 		const int32 NumStrayBits = SendBuffer.GetNumBits();
@@ -1209,9 +1221,17 @@ void UNetConnection::FlushNet(bool bIgnoreSimulation)
 		// Remember the actual time this packet was sent out, so we can compute ping when the ack comes back
 		OutLagPacketId[Index]			= OutPacketId;
 		OutLagTime[Index]				= FPlatformTime::Seconds();
-		OutBytesPerSecondHistory[Index]	= OutBytesPerSecond / 1024;
+	
+		OutBytesPerSecondHistory[Index]	= FMath::Min(OutBytesPerSecond / 1024, 255);
+
+		// Increase outgoing sequence number
+		if (!InternalAck)
+		{
+			PacketNotify.CommitAndIncrementOutSeq();
+			HasDirtyAcks = 0u;
+		}
+		++OutPacketId; 
 
-		OutPacketId++;
 		++OutPackets;
 		++OutTotalPackets;
 		Driver->OutPackets++;
@@ -1237,13 +1257,6 @@ void UNetConnection::FlushNet(bool bIgnoreSimulation)
 
 		InitSendBuffer();
 	}
-
-	// Move acks around.
-	for (int32 i=0; i<QueuedAcks.Num(); i++)
-	{
-		ResendAcks.Add(QueuedAcks[i]);
-	}
-	QueuedAcks.Empty(32);
 }
 
 bool UNetConnection::ShouldDropOutgoingPacketForLossSimulation(int64 NumBits) const
@@ -1271,8 +1284,64 @@ int32 UNetConnection::IsNetReady( bool Saturate )
 void UNetConnection::ReadInput( float DeltaSeconds )
 {}
 
+void UNetConnection::ReceivedAck(int32 AckPacketId)
+{
+	UE_LOG(LogNetTraffic, Verbose, TEXT("   Received ack %i"), AckPacketId);
+
+	// Advance OutAckPacketId
+	OutAckPacketId = AckPacketId;
+
+	// Process the bunch.
+	LastRecvAckTime = Driver->Time;
+
+	if ( PackageMap != NULL )
+	{
+		PackageMap->ReceivedAck( AckPacketId );
+	}
+
+	// Forward the ack to the channel.
+	{
+		SCOPE_CYCLE_COUNTER(Stat_NetConnectionReceivedAcks);
+	
+		for (int32 i = OpenChannels.Num() - 1; i >= 0; i--)
+		{
+			UChannel* const Channel = OpenChannels[i];
+				
+			if (Channel)
+			{
+				if (Channel->OpenPacketId.Last == AckPacketId) // Necessary for unreliable "bNetTemporary" channels.
+				{
+					Channel->OpenAcked = 1;
+				}
+				
+				for (FOutBunch* OutBunch = Channel->OutRec; OutBunch; OutBunch = OutBunch->Next)
+				{
+					if (OutBunch->bOpen)
+					{
+						UE_LOG(LogNet, VeryVerbose, TEXT("Channel %i reset Ackd because open is reliable. "), Channel->ChIndex );
+						Channel->OpenAcked  = 0; // We have a reliable open bunch, don't let the above code set the OpenAcked state,
+												// it must be set in UChannel::ReceivedAcks to verify all open bunches were received.
+					}
+
+					if (OutBunch->PacketId == AckPacketId)
+					{
+						OutBunch->ReceivedAck = 1;
+					}
+				}
+				Channel->ReceivedAcks(); //warning: May destroy Channel.
+			}
+			else
+			{
+				UE_LOG(LogNet, Warning, TEXT("UNetConnection::ReceivedPacket: null channel in OpenChannels array while processing ack. %s"), *Describe());
+			}
+		}
+	}
+}
+
 void UNetConnection::ReceivedNak( int32 NakPacketId )
 {
+	UE_LOG(LogNetTraffic, Verbose, TEXT("   Received nak %i"), NakPacketId);
+
 	SCOPE_CYCLE_COUNTER(Stat_NetConnectionReceivedNak);
 
 	// Update pending NetGUIDs
@@ -1285,7 +1354,7 @@ void UNetConnection::ReceivedNak( int32 NakPacketId )
 		if (Channel)
 		{
 			Channel->ReceivedNak(NakPacketId);
-			if(Channel->OpenPacketId.InRange(NakPacketId))
+			if (Channel->OpenPacketId.InRange(NakPacketId))
 			{
 				Channel->ReceivedAcks(); //warning: May destroy Channel.
 			}
@@ -1295,6 +1364,124 @@ void UNetConnection::ReceivedNak( int32 NakPacketId )
 			UE_LOG(LogNet, Warning, TEXT("UNetConnection::ReceivedNak: null channel in OpenChannels array. %s"), *Describe());
 		}
 	}
+
+	// Stats
+	++OutPacketsLost;
+	++OutTotalPacketsLost;
+	++Driver->OutTotalPacketsLost;
+}
+
+// IMPORTANT:
+// WritePacketHeader must ALWAYS write the exact same number of bits as we go back and rewrite the header
+// right before we put the packet on the wire.
+void UNetConnection::WritePacketHeader(FBitWriter& Writer)
+{
+	// If this is a header refresh, we only serialize the updated serial number information
+	const bool bIsHeaderUpdate = Writer.GetNumBits() > 0u;
+
+	// Header is always written first in the packet
+	FBitWriterMark Reset;
+	FBitWriterMark Restore(Writer);
+	Reset.PopWithoutClear(Writer);
+	
+	// Write notification header or refresh the header if used space is the same.
+	PacketNotify.WriteHeader(Writer, bIsHeaderUpdate);
+
+	// Jump back to where we came from.
+	if (bIsHeaderUpdate)
+	{
+		Restore.PopWithoutClear(Writer);
+	}
+}
+
+void UNetConnection::WritePacketInfo(FBitWriter& Writer) const
+{
+	const uint8 bHasServerFrameTime = Driver->IsServer() ? bLastHasServerFrameTime : ( CVarPingExcludeFrameTime.GetValueOnGameThread() > 0 ? 1u : 0u );
+
+	// Write data used to calculate link latency
+	Writer.WriteBit(bHasServerFrameTime);
+	if (bHasServerFrameTime && Driver->IsServer())
+	{
+		uint8 FrameTimeByte = FMath::Min( FMath::FloorToInt( FrameTime * 1000 ), 255 );
+		Writer << FrameTimeByte;
+	}
+
+	// Notify server of our current rate per second at this time
+	uint8 InKBytesPerSecondByte = FMath::Min(InBytesPerSecond / 1024, 255);
+	Writer << InKBytesPerSecondByte;
+}
+
+bool UNetConnection::ReadPacketInfo(FBitReader& Reader)
+{
+	const bool bHasServerFrameTime = Reader.ReadBit() == 1u;
+	double ServerFrameTime = 0;
+
+	if ( !Driver->IsServer() )
+	{
+		if ( bHasServerFrameTime )
+		{
+			uint8 FrameTimeByte	= 0;
+			Reader << FrameTimeByte;
+			// As a client, our request was granted, read the frame time
+			ServerFrameTime = ( double )FrameTimeByte / 1000;
+		}
+	}
+	else
+	{
+		bLastHasServerFrameTime = bHasServerFrameTime;
+	}
+
+	// limit to known size to know the size of the packet header
+	uint8 RemoteInKBytesPerSecondByte = 0;
+	Reader << RemoteInKBytesPerSecondByte;
+
+	if ( Reader.IsError() )
+	{
+		return false;
+	}
+
+	// Update ping
+	// At this time we have updated OutAckPacketId to the latest received ack.
+	const int32 Index = OutAckPacketId & (ARRAY_COUNT(OutLagPacketId)-1);
+
+	if ( OutLagPacketId[Index] == OutAckPacketId )
+	{
+		OutLagPacketId[Index] = -1;		// Only use the ack once
+
+#if !UE_BUILD_SHIPPING
+		if ( CVarPingDisplayServerTime.GetValueOnAnyThread() > 0 )
+		{
+			UE_LOG( LogNetTraffic, Warning, TEXT( "ServerFrameTime: %2.2f" ), ServerFrameTime * 1000.0f );
+		}
+#endif
+
+		// use FApp's time because it is set closer to the beginning of the frame - we don't care about the time so far of the current frame to process the packet
+		const double CurrentTime = FApp::GetCurrentTime();
+		const double GameTime	 = ServerFrameTime;
+		const double RTT		 = (CurrentTime - OutLagTime[Index] ) - ( CVarPingExcludeFrameTime.GetValueOnAnyThread() ? GameTime : 0.0 );
+		const double NewLag		 = FMath::Max( RTT, 0.0 );
+
+		if ( OutBytesPerSecondHistory[Index] > 0 )
+		{
+			RemoteSaturation = ( 1.0f - FMath::Min( ( float )RemoteInKBytesPerSecondByte / ( float )OutBytesPerSecondHistory[Index], 1.0f ) ) * 100.0f;
+		}
+		else
+		{
+			RemoteSaturation = 0.0f;
+		}
+
+		//UE_LOG( LogNet, Warning, TEXT( "Out: %i, InRemote: %i, Saturation: %f" ), OutBytesPerSecondHistory[Index], RemoteInKBytesPerSecond, RemoteSaturation );
+
+		LagAcc += NewLag;
+		LagCount++;
+
+		if (PlayerController != NULL)
+		{
+			PlayerController->UpdatePing(NewLag);
+		}
+	}
+
+	return true;
 }
 
 void UNetConnection::ReceivedPacket( FBitReader& Reader )
@@ -1319,201 +1506,107 @@ void UNetConnection::ReceivedPacket( FBitReader& Reader )
 	LastReceiveTime		= Driver->Time;
 	LastReceiveRealtime = FPlatformTime::Seconds();
 
-	// Check packet ordering.
-	const int32 PacketId = InternalAck ? InPacketId + 1 : MakeRelative(Reader.ReadInt(MAX_PACKETID),InPacketId,MAX_PACKETID);
-	if( PacketId > InPacketId )
+	if (InternalAck)
 	{
-		const int32 PacketsLost = PacketId - InPacketId - 1;
-		
-		if ( PacketsLost > 10 )
-		{
-			UE_LOG( LogNetTraffic, Log, TEXT( "High single frame packet loss. PacketsLost: %i %s" ), PacketsLost, *Describe() );
-		}
-
-		InPacketsLost += PacketsLost;
-		InTotalPacketsLost += PacketsLost;
-		Driver->InPacketsLost += PacketsLost;
-		Driver->InTotalPacketsLost += PacketsLost;
-		InPacketId = PacketId;
-	}
+		++InPacketId;
+	}	
 	else
 	{
-		Driver->InOutOfOrderPackets++;
-		// Protect against replay attacks
-		// We already protect against this for reliable bunches, and unreliable properties
-		// The only bunch we would process would be unreliable RPC's, which could allow for replay attacks
-		// So rather than add individual protection for unreliable RPC's as well, just kill it at the source, 
-		// which protects everything in one fell swoop
-		return;
-	}
-
-	const bool bIgnoreRPCs = Driver->ShouldIgnoreRPCs();
-
-	bool bSkipAck = false;
-
-	// Track channels that were rejected while processing this packet - used to avoid sending multiple close-channel bunches,
-	// which would cause a disconnect serverside
-	TArray<int32> RejectedChans;
-
-	// Disassemble and dispatch all bunches in the packet.
-	while( !Reader.AtEnd() && State!=USOCK_Closed )
-	{
-		// Parse the bunch.
-		int32 StartPos = Reader.GetPosBits();
-		bool IsAck = !!Reader.ReadBit();
-		if ( Reader.IsError() )
+		// Read packet header
+		FNetPacketNotify::FNotificationHeader Header;
+		if (!PacketNotify.ReadHeader(Header, Reader))
 		{
-			CLOSE_CONNECTION_DUE_TO_SECURITY_VIOLATION(this, ESecurityEvent::Malformed_Packet, TEXT("Bunch missing ack flag"));
+			CLOSE_CONNECTION_DUE_TO_SECURITY_VIOLATION(this, ESecurityEvent::Malformed_Packet, TEXT("Failed to read PacketHeader"));
 			return;
 		}
 
-		// Process the bunch.
-		if( IsAck )
+		// Lambda to dispatch delivery notifications, 
+		auto HandlePacketNotification = [&Header, this](FNetPacketNotify::SequenceNumberT AckedSequence, bool bDelivered)
 		{
-			LastRecvAckTime = Driver->Time;
-
-			// This is an acknowledgment.
-			const int32 AckPacketId = MakeRelative(Reader.ReadInt(MAX_PACKETID),OutAckPacketId,MAX_PACKETID);
+			// Increase LastNotifiedPacketId, this is a full packet Id
+			++LastNotifiedPacketId;
 
-			if( Reader.IsError() )
+			// Sanity check
+			if (FNetPacketNotify::SequenceNumberT(LastNotifiedPacketId) != AckedSequence)
 			{
-				CLOSE_CONNECTION_DUE_TO_SECURITY_VIOLATION(this, ESecurityEvent::Malformed_Packet, TEXT("Bunch missing ack"));
+				CLOSE_CONNECTION_DUE_TO_SECURITY_VIOLATION(this, ESecurityEvent::Malformed_Packet, TEXT("LastNotifiedPacketId != AckedSequence"));
 				return;
 			}
 
-			double ServerFrameTime = 0;
-
-			// If this is the server, we're reading in the request to send them our frame time
-			// If this is the client, we're reading in confirmation that our request to get frame time from server is granted
-			const bool bHasServerFrameTime = !!Reader.ReadBit();
-
-			if ( !Driver->IsServer() )
+			if (bDelivered)
 			{
-				if ( bHasServerFrameTime )
-				{
-					// As a client, our request was granted, read the frame time
-					uint8 FrameTimeByte	= 0;
-					Reader << FrameTimeByte;
-					ServerFrameTime = ( double )FrameTimeByte / 1000;
-				}
+				ReceivedAck(LastNotifiedPacketId);
 			}
 			else
 			{
-				// Server remembers so he can use during SendAck to notify to client of his frame time
-				bLastHasServerFrameTime = bHasServerFrameTime;
-			}
-
-			uint32 RemoteInKBytesPerSecond = 0;
-			Reader.SerializeIntPacked( RemoteInKBytesPerSecond );
+				ReceivedNak(LastNotifiedPacketId);
+			};
+		};
 
-			// Resend any old reliable packets that the receiver hasn't acknowledged.
-			if( AckPacketId>OutAckPacketId )
-			{
-				for (int32 NakPacketId = OutAckPacketId + 1; NakPacketId<AckPacketId; NakPacketId++, OutPacketsLost++, OutTotalPacketsLost++, Driver->OutTotalPacketsLost++)
-				{
-					UE_LOG(LogNetTraffic, Verbose, TEXT("   Received virtual nak %i (%.1f)"), NakPacketId, (Reader.GetPosBits()-StartPos)/8.f );
-					ReceivedNak( NakPacketId );
-				}
-				OutAckPacketId = AckPacketId;
-			}
-			else if( AckPacketId<OutAckPacketId )
+		// Update incoming sequence data and deliver packet notifications
+		// Packet is only accepted if both the incoming sequence number and incoming ack data are valid
+		int PacketSequenceDelta = PacketNotify.Update(Header, HandlePacketNotification);
+		if (PacketSequenceDelta > 0)
+		{
+			const int32 PacketsLost = PacketSequenceDelta - 1;
+		
+			if ( PacketsLost > 10 )
 			{
-				//warning: Double-ack logic makes this unmeasurable.
-				//OutOrdAcc++;
+				UE_LOG( LogNetTraffic, Verbose, TEXT( "High single frame packet loss. PacketsLost: %i %s" ), PacketsLost, *Describe() );
 			}
 
-			// Update ping
-			const int32 Index = AckPacketId & (ARRAY_COUNT(OutLagPacketId)-1);
-
-			if ( OutLagPacketId[Index] == AckPacketId )
-			{
-				OutLagPacketId[Index] = -1;		// Only use the ack once
-
-#if !UE_BUILD_SHIPPING
-				if ( CVarPingDisplayServerTime.GetValueOnAnyThread() > 0 )
-				{
-					UE_LOG( LogNetTraffic, Warning, TEXT( "ServerFrameTime: %2.2f" ), ServerFrameTime * 1000.0f );
-				}
-#endif
-
-				// use FApp's time because it is set closer to the beginning of the frame - we don't care about the time so far of the current frame to process the packet
-				const double CurrentTime = FApp::GetCurrentTime();
-				const double GameTime	 = ServerFrameTime;
-				const double RTT		 = (CurrentTime - OutLagTime[Index] ) - ( CVarPingExcludeFrameTime.GetValueOnAnyThread() ? GameTime : 0.0 );
-				const double NewLag		 = FMath::Max( RTT, 0.0 );
-
-				if ( OutBytesPerSecondHistory[Index] > 0 )
-				{
-					RemoteSaturation = ( 1.0f - FMath::Min( ( float )RemoteInKBytesPerSecond / ( float )OutBytesPerSecondHistory[Index], 1.0f ) ) * 100.0f;
-				}
-				else
-				{
-					RemoteSaturation = 0.0f;
-				}
-
-				//UE_LOG( LogNet, Warning, TEXT( "Out: %i, InRemote: %i, Saturation: %f" ), OutBytesPerSecondHistory[Index], RemoteInKBytesPerSecond, RemoteSaturation );
-
-				LagAcc += NewLag;
-				LagCount++;
-
-				if (PlayerController != NULL)
-				{
-					PlayerController->UpdatePing(NewLag);
-				}
-			}
+			InPacketsLost += PacketsLost;
+			InTotalPacketsLost += PacketsLost;
+			Driver->InPacketsLost += PacketsLost;
+			Driver->InTotalPacketsLost += PacketsLost;
+			InPacketId += PacketSequenceDelta;
 
-			if ( PackageMap != NULL )
+			// Extra information associated with the header
+			if (!ReadPacketInfo(Reader))
 			{
-				PackageMap->ReceivedAck( AckPacketId );
+				CLOSE_CONNECTION_DUE_TO_SECURITY_VIOLATION(this, ESecurityEvent::Malformed_Packet, TEXT("Failed to read PacketHeader"));
+				return;
 			}
+		}
+		else
+		{
+			Driver->InOutOfOrderPackets++;
+			// Protect against replay attacks
+			// We already protect against this for reliable bunches, and unreliable properties
+			// The only bunch we would process would be unreliable RPC's, which could allow for replay attacks
+			// So rather than add individual protection for unreliable RPC's as well, just kill it at the source, 
+			// which protects everything in one fell swoop
+			return;
+		}
+	}
 
-			// Forward the ack to the channel.
-			UE_LOG(LogNetTraffic, Verbose, TEXT("   Received ack %i (%.1f)"), AckPacketId, (Reader.GetPosBits()-StartPos)/8.f );
-
-			{
-				SCOPE_CYCLE_COUNTER(Stat_NetConnectionReceivedAcks);
-
-				for (int32 i = OpenChannels.Num() - 1; i >= 0; i--)
-				{
-					UChannel* const Channel = OpenChannels[i];
+	const bool bIgnoreRPCs = Driver->ShouldIgnoreRPCs();
 
-					if (Channel)
-					{
-						if (Channel->OpenPacketId.Last==AckPacketId) // Necessary for unreliable "bNetTemporary" channels.
-						{
-							Channel->OpenAcked = 1;
-						}
+	bool bSkipAck = false;
 
-						for (FOutBunch* OutBunch = Channel->OutRec; OutBunch; OutBunch = OutBunch->Next)
-						{
-							if (OutBunch->bOpen)
-							{
-								UE_LOG(LogNet, VeryVerbose, TEXT("Channel %i reset Ackd because open is reliable. "), Channel->ChIndex );
-								Channel->OpenAcked  = 0; // We have a reliable open bunch, don't let the above code set the OpenAcked state,
-														 // it must be set in UChannel::ReceivedAcks to verify all open bunches were received.
-							}
+	// Track channels that were rejected while processing this packet - used to avoid sending multiple close-channel bunches,
+	// which would cause a disconnect serverside
+	TArray<int32> RejectedChans;
 
-							if (OutBunch->PacketId==AckPacketId)
-							{
-								OutBunch->ReceivedAck = 1;
-							}
-						}
-						Channel->ReceivedAcks(); //warning: May destroy Channel.
-					}
-					else
-					{
-						UE_LOG(LogNet, Warning, TEXT("UNetConnection::ReceivedPacket: null channel in OpenChannels array while processing ack. %s"), *Describe());
-					}
-				}
-			}	
+	// Disassemble and dispatch all bunches in the packet.
+	while( !Reader.AtEnd() && State!=USOCK_Closed )
+	{
+		// For demo backwards compatibility, old replays still have this bit
+		if (InternalAck && EngineNetworkProtocolVersion < EEngineNetworkVersionHistory::HISTORY_ACKS_INCLUDED_IN_HEADER)
+		{
+			const bool IsAckDummy = Reader.ReadBit() == 1u;
 		}
-		else
+
+		// Parse the bunch.
+		int32 StartPos = Reader.GetPosBits();
+		
+		// Process Received data
 		{
 			// Parse the incoming data.
 			FInBunch Bunch( this );
 			int32 IncomingStartPos		= Reader.GetPosBits();
 			uint8 bControl				= Reader.ReadBit();
-			Bunch.PacketId				= PacketId;
+			Bunch.PacketId				= InPacketId;
 			Bunch.bOpen					= bControl ? Reader.ReadBit() : 0;
 			Bunch.bClose				= bControl ? Reader.ReadBit() : 0;
 			
@@ -1572,7 +1665,7 @@ PRAGMA_ENABLE_DEPRECATION_WARNINGS
 			else if ( Bunch.bPartial )
 			{
 				// If this is an unreliable partial bunch, we simply use packet sequence since we already have it
-				Bunch.ChSequence = PacketId;
+				Bunch.ChSequence = InPacketId;
 			}
 			else
 			{
@@ -1900,8 +1993,36 @@ PRAGMA_ENABLE_DEPRECATION_WARNINGS
 	if ( !bSkipAck )
 	{
 		LastGoodPacketRealtime = FPlatformTime::Seconds();
+	}
+
+	if( !InternalAck )
+	{
+		// We always call AckSequence even if we are explicitly rejecting the packet as this updates the expected InSeq used to drive future acks.
+		if ( bSkipAck )
+		{
+			// Explicit Nak, we treat this packet as dropped but we still report it to the sending side as quickly as possible
+			PacketNotify.NakSeq( InPacketId );
+		}
+		else
+		{
+			PacketNotify.AckSeq( InPacketId );
+
+			// Keep stats happy
+			++OutTotalAcks;
+			++Driver->OutTotalAcks;
+		}
+
+		// We do want to let the other side know about the ack, so even if there are no other outgoing data when we tick the connection we will send an ackpacket.
+		TimeSensitive = 1;
+		++HasDirtyAcks;
 
-		SendAck(PacketId, true);
+		// This is to allow us to recover from hitches were we process more than FNetPacketNotify::SequenceHistoryLength packets in a row withouht sending out any packets.
+		// In most cases this allows us to recover smoothly without reporting any pessimistic naks which is the case is we overshoot our ack history.
+		if (HasDirtyAcks >= FNetPacketNotify::MaxSequenceHistoryLength)
+		{
+			UE_LOG(LogNet, Warning, TEXT("UNetConnection::ReceivedPacket - Too many received acks(%u) since last sent packet. %s "), HasDirtyAcks, *Describe());
+			FlushNet();
+		}
 	}
 }
 
@@ -1945,10 +2066,24 @@ int32 UNetConnection::WriteBitsToSendBuffer(
 	// If this is the start of the queue, make sure to add the packet id
 	if ( SendBuffer.GetNumBits() == 0 && !InternalAck )
 	{
-		SendBuffer.WriteIntWrapped( OutPacketId, MAX_PACKETID );
-		ValidateSendBuffer();
+		// Write Packet Header, before sending the packet we will go back and rewrite the data
+		WritePacketHeader(SendBuffer);
+
+		// Also write server RTT and Received rate
+		WritePacketInfo(SendBuffer);
+
+		// We do not allow the first bunch to merge with the ack data as this will "revert" the ack data.
+		AllowMerge = false;
+	
+		// Update stats for PacketIdBits and ackdata (also including the data used for packet RTT and saturation calculations)
+		int64 BitsWritten = SendBuffer.GetNumBits();
+		NumPacketIdBits += FNetPacketNotify::SequenceNumberT::SeqNumberBits;
+		NumAckBits += BitsWritten - FNetPacketNotify::SequenceNumberT::SeqNumberBits;
+
+		// Report stats to profiler
+		NETWORK_PROFILER( GNetworkProfiler.TrackSendAck( NumAckBits, this ) );
 
-		NumPacketIdBits += SendBuffer.GetNumBits();
+		ValidateSendBuffer();
 	}
 
 	// Add the bits to the queue
@@ -1972,9 +2107,6 @@ int32 UNetConnection::WriteBitsToSendBuffer(
 		case EWriteBitsDataType::Bunch:
 			NumBunchBits += SizeInBits + ExtraSizeInBits;
 			break;
-		case EWriteBitsDataType::Ack:
-			NumAckBits += SizeInBits + ExtraSizeInBits;
-			break;
 		default:
 			break;
 	}
@@ -2022,64 +2154,12 @@ TSharedPtr<FObjectReplicator> UNetConnection::CreateReplicatorForNewActorChannel
 
 void UNetConnection::PurgeAcks()
 {
-	for ( int32 i = 0; i < ResendAcks.Num(); i++ )
-	{
-		SendAck( ResendAcks[i], 0 );
-	}
-
-	ResendAcks.Empty(32);
 }
 
-
 void UNetConnection::SendAck(int32 AckPacketId, bool FirstTime/*=1*/)
 {
-	SCOPE_CYCLE_COUNTER(Stat_NetConnectionSendAck);
-
-	ValidateSendBuffer();
-
-	if( !InternalAck )
-	{
-		if( FirstTime )
-		{
-			PurgeAcks();
-			QueuedAcks.Add(AckPacketId);
-		}
-
-		FBitWriter AckData( 32, true );
-
-		AckData.WriteBit( 1 );
-		AckData.WriteIntWrapped(AckPacketId, MAX_PACKETID);
-		
-		const bool bHasServerFrameTime = Driver->IsServer() ? bLastHasServerFrameTime : ( CVarPingExcludeFrameTime.GetValueOnGameThread() > 0 ? true : false );
-
-		AckData.WriteBit( bHasServerFrameTime ? 1 : 0 );
-
-		if ( Driver->IsServer() && bHasServerFrameTime )
-		{
-			uint8 FrameTimeByte = FMath::Min( FMath::FloorToInt( FrameTime * 1000 ), 255 );
-			AckData << FrameTimeByte;
-		}
-
-		// Notify server of our current rate per second at this time
-		uint32 InKBytesPerSecond = InBytesPerSecond / 1024;
-		AckData.SerializeIntPacked( InKBytesPerSecond );
-
-		NETWORK_PROFILER( GNetworkProfiler.TrackSendAck( AckData.GetNumBits(), this ) );
-
-		WriteBitsToSendBuffer( AckData.GetData(), AckData.GetNumBits(), nullptr, 0, EWriteBitsDataType::Ack );
-
-		AllowMerge = false;
-
-		TimeSensitive = 1;
-
-		UE_LOG(LogNetTraffic, Log, TEXT("   Send ack %i"), AckPacketId);
-	}
-
-	++OutTotalAcks;
-	++Driver->OutTotalAcks;
 }
 
-
 int32 UNetConnection::SendRawBunch( FOutBunch& Bunch, bool InAllowMerge )
 {
 	ValidateSendBuffer();
@@ -2091,7 +2171,7 @@ int32 UNetConnection::SendRawBunch( FOutBunch& Bunch, bool InAllowMerge )
 
 	// Build header.
 	SendBunchHeader.Reset();
-	SendBunchHeader.WriteBit( 0 );
+
 	SendBunchHeader.WriteBit( Bunch.bOpen || Bunch.bClose );
 	if( Bunch.bOpen || Bunch.bClose )
 	{
@@ -2537,8 +2617,6 @@ void UNetConnection::Tick()
 	}
 
 	// Flush.
-	PurgeAcks();
-
 	if ( TimeSensitive || (Driver->Time - LastSendTime) > Driver->KeepAliveTime)
 	{
 		bool bHandlerHandshakeComplete = !Handler.IsValid() || Handler->IsFullyInitialized();
diff --git a/Engine/Source/Runtime/Engine/Public/Net/NetPacketNotify.h b/Engine/Source/Runtime/Engine/Public/Net/NetPacketNotify.h
new file mode 100644
--- /dev/null
+++ b/Engine/Source/Runtime/Engine/Public/Net/NetPacketNotify.h
@@ -0,0 +1,216 @@
+// Copyright 1998-2019 Epic Games, Inc. All Rights Reserved.
+
+#pragma once
+
+#include "EngineLogs.h"
+#include "HAL/UnrealMemory.h"
+#include "Logging/LogMacros.h"
+#include "Templates/IsSigned.h"
+#include "Util/ResizableCircularQueue.h"
+#include "Util/SequenceNumber.h"
+#include "Util/SequenceHistory.h"
+
+#if (UE_BUILD_SHIPPING || UE_BUILD_TEST)
+#	define UE_NET_ENABLE_PACKET_NOTIFY_LOG 0
+#else
+#	define UE_NET_ENABLE_PACKET_NOTIFY_LOG 1
+#endif 
+
+#if UE_NET_ENABLE_PACKET_NOTIFY_LOG
+#	define UE_LOG_PACKET_NOTIFY(Format, ...)  UE_LOG(LogNetTraffic, Verbose, Format, ##__VA_ARGS__)
+#else
+#	define UE_LOG_PACKET_NOTIFY(...)
+#endif
+
+#define UE_LOG_PACKET_NOTIFY_WARNING(Format, ...)  UE_LOG(LogNetTraffic, Warning, Format, ##__VA_ARGS__)
+
+struct FBitWriter;
+struct FBitReader;
+
+/** 
+	FNetPacketNotify - Drives delivery of sequence numbers, acknowledgments and notifications of delivery sequence numbers
+*/
+class FNetPacketNotify
+{
+public:
+	enum { SequenceNumberBits = 14 };
+	enum { MaxSequenceHistoryLength = 256 };
+
+	typedef TSequenceNumber<SequenceNumberBits, uint16> SequenceNumberT;
+	typedef TSequenceHistory<MaxSequenceHistoryLength> SequenceHistoryT;
+
+	struct FNotificationHeader
+	{
+		SequenceHistoryT History;
+		SIZE_T HistoryWordCount;
+		SequenceNumberT Seq;
+		SequenceNumberT AckedSeq;
+	};
+
+	/** Constructor */
+	FNetPacketNotify()
+	: AckRecord(64)
+	, WrittenHistoryWordCount(0)
+	{}
+
+	/** Init notification with expected initial sequence numbers */
+	void Init(SequenceNumberT InitialInSeq, SequenceNumberT InitialOutSeq);
+
+	/** Mark Seq as received and update current InSeq, missing sequence numbers will be marked as lost */
+	void AckSeq(SequenceNumberT Seq) { AckSeq(Seq, true); }
+
+	/** Explicitly mark Seq as not received and update current InSeq, additional missing sequence numbers will be marked as lost */
+	void NakSeq(SequenceNumberT Seq) { AckSeq(Seq, false); }
+
+	/** Increment outgoing seq number and commit data*/
+	SequenceNumberT CommitAndIncrementOutSeq();
+
+	/** Write NotificationHeader, and update outgoing ack record 
+	 if bRefresh is true we will attempt to refresh a previously written header 
+	if the size remains the same */
+	void WriteHeader(FBitWriter& Writer, bool bRefresh = false);
+	
+	/** Read header from stream */
+	bool ReadHeader(FNotificationHeader& Data, FBitReader& Reader) const;
+
+	/** Update state of PacketNotification based on received header and invoke packet notifications for received acks.
+		InFunc is a function in the format void)(FNetPacketNotify::SequenceNumberT AckedSequence, bool delivered) to handle packetNotifications.
+
+		Returns the positive delta of the incoming seq number if it is within half the sequence number space
+		Returns 0 if the received sequence number is out of the current window or if the acknowledged seq received by the remote is invalid.
+	*/
+	template<class Functor>
+	SequenceNumberT::DifferenceT Update(const FNotificationHeader& NotificationData, Functor&& InFunc)
+	{		
+		if (NotificationData.Seq > InSeq && NotificationData.AckedSeq >= OutAckSeq)
+		{
+			UE_LOG_PACKET_NOTIFY(TEXT("FNetPacketNotify::Update - Seq %u, InSeq %u"), NotificationData.Seq.Get(), InSeq.Get());
+
+			const SequenceNumberT::DifferenceT InSeqDelta = SequenceNumberT::Diff(NotificationData.Seq, InSeq);
+			
+			ProcessReceivedAcks(NotificationData, InFunc);
+
+			// accept sequence
+			InSeq = NotificationData.Seq;
+
+			return InSeqDelta;
+		}
+		else
+		{
+			return 0;
+		}
+	}
+
+	/** Get the current SequenceHistory */
+	const SequenceHistoryT& GetInSeqHistory() const { return InSeqHistory; }
+
+	/** Get the last received in sequence number */
+	SequenceNumberT GetInSeq() const { return InSeq; }
+
+	/** Get the last received sequence number that we have accepted, InAckSeq cannot be larger than InSeq */
+	SequenceNumberT GetInAckSeq() const { return InAckSeq; }
+
+	/** Get the current outgoing sequence number */
+	SequenceNumberT GetOutSeq() const { return OutSeq; }
+
+	/** Get the last outgoing sequence number acknowledged by remote */
+	SequenceNumberT GetOutAckSeq() const { return OutAckSeq; }
+
+	/** If we do have more unacknowledged sequence numbers in-flight than our maximum sendwindow we should not send more as the receiving end will not be able to detect if the sequence number has wrapped around */
+	bool CanSend() const { SequenceNumberT NextOutSeq = OutSeq; ++NextOutSeq; return NextOutSeq >= OutAckSeq; }
+
+	/** Get the current sequenceHistory length in bits */
+	SIZE_T GetCurrentSequenceHistoryLength() const;
+
+private:
+	struct FSentAckData
+	{
+		SequenceNumberT OutSeq;	// Not needed... just to verify that things work as expected
+		SequenceNumberT InAckSeq;
+	};
+	typedef TResizableCircularQueue<FSentAckData, TInlineAllocator<128>> AckRecordT;
+
+	AckRecordT AckRecord;				// Track acked seq for each sent packet to track size of ack history
+	SIZE_T WrittenHistoryWordCount;		// Bookkeeping to track if we can update data
+	SequenceNumberT WrittenInAckSeq;	// When we call CommitAndIncrementOutSequence this will be committed along with the current outgoing sequence number for bookkeeping
+
+	// Track incoming sequence data
+	SequenceHistoryT InSeqHistory;		// BitBuffer containing a bitfield describing the history of received packets
+	SequenceNumberT InSeq;				// Last sequence number received and accepted from remote
+	SequenceNumberT InAckSeq;			// Last sequence number received from remote that we have acknowledged, this is needed since we support accepting a packet but explicitly not acknowledge it as received.
+	SequenceNumberT InAckSeqAck;		// Last sequence number received from remote that we have acknowledged and also knows that the remote has received the ack, used to calculate how big our history must be
+
+	// Track outgoing sequence data
+	SequenceNumberT OutSeq;				// Outgoing sequence number
+	SequenceNumberT OutAckSeq;			// Last sequence number that we know that the remote side have received.
+
+private:
+	SequenceNumberT UpdateInAckSeqAck(SequenceNumberT::DifferenceT AckCount, SequenceNumberT AckedSeq)
+	{
+		check((SIZE_T)AckCount <= AckRecord.Count());
+
+		if (AckCount > 1)
+		{
+			AckRecord.PopNoCheck(AckCount - 1);
+		}
+		
+		FSentAckData AckData = AckRecord.PeekNoCheck();
+		AckRecord.PopNoCheck();
+
+		// verify that we have a matching sequence number
+		if (AckData.OutSeq == AckedSeq)
+		{
+			return AckData.InAckSeq;
+		}
+		else
+		{
+			UE_LOG_PACKET_NOTIFY_WARNING(TEXT("FNetPacketNotify::UpdateInAckSeqAck - Failed to find matching AckRecord for %u, (Found %u)"), AckedSeq.Get(), AckData.OutSeq.Get());
+
+			// Pessimistic view, should never occur
+			return SequenceNumberT(AckedSeq.Get() - MaxSequenceHistoryLength);
+		}
+	}
+
+	template<class Functor>
+	void ProcessReceivedAcks(const FNotificationHeader& NotificationData, Functor&& InFunc)
+	{
+		if (NotificationData.AckedSeq > OutAckSeq)
+		{
+			UE_LOG_PACKET_NOTIFY(TEXT("Notification::ProcessReceivedAcks - AckedSeq: %u, OutAckSeq: %u"), NotificationData.AckedSeq.Get(), OutAckSeq.Get());
+
+			SequenceNumberT::DifferenceT AckCount = SequenceNumberT::Diff(NotificationData.AckedSeq, OutAckSeq);
+
+			// Update InAckSeqAck used to track the needed number of bits to transmit our ack history
+			InAckSeqAck = UpdateInAckSeqAck(AckCount, NotificationData.AckedSeq);
+
+			// ExpectedAck = OutAckSeq + 1
+			SequenceNumberT CurrentAck(OutAckSeq);
+			++CurrentAck;
+
+			// Everything not found in the history buffer is treated as lost
+			while (AckCount > (SequenceNumberT::DifferenceT)(SequenceHistoryT::Size))
+			{
+				--AckCount;
+				UE_LOG_PACKET_NOTIFY_WARNING(TEXT("Notification::ProcessReceivedAcks Seq: %u - IsAck: 0 HistoryIndex: N/A"), CurrentAck.Get());
+				InFunc(CurrentAck, false);
+				++CurrentAck;
+			}
+
+			// For sequence numbers contained in the history we lookup the delivery status from the history
+			while (AckCount > 0)
+			{
+				--AckCount;
+				UE_LOG_PACKET_NOTIFY(TEXT("Notification::ProcessReceivedAcks Seq: %u - IsAck: %u HistoryIndex: %u"), CurrentAck.Get(), NotificationData.History.IsDelivered(AckCount) ? 1u : 0u, AckCount);
+				InFunc(CurrentAck, NotificationData.History.IsDelivered(AckCount));
+				++CurrentAck;
+			}
+			OutAckSeq = NotificationData.AckedSeq;
+		}
+	}
+
+	void AckSeq(SequenceNumberT AckedSeq, bool IsAck);
+
+#if WITH_DEV_AUTOMATION_TESTS
+	friend struct FNetPacketNotifyTestUtil;
+#endif
+};
diff --git a/Engine/Source/Runtime/Engine/Public/Net/Util/ResizableCircularQueue.h b/Engine/Source/Runtime/Engine/Public/Net/Util/ResizableCircularQueue.h
new file mode 100644
--- /dev/null
+++ b/Engine/Source/Runtime/Engine/Public/Net/Util/ResizableCircularQueue.h
@@ -0,0 +1,207 @@
+// Copyright 1998-2019 Epic Games, Inc. All Rights Reserved.
+
+#pragma once
+
+#include "CoreTypes.h"
+#include "Misc/AssertionMacros.h"
+#include "Containers/Array.h"
+
+/**
+ * Simple ResisableCircularQueue for trivially copyable types
+ * Relies on unsigned arithmetics and ever increasing head and tail indices to avoid having to store 
+ * an extra element or maintain explicit empty state.
+ * Capacity most be power of two and it supports growing
+ */
+class FDefaultAllocator;
+
+template<typename T, typename AllocatorT = FDefaultAllocator>
+class TResizableCircularQueue
+{
+public:
+	typedef T ElementT;
+
+	/** Construct Empty Queue with the given initial Capacity, Capacity must be a power of two since we rely on unsigned arithmetics for wraparound */
+	explicit TResizableCircularQueue(SIZE_T InitialCapacity);
+
+	/** Returns true if the queue is empty */
+	bool IsEmpty() const { return Head == Tail; }
+
+	/** Gets the number of elements in the queue. */
+	SIZE_T Count() const { return Head - Tail; }
+
+	/** Current allocated Capacity */
+	SIZE_T AllocatedCapacity() const { return Storage.Num(); }
+
+	/** Push single element to the back of the Queue */
+	void Enqueue(const ElementT& SrcData);
+
+	/** Pop elements from the front of the queue */
+	void Pop();
+
+	/** Pop Count elements from the front of the queue */
+	void Pop(SIZE_T Count);
+
+	/**  Unchecked version, Pop a single element from the front of the queue */
+	void PopNoCheck() { ++Tail; }
+
+	/** Unchecked version, Pop Count elements from the front of the queue */
+	void PopNoCheck(SIZE_T Count) { Tail += Count; }
+
+	/** Peek with the given offset from the front of the queue */
+	const ElementT& PeekAtOffset(SIZE_T Offset = 0) const { check(Offset < Count()); return PeekAtOffsetNoCheck(Offset); }
+
+	/** Peek at front element */
+	const ElementT& Peek() const { return PeekAtOffset(0); }
+
+	/**  Unchecked version, Peek with the given offset from the front of the queue */
+	const ElementT& PeekAtOffsetNoCheck(SIZE_T Offset = 0) const { return Storage.GetData()[(Tail + Offset) & IndexMask]; }
+
+	/** Peek at front element with no check */
+	const ElementT& PeekNoCheck() const { return PeekAtOffsetNoCheck(0); }
+
+	/**  Trim memory usage to the next power of two for the current size */
+	void Trim();
+
+	/** Empty queue without releasing memory */
+	void Reset() { Head = Tail = 0u; }
+
+	/** Empty queue and release memory */
+	void Empty() { Reset(); Storage.Empty(); IndexMask = IndexT(-1); }
+	
+private:
+
+#if WITH_DEV_AUTOMATION_TESTS
+	friend struct FResizableCircularQueueTestUtil;
+#endif
+
+	// Resize buffer maintaining validity of stored data.
+	void SetCapacity(SIZE_T NewCapacity);
+
+	// TIsTriviallyCopyable (a.k.a. std::is_trivially_copyable) should be used but
+	// TIsTriviallyCopyable is not provided by the core module at the moment. Core's
+	// implementation of TIsTrivial is actually equivalent to std::is_trivially_copyable
+	// since std::is_trivial<T> requires T to have a trivial constructor but TIsTrivial
+	// doesn't
+	static_assert(TIsTrivial<ElementT>::Value, "ResizableCircularQueue only supports trivially copyable elements");
+
+	typedef uint32 IndexT;
+	typedef TArray<ElementT, AllocatorT> StorageT;
+
+	IndexT Head;
+	IndexT Tail;
+
+	IndexT IndexMask;
+	StorageT Storage;
+};
+
+
+template<typename T, typename AllocatorT>
+TResizableCircularQueue<T, AllocatorT>::TResizableCircularQueue(SIZE_T InitialCapacity)
+: Head(0u)
+, Tail(0u)
+, IndexMask(0u)
+{
+	// Capacity should be power of two, it will be rounded up but lets warn the user anyway in debug builds.
+	checkSlow((InitialCapacity & (InitialCapacity - 1)) == 0);
+	if (InitialCapacity > 0)
+	{
+		SetCapacity(InitialCapacity);
+	}
+}
+
+template<typename T, typename AllocatorT>
+void TResizableCircularQueue<T, AllocatorT>::Enqueue(const ElementT& SrcData)
+{ 
+	const SIZE_T RequiredCapacity = Count() + 1;
+	if (RequiredCapacity > AllocatedCapacity())
+	{
+		// Capacity must be power of two
+		SetCapacity(RequiredCapacity);
+	}
+
+	T* DstData = Storage.GetData();
+	const IndexT MaskedIndex = Head & IndexMask;
+	DstData[MaskedIndex] = SrcData;
+	++Head;		
+}
+
+template<typename T, typename AllocatorT>
+void
+TResizableCircularQueue<T, AllocatorT>::Pop(SIZE_T PopCount)
+{
+	if (ensure(Count() > PopCount))
+	{
+		PopNoCheck(PopCount);
+	}
+}
+
+template<typename T, typename AllocatorT>
+void TResizableCircularQueue<T, AllocatorT>::Pop()
+{
+	if (ensure(Count() > 0))
+	{
+		PopNoCheck();
+	}
+}
+
+template<typename T, typename AllocatorT>
+void TResizableCircularQueue<T, AllocatorT>::SetCapacity(SIZE_T RequiredCapacity)
+{
+	SIZE_T NewCapacity = FMath::RoundUpToPowerOfTwo(RequiredCapacity);
+
+	if ((NewCapacity == Storage.Num()) || (NewCapacity < Count()))
+	{
+		return;
+	}
+
+	if (Storage.Num() > 0)
+	{
+		StorageT NewStorage;
+		NewStorage.AddUninitialized(NewCapacity);
+		
+		// copy data to new storage
+		const IndexT MaskedTail = Tail & IndexMask;
+		const IndexT MaskedHead = Head & IndexMask;
+
+		const ElementT* SrcData = Storage.GetData();
+		const SIZE_T SrcCapacity = Storage.Num();
+		const SIZE_T SrcSize = Count();
+		ElementT* DstData = NewStorage.GetData();
+
+		if (MaskedTail >= MaskedHead)
+		{
+			const SIZE_T CopyCount = (SrcCapacity - MaskedTail);
+			FPlatformMemory::Memcpy(DstData, SrcData + MaskedTail, CopyCount * sizeof(ElementT));
+			FPlatformMemory::Memcpy(DstData + CopyCount, SrcData, MaskedHead * sizeof(ElementT));
+		}
+		else
+		{
+			FPlatformMemory::Memcpy(DstData, SrcData + MaskedTail, SrcSize * sizeof(ElementT));
+		}
+
+		this->Storage = MoveTemp(NewStorage);
+		IndexMask = NewCapacity - 1;
+		Tail = 0u;
+		Head = SrcSize;
+	}
+	else
+	{
+		IndexMask = NewCapacity - 1;
+		Storage.AddUninitialized(NewCapacity);
+	}
+}
+
+template<typename T, typename AllocatorT>
+void TResizableCircularQueue<T, AllocatorT>::Trim()
+{
+	if (IsEmpty())
+	{
+		Empty();
+	}
+	else
+	{
+		SetCapacity(Count());
+	}
+}
+
+
diff --git a/Engine/Source/Runtime/Engine/Public/Net/Util/SequenceHistory.h b/Engine/Source/Runtime/Engine/Public/Net/Util/SequenceHistory.h
new file mode 100644
--- /dev/null
+++ b/Engine/Source/Runtime/Engine/Public/Net/Util/SequenceHistory.h
@@ -0,0 +1,126 @@
+// Copyright 1998-2019 Epic Games, Inc. All Rights Reserved.
+
+#pragma once
+
+#include "CoreTypes.h"
+#include "HAL/UnrealMemory.h"
+#include "Misc/AssertionMacros.h"
+#include "Templates/IsSigned.h"
+#include "Serialization/BitWriter.h"
+#include "Serialization/BitReader.h"
+
+/** Util class to manage history of received sequence numbers */
+template <SIZE_T HistorySize>
+class TSequenceHistory
+{
+public:	
+	typedef uint32 WordT;
+	
+	static const SIZE_T BitsPerWord = sizeof(WordT) * 8;
+	static const SIZE_T WordCount = HistorySize / BitsPerWord;
+	static const SIZE_T Size = HistorySize;
+
+	static_assert(HistorySize > 0, "HistorySize must be > 0");
+	static_assert(HistorySize % BitsPerWord == 0, "InMaxHistorySize must be a modulo of the wordsize");
+
+public:	
+	TSequenceHistory();
+
+#if WITH_DEV_AUTOMATION_TESTS
+	explicit TSequenceHistory(WordT Value, SIZE_T Count = 1);
+	WordT* Data() { return &Storage[0]; }
+#endif
+	/** Reset */
+	void Reset();
+
+	/** Store delivery status, oldest will be dropped */
+	void AddDeliveryStatus(bool Delivered);
+
+	/** Query the status of a specific index, index 0 is last stored status */
+	bool IsDelivered(SIZE_T Index) const;
+	
+	bool operator==(const TSequenceHistory& Other) const { return FMemory::Memcmp(Storage, Other.Storage, WordCount * sizeof (WordT)) == 0; }
+
+	bool operator!=(const TSequenceHistory& Other) const { return FMemory::Memcmp(Storage, Other.Storage, WordCount * sizeof (WordT)) != 0; }
+
+	/** Write history to BitStream */
+	void Write(FBitWriter& Writer, SIZE_T NumWords) const;
+
+	/** Read history from BitStream */
+	void Read(FBitReader& Reader, SIZE_T NumWords);
+
+private:
+	WordT Storage[WordCount];
+};
+
+template <SIZE_T HistorySize>
+TSequenceHistory<HistorySize>::TSequenceHistory()
+{
+	Reset();
+}
+
+#if WITH_DEV_AUTOMATION_TESTS
+template <SIZE_T HistorySize>
+TSequenceHistory<HistorySize>::TSequenceHistory(WordT Value, SIZE_T Count)
+{
+	Reset();
+	for (SIZE_T CurrentWordIt = 0; CurrentWordIt < Count; ++CurrentWordIt)
+	{
+		Storage[CurrentWordIt] = Value;
+	}	
+}
+#endif 
+
+template <SIZE_T HistorySize>
+void TSequenceHistory<HistorySize>::Reset()
+{
+	FPlatformMemory::Memset(&Storage[0], 0, WordCount * sizeof(WordT));
+}
+
+template <SIZE_T HistorySize>
+void TSequenceHistory<HistorySize>::AddDeliveryStatus(bool Delivered)
+{
+	WordT Carry = Delivered ? 1u : 0u;
+	const WordT ValueMask = 1u << (BitsPerWord - 1);
+	
+	for (SIZE_T CurrentWordIt = 0; CurrentWordIt < WordCount; ++CurrentWordIt)
+	{
+		const WordT OldValue = Carry;
+		
+		// carry over highest bit in each word to the next word
+		Carry = (Storage[CurrentWordIt] & ValueMask) >> (BitsPerWord - 1);
+		Storage[CurrentWordIt] = (Storage[CurrentWordIt] << 1u) | OldValue;
+	}
+}
+
+template <SIZE_T HistorySize>
+bool TSequenceHistory<HistorySize>::IsDelivered(SIZE_T Index) const
+{
+	check(Index < Size);
+
+	const SIZE_T WordIndex = Index / BitsPerWord;
+	const WordT WordMask = (WordT(1) << (Index & (BitsPerWord - 1)));
+	
+	return (Storage[WordIndex] & WordMask) != 0u;
+}
+
+template <SIZE_T HistorySize>
+void TSequenceHistory<HistorySize>::Write(FBitWriter& Writer, SIZE_T NumWords) const
+{
+	NumWords = FPlatformMath::Min(NumWords, WordCount);
+	for (SIZE_T CurrentWordIt = 0; CurrentWordIt < NumWords; ++CurrentWordIt)
+	{
+		WordT temp = Storage[CurrentWordIt];
+		Writer << temp;
+	}
+}
+
+template <SIZE_T HistorySize>
+void TSequenceHistory<HistorySize>::Read(FBitReader& Reader, SIZE_T NumWords)
+{
+	NumWords = FPlatformMath::Min(NumWords, WordCount);
+	for (SIZE_T CurrentWordIt = 0; CurrentWordIt < NumWords; ++CurrentWordIt)
+	{
+		Reader << Storage[CurrentWordIt];
+	}
+}
diff --git a/Engine/Source/Runtime/Engine/Public/Net/Util/SequenceNumber.h b/Engine/Source/Runtime/Engine/Public/Net/Util/SequenceNumber.h
new file mode 100644
--- /dev/null
+++ b/Engine/Source/Runtime/Engine/Public/Net/Util/SequenceNumber.h
@@ -0,0 +1,70 @@
+// Copyright 1998-2019 Epic Games, Inc. All Rights Reserved.
+
+#pragma once
+
+#include "CoreTypes.h"
+#include "HAL/UnrealMemory.h"
+#include "Misc/AssertionMacros.h"
+#include "Templates/IsSigned.h"
+
+/** Helper class to work with sequence numbers */
+template <SIZE_T NumBits, typename SequenceType>
+class TSequenceNumber
+{
+	static_assert(TIsSigned<SequenceType>::Value == false, "The base type for sequence numbers must be unsigned");
+
+public:
+	typedef SequenceType SequenceT;
+	typedef int32 DifferenceT;
+
+	// Constants
+	enum { SeqNumberBits = NumBits };
+	enum { SeqNumberCount = SequenceT(1) << NumBits };
+	enum { SeqNumberHalf = SequenceT(1) << (NumBits - 1) };
+	enum { SeqNumberMax = SeqNumberCount - 1u };
+	enum { SeqNumberMask = SeqNumberMax };
+
+	/** Default constructor */
+	TSequenceNumber() : Value(0u) {}
+
+	/** Constructor with given value */
+	TSequenceNumber(SequenceT ValueIn) : Value(ValueIn & SeqNumberMask) {}
+	
+	/** Get Current Value */	
+	SequenceT Get() const { return Value; }
+
+	/** Diff between sequence numbers (A - B) only valid if (A - B) < SeqNumberHalf */
+	static DifferenceT Diff(TSequenceNumber A, TSequenceNumber B);
+	
+	/** return true if this is > Other, this is only considered to be the case if (A - B) < SeqNumberHalf since we have to be able to detect wraparounds */
+	bool operator>(const TSequenceNumber& Other) const { return (Value != Other.Value) && (((Value - Other.Value) & SeqNumberMask) < SeqNumberHalf); }
+
+	/** Check if this is >= Other, See above */
+	bool operator>=(const TSequenceNumber& Other) const { return ((Value - Other.Value) & SeqNumberMask) < SeqNumberHalf; }
+
+	/** Equals, NOTE that sequence numbers wrap around so 0 == 0 + SequenceNumberCount */
+	bool operator==(const TSequenceNumber& Other) const { return Value == Other.Value; }
+
+	bool operator!=(const TSequenceNumber& Other) const { return Value != Other.Value; }
+
+	/** Pre-increment and wrap around */
+	TSequenceNumber& operator++() { Increment(1u); return *this; }
+	
+	/** Post-increment and wrap around */
+	TSequenceNumber operator++(int) { TSequenceNumber Tmp(*this); Increment(1u); return Tmp; }
+
+private:
+	void Increment(SequenceT InValue) { *this = TSequenceNumber(Value + InValue); }
+	SequenceT Value;
+};
+
+template <SIZE_T NumBits, typename SequenceType>
+typename TSequenceNumber<NumBits, SequenceType>::DifferenceT TSequenceNumber<NumBits, SequenceType>::Diff(TSequenceNumber A, TSequenceNumber B) 
+{ 
+	const SIZE_T ShiftValue = sizeof(DifferenceT)*8 - NumBits;
+
+	const SequenceT ValueA = A.Value;
+	const SequenceT ValueB = B.Value;
+
+	return (DifferenceT)((ValueA - ValueB) << ShiftValue) >> ShiftValue;
+};
diff --git a/Engine/Source/Runtime/MovieSceneTracks/Private/Evaluation/MovieScenePrimitiveMaterialTemplate.cpp b/Engine/Source/Runtime/MovieSceneTracks/Private/Evaluation/MovieScenePrimitiveMaterialTemplate.cpp
--- a/Engine/Source/Runtime/MovieSceneTracks/Private/Evaluation/MovieScenePrimitiveMaterialTemplate.cpp
+++ b/Engine/Source/Runtime/MovieSceneTracks/Private/Evaluation/MovieScenePrimitiveMaterialTemplate.cpp
@@ -5,6 +5,7 @@
 #include "Sections/MovieScenePrimitiveMaterialSection.h"
 #include "UObject/StrongObjectPtr.h"
 #include "Materials/MaterialInterface.h"
+#include "Materials/MaterialInstanceDynamic.h"
 #include "Components/PrimitiveComponent.h"
 
 struct FSetMaterialToken : IMovieScenePreAnimatedToken
@@ -57,9 +58,17 @@ struct FPrimitiveMaterialExecToken : IMovieSceneExecutionToken
 			UPrimitiveComponent* PrimitiveComponent = Cast<UPrimitiveComponent>(WeakObject.Get());
 			if (PrimitiveComponent && MaterialIndex >= 0 && MaterialIndex < PrimitiveComponent->GetNumMaterials())
 			{
-				UMaterialInterface* ExistingMaterial = PrimitiveComponent->GetMaterial(MaterialIndex);
-				Player.SavePreAnimatedState(*PrimitiveComponent, TMovieSceneAnimTypeID<FPrimitiveMaterialExecToken>(), FSetMaterialTokenProducer(MaterialIndex, ExistingMaterial));
+				UMaterialInterface*        ExistingMaterial = PrimitiveComponent->GetMaterial(MaterialIndex);
+				UMaterialInstanceDynamic*  MaterialInstance = Cast<UMaterialInstanceDynamic>(ExistingMaterial);
 
+				if (MaterialInstance && MaterialInstance->Parent && MaterialInstance->Parent == NewMaterial)
+				{
+					// Do not re-assign materials when a dynamic instance is already assigned with the same parent (since that's basically the same material, just with animated parameters)
+					// This is required for supporting material switchers alongside parameter tracks
+					continue;
+				}
+				
+				Player.SavePreAnimatedState(*PrimitiveComponent, TMovieSceneAnimTypeID<FPrimitiveMaterialExecToken>(), FSetMaterialTokenProducer(MaterialIndex, ExistingMaterial));
 				if (NewMaterial != ExistingMaterial)
 				{
 					PrimitiveComponent->SetMaterial(MaterialIndex, NewMaterial);
diff --git a/Engine/Source/Runtime/MovieSceneTracks/Private/Tracks/MovieSceneMaterialTrack.cpp b/Engine/Source/Runtime/MovieSceneTracks/Private/Tracks/MovieSceneMaterialTrack.cpp
--- a/Engine/Source/Runtime/MovieSceneTracks/Private/Tracks/MovieSceneMaterialTrack.cpp
+++ b/Engine/Source/Runtime/MovieSceneTracks/Private/Tracks/MovieSceneMaterialTrack.cpp
@@ -3,7 +3,7 @@
 #include "Tracks/MovieSceneMaterialTrack.h"
 #include "MovieSceneCommonHelpers.h"
 #include "Evaluation/MovieSceneParameterTemplate.h"
-
+#include "Evaluation/MovieSceneEvaluationTrack.h"
 
 UMovieSceneMaterialTrack::UMovieSceneMaterialTrack(const FObjectInitializer& ObjectInitializer)
 	: Super(ObjectInitializer)
@@ -105,6 +105,11 @@ FMovieSceneEvalTemplatePtr UMovieSceneComponentMaterialTrack::CreateTemplateForS
 	return FMovieSceneComponentMaterialSectionTemplate(*CastChecked<UMovieSceneParameterSection>(&InSection), *this);
 }
 
+void UMovieSceneComponentMaterialTrack::PostCompile(FMovieSceneEvaluationTrack& OutTrack, const FMovieSceneTrackCompilerArgs& Args) const
+{
+	OutTrack.SetEvaluationPriority(EvaluationPriority);
+}
+
 
 #if WITH_EDITORONLY_DATA
 FText UMovieSceneComponentMaterialTrack::GetDefaultDisplayName() const
diff --git a/Engine/Source/Runtime/MovieSceneTracks/Private/Tracks/MovieScenePrimitiveMaterialTrack.cpp b/Engine/Source/Runtime/MovieSceneTracks/Private/Tracks/MovieScenePrimitiveMaterialTrack.cpp
--- a/Engine/Source/Runtime/MovieSceneTracks/Private/Tracks/MovieScenePrimitiveMaterialTrack.cpp
+++ b/Engine/Source/Runtime/MovieSceneTracks/Private/Tracks/MovieScenePrimitiveMaterialTrack.cpp
@@ -1,8 +1,10 @@
 // Copyright 1998-2019 Epic Games, Inc. All Rights Reserved.
 
 #include "Tracks/MovieScenePrimitiveMaterialTrack.h"
+#include "Tracks/MovieSceneMaterialTrack.h"
 #include "Sections/MovieScenePrimitiveMaterialSection.h"
 #include "Evaluation/MovieScenePrimitiveMaterialTemplate.h"
+#include "Evaluation/MovieSceneEvaluationTrack.h"
 
 
 UMovieScenePrimitiveMaterialTrack::UMovieScenePrimitiveMaterialTrack(const FObjectInitializer& ObjInit)
@@ -23,3 +25,9 @@ FMovieSceneEvalTemplatePtr UMovieScenePrimitiveMaterialTrack::CreateTemplateForS
 {
 	return FMovieScenePrimitiveMaterialTemplate(*CastChecked<UMovieScenePrimitiveMaterialSection>(&InSection), *this);
 }
+
+void UMovieScenePrimitiveMaterialTrack::PostCompile(FMovieSceneEvaluationTrack& OutTrack, const FMovieSceneTrackCompilerArgs& Args) const
+{
+	// Must evaluate before material parameter tracks
+	OutTrack.SetEvaluationPriority(UMovieSceneComponentMaterialTrack::EvaluationPriority + 1);
+}
\ No newline at end of file
diff --git a/Engine/Source/Runtime/MovieSceneTracks/Public/Tracks/MovieSceneMaterialTrack.h b/Engine/Source/Runtime/MovieSceneTracks/Public/Tracks/MovieSceneMaterialTrack.h
--- a/Engine/Source/Runtime/MovieSceneTracks/Public/Tracks/MovieSceneMaterialTrack.h
+++ b/Engine/Source/Runtime/MovieSceneTracks/Public/Tracks/MovieSceneMaterialTrack.h
@@ -69,7 +69,10 @@ public:
 
 	// UMovieSceneTrack interface
 
+	static const uint16 EvaluationPriority = 1000;
+
 	virtual FMovieSceneEvalTemplatePtr CreateTemplateForSection(const UMovieSceneSection& InSection) const override;
+	virtual void PostCompile(FMovieSceneEvaluationTrack& OutTrack, const FMovieSceneTrackCompilerArgs& Args) const override;
 	virtual FName GetTrackName() const { return FName( *FString::FromInt(MaterialIndex) ); }
 
 #if WITH_EDITORONLY_DATA
diff --git a/Engine/Source/Runtime/MovieSceneTracks/Public/Tracks/MovieScenePrimitiveMaterialTrack.h b/Engine/Source/Runtime/MovieSceneTracks/Public/Tracks/MovieScenePrimitiveMaterialTrack.h
--- a/Engine/Source/Runtime/MovieSceneTracks/Public/Tracks/MovieScenePrimitiveMaterialTrack.h
+++ b/Engine/Source/Runtime/MovieSceneTracks/Public/Tracks/MovieScenePrimitiveMaterialTrack.h
@@ -21,4 +21,5 @@ public:
 	/*~ UMovieSceneTrack interface */
 	virtual UMovieSceneSection* CreateNewSection() override;
 	virtual FMovieSceneEvalTemplatePtr CreateTemplateForSection(const UMovieSceneSection& InSection) const override;
+	virtual void PostCompile(FMovieSceneEvaluationTrack& OutTrack, const FMovieSceneTrackCompilerArgs& Args) const override;
 };
diff --git a/Engine/Source/Runtime/Renderer/Private/SceneRendering.h b/Engine/Source/Runtime/Renderer/Private/SceneRendering.h
--- a/Engine/Source/Runtime/Renderer/Private/SceneRendering.h
+++ b/Engine/Source/Runtime/Renderer/Private/SceneRendering.h
@@ -1263,7 +1263,12 @@ public:
 
 	/** Custom Data Memstack functions.	*/
 	FORCEINLINE FMemStackBase& GetCustomDataGlobalMemStack() { return PrimitiveCustomDataMemStack[0]; }
-	FORCEINLINE FMemStackBase& AllocateCustomDataMemStack() { return *new(PrimitiveCustomDataMemStack) FMemStackBase(0); }
+	FORCEINLINE FMemStackBase& AllocateCustomDataMemStack() 
+	{ 
+		// Don't reallocate since we keep references in FRelevancePacket.
+		check(PrimitiveCustomDataMemStack.GetSlack() > 0); 
+		return *new(PrimitiveCustomDataMemStack) FMemStackBase(0);
+	}
 
 private:
 	// Cache of TEXTUREGROUP_World to create view's samplers on render thread.
diff --git a/Engine/Source/Runtime/Renderer/Private/SceneVisibility.cpp b/Engine/Source/Runtime/Renderer/Private/SceneVisibility.cpp
--- a/Engine/Source/Runtime/Renderer/Private/SceneVisibility.cpp
+++ b/Engine/Source/Runtime/Renderer/Private/SceneVisibility.cpp
@@ -2149,12 +2149,6 @@ static void ComputeAndMarkRelevanceForViewParallel(
 
 	bool WillExecuteInParallel = FApp::ShouldUseThreadingForPerformance() && CVarParallelInitViews.GetValueOnRenderThread() > 0;
 
-	if (WillExecuteInParallel)
-	{
-		// We must reserve to prevent realloc otherwise it will cause memory leak if WillExecuteInParallel == true
-		View.PrimitiveCustomDataMemStack.Reserve(View.PrimitiveCustomDataMemStack.Num() + FMath::TruncToInt((float)NumMesh / (float)FRelevancePrimSet<int32>::MaxInputPrims + 1.0f));
-	}
-
 	{
 		FSceneSetBitIterator BitIt(View.PrimitiveVisibilityMap);
 		if (BitIt)
@@ -2871,6 +2865,11 @@ void FSceneRenderer::ComputeViewVisibility(FRHICommandListImmediate& RHICmdList)
 
 		View.PrimitivesCustomData.Init(nullptr, Scene->Primitives.Num());
 		View.PrimitivesWithCustomData.Reserve(Scene->Primitives.Num());
+
+		// We must reserve to prevent realloc otherwise it will cause memory leak if we Execute In Parallel
+		bool WillExecuteInParallel = FApp::ShouldUseThreadingForPerformance() && CVarParallelInitViews.GetValueOnRenderThread() > 0;
+		View.PrimitiveCustomDataMemStack.Reserve(WillExecuteInParallel ? FMath::CeilToInt(((float)View.PrimitiveVisibilityMap.Num() / (float)FRelevancePrimSet<int32>::MaxInputPrims)) + 1 : 1);
+
 		View.AllocateCustomDataMemStack();
 
 		View.VisibleLightInfos.Empty(Scene->Lights.GetMaxIndex());
