Store UActorChannel ActorReplicator as a SharedPtr instead of a raw pointer.
This should hopefully prevent cases where the Object Replicator get's destroyed silently.

diff --git a/Engine/Source/Runtime/Engine/Classes/Engine/ActorChannel.h b/Engine/Source/Runtime/Engine/Classes/Engine/ActorChannel.h
--- a/Engine/Source/Runtime/Engine/Classes/Engine/ActorChannel.h
+++ b/Engine/Source/Runtime/Engine/Classes/Engine/ActorChannel.h
@@ -83,7 +83,7 @@ public:
 	bool GetSkipRoleSwap() const { return !!bSkipRoleSwap; }
 	void SetSkipRoleSwap(const bool bShouldSkip) { bSkipRoleSwap = bShouldSkip; }
 
-	FObjectReplicator* ActorReplicator;
+	TSharedPtr<FObjectReplicator> ActorReplicator;
 
 	TMap< UObject*, TSharedRef< FObjectReplicator > > ReplicationMap;
 
diff --git a/Engine/Source/Runtime/Engine/Private/DataChannel.cpp b/Engine/Source/Runtime/Engine/Private/DataChannel.cpp
--- a/Engine/Source/Runtime/Engine/Private/DataChannel.cpp
+++ b/Engine/Source/Runtime/Engine/Private/DataChannel.cpp
@@ -1778,7 +1778,7 @@ void UActorChannel::CleanupReplicators( const bool bKeepReplicators )
 
 	ReplicationMap.Empty();
 
-	ActorReplicator = NULL;
+	ActorReplicator.Reset();
 }
 
 static TAutoConsoleVariable<int32> CVarRelinkMappedReferences( TEXT( "net.RelinkMappedReferences" ), 1, TEXT( "" ) );
@@ -2106,7 +2106,7 @@ void UActorChannel::SetChannelActor( AActor* InActor )
 		check( !ReplicationMap.Contains( Actor ) );
 
 		// Create the actor replicator, and store a quick access pointer to it
-		ActorReplicator = &FindOrCreateReplicator( Actor ).Get();
+		ActorReplicator = FindOrCreateReplicator( Actor );
 
 		// Remove from connection's dormancy lists
 		Connection->Driver->GetNetworkObjectList().MarkActive(Actor, Connection, Connection->Driver);
@@ -2699,11 +2699,12 @@ int64 UActorChannel::ReplicateActor()
 	// If our Actor is PendingKill, that's bad. It means that somehow it wasn't properly removed
 	// from the NetDriver or ReplicationDriver.
 	// TODO: Maybe notify the NetDriver / RepDriver about this, and have the channel close?
-	else if (Actor->IsPendingKill())
+	else if (Actor->IsPendingKillOrUnreachable())
 	{
 		bActorIsPendingKill = true;
+		ActorReplicator.Reset();
 		FString Error(FString::Printf(TEXT("ReplicateActor called with PendingKill Actor! %s"), *Describe()));
-		UE_CLOG(!bActorIsPendingKill, LogNet, Log, TEXT("%s"), *Error);
+		UE_LOG(LogNet, Log, TEXT("%s"), *Error);
 		ensureMsgf(false, TEXT("%s"), *Error);
 		return 0;
 	}
@@ -2878,14 +2879,17 @@ int64 UActorChannel::ReplicateActor()
 		}
 
 		// Look for deleted subobjects
+		FObjectReplicator* LocalActorReplicator = ActorReplicator.Get();
 		for (auto RepComp = ReplicationMap.CreateIterator(); RepComp; ++RepComp)
 		{
-			if (!RepComp.Value()->GetWeakObjectPtr().IsValid())
+			TSharedRef<FObjectReplicator>& LocalReplicator = RepComp.Value();
+
+			if (!LocalReplicator->GetWeakObjectPtr().IsValid())
 			{
-				if (RepComp.Value()->ObjectNetGUID.IsValid())
+				if (LocalReplicator->ObjectNetGUID.IsValid())
 				{
 					// Write a deletion content header:
-					WriteContentBlockForSubObjectDelete(Bunch, RepComp.Value()->ObjectNetGUID);
+					WriteContentBlockForSubObjectDelete(Bunch, LocalReplicator->ObjectNetGUID);
 
 					WroteSomethingImportant = true;
 					Bunch.bReliable = true;
@@ -2895,7 +2899,15 @@ int64 UActorChannel::ReplicateActor()
 					UE_LOG(LogNetTraffic, Error, TEXT("Unable to write subobject delete for (%s), object replicator has invalid NetGUID"), *GetPathNameSafe(Actor));
 				}
 
-				RepComp.Value()->CleanUp();
+				// The only way this case would be possible is if someone tried destroying the Actor as a part of
+				// a Subobject's Pre / Post replication, during Replicate Subobjects, or OnSerializeNewActor.
+				// All of those are bad.
+				if (!ensureMsgf(LocalActorReplicator != &LocalReplicator.Get(), TEXT("UActorChannel::ReplicateActor: Actor was deleting during replication: %s"), *Describe()))
+				{
+					ActorReplicator.Reset();
+				}
+
+				LocalReplicator->CleanUp();
 				RepComp.RemoveCurrent();
 			}
 		}
@@ -3663,7 +3675,7 @@ FNetFieldExportGroup* UActorChannel::GetNetFieldExportGroupForClassNetCache( UCl
 
 FObjectReplicator & UActorChannel::GetActorReplicationData()
 {
-	check(ActorReplicator != nullptr);
+	// TSharedPtr will do a check before dereference, so no need to explicitly check here.
 	return *ActorReplicator;
 }
 
diff --git a/Engine/Source/Runtime/Engine/Private/DemoNetDriver.cpp b/Engine/Source/Runtime/Engine/Private/DemoNetDriver.cpp
--- a/Engine/Source/Runtime/Engine/Private/DemoNetDriver.cpp
+++ b/Engine/Source/Runtime/Engine/Private/DemoNetDriver.cpp
@@ -3853,14 +3853,16 @@ void UDemoNetDriver::FinalizeFastForward( const double StartTime )
 				continue;
 			}
 
-			const FObjectReplicator* const ActorReplicator = ActorChannel->ActorReplicator;
-			if (Actor->IsNetStartupActor() && ActorReplicator)
+			if (const FObjectReplicator* const ActorReplicator = ActorChannel->ActorReplicator.Get())
 			{
-				FReceivingRepState* ReceivingRepState = ActorReplicator->RepState->GetReceivingRepState();
-				FRepShadowDataBuffer ShadowData(ReceivingRepState->StaticBuffer.GetData());
-				FConstRepObjectDataBuffer ActorData(Actor);
+				if (Actor->IsNetStartupActor())
+				{
+					FReceivingRepState* ReceivingRepState = ActorReplicator->RepState->GetReceivingRepState();
+					FRepShadowDataBuffer ShadowData(ReceivingRepState->StaticBuffer.GetData());
+					FConstRepObjectDataBuffer ActorData(Actor);
 
-				ActorReplicator->RepLayout->DiffProperties(&(ReceivingRepState->RepNotifies), ShadowData, ActorData, EDiffPropertiesFlags::Sync);
+					ActorReplicator->RepLayout->DiffProperties(&(ReceivingRepState->RepNotifies), ShadowData, ActorData, EDiffPropertiesFlags::Sync);
+				}
 			}
 		}
 	}
@@ -4599,7 +4601,7 @@ bool UDemoNetDriver::FastForwardLevels(const FGotoResult& GotoResult)
 			}
 
 			ChannelsToUpdate.Add(ActorChannel);
-			if (const FObjectReplicator* const ActorReplicator = ActorChannel->ActorReplicator)
+			if (const FObjectReplicator* const ActorReplicator = ActorChannel->ActorReplicator.Get())
 			{
 				FReceivingRepState* ReceivingRepState = ActorReplicator->RepState->GetReceivingRepState();
 				FRepShadowDataBuffer ShadowData(ReceivingRepState->StaticBuffer.GetData());
