ML - Clan probably does de-virtualization optimization and wants to call InitRHI directly instead though the base.

diff --git a/Engine/Source/Runtime/Engine/Public/RawIndexBuffer.h b/Engine/Source/Runtime/Engine/Public/RawIndexBuffer.h
--- a/Engine/Source/Runtime/Engine/Public/RawIndexBuffer.h
+++ b/Engine/Source/Runtime/Engine/Public/RawIndexBuffer.h
@@ -267,7 +267,7 @@ public:
     void Discard();
     
 	// FRenderResource interface.
-	virtual void InitRHI() override;
+	ENGINE_API virtual void InitRHI() override;
 
 	inline bool Is32Bit() const { return b32Bit; }
 
