Rewritten asset streaming behavior when using SCOPED_SUSPEND_RENDERING_THREAD, StopRenderingThread and StartRenderingThread.
Now the asset streaming is only suspended instead of being flushed.

UTexture2D::CancelPendingTextureStreaming() doesn't flush streaming anymore as the effect of calling  UTexture2D::CancelPendingMipChangeRequest() is immediate.

Added a new counter in "stat streaming" named "Pending 2D Update" to show how many pending update are inflight.

diff --git a/Engine/Source/Runtime/Engine/Classes/Engine/StaticMesh.h b/Engine/Source/Runtime/Engine/Classes/Engine/StaticMesh.h
--- a/Engine/Source/Runtime/Engine/Classes/Engine/StaticMesh.h
+++ b/Engine/Source/Runtime/Engine/Classes/Engine/StaticMesh.h
@@ -937,6 +937,7 @@ public:
 	virtual int32 GetNumRequestedMips() const final override;
 	virtual bool CancelPendingMipChangeRequest() final override;
 	virtual bool HasPendingUpdate() const final override;
+	virtual bool IsPendingUpdateLocked() const final override;
 	virtual bool StreamOut(int32 NewMipCount) final override;
 	virtual bool StreamIn(int32 NewMipCount, bool bHighPrio) final override;
 	virtual bool UpdateStreamingStatus(bool bWaitForMipFading = false) final override;
diff --git a/Engine/Source/Runtime/Engine/Classes/Engine/StreamableRenderAsset.h b/Engine/Source/Runtime/Engine/Classes/Engine/StreamableRenderAsset.h
--- a/Engine/Source/Runtime/Engine/Classes/Engine/StreamableRenderAsset.h
+++ b/Engine/Source/Runtime/Engine/Classes/Engine/StreamableRenderAsset.h
@@ -93,6 +93,12 @@ public:
 		return false;
 	}
 
+	virtual bool IsPendingUpdateLocked() const 
+	{ 
+		STREAMABLERENDERASSET_NODEFAULT(IsPendingUpdateLocked);
+		return false;
+	}
+
 	/**
 	* Unload some mips from memory. Only usable if the asset is streamable.
 	*
diff --git a/Engine/Source/Runtime/Engine/Classes/Engine/Texture.h b/Engine/Source/Runtime/Engine/Classes/Engine/Texture.h
--- a/Engine/Source/Runtime/Engine/Classes/Engine/Texture.h
+++ b/Engine/Source/Runtime/Engine/Classes/Engine/Texture.h
@@ -608,6 +608,12 @@ public:
 	{
 	}
 
+	virtual bool HasPendingUpdate() const override 
+	{ 
+		return false; // Overriden in UTexture2D
+	}
+
+
 	/**
 	 * Textures that use the derived data cache must override this function and
 	 * provide a pointer to the linked list of platform data.
diff --git a/Engine/Source/Runtime/Engine/Classes/Engine/Texture2D.h b/Engine/Source/Runtime/Engine/Classes/Engine/Texture2D.h
--- a/Engine/Source/Runtime/Engine/Classes/Engine/Texture2D.h
+++ b/Engine/Source/Runtime/Engine/Classes/Engine/Texture2D.h
@@ -157,6 +157,7 @@ public:
 
 	/** true if the texture is currently being updated through StreamIn() or StreamOut(). */
 	virtual bool HasPendingUpdate() const final override { return PendingUpdate != nullptr; }
+	virtual bool IsPendingUpdateLocked() const final override;
 
 	virtual bool StreamOut(int32 NewMipCount) final override;
 	virtual bool StreamIn(int32 NewMipCount, bool bHighPrio) final override;
diff --git a/Engine/Source/Runtime/Engine/Private/StaticMesh.cpp b/Engine/Source/Runtime/Engine/Private/StaticMesh.cpp
--- a/Engine/Source/Runtime/Engine/Private/StaticMesh.cpp
+++ b/Engine/Source/Runtime/Engine/Private/StaticMesh.cpp
@@ -4740,6 +4740,11 @@ bool UStaticMesh::HasPendingUpdate() const
 	return !!PendingUpdate;
 }
 
+bool UStaticMesh::IsPendingUpdateLocked() const 
+{ 
+	return PendingUpdate && PendingUpdate->IsLocked(); 
+}
+
 bool UStaticMesh::StreamOut(int32 NewMipCount)
 {
 	check(IsInGameThread());
diff --git a/Engine/Source/Runtime/Engine/Private/Streaming/RenderAssetUpdate.cpp b/Engine/Source/Runtime/Engine/Private/Streaming/RenderAssetUpdate.cpp
--- a/Engine/Source/Runtime/Engine/Private/Streaming/RenderAssetUpdate.cpp
+++ b/Engine/Source/Runtime/Engine/Private/Streaming/RenderAssetUpdate.cpp
@@ -5,18 +5,106 @@ RenderAssetUpdate.cpp: Base class of helpers to stream in and out texture/mesh L
 =============================================================================*/
 
 #include "Streaming/RenderAssetUpdate.h"
+#include "Engine/StreamableRenderAsset.h"
+#include "Streaming/TextureStreamingHelpers.h"
+#include "UObject/UObjectIterator.h"
 
-volatile int32 GRenderAssetStreamingSuspendRenderThreadTasks;
+float GStreamingFlushTimeOut = 3.00f;
+static FAutoConsoleVariableRef CVarStreamingFlushTimeOut(
+	TEXT("r.Streaming.FlushTimeOut"),
+	GStreamingFlushTimeOut,
+	TEXT("Time before we timeout when flushing streaming (default=3)"),
+	ECVF_Default
+);
 
-void SuspendRenderAssetStreamingRenderTasksInternal()
+
+volatile int32 GRenderAssetStreamingSuspension = 0;
+
+bool IsAssetStreamingSuspended()
+{
+	return GRenderAssetStreamingSuspension > 0;
+}
+
+
+void SuspendRenderAssetStreaming()
 {
-	// This doesn't prevent from having a task pushed immediately after as some threads could already be deep in FTexture2DUpdate::PushTask.
-	// This is why GSuspendRenderThreadTasks also gets checked in FTexture2DUpdate::Tick. The goal being to avoid accessing the RHI more
-	// than avoiding pushing new render commands. The reason behind is that some code paths access the RHI outside the render thread.
-	FPlatformAtomics::InterlockedIncrement(&GRenderAssetStreamingSuspendRenderThreadTasks);
+	ensure(IsInGameThread());
+
+	if (FPlatformAtomics::InterlockedIncrement(&GRenderAssetStreamingSuspension) == 1)
+	{
+		bool bHasPendingUpdate = false;
+
+		// Wait for all assets to have their update lock unlocked. 
+		TArray<UStreamableRenderAsset*> LockedAssets;
+		for (TObjectIterator<UStreamableRenderAsset> It; It; ++It)
+		{
+			UStreamableRenderAsset* CurrentAsset = *It;
+			if (CurrentAsset && CurrentAsset->HasPendingUpdate())
+			{
+				bHasPendingUpdate = true;
+				if (CurrentAsset->IsPendingUpdateLocked())
+				{
+					LockedAssets.Add(CurrentAsset);
+				}
+			}
+		}
+
+		// If an asset stays locked for  GStreamingFlushTimeOut, 
+		// we conclude there is a deadlock or that the object is never going to recover.
+
+		const float TimeIncrement = 0.010f;
+		float TimeLimit = GStreamingFlushTimeOut;
+
+		while (LockedAssets.Num() && (TimeLimit > 0 || GStreamingFlushTimeOut <= 0))
+		{
+			FPlatformProcess::Sleep(TimeIncrement);
+			FlushRenderingCommands();
+			
+			TimeLimit -= TimeIncrement;
+
+			for (int32 LockedIndex = 0; LockedIndex < LockedAssets.Num(); ++LockedIndex)
+			{
+				UStreamableRenderAsset* CurrentAsset = LockedAssets[LockedIndex];
+				if (!CurrentAsset || !CurrentAsset->IsPendingUpdateLocked())
+				{
+					LockedAssets.RemoveAtSwap(LockedIndex);
+					--LockedIndex;
+				}
+			}
+		}
+
+		if (TimeLimit <= 0 && GStreamingFlushTimeOut > 0)
+		{
+			UE_LOG(LogContentStreaming, Error, TEXT("SuspendRenderAssetStreaming timed out while waiting for asset:"));
+			for (int32 LockedIndex = 0; LockedIndex < LockedAssets.Num(); ++LockedIndex)
+			{
+				UStreamableRenderAsset* CurrentAsset = LockedAssets[LockedIndex];
+				if (CurrentAsset)
+				{
+					if (CurrentAsset->IsPendingKill() || CurrentAsset->HasAnyFlags(RF_BeginDestroyed|RF_FinishDestroyed))
+					{
+						UE_LOG(LogContentStreaming, Error, TEXT("	%s"), *CurrentAsset->GetFullName());
+					}
+					else
+					{
+						UE_LOG(LogContentStreaming, Error, TEXT("	%s (PendingKill)"), *CurrentAsset->GetFullName());
+					}
+				}
+			}
+		}
+
+		// At this point, no more rendercommands or IO requests can be generated before a call to ResumeRenderAssetStreamingRenderTasksInternal().
+
+		if (bHasPendingUpdate)
+		{
+			// Ensure any pending render command executes.
+			FlushRenderingCommands();
+		}
+	}
 }
 
-void ResumeRenderAssetStreamingRenderTasksInternal()
+void ResumeRenderAssetStreaming()
 {
-	FPlatformAtomics::InterlockedDecrement(&GRenderAssetStreamingSuspendRenderThreadTasks);
+	FPlatformAtomics::InterlockedDecrement(&GRenderAssetStreamingSuspension);
+	ensure(GRenderAssetStreamingSuspension >= 0);
 }
diff --git a/Engine/Source/Runtime/Engine/Private/Streaming/RenderAssetUpdate.h b/Engine/Source/Runtime/Engine/Private/Streaming/RenderAssetUpdate.h
--- a/Engine/Source/Runtime/Engine/Private/Streaming/RenderAssetUpdate.h
+++ b/Engine/Source/Runtime/Engine/Private/Streaming/RenderAssetUpdate.h
@@ -15,8 +15,10 @@ RenderAssetUpdate.h: Base class of helpers to stream in and out texture/mesh LOD
 /** SRA stands for StreamableRenderAsset */
 #define SRA_UPDATE_CALLBACK(FunctionName) [this](const FContext& C){ FunctionName(C); }
 
-/** A counter defining whether render thread tasks should be postponed The callback handling the cancellation of the update (only if the update gets cancelled). */
-extern volatile int32 GRenderAssetStreamingSuspendRenderThreadTasks;
+ // Allows yield to lower priority threads
+#define RENDER_ASSET_STREAMING_SLEEP_DT (0.010f)
+
+ENGINE_API bool IsAssetStreamingSuspended();
 
 /**
 * This class provides a framework for loading and unloading the texture/mesh LODs.
@@ -104,7 +106,7 @@ public:
 			// If the task is ready to execute. 
 			// If this is the renderthread and we shouldn't run renderthread task, mark as pending.
 			// This will require a tick from the game thread to reschedule the task.
-			if (TaskSynchronization.GetValue() <= 0 && !(GRenderAssetStreamingSuspendRenderThreadTasks > 0 && InCurrentThread == TT_Render))
+			if (TaskSynchronization.GetValue() <= 0 && !IsAssetStreamingSuspended())
 			{
 				FContext Context(InAsset, InCurrentThread);
 
@@ -214,6 +216,10 @@ public:
 		TaskState = CachedPendingTaskState;
 	}
 
+	bool IsLocked() const
+	{
+		return TaskState == TS_Locked;
+	}
 	/** Get the number of requested mips for this update, ignoring cancellation attempts. */
 	int32 GetNumRequestedMips() const
 	{
@@ -254,8 +260,8 @@ protected:
 		const EThreadType RelevantThread = !bCachedIsCancelled ? InTaskThread : InCancelationThread;
 
 		// TaskSynchronization is expected to be set before call this.
-		// If the rendering thread is suspended, delay the scheduling until not suspended anymore.
-		const bool bCanExecuteNow = TaskSynchronization.GetValue() <= 0 && !(GRenderAssetStreamingSuspendRenderThreadTasks > 0 && RelevantThread == TT_Render);
+		// If the update is suspended, delay the scheduling until not suspended anymore.
+		const bool bCanExecuteNow = TaskSynchronization.GetValue() <= 0 && !IsAssetStreamingSuspended();
 
 		if (RelevantThread == TT_None)
 		{
@@ -470,5 +476,5 @@ protected:
 	FCallback CancelationCallback;
 };
 
-void SuspendRenderAssetStreamingRenderTasksInternal();
-void ResumeRenderAssetStreamingRenderTasksInternal();
+void SuspendRenderAssetStreaming();
+void ResumeRenderAssetStreaming();
diff --git a/Engine/Source/Runtime/Engine/Private/Streaming/StreamingManagerTexture.cpp b/Engine/Source/Runtime/Engine/Private/Streaming/StreamingManagerTexture.cpp
--- a/Engine/Source/Runtime/Engine/Private/Streaming/StreamingManagerTexture.cpp
+++ b/Engine/Source/Runtime/Engine/Private/Streaming/StreamingManagerTexture.cpp
@@ -23,6 +23,7 @@
 #include "ProfilingDebugging/CsvProfiler.h"
 #include "Interfaces/ITargetPlatformManagerModule.h"
 #include "Interfaces/ITargetPlatform.h"
+#include "Streaming/RenderAssetUpdate.h"
 
 CSV_DECLARE_CATEGORY_MODULE_EXTERN(CORE_API, Basic);
 
@@ -1469,10 +1470,9 @@ void FRenderAssetStreamingManager::UpdateResourceStreaming( float DeltaTime, boo
 int32 FRenderAssetStreamingManager::BlockTillAllRequestsFinished( float TimeLimit /*= 0.0f*/, bool bLogResults /*= false*/ )
 {
 	FScopeLock ScopeLock(&CriticalSection);
-
 	double StartTime = FPlatformTime::Seconds();
 
-	while (true) 
+	while (ensure(!IsAssetStreamingSuspended()))
 	{
 		int32 NumOfInFlights = 0;
 
@@ -1488,18 +1488,19 @@ int32 FRenderAssetStreamingManager::BlockTillAllRequestsFinished( float TimeLimi
 		if (NumOfInFlights && (TimeLimit == 0 || (float)(FPlatformTime::Seconds() - StartTime) < TimeLimit))
 		{
 			FlushRenderingCommands();
-			FPlatformProcess::Sleep( 0.010f );
+			FPlatformProcess::Sleep(RENDER_ASSET_STREAMING_SLEEP_DT);
 		}
 		else
 		{
 			if (bLogResults)
 			{
-				UE_LOG(LogContentStreaming, Log, TEXT("Blocking on texture streaming: %.1f ms (%d still in flight)"), (float)(FPlatformTime::Seconds() - StartTime) * 1000, NumOfInFlights );
+				UE_LOG(LogContentStreaming, Log, TEXT("Blocking on texture streaming: %.1f ms (%d still in flight)"), (float)(FPlatformTime::Seconds() - StartTime) * 1000, NumOfInFlights);
 
 			}
 			return NumOfInFlights;
 		}
 	}
+	return 0;
 }
 
 void FRenderAssetStreamingManager::GetObjectReferenceBounds(const UObject* RefObject, TArray<FBox>& AssetBoxes)
diff --git a/Engine/Source/Runtime/Engine/Private/Streaming/Texture2DUpdate.cpp b/Engine/Source/Runtime/Engine/Private/Streaming/Texture2DUpdate.cpp
--- a/Engine/Source/Runtime/Engine/Private/Streaming/Texture2DUpdate.cpp
+++ b/Engine/Source/Runtime/Engine/Private/Streaming/Texture2DUpdate.cpp
@@ -8,6 +8,11 @@ Texture2DUpdate.cpp: Helpers to stream in and out mips.
 #include "RenderUtils.h"
 #include "Containers/ResourceArray.h"
 
+#if STATS
+extern volatile int64 GPending2DUpdateCount;
+volatile int64 GPending2DUpdateCount = 0;
+#endif
+
 FTexture2DUpdateContext::FTexture2DUpdateContext(UTexture2D* InTexture, EThreadType InCurrentThread)
 	: Texture(InTexture)
 	, CurrentThread(InCurrentThread)
@@ -34,11 +39,15 @@ FTexture2DUpdate::FTexture2DUpdate(UTexture2D* InTexture, int32 InRequestedMips)
 		PendingFirstMip = INDEX_NONE;
 		bIsCancelled = true;
 	}
+
+	STAT(FPlatformAtomics::InterlockedIncrement(&GPending2DUpdateCount));
 }
 
 FTexture2DUpdate::~FTexture2DUpdate()
 {
 	ensure(!IntermediateTextureRHI);
+
+	STAT(FPlatformAtomics::InterlockedDecrement(&GPending2DUpdateCount));
 }
 
 
diff --git a/Engine/Source/Runtime/Engine/Private/Streaming/TextureStreamingHelpers.cpp b/Engine/Source/Runtime/Engine/Private/Streaming/TextureStreamingHelpers.cpp
--- a/Engine/Source/Runtime/Engine/Private/Streaming/TextureStreamingHelpers.cpp
+++ b/Engine/Source/Runtime/Engine/Private/Streaming/TextureStreamingHelpers.cpp
@@ -32,6 +32,7 @@ DECLARE_CYCLE_STAT(TEXT("Setup Async Task"), STAT_Streaming01_SetupAsyncTask, ST
 DECLARE_CYCLE_STAT(TEXT("Update Streaming Data"), STAT_Streaming02_UpdateStreamingData, STATGROUP_Streaming);
 DECLARE_CYCLE_STAT(TEXT("Streaming Texture"), STAT_Streaming03_StreamTextures, STATGROUP_Streaming);
 DECLARE_CYCLE_STAT(TEXT("Notifications"), STAT_Streaming04_Notifications, STATGROUP_Streaming);
+DECLARE_DWORD_COUNTER_STAT(TEXT("Pending 2D Update"), STAT_Streaming05_Pending2DUpdate, STATGROUP_Streaming);
 
 /** Streaming Overview stats */
 
@@ -338,6 +339,7 @@ bool GNeverStreamOutRenderAssets = false;
 #if STATS
 extern int64 GUITextureMemory;
 extern int64 GNeverStreamTextureMemory;
+extern volatile int64 GPending2DUpdateCount;
 
 int64 GRequiredPoolSizeSum = 0;
 int64 GRequiredPoolSizeCount= 0;
@@ -372,6 +374,7 @@ void FRenderAssetStreamingStats::Apply()
 	SET_CYCLE_COUNTER(STAT_Streaming02_UpdateStreamingData, UpdateStreamingDataCycles);
 	SET_CYCLE_COUNTER(STAT_Streaming03_StreamTextures, StreamTexturesCycles);
 	SET_CYCLE_COUNTER(STAT_Streaming04_Notifications, CallbacksCycles);
+	INC_DWORD_STAT_BY(STAT_Streaming05_Pending2DUpdate, GPending2DUpdateCount);
 
 	/** Streaming Overview stats */
 
diff --git a/Engine/Source/Runtime/Engine/Private/Texture2D.cpp b/Engine/Source/Runtime/Engine/Private/Texture2D.cpp
--- a/Engine/Source/Runtime/Engine/Private/Texture2D.cpp
+++ b/Engine/Source/Runtime/Engine/Private/Texture2D.cpp
@@ -433,15 +433,13 @@ void UTexture2D::UnlinkStreaming()
 
 void UTexture2D::CancelPendingTextureStreaming()
 {
-	FlushRenderingCommands();
-
 	for( TObjectIterator<UTexture2D> It; It; ++It )
 	{
 		UTexture2D* CurrentTexture = *It;
 		CurrentTexture->CancelPendingMipChangeRequest();
 	}
 
-	FlushResourceStreaming();
+	// No need to call FlushResourceStreaming(), since calling CancelPendingMipChangeRequest has an immediate effect.
 }
 
 void UTexture2D::PostLoad()
@@ -491,11 +489,11 @@ void UTexture2D::GetAssetRegistryTags(TArray<FAssetRegistryTag>& OutTags) const
 void UTexture2D::UpdateResource()
 {
 	// Make sure there are no pending requests in flight.
-	while( UpdateStreamingStatus() == true )
+	while (UpdateStreamingStatus() && ensure(!IsAssetStreamingSuspended()))
 	{
 		// Force flush the RHI threads to execute all commands issued for texture streaming, and give up timeslice.
 		FlushRenderingCommands();
-		FPlatformProcess::Sleep(0);
+		FPlatformProcess::Sleep(RENDER_ASSET_STREAMING_SLEEP_DT);
 	}
 
 #if WITH_EDITOR
@@ -564,14 +562,14 @@ FString UTexture2D::GetDesc()
 
 void UTexture2D::WaitForStreaming()
 {
-	if (bIsStreamable)
+	if (bIsStreamable && ensure(!IsAssetStreamingSuspended()))
 	{
 		// Make sure there are no pending requests in flight otherwise calling UpdateIndividualTexture could be prevented to defined a new requested mip.
 		while (	!IsReadyForStreaming() || UpdateStreamingStatus() ) 
 		{
 			// Force flush the RHI threads to execute all commands issued for texture streaming, and give up timeslice.
 			FlushRenderingCommands();
-			FPlatformProcess::Sleep(0);
+			FPlatformProcess::Sleep(RENDER_ASSET_STREAMING_SLEEP_DT);
 		}
 
 		// Update the wanted mip and stream in..		
@@ -583,7 +581,7 @@ void UTexture2D::WaitForStreaming()
 			{
 				// Force flush the RHI threads to execute all commands issued for texture streaming, and give up timeslice.
 				FlushRenderingCommands();
-				FPlatformProcess::Sleep(0);
+				FPlatformProcess::Sleep(RENDER_ASSET_STREAMING_SLEEP_DT);
 			}
 		}
 	}
@@ -1131,7 +1129,6 @@ void UTexture2D::TemporarilyDisableStreaming()
 		UpdateResource();
 	}
 }
-
 #endif
 
 
@@ -1588,6 +1585,11 @@ bool UTexture2D::StreamIn(int32 NewMipCount, bool bHighPrio)
 	return false;
 }
 
+bool UTexture2D::IsPendingUpdateLocked() const 
+{ 
+	return PendingUpdate && PendingUpdate->IsLocked(); 
+}
+
 bool UTexture2D::StreamOut(int32 NewMipCount)
 {
 	check(IsInGameThread());
diff --git a/Engine/Source/Runtime/Engine/Private/UnrealClient.cpp b/Engine/Source/Runtime/Engine/Private/UnrealClient.cpp
--- a/Engine/Source/Runtime/Engine/Private/UnrealClient.cpp
+++ b/Engine/Source/Runtime/Engine/Private/UnrealClient.cpp
@@ -1796,9 +1796,6 @@ void FViewport::GetActorsAndModelsInHitProxy(FIntRect InRect, TSet<AActor*>& Out
 
 void FViewport::UpdateViewportRHI(bool bDestroyed, uint32 NewSizeX, uint32 NewSizeY, EWindowMode::Type NewWindowMode, EPixelFormat PreferredPixelFormat)
 {
-	// Make sure we're not in the middle of streaming textures.
-	(*GFlushStreamingFunc)();
-
 	{
 		// Temporarily stop rendering thread.
 		SCOPED_SUSPEND_RENDERING_THREAD(true);
diff --git a/Engine/Source/Runtime/Engine/Private/UnrealEngine.cpp b/Engine/Source/Runtime/Engine/Private/UnrealEngine.cpp
--- a/Engine/Source/Runtime/Engine/Private/UnrealEngine.cpp
+++ b/Engine/Source/Runtime/Engine/Private/UnrealEngine.cpp
@@ -254,8 +254,8 @@ void FEngineModule::StartupModule()
 		CVARShowMaterialDrawEvents->SetOnChangedCallback(FConsoleVariableDelegate::CreateStatic(&OnChangeEngineCVarRequiringRecreateRenderState));
 	}
 
-	SuspendTextureStreamingRenderTasks = &SuspendRenderAssetStreamingRenderTasksInternal;
-	ResumeTextureStreamingRenderTasks = &ResumeRenderAssetStreamingRenderTasksInternal;
+	SuspendTextureStreamingRenderTasks = &SuspendRenderAssetStreaming;
+	ResumeTextureStreamingRenderTasks = &ResumeRenderAssetStreaming;
 
 	FParticleSystemWorldManager::OnStartup();
 
diff --git a/Engine/Source/Runtime/RenderCore/Private/RenderingThread.cpp b/Engine/Source/Runtime/RenderCore/Private/RenderingThread.cpp
--- a/Engine/Source/Runtime/RenderCore/Private/RenderingThread.cpp
+++ b/Engine/Source/Runtime/RenderCore/Private/RenderingThread.cpp
@@ -86,6 +86,9 @@ FSuspendRenderingThread::FSuspendRenderingThread( bool bInRecreateThread )
 		SuspendAsyncLoading();
 	}
 
+	// Pause asset streaming to prevent rendercommands from being enqueued.
+	SuspendTextureStreamingRenderTasks();
+
 	bRecreateThread = bInRecreateThread;
 	bUseRenderingThread = GUseThreadedRendering;
 	bWasRenderingThreadRunning = GIsThreadedRendering;
@@ -102,7 +105,6 @@ FSuspendRenderingThread::FSuspendRenderingThread( bool bInRecreateThread )
 		if ( GIsRenderingThreadSuspended.Load(EMemoryOrder::Relaxed) == 0 )
 		{
 			// First tell the render thread to finish up all pending commands and then suspend its activities.
-			
 			// this ensures that async stuff will be completed too
 			FlushRenderingCommands();
 			
@@ -183,6 +185,10 @@ FSuspendRenderingThread::~FSuspendRenderingThread()
 		// Resume the render thread again.
 		--GIsRenderingThreadSuspended;
 	}
+
+	// Resume any asset streaming
+	ResumeTextureStreamingRenderTasks();
+
 	if (IsAsyncLoadingMultithreaded())
 	{
 		ResumeAsyncLoading();
@@ -655,6 +661,9 @@ void StartRenderingThread()
 
 	check(!GRHIThread_InternalUseOnly && !GIsRunningRHIInSeparateThread_InternalUseOnly && !GIsRunningRHIInDedicatedThread_InternalUseOnly && !GIsRunningRHIInTaskThread_InternalUseOnly);
 
+	// Pause asset streaming to prevent rendercommands from being enqueued.
+	SuspendTextureStreamingRenderTasks();
+
 	// Flush GT since render commands issued by threads other than GT are sent to
 	// the main queue of GT when RT is disabled. Without this flush, those commands
 	// will run on GT after RT is enabled
@@ -711,6 +720,9 @@ void StartRenderingThread()
 	GRenderingThreadHeartbeat = FRunnableThread::Create(GRenderingThreadRunnableHeartbeat, *FString::Printf(TEXT("RTHeartBeat %d"), ThreadCount), 16 * 1024, TPri_AboveNormal, FPlatformAffinity::GetRTHeartBeatMask());
 
 	ThreadCount++;
+
+	// Update can now resume.
+	ResumeTextureStreamingRenderTasks();
 }
 
 
@@ -740,7 +752,7 @@ void StopRenderingThread()
 		FPendingCleanupObjects* PendingCleanupObjects = GetPendingCleanupObjects();
 
 		// Make sure we're not in the middle of streaming textures.
-		(*GFlushStreamingFunc)();
+		SuspendTextureStreamingRenderTasks();
 
 		// Wait for the rendering thread to finish executing all enqueued commands.
 		FlushRenderingCommands();
@@ -803,6 +815,9 @@ void StopRenderingThread()
 
 		// Delete the pending cleanup objects which were in use by the rendering thread.
 		delete PendingCleanupObjects;
+
+		// Update can now resume with renderthread being the gamethread.
+		ResumeTextureStreamingRenderTasks();
 	}
 
 	check(!GRHIThread_InternalUseOnly);
