Move priority settings on USoundBase to own category

diff --git a/Engine/Source/Runtime/Engine/Classes/Sound/SoundBase.h b/Engine/Source/Runtime/Engine/Classes/Sound/SoundBase.h
--- a/Engine/Source/Runtime/Engine/Classes/Sound/SoundBase.h
+++ b/Engine/Source/Runtime/Engine/Classes/Sound/SoundBase.h
@@ -71,7 +71,7 @@ public:
 	uint8 bHasVirtualizeWhenSilent:1;
 
 	/** Allows this sound to bypass volume-weighting for the max channel resolution. */
-	UPROPERTY(EditAnywhere, BlueprintReadWrite, Category = Concurrency)
+	UPROPERTY(EditAnywhere, BlueprintReadWrite, Category = Priority)
 	uint8 bBypassVolumeScaleForPriority : 1;
 
 #if WITH_EDITORONLY_DATA
@@ -112,8 +112,10 @@ public:
 	UPROPERTY(Category = Info, AssetRegistrySearchable, VisibleAnywhere, BlueprintReadOnly)
 	float TotalSamples;
 
-	/** Sound priority (higher value is higher priority) used for concurrency resolution. This priority value is weighted against the final volume of the sound. */
-	UPROPERTY(EditAnywhere, BlueprintReadWrite, Category = Concurrency, meta = (ClampMin = "0.0", UIMin = "0.0", ClampMax = "100.0", UIMax = "100.0") )
+	/** Used to determine whether sound can play or remain active, where higher value is higher priority (see Audio Settings 'Max Channel Resolution'). 
+	  * Value is weighted with the final volume of the sound to produce final runtime priority value unless bypassed.
+	  */
+	UPROPERTY(EditAnywhere, BlueprintReadWrite, Category = Priority, meta = (ClampMin = "0.0", UIMin = "0.0", ClampMax = "100.0", UIMax = "100.0"))
 	float Priority;
 
 	/** Attenuation settings package for the sound */
@@ -127,7 +129,7 @@ public:
 	UPROPERTY(EditAnywhere, BlueprintReadWrite, Category = Effects, meta = (DisplayName = "Sound Submix"))
 	USoundSubmix* SoundSubmixObject;
 
-	/** An array of submix sends. Audio from this sound will send a portion of its audio to these effects.  */
+	/** An array of submix sends. Audio from this sound will send a portion of its audio to these effects. */
 	UPROPERTY(EditAnywhere, BlueprintReadWrite, Category = Effects)
 	TArray<FSoundSubmixSendInfo> SoundSubmixSends;
 
@@ -135,11 +137,11 @@ public:
 	UPROPERTY(EditAnywhere, BlueprintReadWrite, Category = Effects)
 	USoundEffectSourcePresetChain* SourceEffectChain;
 
-	/** This sound will send its audio output to this list of buses if there are bus instances playing after source effects are processed.  */
+	/** This sound will send its audio output to this list of buses if there are bus instances playing after source effects are processed. */
 	UPROPERTY(EditAnywhere, BlueprintReadWrite, Category = Effects, meta = (DisplayName = "Post-Effect Bus Sends"))
 	TArray<FSoundSourceBusSendInfo> BusSends;
 
-	/** This sound will send its audio output to this list of buses if there are bus instances playing before source effects are processed.  */
+	/** This sound will send its audio output to this list of buses if there are bus instances playing before source effects are processed. */
 	UPROPERTY(EditAnywhere, BlueprintReadWrite, Category = Effects, meta = (DisplayName = "Pre-Effect Bus Sends"))
 	TArray<FSoundSourceBusSendInfo> PreEffectBusSends;
 
