Clarify comments for audio priority culling properties

diff --git a/Engine/Source/Runtime/Engine/Classes/Sound/SoundBase.h b/Engine/Source/Runtime/Engine/Classes/Sound/SoundBase.h
--- a/Engine/Source/Runtime/Engine/Classes/Sound/SoundBase.h
+++ b/Engine/Source/Runtime/Engine/Classes/Sound/SoundBase.h
@@ -70,7 +70,7 @@ public:
 	UPROPERTY()
 	uint8 bHasVirtualizeWhenSilent:1;
 
-	/** Allows this sound to bypass volume-weighting for the max channel resolution. */
+	/** Bypass volume weighting priority upon evaluating whether sound should remain active when max channel count is met (See platform Audio Settings). */
 	UPROPERTY(EditAnywhere, BlueprintReadWrite, Category = Priority)
 	uint8 bBypassVolumeScaleForPriority : 1;
 
@@ -112,8 +112,9 @@ public:
 	UPROPERTY(Category = Info, AssetRegistrySearchable, VisibleAnywhere, BlueprintReadOnly)
 	float TotalSamples;
 
-	/** Used to determine whether sound can play or remain active, where higher value is higher priority (see Audio Settings 'Max Channel Resolution'). 
-	  * Value is weighted with the final volume of the sound to produce final runtime priority value unless bypassed.
+	/** Used to determine whether sound can play or remain active if channel limit is met, where higher value is higher priority
+	  * (see platform's Audio Settings 'Max Channels' property). Unless bypassed, value is weighted with the final volume of the
+	  * sound to produce final runtime priority value.
 	  */
 	UPROPERTY(EditAnywhere, BlueprintReadWrite, Category = Priority, meta = (ClampMin = "0.0", UIMin = "0.0", ClampMax = "100.0", UIMax = "100.0"))
 	float Priority;
