#UE4 Silence warning about lighting being unbuilt due to an inconsistency in static mesh component instance data when static lighting is not enabled.

diff --git a/Engine/Source/Runtime/Engine/Private/Components/StaticMeshComponent.cpp b/Engine/Source/Runtime/Engine/Private/Components/StaticMeshComponent.cpp
--- a/Engine/Source/Runtime/Engine/Private/Components/StaticMeshComponent.cpp
+++ b/Engine/Source/Runtime/Engine/Private/Components/StaticMeshComponent.cpp
@@ -37,6 +37,7 @@
 #include "ComponentRecreateRenderStateContext.h"
 #include "Engine/StaticMesh.h"
 #include "HAL/LowLevelMemTracker.h"
+#include "HAL/IConsoleManager.h"
 #include "Algo/AllOf.h"
 #include "Algo/Transform.h"
 
@@ -2282,11 +2283,18 @@ void UStaticMeshComponent::ApplyComponentInstanceData(FStaticMeshComponentInstan
 		}
 		else
 		{
-			UE_ASSET_LOG(LogStaticMesh, Warning, this,
-				TEXT("Cached component instance data transform did not match!  Discarding cached lighting data which will cause lighting to be unbuilt.\n%s\nCurrent: %s Cached: %s"), 
-				*GetPathName(),
-				*GetComponentTransform().ToString(),
-				*StaticMeshInstanceData->GetComponentTransform().ToString());
+			static const auto AllowStaticLightingVar = IConsoleManager::Get().FindTConsoleVariableDataInt(TEXT("r.AllowStaticLighting"));
+			const bool bAllowStaticLighting = (!AllowStaticLightingVar || AllowStaticLightingVar->GetValueOnAnyThread() != 0);
+
+			// Only warn if static lighting is enabled in the project
+			if (bAllowStaticLighting)
+			{
+				UE_ASSET_LOG(LogStaticMesh, Warning, this,
+					TEXT("Cached component instance data transform did not match!  Discarding cached lighting data which will cause lighting to be unbuilt.\n%s\nCurrent: %s Cached: %s"),
+					*GetPathName(),
+					*GetComponentTransform().ToString(),
+					*StaticMeshInstanceData->GetComponentTransform().ToString());
+			}
 		}
 	}
 
