Temporarily disable reporting of dynamic RepLayout memory.

diff --git a/Engine/Source/Runtime/Engine/Private/RepLayout.cpp b/Engine/Source/Runtime/Engine/Private/RepLayout.cpp
--- a/Engine/Source/Runtime/Engine/Private/RepLayout.cpp
+++ b/Engine/Source/Runtime/Engine/Private/RepLayout.cpp
@@ -452,10 +452,10 @@ void FRepStateStaticBuffer::CountBytes(FArchive& Ar) const
 
 	GRANULAR_NETWORK_MEMORY_TRACKING_TRACK("Static Memory", Buffer.CountBytes(Ar));
 	GRANULAR_NETWORK_MEMORY_TRACKING_TRACK("Dynamic Memory (Undercounts!)",
-		FCountBytesHelper CountBytesHelper(Ar, Buffer.GetData(), RepLayout->Parents, RepLayout->Cmds);
-		CountBytesHelper.CountBytes();
-		GRANULAR_NETWORK_MEMORY_TRACKING_CUSTOM_WORK("OnRepMemory", CountBytesHelper.OnRepMemory);
-		GRANULAR_NETWORK_MEMORY_TRACKING_CUSTOM_WORK("NonRepMemory", CountBytesHelper.NonRepMemory);
+		// FCountBytesHelper CountBytesHelper(Ar, Buffer.GetData(), RepLayout->Parents, RepLayout->Cmds);
+		// CountBytesHelper.CountBytes();
+		// GRANULAR_NETWORK_MEMORY_TRACKING_CUSTOM_WORK("OnRepMemory", CountBytesHelper.OnRepMemory);
+		// GRANULAR_NETWORK_MEMORY_TRACKING_CUSTOM_WORK("NonRepMemory", CountBytesHelper.NonRepMemory);
 	);
 }
 
