Fix FLocalLightSceneProxy compile error

diff --git a/Engine/Source/Runtime/Engine/Public/LocalLightSceneProxy.h b/Engine/Source/Runtime/Engine/Public/LocalLightSceneProxy.h
--- a/Engine/Source/Runtime/Engine/Public/LocalLightSceneProxy.h
+++ b/Engine/Source/Runtime/Engine/Public/LocalLightSceneProxy.h
@@ -43,7 +43,7 @@ public:
 
 	virtual bool GetScissorRect(FIntRect& ScissorRect, const FSceneView& View, const FIntRect& ViewRect) const override;
 
-	virtual void SetScissorRect(FRHICommandList& RHICmdList, const FSceneView& View, const FIntRect& ViewRect) const override;
+	virtual bool SetScissorRect(FRHICommandList& RHICmdList, const FSceneView& View, const FIntRect& ViewRect, FIntRect* OutScissorRect = nullptr) const override;
 
 	virtual FSphere GetBoundingSphere() const;
 
