A new function was added to ViewportClient class to validate if permanent Capture is required. By overloading it in GameViewPortClient we can prevent the game from getting permanent capture of the mouse if the console is active.

diff --git a/Engine/Source/Runtime/Engine/Classes/Engine/GameViewportClient.h b/Engine/Source/Runtime/Engine/Classes/Engine/GameViewportClient.h
--- a/Engine/Source/Runtime/Engine/Classes/Engine/GameViewportClient.h
+++ b/Engine/Source/Runtime/Engine/Classes/Engine/GameViewportClient.h
@@ -162,6 +162,7 @@ public:
 	virtual bool IsFocused(FViewport* Viewport) override;
 	virtual void Activated(FViewport* InViewport, const FWindowActivateEvent& InActivateEvent) override;
 	virtual void Deactivated(FViewport* InViewport, const FWindowActivateEvent& InActivateEvent) override;
+	virtual bool IsInPermanentCapture() override;
 	virtual bool WindowCloseRequested() override;
 	virtual void CloseRequested(FViewport* Viewport) override;
 	virtual bool RequiresHitProxyStorage() override { return 0; }
diff --git a/Engine/Source/Runtime/Engine/Private/GameViewportClient.cpp b/Engine/Source/Runtime/Engine/Private/GameViewportClient.cpp
--- a/Engine/Source/Runtime/Engine/Private/GameViewportClient.cpp
+++ b/Engine/Source/Runtime/Engine/Private/GameViewportClient.cpp
@@ -1833,6 +1833,16 @@ void UGameViewportClient::Deactivated(FViewport* InViewport, const FWindowActiva
 	LostFocus(InViewport);
 }
 
+bool UGameViewportClient::IsInPermanentCapture()
+{
+	bool bIsInPermanentCapture = FViewportClient::IsInPermanentCapture();
+	if (ViewportConsole)
+	{
+		bIsInPermanentCapture = !ViewportConsole->ConsoleActive() && bIsInPermanentCapture;
+	}
+	return bIsInPermanentCapture;
+}
+
 bool UGameViewportClient::WindowCloseRequested()
 {
 	return !WindowCloseRequestedDelegate.IsBound() || WindowCloseRequestedDelegate.Execute();
diff --git a/Engine/Source/Runtime/Engine/Private/Slate/SceneViewport.cpp b/Engine/Source/Runtime/Engine/Private/Slate/SceneViewport.cpp
--- a/Engine/Source/Runtime/Engine/Private/Slate/SceneViewport.cpp
+++ b/Engine/Source/Runtime/Engine/Private/Slate/SceneViewport.cpp
@@ -1155,9 +1155,8 @@ FReply FSceneViewport::OnViewportActivated(const FWindowActivateEvent& InActivat
 		ViewportClient->Activated(this, InActivateEvent);
 		
 		// Determine if we're in permanent capture mode.  This cannot be cached as part of bShouldCaptureMouseOnActivate because it could change between window activate and deactivate
-		const bool bPermanentCapture =
-			!GIsEditor && ((ViewportClient->CaptureMouseOnClick() == EMouseCaptureMode::CapturePermanently) ||
-			(ViewportClient->CaptureMouseOnClick() == EMouseCaptureMode::CapturePermanently_IncludingInitialMouseDown));
+		const bool bPermanentCapture = ViewportClient->IsInPermanentCapture();
+
 
 		// If we are activating and had Mouse Capture on deactivate then we should get focus again
 		// It's important to note in the case of:
@@ -1175,6 +1174,7 @@ FReply FSceneViewport::OnViewportActivated(const FWindowActivateEvent& InActivat
 
 	return FReply::Unhandled();
 }
+
 void FSceneViewport::OnViewportDeactivated(const FWindowActivateEvent& InActivateEvent)
 {
 	// We backup if we have capture for us on activation, however we also maintain "true" if it's already true!
diff --git a/Engine/Source/Runtime/Engine/Public/Slate/SceneViewport.h b/Engine/Source/Runtime/Engine/Public/Slate/SceneViewport.h
--- a/Engine/Source/Runtime/Engine/Public/Slate/SceneViewport.h
+++ b/Engine/Source/Runtime/Engine/Public/Slate/SceneViewport.h
@@ -374,6 +374,10 @@ private:
 	 */
 	void OnPostResizeWindowBackbuffer(void* Backbuffer);
 
+
+	/** @return Returns true if the viewport needs permanent capture. */
+	bool IsInPermanentCapture();
+
 private:
 	/** An intermediate reply state that is reset whenever an input event is generated */
 	FReply CurrentReplyState;
diff --git a/Engine/Source/Runtime/Engine/Public/UnrealClient.h b/Engine/Source/Runtime/Engine/Public/UnrealClient.h
--- a/Engine/Source/Runtime/Engine/Public/UnrealClient.h
+++ b/Engine/Source/Runtime/Engine/Public/UnrealClient.h
@@ -917,6 +917,12 @@ public:
 	virtual void Activated(FViewport* Viewport, const FWindowActivateEvent& InActivateEvent) {}
 	virtual void Deactivated(FViewport* Viewport, const FWindowActivateEvent& InActivateEvent) {}
 
+	virtual bool IsInPermanentCapture()
+	{ 
+		return  !GIsEditor && ((CaptureMouseOnClick() == EMouseCaptureMode::CapturePermanently) ||
+			(CaptureMouseOnClick() == EMouseCaptureMode::CapturePermanently_IncludingInitialMouseDown));
+	}
+
 	/**
 	 * Called when the top level window associated with the viewport has been requested to close.
 	 * At this point, the viewport has not been closed and the operation may be canceled.
