Fix streamer calculating mesh screen size incorrectly

diff --git a/Engine/Source/Runtime/Engine/Private/Streaming/StreamingTexture.cpp b/Engine/Source/Runtime/Engine/Private/Streaming/StreamingTexture.cpp
--- a/Engine/Source/Runtime/Engine/Private/Streaming/StreamingTexture.cpp
+++ b/Engine/Source/Runtime/Engine/Private/Streaming/StreamingTexture.cpp
@@ -73,7 +73,10 @@ void FStreamingRenderAsset::UpdateStaticData(const FRenderAssetStreamingSettings
 				for (int32 Idx = 0; Idx < MaxNumMeshLODs; ++Idx)
 				{
 					const int32 LODIdx = FMath::Max(MipCount - Idx - 1, 0);
-					LODScreenSizes[Idx] = StaticMesh->RenderData->ScreenSize[LODIdx].GetValueForFeatureLevel(GMaxRHIFeatureLevel);
+					// Screen sizes stored on assets are 2R/D where R is the radius of bounding spheres and D is the
+					// distance from view origins to bounds origins. The factor calculated by the streamer, however,
+					// is R/D so multiply 0.5 here
+					LODScreenSizes[Idx] = StaticMesh->RenderData->ScreenSize[LODIdx].GetValueForFeatureLevel(GMaxRHIFeatureLevel) * 0.5f;
 				}
 			}
 			else // AT_SkeletalMesh
@@ -83,7 +86,7 @@ void FStreamingRenderAsset::UpdateStaticData(const FRenderAssetStreamingSettings
 				for (int32 Idx = 0; Idx < MaxNumMeshLODs; ++Idx)
 				{
 					const int32 LODIdx = FMath::Max(MipCount - Idx - 1, 0);
-					LODScreenSizes[Idx] = LODInfos[LODIdx].ScreenSize.GetValueForFeatureLevel(GMaxRHIFeatureLevel);
+					LODScreenSizes[Idx] = LODInfos[LODIdx].ScreenSize.GetValueForFeatureLevel(GMaxRHIFeatureLevel) * 0.5f;
 				}
 			}
 		}
