Fix for deadlock in streaming manager, memleak, and missing stream chunk loading.

diff --git a/Engine/Source/Runtime/Engine/Private/AudioStreaming.cpp b/Engine/Source/Runtime/Engine/Private/AudioStreaming.cpp
--- a/Engine/Source/Runtime/Engine/Private/AudioStreaming.cpp
+++ b/Engine/Source/Runtime/Engine/Private/AudioStreaming.cpp
@@ -530,68 +530,64 @@ void FAudioStreamingManager::ProcessPendingAsyncFileResults()
 
 void FAudioStreamingManager::UpdateResourceStreaming(float DeltaTime, bool bProcessEverything /*= false*/)
 {
-	LLM_SCOPE(ELLMTag::AudioMisc);
+	LLM_SCOPE(ELLMTag::Audio);
 
-	if (CriticalSection.TryLock())
+	FScopeLock Lock(&CriticalSection);
+
+	for (auto& WavePair : StreamingSoundWaves)
 	{
-		for (auto& WavePair : StreamingSoundWaves)
-		{
-			WavePair.Value->UpdateStreamingStatus();
-		}
+		WavePair.Value->UpdateStreamingStatus();
+	}
 
-		// Process any async file requests after updating the stream status
-		ProcessPendingAsyncFileResults();
+	// Process any async file requests after updating the stream status
+	ProcessPendingAsyncFileResults();
 
-		for (auto Source : StreamingSoundSources)
+	for (auto Source : StreamingSoundSources)
+	{
+		const FWaveInstance* WaveInstance = Source->GetWaveInstance();
+		USoundWave* Wave = WaveInstance ? WaveInstance->WaveData : nullptr;
+		if (Wave)
 		{
-			const FWaveInstance* WaveInstance = Source->GetWaveInstance();
-			USoundWave* Wave = WaveInstance ? WaveInstance->WaveData : nullptr;
-			if (Wave)
-			{
-				FStreamingWaveData** WaveDataPtr = StreamingSoundWaves.Find(Wave);
+			FStreamingWaveData** WaveDataPtr = StreamingSoundWaves.Find(Wave);
 
-				if (WaveDataPtr && (*WaveDataPtr)->PendingChunkChangeRequestStatus.GetValue() == AudioState_ReadyFor_Requests)
+			if (WaveDataPtr && (*WaveDataPtr)->PendingChunkChangeRequestStatus.GetValue() == AudioState_ReadyFor_Requests)
+			{
+				FStreamingWaveData* WaveData = *WaveDataPtr;
+				// Request the chunk the source is using and the one after that
+				FWaveRequest& WaveRequest = GetWaveRequest(Wave);
+				const FSoundBuffer* SoundBuffer = Source->GetBuffer();
+				if (SoundBuffer)
 				{
-					FStreamingWaveData* WaveData = *WaveDataPtr;
-					// Request the chunk the source is using and the one after that
-					FWaveRequest& WaveRequest = GetWaveRequest(Wave);
-					const FSoundBuffer* SoundBuffer = Source->GetBuffer();
-					if (SoundBuffer)
+					int32 SourceChunk = SoundBuffer->GetCurrentChunkIndex();
+					if (SourceChunk >= 0 && SourceChunk < Wave->RunningPlatformData->NumChunks)
+					{
+						WaveRequest.RequiredIndices.AddUnique(SourceChunk);
+						WaveRequest.RequiredIndices.AddUnique((SourceChunk + 1) % Wave->RunningPlatformData->NumChunks);
+						WaveRequest.bPrioritiseRequest = true;
+					}
+					else
 					{
-						int32 SourceChunk = SoundBuffer->GetCurrentChunkIndex();
-						if (SourceChunk >= 0 && SourceChunk < Wave->RunningPlatformData->NumChunks)
-						{
-							WaveRequest.RequiredIndices.AddUnique(SourceChunk);
-							WaveRequest.RequiredIndices.AddUnique((SourceChunk + 1) % Wave->RunningPlatformData->NumChunks);
-							WaveRequest.bPrioritiseRequest = true;
-						}
-						else
-						{
-							UE_LOG(LogAudio, Log, TEXT("Invalid chunk request curIndex=%d numChunks=%d\n"), SourceChunk, Wave->RunningPlatformData->NumChunks);
-						}
+						UE_LOG(LogAudio, Log, TEXT("Invalid chunk request curIndex=%d numChunks=%d\n"), SourceChunk, Wave->RunningPlatformData->NumChunks);
 					}
 				}
 			}
 		}
+	}
+	for (auto Iter = WaveRequests.CreateIterator(); Iter; ++Iter)
+	{
+		USoundWave* Wave = Iter.Key();
+		FStreamingWaveData* WaveData = StreamingSoundWaves.FindRef(Wave);
 
-		for (auto Iter = WaveRequests.CreateIterator(); Iter; ++Iter)
+		if (WaveData && WaveData->PendingChunkChangeRequestStatus.GetValue() == AudioState_ReadyFor_Requests)
 		{
-			USoundWave* Wave = Iter.Key();
-			FStreamingWaveData* WaveData = StreamingSoundWaves.FindRef(Wave);
-
-			if (WaveData && WaveData->PendingChunkChangeRequestStatus.GetValue() == AudioState_ReadyFor_Requests)
-			{
-				WaveData->UpdateChunkRequests(Iter.Value());
-				WaveData->UpdateStreamingStatus();
-				Iter.RemoveCurrent();
-			}
+			WaveData->UpdateChunkRequests(Iter.Value());
+			WaveData->UpdateStreamingStatus();
+			Iter.RemoveCurrent();
 		}
-
-		// Process any async file requests after updating the streaming wave data stream statuses
-		ProcessPendingAsyncFileResults();
-
-		CriticalSection.Unlock();
 	}
+ 
+	// Process any async file requests after updating the streaming wave data stream statuses
+	ProcessPendingAsyncFileResults();
 }
 
 int32 FAudioStreamingManager::BlockTillAllRequestsFinished(float TimeLimit, bool)
@@ -803,14 +799,19 @@ bool FAudioStreamingManager::IsManagedStreamingSoundSource(const FSoundSource* S
 
 const uint8* FAudioStreamingManager::GetLoadedChunk(const USoundWave* SoundWave, uint32 ChunkIndex, uint32* OutChunkSize) const
 {
-	FScopeLock Lock(&CriticalSection);
-
 	// Check for the spoof of failing to load a stream chunk
 	if (SpoofFailedStreamChunkLoad > 0)
 	{
 		return nullptr;
 	}
 
+	// Get a critical section lock, but don't block others. This avoids rare conditions of deadlocking in this streaming manager with UpdateStreamingResources.
+	// But allows us to make sure we can get the requested loaded chunk. Most of the time, this function is called in an async task.
+	while (!CriticalSection.TryLock())
+	{
+		FPlatformProcess::Sleep(0.001f);
+	}
+
 	const FStreamingWaveData* WaveData = StreamingSoundWaves.FindRef(SoundWave);
 	if (WaveData)
 	{
@@ -820,17 +821,20 @@ const uint8* FAudioStreamingManager::GetLoadedChunk(const USoundWave* SoundWave,
 			{
 				if (WaveData->LoadedChunks[Index].Index == ChunkIndex)
 				{
-					if(OutChunkSize != NULL)
+					if (OutChunkSize != NULL)
 					{
 						// Return the size of the audio data within the chunk, not the chunk itself since this chunk could be zero-padded
 						*OutChunkSize = WaveData->LoadedChunks[Index].AudioDataSize;
 					}
-					
+
+					CriticalSection.Unlock();
 					return WaveData->LoadedChunks[Index].Data;
 				}
 			}
 		}
 	}
+
+	CriticalSection.Unlock();
 	return NULL;
 }
 
