Support twist swing angle from Delta Rotation based on Reference pose rotation

diff --git a/Engine/Source/Runtime/Engine/Classes/Components/SkinnedMeshComponent.h b/Engine/Source/Runtime/Engine/Classes/Components/SkinnedMeshComponent.h
--- a/Engine/Source/Runtime/Engine/Classes/Components/SkinnedMeshComponent.h
+++ b/Engine/Source/Runtime/Engine/Classes/Components/SkinnedMeshComponent.h
@@ -689,6 +689,21 @@ public:
 	UFUNCTION(BlueprintCallable, Category = "Components|SkinnedMesh")
 	FTransform GetDeltaTransformFromRefPose(FName BoneName, FName BaseName = NAME_None) const;
 
+	/** 
+	 * Get Twist and Swing Angle in Degree of Delta Rotation from Reference Pose in Local space 
+	 *
+	 * First this function gets rotation of current, and rotation of ref pose in local space, and 
+	 * And gets twist/swing angle value from refpose aligned. 
+	 * 
+	 * @param BoneName Name of the bone
+	 * @param OutTwistAngle TwistAngle in degree
+	 * @param OutSwingAngle SwingAngle in degree
+	 *
+	 * @return true if succeed. False otherwise. Often due to incorrect bone name. 
+	 */
+	UFUNCTION(BlueprintCallable, Category = "Components|SkinnedMesh")
+	bool GetTwistAndSwingAngleOfDeltaRotationFromRefPose(FName BoneName, float& OutTwistAngle, float& OutSwingAngle) const;
+
 public:
 	//~ Begin UObject Interface
 	virtual void BeginDestroy() override;
diff --git a/Engine/Source/Runtime/Engine/Private/Components/SkinnedMeshComponent.cpp b/Engine/Source/Runtime/Engine/Private/Components/SkinnedMeshComponent.cpp
--- a/Engine/Source/Runtime/Engine/Private/Components/SkinnedMeshComponent.cpp
+++ b/Engine/Source/Runtime/Engine/Private/Components/SkinnedMeshComponent.cpp
@@ -1232,6 +1232,40 @@ FTransform USkinnedMeshComponent::GetDeltaTransformFromRefPose(FName BoneName, F
 	return FTransform::Identity;
 }
 
+bool USkinnedMeshComponent::GetTwistAndSwingAngleOfDeltaRotationFromRefPose(FName BoneName, float& OutTwistAngle, float& OutSwingAngle) const
+{
+	const FReferenceSkeleton& RefSkeleton = SkeletalMesh->RefSkeleton;
+	const int32 BoneIndex = GetBoneIndex(BoneName);
+	if (BoneIndex != INDEX_NONE)
+	{
+		FTransform LocalTransform = GetComponentSpaceTransforms()[BoneIndex];
+		FTransform ReferenceTransform = RefSkeleton.GetRefBonePose()[BoneIndex];
+		FName ParentName = GetParentBone(BoneName);
+		int32 ParentIndex = INDEX_NONE;
+		if (ParentName != NAME_None)
+		{
+			ParentIndex = GetBoneIndex(ParentName);
+		}
+
+		if (ParentIndex != INDEX_NONE)
+		{
+			LocalTransform = LocalTransform.GetRelativeTransform(GetComponentSpaceTransforms()[ParentIndex]);
+		}
+
+		FQuat Swing, Twist;
+
+		// figure out based on ref pose rotation, and calculate twist based on that 
+		FVector TwistAxis = ReferenceTransform.GetRotation().Vector();
+		ensure(TwistAxis.IsNormalized());
+		LocalTransform.GetRotation().ToSwingTwist(TwistAxis, Swing, Twist);
+		OutTwistAngle = FMath::RadiansToDegrees(Twist.GetAngle());
+		OutSwingAngle = FMath::RadiansToDegrees(Swing.GetAngle());
+		return true;
+	}
+
+	return false;
+}
+
 void USkinnedMeshComponent::GetBoneNames(TArray<FName>& BoneNames)
 {
 	if (SkeletalMesh == NULL)
