Added engine cvar r.SkylightIntensityMultiplier since it's generally useful for scalability in dynamically lit games

diff --git a/Engine/Source/Runtime/Engine/Private/Components/SkyLightComponent.cpp b/Engine/Source/Runtime/Engine/Private/Components/SkyLightComponent.cpp
--- a/Engine/Source/Runtime/Engine/Private/Components/SkyLightComponent.cpp
+++ b/Engine/Source/Runtime/Engine/Private/Components/SkyLightComponent.cpp
@@ -58,6 +58,14 @@ FAutoConsoleVariableRef CVarUpdateSkylightsEveryFrame(
 	ECVF_Default
 	);
 
+float GSkylightIntensityMultiplier = 1.0f;
+FAutoConsoleVariableRef CVarSkylightIntensityMultiplier(
+	TEXT("r.SkylightIntensityMultiplier"),
+	GSkylightIntensityMultiplier,
+	TEXT("Intensity scale on Stationary and Movable skylights.  This is useful to control overall lighting contrast in dynamically lit games with scalability levels which disable Ambient Occlusion.  For example, if medium quality disables SSAO and DFAO, reduce skylight intensity."),
+	ECVF_Scalability | ECVF_RenderThreadSafe
+	);
+
 void FSkyTextureCubeResource::InitRHI()
 {
 	if (GetFeatureLevel() >= ERHIFeatureLevel::SM4 || GSupportsRenderTargetFormat_PF_FloatRGBA)
@@ -145,6 +153,11 @@ void FSkyLightSceneProxy::Initialize(
 #endif
 }
 
+FLinearColor FSkyLightSceneProxy::GetEffectiveLightColor() const
+{
+	return LightColor * GSkylightIntensityMultiplier;
+}
+
 FSkyLightSceneProxy::FSkyLightSceneProxy(const USkyLightComponent* InLightComponent)
 	: LightComponent(InLightComponent)
 	, ProcessedTexture(InLightComponent->ProcessedSkyTexture)
@@ -339,7 +352,7 @@ void USkyLightComponent::UpdateLimitedRenderingStateFast()
 		ENQUEUE_RENDER_COMMAND(FFastUpdateSkyLightCommand)(
 			[LightSceneProxy, InLightColor, InIndirectLightingIntensity, InVolumetricScatteringIntensity](FRHICommandList& RHICmdList)
 			{
-				LightSceneProxy->LightColor = InLightColor;
+				LightSceneProxy->SetLightColor(InLightColor);
 				LightSceneProxy->IndirectLightingIntensity = InIndirectLightingIntensity;
 				LightSceneProxy->VolumetricScatteringIntensity = InVolumetricScatteringIntensity;
 			});
diff --git a/Engine/Source/Runtime/Engine/Public/SceneManagement.h b/Engine/Source/Runtime/Engine/Public/SceneManagement.h
--- a/Engine/Source/Runtime/Engine/Public/SceneManagement.h
+++ b/Engine/Source/Runtime/Engine/Public/SceneManagement.h
@@ -981,7 +981,6 @@ public:
 	uint8 bHasStaticLighting:1;
 	uint8 bCastVolumetricShadow:1;
 	TEnumAsByte<EOcclusionCombineMode> OcclusionCombineMode;
-	FLinearColor LightColor;
 	float AverageBrightness;
 	float IndirectLightingIntensity;
 	float VolumetricScatteringIntensity;
@@ -1016,6 +1015,15 @@ public:
 	FRWBuffer SkyLightMipTreePdfNegZ;
 	FRWBuffer SolidAnglePdf;
 #endif
+
+	void SetLightColor(const FLinearColor& InColor)
+	{
+		LightColor = InColor;
+	}
+	FLinearColor GetEffectiveLightColor() const;
+
+private:
+	FLinearColor LightColor;
 };
 
 
diff --git a/Engine/Source/Runtime/Renderer/Private/SceneRendering.cpp b/Engine/Source/Runtime/Renderer/Private/SceneRendering.cpp
--- a/Engine/Source/Runtime/Renderer/Private/SceneRendering.cpp
+++ b/Engine/Source/Runtime/Renderer/Private/SceneRendering.cpp
@@ -1367,7 +1367,7 @@ void FViewInfo::SetupUniformBufferParameters(
 	{
 		FSkyLightSceneProxy* SkyLight = Scene->SkyLight;
 
-		ViewUniformShaderParameters.SkyLightColor = SkyLight->LightColor;
+		ViewUniformShaderParameters.SkyLightColor = SkyLight->GetEffectiveLightColor();
 
 		bool bApplyPrecomputedBentNormalShadowing = 
 			SkyLight->bCastShadows 
