#ue4 - CharacterMovementComponent saved moves are now allowed to combine when changing between static base components. Since we don't send the component any longer in the static case (ie ServerMoveNoBase), we don't care if the static base changes. This saves bandwidth and reduces server RPC cost when running across many static pieces. Changing dynamic pieces still requires flushing the move since relative location changes.

CVar "p.NetPreventMoveCombiningOnStaticBaseChange" can be enabled to revert to old behavior.

diff --git a/Engine/Source/Runtime/Engine/Private/Components/CharacterMovementComponent.cpp b/Engine/Source/Runtime/Engine/Private/Components/CharacterMovementComponent.cpp
--- a/Engine/Source/Runtime/Engine/Private/Components/CharacterMovementComponent.cpp
+++ b/Engine/Source/Runtime/Engine/Private/Components/CharacterMovementComponent.cpp
@@ -137,6 +137,14 @@ namespace CharacterMovementCVars
 		TEXT("0: Disable, 1: Enable"),
 		ECVF_Default);
 
+	static int32 NetPreventMoveCombiningOnStaticBaseChange = 0;
+	FAutoConsoleVariableRef CVarNetPreventMoveCombiningOnStaticBaseChange(
+		TEXT("p.NetPreventMoveCombiningOnStaticBaseChange"),
+		NetPreventMoveCombiningOnStaticBaseChange,
+		TEXT("Whether to enable move combining on the client to reduce bandwidth by combining similar moves.\n")
+		TEXT("0: Disable, 1: Enable"),
+		ECVF_Default);
+
 	static int32 NetUseClientTimestampForReplicatedTransform = 1;
 	FAutoConsoleVariableRef CVarNetUseClientTimestampForReplicatedTransform(
 		TEXT("p.NetUseClientTimestampForReplicatedTransform"),
@@ -8470,10 +8478,12 @@ void UCharacterMovementComponent::ServerMove(float TimeStamp, FVector_NetQuantiz
 {
 	if (MovementBaseUtility::IsDynamicBase(ClientMovementBase))
 	{
+		//UE_LOG(LogCharacterMovement, Log, TEXT("ServerMove: base %s"), *ClientMovementBase->GetName());
 		CharacterOwner->ServerMove(TimeStamp, InAccel, ClientLoc, CompressedMoveFlags, ClientRoll, View, ClientMovementBase, ClientBaseBoneName, ClientMovementMode);
 	}
 	else
 	{
+		//UE_LOG(LogCharacterMovement, Log, TEXT("ServerMoveNoBase"));
 		CharacterOwner->ServerMoveNoBase(TimeStamp, InAccel, ClientLoc, CompressedMoveFlags, ClientRoll, View, ClientMovementMode);
 	}
 }
@@ -8751,10 +8761,12 @@ void UCharacterMovementComponent::ServerMoveDual(float TimeStamp0, FVector_NetQu
 {
 	if (MovementBaseUtility::IsDynamicBase(ClientMovementBase))
 	{
+		//UE_LOG(LogCharacterMovement, Log, TEXT("ServerMoveDual: base %s"), *ClientMovementBase->GetName());
 		CharacterOwner->ServerMoveDual(TimeStamp0, InAccel0, PendingFlags, View0, TimeStamp, InAccel, ClientLoc, NewFlags, ClientRoll, View, ClientMovementBase, ClientBaseBoneName, ClientMovementMode);
 	}
 	else
 	{
+		//UE_LOG(LogCharacterMovement, Log, TEXT("ServerMoveDualNoBase"));
 		CharacterOwner->ServerMoveDualNoBase(TimeStamp0, InAccel0, PendingFlags, View0, TimeStamp, InAccel, ClientLoc, NewFlags, ClientRoll, View, ClientMovementMode);
 	}
 }
@@ -10671,14 +10683,30 @@ bool FSavedMove_Character::CanCombineWith(const FSavedMovePtr& NewMovePtr, AChar
 		return false;
 	}
 
-	if (StartBase != NewMove->StartBase)
+	const UPrimitiveComponent* OldBasePtr = StartBase.Get();
+	const UPrimitiveComponent* NewBasePtr = NewMove->StartBase.Get();
+	const bool bDynamicBaseOld = MovementBaseUtility::IsDynamicBase(OldBasePtr);
+	const bool bDynamicBaseNew = MovementBaseUtility::IsDynamicBase(NewBasePtr);
+
+	// Change between static/dynamic requires separate moves (position sent as world vs relative)
+	if (bDynamicBaseOld != bDynamicBaseNew)
 	{
 		return false;
 	}
 
-	if (StartBoneName != NewMove->StartBoneName)
+	// Only need to prevent combining when on a dynamic base that changes (unless forced off via CVar). Again, because relative location can change.
+	const bool bPreventOnStaticBaseChange = (CharacterMovementCVars::NetPreventMoveCombiningOnStaticBaseChange != 0);
+	if (bPreventOnStaticBaseChange || (bDynamicBaseOld || bDynamicBaseNew))
 	{
-		return false;
+		if (OldBasePtr != NewBasePtr)
+		{
+			return false;
+		}
+
+		if (StartBoneName != NewMove->StartBoneName)
+		{
+			return false;
+		}
 	}
 
 	if (EndPackedMovementMode != NewMove->StartPackedMovementMode)
