Expose configuration settings for Bad and Severe thresholds for PacketLoss and Ping.
These are primarily used for analytic type events and / or potential cheat detection.

diff --git a/Engine/Source/Runtime/Engine/Classes/GameFramework/GameNetworkManager.h b/Engine/Source/Runtime/Engine/Classes/GameFramework/GameNetworkManager.h
--- a/Engine/Source/Runtime/Engine/Classes/GameFramework/GameNetworkManager.h
+++ b/Engine/Source/Runtime/Engine/Classes/GameFramework/GameNetworkManager.h
@@ -28,6 +28,25 @@ class ENGINE_API AGameNetworkManager : public AInfo
 {
 	GENERATED_UCLASS_BODY()
 
+	//======================================================================================================================
+	// Analytics
+
+	/* If packet loss goes over this value, we have bad packet loss. Value is between 0 and 1.*/
+	UPROPERTY(GlobalConfig)
+	float BadPacketLossThreshold = 0.05f;
+
+	/* If the packet loss goes over this threshold, we have severe packet loss. Value is between 0 and 1*/
+	UPROPERTY(GlobalConfig)
+	float SeverePacketLossThreshold = 0.15f;
+
+	/** The point we determine the server is either delaying packets or has bad upstream. */
+	UPROPERTY(GlobalConfig)
+	int32 BadPingThreshold;
+
+	/** Similar to BadPingThreshold, but used to track exceptionally bad pings. */
+	UPROPERTY(GlobalConfig)
+	int32 SeverePingThreshold;
+
 	//======================================================================================================================
 	// Listen server dynamic netspeed adjustment
 	
@@ -70,10 +89,6 @@ class ENGINE_API AGameNetworkManager : public AInfo
 	UPROPERTY(config)
 	float StandbyTxCheatTime;
 
-	/** The point we determine the server is either delaying packets or has bad upstream */
-	UPROPERTY(config)
-	int32 BadPingThreshold;
-
 	/** The percentage of clients missing RX data before triggering the standby code */
 	UPROPERTY(config)
 	float PercentMissingForRxStandby;
diff --git a/Engine/Source/Runtime/Engine/Private/GameNetworkManager.cpp b/Engine/Source/Runtime/Engine/Private/GameNetworkManager.cpp
--- a/Engine/Source/Runtime/Engine/Private/GameNetworkManager.cpp
+++ b/Engine/Source/Runtime/Engine/Private/GameNetworkManager.cpp
@@ -17,6 +17,11 @@ DEFINE_LOG_CATEGORY_STATIC(LogGameNetworkManager, Log, All);
 AGameNetworkManager::AGameNetworkManager(const FObjectInitializer& ObjectInitializer)
 	: Super(ObjectInitializer)
 {
+	BadPacketLossThreshold = 0.05f;
+	SeverePacketLossThreshold = 0.15f;
+	BadPingThreshold = 200;
+	SeverePingThreshold = 500;
+
 	MoveRepSize = 42.0f;
 	MAXPOSITIONERRORSQUARED = 3.0f;
 	MAXNEARZEROVELOCITYSQUARED = 9.0f;
