#ue4 - Client character move combining fixes for bandwidth and server perf.

- Allow character client move combining when attached to objects even if the previous location was "overlapping" something blocking. Attachment carries you along anyway regardless of those cases, and you can be "blocked" by what you're attached to as well. This was preventing client throttling when attached in this way.

- Move combining is allowed when attachment relative location and rotation changes in a very small threshold, to account for jitter from transform updates especially when attached to simulating objects.

- GetClientNetSendDeltaTime() allows small control rotation changes when using the "stationary" send rate, to account for jitter when attached/standing on objects.

diff --git a/Engine/Source/Runtime/Engine/Private/Components/CharacterMovementComponent.cpp b/Engine/Source/Runtime/Engine/Private/Components/CharacterMovementComponent.cpp
--- a/Engine/Source/Runtime/Engine/Private/Components/CharacterMovementComponent.cpp
+++ b/Engine/Source/Runtime/Engine/Private/Components/CharacterMovementComponent.cpp
@@ -137,14 +137,35 @@ namespace CharacterMovementCVars
 		TEXT("0: Disable, 1: Enable"),
 		ECVF_Default);
 
-	static int32 NetPreventMoveCombiningOnStaticBaseChange = 0;
-	FAutoConsoleVariableRef CVarNetPreventMoveCombiningOnStaticBaseChange(
-		TEXT("p.NetPreventMoveCombiningOnStaticBaseChange"),
-		NetPreventMoveCombiningOnStaticBaseChange,
-		TEXT("Whether to enable move combining on the client to reduce bandwidth by combining similar moves.\n")
+	static int32 NetEnableMoveCombiningOnStaticBaseChange = 1;
+	FAutoConsoleVariableRef CVarNetEnableMoveCombiningOnStaticBaseChange(
+		TEXT("p.NetEnableMoveCombiningOnStaticBaseChange"),
+		NetEnableMoveCombiningOnStaticBaseChange,
+		TEXT("Whether to allow combining client moves when moving between static geometry.\n")
 		TEXT("0: Disable, 1: Enable"),
 		ECVF_Default);
 
+	static float NetMoveCombiningAttachedLocationTolerance = 0.01f;
+	FAutoConsoleVariableRef CVarNetMoveCombiningAttachedLocationTolerance(
+		TEXT("p.NetMoveCombiningAttachedLocationTolerance"),
+		NetMoveCombiningAttachedLocationTolerance,
+		TEXT("Tolerance for relative location attachment change when combining moves. Small tolerances allow for very slight jitter due to transform updates."),
+		ECVF_Default);
+
+	static float NetMoveCombiningAttachedRotationTolerance = 0.01f;
+	FAutoConsoleVariableRef CVarNetMoveCombiningAttachedRotationTolerance(
+		TEXT("p.NetMoveCombiningAttachedRotationTolerance"),
+		NetMoveCombiningAttachedRotationTolerance,
+		TEXT("Tolerance for relative rotation attachment change when combining moves. Small tolerances allow for very slight jitter due to transform updates."),
+		ECVF_Default);
+
+	static float NetStationaryRotationTolerance = 0.1f;
+	FAutoConsoleVariableRef CVarNetStationaryRotationTolerance(
+		TEXT("p.NetStationaryRotationTolerance"),
+		NetStationaryRotationTolerance,
+		TEXT("Tolerance for GetClientNetSendDeltaTime() to remain throttled when small control rotation changes occur."),
+		ECVF_Default);
+
 	static int32 NetUseClientTimestampForReplicatedTransform = 1;
 	FAutoConsoleVariableRef CVarNetUseClientTimestampForReplicatedTransform(
 		TEXT("p.NetUseClientTimestampForReplicatedTransform"),
@@ -7889,13 +7910,14 @@ void UCharacterMovementComponent::ReplicateMoveToServer(float DeltaTime, const F
 	// do not combine moves which have different TimeStamps (before and after reset).
 	if (const FSavedMove_Character* PendingMove = ClientData->PendingMove.Get())
 	{
-		if (!PendingMove->bOldTimeStampBeforeReset && PendingMove->CanCombineWith(NewMovePtr, CharacterOwner, ClientData->MaxMoveDeltaTime * CharacterOwner->GetActorTimeDilation(*MyWorld)))
+		if (PendingMove->CanCombineWith(NewMovePtr, CharacterOwner, ClientData->MaxMoveDeltaTime * CharacterOwner->GetActorTimeDilation(*MyWorld)))
 		{
 			SCOPE_CYCLE_COUNTER(STAT_CharacterMovementCombineNetMove);
 
 			// Only combine and move back to the start location if we don't move back in to a spot that would make us collide with something new.
 			const FVector OldStartLocation = PendingMove->GetRevertedLocation();
-			if (!OverlapTest(OldStartLocation, PendingMove->StartRotation.Quaternion(), UpdatedComponent->GetCollisionObjectType(), GetPawnCapsuleCollisionShape(SHRINK_None), CharacterOwner))
+			const bool bAttachedToObject = (NewMovePtr->StartAttachParent != nullptr);
+			if (bAttachedToObject || !OverlapTest(OldStartLocation, PendingMove->StartRotation.Quaternion(), UpdatedComponent->GetCollisionObjectType(), GetPawnCapsuleCollisionShape(SHRINK_None), CharacterOwner))
 			{
 				// Avoid updating Mesh bones to physics during the teleport back, since PerformMovement() will update it right away anyway below.
 				// Note: this must be before the FScopedMovementUpdate below, since that scope is what actually moves the character and mesh.
@@ -10598,7 +10620,7 @@ float UCharacterMovementComponent::GetClientNetSendDeltaTime(const APlayerContro
 		}
 
 		// Lower frequency for standing still and not rotating camera
-		if (Acceleration.IsZero() && Velocity.IsZero() && ClientData->LastAckedMove.IsValid() && ClientData->LastAckedMove->StartControlRotation.Equals(PC->GetControlRotation()))
+		if (Acceleration.IsZero() && Velocity.IsZero() && ClientData->LastAckedMove.IsValid() && ClientData->LastAckedMove->StartControlRotation.Equals(PC->GetControlRotation(), CharacterMovementCVars::NetStationaryRotationTolerance))
 		{
 			NetMoveDelta = FMath::Max(GameNetworkManager->ClientNetSendMoveDeltaTimeStationary, NetMoveDelta);
 		}
@@ -10616,6 +10638,11 @@ bool FSavedMove_Character::CanCombineWith(const FSavedMovePtr& NewMovePtr, AChar
 		return false;
 	}
 
+	if (bOldTimeStampBeforeReset)
+	{
+		return false;
+	}
+
 	// Cannot combine moves which contain root motion for now.
 	// @fixme laurent - we should be able to combine most of them though, but current scheme of resetting pawn location and resimulating forward doesn't work.
 	// as we don't want to tick montage twice (so we don't fire events twice). So we need to rearchitecture this so we tick only the second part of the move, and reuse the first part.
@@ -10695,7 +10722,7 @@ bool FSavedMove_Character::CanCombineWith(const FSavedMovePtr& NewMovePtr, AChar
 	}
 
 	// Only need to prevent combining when on a dynamic base that changes (unless forced off via CVar). Again, because relative location can change.
-	const bool bPreventOnStaticBaseChange = (CharacterMovementCVars::NetPreventMoveCombiningOnStaticBaseChange != 0);
+	const bool bPreventOnStaticBaseChange = (CharacterMovementCVars::NetEnableMoveCombiningOnStaticBaseChange == 0);
 	if (bPreventOnStaticBaseChange || (bDynamicBaseOld || bDynamicBaseNew))
 	{
 		if (OldBasePtr != NewBasePtr)
@@ -10742,14 +10769,16 @@ bool FSavedMove_Character::CanCombineWith(const FSavedMovePtr& NewMovePtr, AChar
 	if (NewStartAttachParent != nullptr)
 	{
 		// If attached, no combining if relative location changed.
-		if (StartAttachRelativeLocation != NewMove->StartAttachRelativeLocation)
+		const FVector RelativeLocationDelta = (StartAttachRelativeLocation - NewMove->StartAttachRelativeLocation);
+		if (!RelativeLocationDelta.IsNearlyZero(CharacterMovementCVars::NetMoveCombiningAttachedLocationTolerance))
 		{
+			//UE_LOG(LogCharacterMovement, Warning, TEXT("NoCombine: DeltaLocation(%s)"), *RelativeLocationDelta.ToString());
 			return false;
 		}
 		// For rotation, Yaw doesn't matter for capsules
 		FRotator RelativeRotationDelta = StartAttachRelativeRotation - NewMove->StartAttachRelativeRotation;
 		RelativeRotationDelta.Yaw = 0.0f;
-		if (!RelativeRotationDelta.IsNearlyZero())
+		if (!RelativeRotationDelta.IsNearlyZero(CharacterMovementCVars::NetMoveCombiningAttachedRotationTolerance))
 		{
 			return false;
 		}
