Add Legacy Reverb Disable CVar for old console audio backends (XAudio2 & PS4)

diff --git a/Engine/Source/Runtime/Engine/Private/AudioDevice.cpp b/Engine/Source/Runtime/Engine/Private/AudioDevice.cpp
--- a/Engine/Source/Runtime/Engine/Private/AudioDevice.cpp
+++ b/Engine/Source/Runtime/Engine/Private/AudioDevice.cpp
@@ -119,6 +119,15 @@ FAutoConsoleVariableRef CVarAllowAudioSpatializationCVar(
 	TEXT("0: Disable, >0: Enable"),
 	ECVF_Default);
 
+static int32 DisableLegacyReverb = 0;
+FAutoConsoleVariableRef CVarDisableLegacyReverb(
+	TEXT("au.DisableLegacyReverb"),
+	DisableLegacyReverb,
+	TEXT("Disables reverb on legacy audio backends.\n")
+	TEXT("0: Enabled, 1: Disabled"),
+	ECVF_Default);
+
+
 /*-----------------------------------------------------------------------------
 FDynamicParameter implementation.
 -----------------------------------------------------------------------------*/
@@ -5458,6 +5467,11 @@ void FAudioDevice::StopSoundsUsingResource(USoundWave* SoundWave, TArray<UAudioC
 	}
 }
 
+bool FAudioDevice::LegacyReverbDisabled()
+{
+	return DisableLegacyReverb != 0;
+}
+
 void FAudioDevice::RegisterPluginListener(const TAudioPluginListenerPtr PluginListener)
 {
 	PluginListeners.AddUnique(PluginListener);
diff --git a/Engine/Source/Runtime/Engine/Public/AudioDevice.h b/Engine/Source/Runtime/Engine/Public/AudioDevice.h
--- a/Engine/Source/Runtime/Engine/Public/AudioDevice.h
+++ b/Engine/Source/Runtime/Engine/Public/AudioDevice.h
@@ -598,6 +598,8 @@ public:
 	 */
 	void StopSoundsUsingResource(USoundWave* SoundWave, TArray<UAudioComponent*>* StoppedComponents = nullptr);
 
+	static bool LegacyReverbDisabled();
+
 #if WITH_EDITOR
 	/** Deals with anything audio related that should happen when PIE starts */
 	void OnBeginPIE(const bool bIsSimulating);
diff --git a/Engine/Source/Runtime/Windows/XAudio2/Private/XAudio2Source.cpp b/Engine/Source/Runtime/Windows/XAudio2/Private/XAudio2Source.cpp
--- a/Engine/Source/Runtime/Windows/XAudio2/Private/XAudio2Source.cpp
+++ b/Engine/Source/Runtime/Windows/XAudio2/Private/XAudio2Source.cpp
@@ -20,6 +20,7 @@
 #include "ContentStreaming.h"
 #include "HAL/LowLevelMemTracker.h"
 
+
 /*------------------------------------------------------------------------------------
 	For muting user soundtracks during cinematics
 ------------------------------------------------------------------------------------*/
@@ -558,7 +559,10 @@ bool FXAudio2SoundSource::Init(FWaveInstance* InWaveInstance)
 		SCOPE_CYCLE_COUNTER(STAT_AudioSourceInitTime);
 
 		// Set whether to apply reverb
-		SetReverbApplied(Effects->ReverbEffectVoice != nullptr);
+		if (!FAudioDevice::LegacyReverbDisabled())
+		{
+			SetReverbApplied(Effects->ReverbEffectVoice != nullptr);
+		}
 
 		// Create a new source if we haven't already
 		if (CreateSource())
