Adding demo.LoopCount for playing X times rather than infinitely

diff --git a/Engine/Source/Runtime/Engine/Private/DemoNetDriver.cpp b/Engine/Source/Runtime/Engine/Private/DemoNetDriver.cpp
--- a/Engine/Source/Runtime/Engine/Private/DemoNetDriver.cpp
+++ b/Engine/Source/Runtime/Engine/Private/DemoNetDriver.cpp
@@ -65,6 +65,9 @@ static TAutoConsoleVariable<int32> CVarLoopDemo(TEXT("demo.Loop"), 0, TEXT("<1>
 static TAutoConsoleVariable<int32> CVarDemoFastForwardIgnoreRPCs( TEXT( "demo.FastForwardIgnoreRPCs" ), 1, TEXT( "If true, RPCs will be discarded during playback fast forward." ) );
 static TAutoConsoleVariable<int32> CVarDemoLateActorDormancyCheck(TEXT("demo.LateActorDormancyCheck"), 1, TEXT("If true, check if an actor should become dormant as late as possible- when serializing it to the demo archive."));
 
+static int32 GDemoLoopCount = 0;
+static FAutoConsoleVariableRef CVarDemoLoopCount( TEXT( "demo.LoopCount" ), GDemoLoopCount, TEXT( "If > 1, will play the replay that many times before stopping." ) );
+
 static int32 GDemoSaveRollbackActorState = 1;
 static FAutoConsoleVariableRef CVarDemoSaveRollbackActorState( TEXT( "demo.SaveRollbackActorState" ), GDemoSaveRollbackActorState, TEXT( "If true, rollback actors will save some replicated state to apply when respawned." ) );
 
@@ -3735,15 +3738,25 @@ void UDemoNetDriver::TickDemoPlayback( float DeltaSeconds )
 		{
 			OnDemoFinishPlaybackDelegate.Broadcast();
 
-			if (FParse::Param(FCommandLine::Get(), TEXT("ExitAfterReplay")))
+			// checking against 1 so the count will mean total number of playthroughs, not additional loops
+			if (GDemoLoopCount > 1)
 			{
-				FPlatformMisc::RequestExit(false);
-			}
+				--GDemoLoopCount;
 
-			if (CVarLoopDemo.GetValueOnGameThread() > 0)
-			{
 				GotoTimeInSeconds(0.0f);
 			}
+			else
+			{
+				if (FParse::Param(FCommandLine::Get(), TEXT("ExitAfterReplay")))
+				{
+					FPlatformMisc::RequestExit(false);
+				}
+
+				if (CVarLoopDemo.GetValueOnGameThread() > 0)
+				{
+					GotoTimeInSeconds(0.0f);
+				}
+			}
 		}
 	}
 
