Deprecated AWorldSettings::Pauser, and replaced with PauserPlayerState, Get/SetPauserPlayerState().
[RN] fix general

diff --git a/Engine/Plugins/NetcodeUnitTest/NetcodeUnitTest/Source/NetcodeUnitTest/Private/Net/NUTUtilNet.cpp b/Engine/Plugins/NetcodeUnitTest/NetcodeUnitTest/Source/NetcodeUnitTest/Private/Net/NUTUtilNet.cpp
--- a/Engine/Plugins/NetcodeUnitTest/NetcodeUnitTest/Source/NetcodeUnitTest/Private/Net/NUTUtilNet.cpp
+++ b/Engine/Plugins/NetcodeUnitTest/NetcodeUnitTest/Source/NetcodeUnitTest/Private/Net/NUTUtilNet.cpp
@@ -534,7 +534,7 @@ UWorld* NUTNet::CreateUnitTestWorld(bool bHookTick/*=true*/)
 
 				if (PrimState != NULL)
 				{
-					CurSettings->Pauser = PrimState;
+					CurSettings->SetPauserPlayerState(PrimState);
 				}
 			}
 		}
diff --git a/Engine/Source/Runtime/Engine/Classes/GameFramework/WorldSettings.h b/Engine/Source/Runtime/Engine/Classes/GameFramework/WorldSettings.h
--- a/Engine/Source/Runtime/Engine/Classes/GameFramework/WorldSettings.h
+++ b/Engine/Source/Runtime/Engine/Classes/GameFramework/WorldSettings.h
@@ -682,6 +682,7 @@ public:
 	FBroadphaseSettings BroadphaseSettings;
 
 	// If paused, FName of person pausing the game.
+	UE_DEPRECATED(4.23, "This property is deprecated. Please use Get/SetPauserPlayerState().")
 	UPROPERTY(transient, replicated)
 	class APlayerState* Pauser;
 
@@ -703,6 +704,11 @@ public:
 	UPROPERTY()
 	TArray<UAssetUserData*> AssetUserData;
 
+private:
+	// If paused, PlayerState of person pausing the game.
+	UPROPERTY(transient, replicated)
+	class APlayerState* PauserPlayerState;
+
 public:
 	//~ Begin UObject Interface.
 	virtual void PostLoad() override;
@@ -778,6 +784,9 @@ public:
 	UMaterialInterface* GetHierarchicalLODBaseMaterial() const;
 #endif // WITH EDITOR
 
+	FORCEINLINE class APlayerState* GetPauserPlayerState() const { return PauserPlayerState; }
+	FORCEINLINE virtual void SetPauserPlayerState(class APlayerState* PlayerState) { PauserPlayerState = PlayerState; }
+
 private:
 
 	// Hidden functions that don't make sense to use on this class.
diff --git a/Engine/Source/Runtime/Engine/Private/Components/CharacterMovementComponent.cpp b/Engine/Source/Runtime/Engine/Private/Components/CharacterMovementComponent.cpp
--- a/Engine/Source/Runtime/Engine/Private/Components/CharacterMovementComponent.cpp
+++ b/Engine/Source/Runtime/Engine/Private/Components/CharacterMovementComponent.cpp
@@ -8578,7 +8578,7 @@ void UCharacterMovementComponent::ServerMove_Implementation(
 	}
 
 	// Perform actual movement
-	if ((MyWorld->GetWorldSettings()->Pauser == NULL) && (DeltaTime > 0.f))
+	if ((MyWorld->GetWorldSettings()->GetPauserPlayerState() == NULL) && (DeltaTime > 0.f))
 	{
 		if (PC)
 		{
diff --git a/Engine/Source/Runtime/Engine/Private/DemoNetDriver.cpp b/Engine/Source/Runtime/Engine/Private/DemoNetDriver.cpp
--- a/Engine/Source/Runtime/Engine/Private/DemoNetDriver.cpp
+++ b/Engine/Source/Runtime/Engine/Private/DemoNetDriver.cpp
@@ -3761,7 +3761,7 @@ void UDemoNetDriver::TickDemoPlayback( float DeltaSeconds )
 	}
 
 	// Advance demo time by seconds passed if we're not paused
-	if ( World && World->GetWorldSettings() && World->GetWorldSettings()->Pauser == nullptr )
+	if ( World && World->GetWorldSettings() && World->GetWorldSettings()->GetPauserPlayerState() == nullptr )
 	{
 		DemoCurrentTime += DeltaSeconds;
 	}
diff --git a/Engine/Source/Runtime/Engine/Private/GameModeBase.cpp b/Engine/Source/Runtime/Engine/Private/GameModeBase.cpp
--- a/Engine/Source/Runtime/Engine/Private/GameModeBase.cpp
+++ b/Engine/Source/Runtime/Engine/Private/GameModeBase.cpp
@@ -201,9 +201,9 @@ bool AGameModeBase::SetPause(APlayerController* PC, FCanUnpause CanUnpauseDelega
 
 		// Let the first one in "own" the pause state
 		AWorldSettings * WorldSettings = GetWorldSettings();
-		if (WorldSettings->Pauser == nullptr)
+		if (WorldSettings->GetPauserPlayerState() == nullptr)
 		{
-			WorldSettings->Pauser = PC->PlayerState;
+			WorldSettings->SetPauserPlayerState(PC->PlayerState);
 		}
 		return true;
 	}
@@ -243,7 +243,7 @@ bool AGameModeBase::ClearPause()
 	// Clear the pause state if the list is empty
 	if (Pausers.Num() == 0)
 	{
-		GetWorldSettings()->Pauser = nullptr;
+		GetWorldSettings()->SetPauserPlayerState(nullptr);
 	}
 
 	return bPauseCleared;
@@ -272,9 +272,9 @@ void AGameModeBase::ForceClearUnpauseDelegates(AActor* PauseActor)
 
 		APlayerController* PC = Cast<APlayerController>(PauseActor);
 		AWorldSettings * WorldSettings = GetWorldSettings();
-		if (PC != nullptr && PC->PlayerState != nullptr && WorldSettings != nullptr && WorldSettings->Pauser == PC->PlayerState)
+		if (PC != nullptr && PC->PlayerState != nullptr && WorldSettings != nullptr && WorldSettings->GetPauserPlayerState() == PC->PlayerState)
 		{
-			// Try to find another player to be the worldsettings's Pauser
+			// Try to find another player to be the worldsettings's PauserPlayerState
 			for (FConstPlayerControllerIterator Iterator = GetWorld()->GetPlayerControllerIterator(); Iterator; ++Iterator)
 			{
 				APlayerController* Player = Iterator->Get();
@@ -282,15 +282,15 @@ void AGameModeBase::ForceClearUnpauseDelegates(AActor* PauseActor)
 					&&	Player->PlayerState != PC->PlayerState
 					&& !Player->IsPendingKillPending() && !Player->PlayerState->IsPendingKillPending())
 				{
-					WorldSettings->Pauser = Player->PlayerState;
+					WorldSettings->SetPauserPlayerState(Player->PlayerState);
 					break;
 				}
 			}
 
 			// If it's still pointing to the original player's PlayerState, clear it completely
-			if (WorldSettings->Pauser == PC->PlayerState)
+			if (WorldSettings->GetPauserPlayerState() == PC->PlayerState)
 			{
-				WorldSettings->Pauser = nullptr;
+				WorldSettings->SetPauserPlayerState(nullptr);
 			}
 		}
 	}
diff --git a/Engine/Source/Runtime/Engine/Private/LevelTick.cpp b/Engine/Source/Runtime/Engine/Private/LevelTick.cpp
--- a/Engine/Source/Runtime/Engine/Private/LevelTick.cpp
+++ b/Engine/Source/Runtime/Engine/Private/LevelTick.cpp
@@ -455,7 +455,7 @@ bool UWorld::IsPaused() const
 {
 	// pause if specifically set or if we're waiting for the end of the tick to perform streaming level loads (so actors don't fall through the world in the meantime, etc)
 	const AWorldSettings* Info = GetWorldSettings(/*bCheckStreamingPersistent=*/false, /*bChecked=*/false);
-	return ( (Info && Info->Pauser != nullptr && TimeSeconds >= PauseDelay) ||
+	return ( (Info && Info->GetPauserPlayerState() != nullptr && TimeSeconds >= PauseDelay) ||
 				(bRequestedBlockOnAsyncLoading && GetNetMode() == NM_Client) ||
 				(GEngine->ShouldCommitPendingMapChange(this)) ||
 				(IsPlayInEditor() && bDebugPauseExecution) );
diff --git a/Engine/Source/Runtime/Engine/Private/PlayerController.cpp b/Engine/Source/Runtime/Engine/Private/PlayerController.cpp
--- a/Engine/Source/Runtime/Engine/Private/PlayerController.cpp
+++ b/Engine/Source/Runtime/Engine/Private/PlayerController.cpp
@@ -987,7 +987,7 @@ void APlayerController::ServerShortTimeout_Implementation()
 		bShortConnectTimeOut = true;
 
 		// quick update of pickups and gameobjectives since this player is now relevant
-		if (GetWorldSettings()->Pauser != NULL)
+		if (GetWorldSettings()->GetPauserPlayerState() != NULL)
 		{
 			// update everything immediately, as TimeSeconds won't get advanced while paused
 			// so otherwise it won't happen at all until the game is unpaused
@@ -1770,7 +1770,7 @@ bool APlayerController::SetPause( bool bPause, FCanUnpause CanUnpauseDelegate)
 
 bool APlayerController::IsPaused() const
 {
-	return GetWorldSettings()->Pauser != NULL;
+	return GetWorldSettings()->GetPauserPlayerState() != NULL;
 }
 
 void APlayerController::Pause()
@@ -4714,7 +4714,7 @@ FString APlayerController::GetServerNetworkAddress()
 
 bool APlayerController::DefaultCanUnpause()
 {
-	return GetWorldSettings() != NULL && GetWorldSettings()->Pauser == PlayerState;
+	return GetWorldSettings() != NULL && GetWorldSettings()->GetPauserPlayerState() == PlayerState;
 }
 
 void APlayerController::StartSpectatingOnly()
diff --git a/Engine/Source/Runtime/Engine/Private/UnrealEngine.cpp b/Engine/Source/Runtime/Engine/Private/UnrealEngine.cpp
--- a/Engine/Source/Runtime/Engine/Private/UnrealEngine.cpp
+++ b/Engine/Source/Runtime/Engine/Private/UnrealEngine.cpp
@@ -12749,7 +12749,7 @@ void UEngine::UpdateTransitionType(UWorld *CurrentWorld)
 	else if(TransitionType == TT_None || TransitionType == TT_Paused)
 	{
 		// Display a paused screen if the game is paused.
-		TransitionType = (CurrentWorld->GetWorldSettings()->Pauser != NULL) ? TT_Paused : TT_None;
+		TransitionType = (CurrentWorld->GetWorldSettings()->GetPauserPlayerState() != NULL) ? TT_Paused : TT_None;
 	}
 }
 
diff --git a/Engine/Source/Runtime/Engine/Private/World.cpp b/Engine/Source/Runtime/Engine/Private/World.cpp
--- a/Engine/Source/Runtime/Engine/Private/World.cpp
+++ b/Engine/Source/Runtime/Engine/Private/World.cpp
@@ -3496,7 +3496,7 @@ bool UWorld::HandleDemoScrubCommand(const TCHAR* Cmd, FOutputDevice& Ar, UWorld*
 		APlayerController* PlayerController = Cast<APlayerController>(DemoNetDriver->ServerConnection->OwningActor);
 		if (PlayerController != nullptr)
 		{
-			GetWorldSettings()->Pauser = PlayerController->PlayerState;
+			GetWorldSettings()->SetPauserPlayerState(PlayerController->PlayerState);
 			const uint32 Time = FCString::Atoi(*TimeString);
 			DemoNetDriver->GotoTimeInSeconds(Time);
 		}
@@ -3511,20 +3511,20 @@ bool UWorld::HandleDemoPauseCommand(const TCHAR* Cmd, FOutputDevice& Ar, UWorld*
 	AWorldSettings* WorldSettings = GetWorldSettings();
 	check(WorldSettings != nullptr);
 
-	if (WorldSettings->Pauser == nullptr)
+	if (WorldSettings->GetPauserPlayerState() == nullptr)
 	{
 		if (DemoNetDriver != nullptr && DemoNetDriver->ServerConnection != nullptr && DemoNetDriver->ServerConnection->OwningActor != nullptr)
 		{
 			APlayerController* PlayerController = Cast<APlayerController>(DemoNetDriver->ServerConnection->OwningActor);
 			if (PlayerController != nullptr)
 			{
-				WorldSettings->Pauser = PlayerController->PlayerState;
+				WorldSettings->SetPauserPlayerState(PlayerController->PlayerState);
 			}
 		}
 	}
 	else
 	{
-		WorldSettings->Pauser = nullptr;
+		WorldSettings->SetPauserPlayerState(nullptr);
 	}
 	return true;
 }
diff --git a/Engine/Source/Runtime/Engine/Private/WorldSettings.cpp b/Engine/Source/Runtime/Engine/Private/WorldSettings.cpp
--- a/Engine/Source/Runtime/Engine/Private/WorldSettings.cpp
+++ b/Engine/Source/Runtime/Engine/Private/WorldSettings.cpp
@@ -266,7 +266,7 @@ void AWorldSettings::GetLifetimeReplicatedProps( TArray< FLifetimeProperty > & O
 {
 	Super::GetLifetimeReplicatedProps( OutLifetimeProps );
 
-	DOREPLIFETIME( AWorldSettings, Pauser );
+	DOREPLIFETIME( AWorldSettings, PauserPlayerState );
 	DOREPLIFETIME( AWorldSettings, TimeDilation );
 	DOREPLIFETIME( AWorldSettings, MatineeTimeDilation );
 	DOREPLIFETIME( AWorldSettings, WorldGravityZ );
