If UAssetManager::ChangeBundleStateForPrimaryAsset is called on an asset id that is currently streaming and the new bundle state is a superset of the current one, just create a combined handle on top of it. If the new bundle state is not equal or a superset, then cancel the pending stream. Fixes an issue where two different systems were changing bundle state close enough to the same time that streaming wasn't complete, but we were missing the completion delegate from the first state change.

diff --git a/Engine/Source/Runtime/Engine/Private/AssetManager.cpp b/Engine/Source/Runtime/Engine/Private/AssetManager.cpp
--- a/Engine/Source/Runtime/Engine/Private/AssetManager.cpp
+++ b/Engine/Source/Runtime/Engine/Private/AssetManager.cpp
@@ -1131,8 +1131,30 @@ TSharedPtr<FStreamableHandle> UAssetManager::ChangeBundleStateForPrimaryAssets(c
 					ExistingHandles.Add(NameData->PendingState.Handle);
 					continue;
 				}
-				// Clear pending state
-				NameData->PendingState.Reset(true);
+
+				bool bIsSuperSet = false;
+				if (NewBundleState.Num() > NameData->PendingState.BundleNames.Num())
+				{
+					bIsSuperSet = true;
+					for (int i = 0; i < NameData->PendingState.BundleNames.Num(); i++)
+					{
+						if (!NewBundleState.Contains(NameData->PendingState.BundleNames[i]))
+						{
+							bIsSuperSet = false;
+							break;
+						}
+					}
+				}
+
+				if (bIsSuperSet)
+				{
+					ExistingHandles.Add(NameData->PendingState.Handle);
+				}
+				else
+				{
+					// Clear pending state
+					NameData->PendingState.Reset(true);
+				}
 			}
 			else if (NameData->CurrentState.IsValid() && NameData->CurrentState.BundleNames == NewBundleState)
 			{
@@ -1185,6 +1207,8 @@ TSharedPtr<FStreamableHandle> UAssetManager::ChangeBundleStateForPrimaryAssets(c
 				DebugName += TEXT(")");
 			}
 
+			UE_LOG(LogAssetManager, Verbose, TEXT("%s bundle state change: %s"), *PrimaryAssetId.ToString(), *DebugName);
+
 			NewHandle = LoadAssetList(PathsToLoad.Array(), FStreamableDelegate(), Priority, DebugName);
 
 			if (!NewHandle.IsValid())
