Add support tracking for approximate input latency time on platforms that support flip tracking.

- Input time is sampled just before slate polls/process the input, this is value is then passed to the render thread which passes it to the rhi thread.
- The RHI thread can then measure the end to end time from sampling the input until the flip on supported platforms.
- This is a step towards easier debug of flip tracking.

diff --git a/Engine/Source/Runtime/Core/Private/Misc/CoreGlobals.cpp b/Engine/Source/Runtime/Core/Private/Misc/CoreGlobals.cpp
--- a/Engine/Source/Runtime/Core/Private/Misc/CoreGlobals.cpp
+++ b/Engine/Source/Runtime/Core/Private/Misc/CoreGlobals.cpp
@@ -247,6 +247,8 @@ bool					GEventDrivenLoaderEnabled = false;
 /** Steadily increasing frame counter.																		*/
 TSAN_ATOMIC(uint64)		GFrameCounter(0);
 uint64					GLastGCFrame					= 0;
+/** The time input was sampled, in cycles. */
+uint64					GInputTime					= 0;
 /** Incremented once per frame before the scene is being rendered. In split screen mode this is incremented once for all views (not for each view). */
 uint32					GFrameNumber					= 1;
 /** NEED TO RENAME, for RT version of GFrameTime use View.ViewFamily->FrameNumber or pass down from RT from GFrameTime). */
diff --git a/Engine/Source/Runtime/Core/Public/CoreGlobals.h b/Engine/Source/Runtime/Core/Public/CoreGlobals.h
--- a/Engine/Source/Runtime/Core/Public/CoreGlobals.h
+++ b/Engine/Source/Runtime/Core/Public/CoreGlobals.h
@@ -322,6 +322,9 @@ extern CORE_API TSAN_ATOMIC(uint64) GFrameCounter;
 /** GFrameCounter the last time GC was run. */
 extern CORE_API uint64 GLastGCFrame;
 
+/** The time input was sampled, in cycles. */
+extern CORE_API uint64 GInputTime;
+
 /** Incremented once per frame before the scene is being rendered. In split screen mode this is incremented once for all views (not for each view). */
 extern CORE_API uint32 GFrameNumber;
 
diff --git a/Engine/Source/Runtime/Core/Public/HAL/Platform.h b/Engine/Source/Runtime/Core/Public/HAL/Platform.h
--- a/Engine/Source/Runtime/Core/Public/HAL/Platform.h
+++ b/Engine/Source/Runtime/Core/Public/HAL/Platform.h
@@ -465,6 +465,10 @@
 	#define PLATFORM_NEEDS_RHIRESOURCELIST 1
 #endif
 
+#ifndef PLATFORM_SUPPORTS_FLIP_TRACKING
+	#define PLATFORM_SUPPORTS_FLIP_TRACKING 0
+#endif
+
 // deprecated, do not use
 #define PLATFORM_HAS_THREADSAFE_RHIGetRenderQueryResult	#
 #define PLATFORM_SUPPORTS_RHI_THREAD #
diff --git a/Engine/Source/Runtime/Engine/Private/UnrealClient.cpp b/Engine/Source/Runtime/Engine/Private/UnrealClient.cpp
--- a/Engine/Source/Runtime/Engine/Private/UnrealClient.cpp
+++ b/Engine/Source/Runtime/Engine/Private/UnrealClient.cpp
@@ -412,6 +412,9 @@ int32 FStatUnitData::DrawStat(FViewport* InViewport, FCanvas* InCanvas, int32 In
 	RawRHITTime = FPlatformTime::ToMilliseconds(GRHIThreadTime);
 	RHITTime = 0.9 * RHITTime + 0.1 * RawRHITTime;
 
+	RawInputLatencyTime = FPlatformTime::ToMilliseconds64(GInputLatencyTime);
+	InputLatencyTime = 0.9 * InputLatencyTime + 0.1 * RawInputLatencyTime;
+
 	FDynamicResolutionStateInfos DynamicResolutionStateInfos;
 	GEngine->GetDynamicResolutionCurrentStateInfos(/* out */ DynamicResolutionStateInfos);
 
@@ -425,6 +428,7 @@ int32 FStatUnitData::DrawStat(FViewport* InViewport, FCanvas* InCanvas, int32 In
 	SET_FLOAT_STAT(STAT_UnitRHIT, RHITTime);
 	SET_FLOAT_STAT(STAT_UnitGame, GameThreadTime);
 	SET_FLOAT_STAT(STAT_UnitGPU, GPUFrameTime);
+	SET_FLOAT_STAT(STAT_InputLatencyTime, InputLatencyTime);
 
 	GEngine->SetAverageUnitTimes(FrameTime, RenderThreadTime, GameThreadTime, GPUFrameTime, RHITTime);
 
@@ -433,6 +437,7 @@ int32 FStatUnitData::DrawStat(FViewport* InViewport, FCanvas* InCanvas, int32 In
 	float Max_GPUFrameTime = 0.0f;
 	float Max_FrameTime = 0.0f;
 	float Max_RHITTime = 0.0f;
+	float Max_InputLatencyTime = 0.0f;
 
 	const bool bShowUnitMaxTimes = InViewport->GetClient() ? InViewport->GetClient()->IsStatEnabled(TEXT("UnitMax")) : false;
 #if !UE_BUILD_SHIPPING
@@ -442,6 +447,7 @@ int32 FStatUnitData::DrawStat(FViewport* InViewport, FCanvas* InCanvas, int32 In
 	GPUFrameTimes[CurrentIndex] = bShowRawUnitTimes ? RawGPUFrameTime : GPUFrameTime;
 	FrameTimes[CurrentIndex] = bShowRawUnitTimes ? RawFrameTime : FrameTime;
 	RHITTimes[CurrentIndex] = bShowRawUnitTimes ? RawRHITTime : RHITTime;
+	InputLatencyTimes[CurrentIndex] = bShowRawUnitTimes ? RawInputLatencyTime : InputLatencyTime;
 	ResolutionFractions[CurrentIndex] = DynamicResolutionStateInfos.ResolutionFractionApproximation;
 	CurrentIndex++;
 	if (CurrentIndex == NumberOfSamples)
@@ -474,6 +480,10 @@ int32 FStatUnitData::DrawStat(FViewport* InViewport, FCanvas* InCanvas, int32 In
 			{
 				Max_RHITTime = RHITTimes[MaxIndex];
 			}
+			if (Max_InputLatencyTime < InputLatencyTimes[MaxIndex])
+			{
+				Max_InputLatencyTime = InputLatencyTimes[MaxIndex];
+			}
 		}
 	}
 #endif // #if !(UE_BUILD_SHIPPING || UE_BUILD_TEST)
@@ -485,6 +495,7 @@ int32 FStatUnitData::DrawStat(FViewport* InViewport, FCanvas* InCanvas, int32 In
 
 	const bool bShowUnitTimeGraph = InViewport->GetClient() ? InViewport->GetClient()->IsStatEnabled(TEXT("UnitGraph")) : false;
 	const bool bHaveGPUData = GPUCycles > 0;
+	const bool bHaveInputLatencyData = InputLatencyTime > 0;
 
 	const float AlertResolutionFraction = 0.70f; // Truncation of sqrt(0.5) for easier remembering.
 
@@ -590,6 +601,19 @@ int32 FStatUnitData::DrawStat(FViewport* InViewport, FCanvas* InCanvas, int32 In
 			}
 			InY += RowHeight;
 		}
+		if (bHaveInputLatencyData)
+		{
+			const float ReasonableInputLatencyFactor = 2.5f;
+			const FColor InputLatencyAverageColor = GEngine->GetFrameTimeDisplayColor(InputLatencyTime / ReasonableInputLatencyFactor);
+			InCanvas->DrawShadowedString(X1, InY, TEXT("Input:"), Font, bShowUnitTimeGraph ? FColor(255, 255, 100) : FColor::White);
+			InCanvas->DrawShadowedString(X2, InY, *FString::Printf(TEXT("%3.2f ms"), InputLatencyTime), Font, InputLatencyAverageColor);
+			if (bShowUnitMaxTimes)
+			{
+				const FColor InputLatencyMaxColor = GEngine->GetFrameTimeDisplayColor(Max_InputLatencyTime / ReasonableInputLatencyFactor);
+				InCanvas->DrawShadowedString(X3, InY, *FString::Printf(TEXT("%4.2f ms"), Max_InputLatencyTime), Font, InputLatencyMaxColor);
+			}
+			InY += RowHeight;
+		}
 		{
 			uint64 MemoryUsed = FPlatformMemory::GetMemoryUsedFast();
 			if (MemoryUsed > 0)
diff --git a/Engine/Source/Runtime/Engine/Private/UnrealEngine.cpp b/Engine/Source/Runtime/Engine/Private/UnrealEngine.cpp
--- a/Engine/Source/Runtime/Engine/Private/UnrealEngine.cpp
+++ b/Engine/Source/Runtime/Engine/Private/UnrealEngine.cpp
@@ -10179,6 +10179,20 @@ void FFrameEndSync::Sync( bool bAllowOneFrameThreadLag )
 {
 	check(IsInGameThread());			
 
+#if !UE_BUILD_SHIPPING && PLATFORM_SUPPORTS_FLIP_TRACKING
+	// Set the FrameDebugInfo on platforms that have accurate frame tracking.
+	ENQUEUE_RENDER_COMMAND(FrameDebugInfo)(
+		[CurrentFrameCounter = GFrameCounter, CurrentInputTime = GInputTime](FRHICommandList&)
+	{
+		GRHICommandList.GetImmediateCommandList().EnqueueLambda(
+			[CurrentFrameCounter, CurrentInputTime](FRHICommandList&)
+		{
+			// Set the FrameCount and InputTime for input latency stats and flip debugging.
+			RHISetFrameDebugInfo(GRHIPresentCounter - 1, CurrentFrameCounter - 1, CurrentInputTime);
+		});
+	});
+#endif
+
 	// Since this is the frame end sync, allow sync with the RHI and GPU (true).
 	Fence[EventIndex].BeginFence(true);
 
@@ -12749,7 +12763,7 @@ void UEngine::UpdateTransitionType(UWorld *CurrentWorld)
 	else if(TransitionType == TT_None || TransitionType == TT_Paused)
 	{
 		// Display a paused screen if the game is paused.
-		TransitionType = (CurrentWorld->GetWorldSettings()->GetPauserPlayerState() != NULL) ? TT_Paused : TT_None;
+		TransitionType = (CurrentWorld->GetWorldSettings()->Pauser != NULL) ? TT_Paused : TT_None;
 	}
 }
 
diff --git a/Engine/Source/Runtime/Engine/Public/UnrealClient.h b/Engine/Source/Runtime/Engine/Public/UnrealClient.h
--- a/Engine/Source/Runtime/Engine/Public/UnrealClient.h
+++ b/Engine/Source/Runtime/Engine/Public/UnrealClient.h
@@ -205,6 +205,7 @@ struct FStatUnitData
 	float GPUFrameTime;
 	float FrameTime;
 	float RHITTime;
+	float InputLatencyTime;
 
 	/** Raw equivalents of the above variables */
 	float RawRenderThreadTime;
@@ -212,6 +213,7 @@ struct FStatUnitData
 	float RawGPUFrameTime;
 	float RawFrameTime;
 	float RawRHITTime;
+	float RawInputLatencyTime;
 
 	/** Time that has transpired since the last draw call */
 	double LastTime;
@@ -225,6 +227,7 @@ struct FStatUnitData
 	TArray<float> GPUFrameTimes;
 	TArray<float> FrameTimes;
 	TArray<float> RHITTimes;
+	TArray<float> InputLatencyTimes;
 	TArray<float> ResolutionFractions;
 #endif //!UE_BUILD_SHIPPING
 
@@ -234,10 +237,13 @@ struct FStatUnitData
 		, GPUFrameTime(0.0f)
 		, FrameTime(0.0f)
 		, RHITTime(0.0f)
+		, InputLatencyTime(0.0f)
 		, RawRenderThreadTime(0.0f)
 		, RawGameThreadTime(0.0f)
 		, RawGPUFrameTime(0.0f)
 		, RawFrameTime(0.0f)
+		, RawRHITTime(0.0f)
+		, RawInputLatencyTime(0.0f)
 		, LastTime(0.0)
 	{
 #if !UE_BUILD_SHIPPING
@@ -247,6 +253,7 @@ struct FStatUnitData
 		GPUFrameTimes.AddZeroed(NumberOfSamples);
 		FrameTimes.AddZeroed(NumberOfSamples);
 		RHITTimes.AddZeroed(NumberOfSamples);
+		InputLatencyTimes.AddZeroed(NumberOfSamples);
 		ResolutionFractions.Reserve(NumberOfSamples);
 		for (int32 i = 0; i < NumberOfSamples; i++)
 		{
diff --git a/Engine/Source/Runtime/Launch/Private/LaunchEngineLoop.cpp b/Engine/Source/Runtime/Launch/Private/LaunchEngineLoop.cpp
--- a/Engine/Source/Runtime/Launch/Private/LaunchEngineLoop.cpp
+++ b/Engine/Source/Runtime/Launch/Private/LaunchEngineLoop.cpp
@@ -983,6 +983,10 @@ static void UpdateCoreCsvStats()
 	CSV_CUSTOM_STAT_GLOBAL(GameThreadTime, FPlatformTime::ToMilliseconds(GGameThreadTime), ECsvCustomStatOp::Set);
 	CSV_CUSTOM_STAT_GLOBAL(GPUTime, FPlatformTime::ToMilliseconds(GGPUFrameTime), ECsvCustomStatOp::Set);
 	CSV_CUSTOM_STAT_GLOBAL(RHIThreadTime, FPlatformTime::ToMilliseconds(GRHIThreadTime), ECsvCustomStatOp::Set);
+	if (GInputLatencyTime > 0)
+	{
+		CSV_CUSTOM_STAT_GLOBAL(InputLatencyTime, FPlatformTime::ToMilliseconds(GInputLatencyTime), ECsvCustomStatOp::Set);
+	}
 	FPlatformMemoryStats MemoryStats = FPlatformMemory::GetStats();
 	float PhysicalMBFree = float(MemoryStats.AvailablePhysical / 1024) / 1024.0f;
 	CSV_CUSTOM_STAT_GLOBAL(MemoryFreeMB, PhysicalMBFree, ECsvCustomStatOp::Set);
@@ -3913,6 +3917,9 @@ void FEngineLoop::Tick()
 		// tick active platform files
 		FPlatformFileManager::Get().TickActivePlatformFile();
 
+		// Roughly track the time when the input was sampled
+		GInputTime = FPlatformTime::Cycles64();
+
 		// process accumulated Slate input
 		if (FSlateApplication::IsInitialized() && !bIdleMode)
 		{
diff --git a/Engine/Source/Runtime/RHI/Private/RHICommandList.cpp b/Engine/Source/Runtime/RHI/Private/RHICommandList.cpp
--- a/Engine/Source/Runtime/RHI/Private/RHICommandList.cpp
+++ b/Engine/Source/Runtime/RHI/Private/RHICommandList.cpp
@@ -127,6 +127,9 @@ uint32 GWorkingRHIThreadTime = 0;
 uint32 GWorkingRHIThreadStallTime = 0;
 uint32 GWorkingRHIThreadStartCycles = 0;
 
+/** How many cycles the from sampling input to the frame being flipped. */
+uint64 GInputLatencyTime = 0;
+
 RHI_API bool GEnableAsyncCompute = true;
 RHI_API FRHICommandListExecutor GRHICommandList;
 
diff --git a/Engine/Source/Runtime/RHI/Private/RHIUtilities.cpp b/Engine/Source/Runtime/RHI/Private/RHIUtilities.cpp
--- a/Engine/Source/Runtime/RHI/Private/RHIUtilities.cpp
+++ b/Engine/Source/Runtime/RHI/Private/RHIUtilities.cpp
@@ -159,6 +159,16 @@ struct FRHIFrameOffsetThread : public FRunnable
 
 	FEvent* WaitEvent;
 
+#if !UE_BUILD_SHIPPING
+	struct FFrameDebugInfo
+	{
+		uint64 PresentIndex;
+		uint64 FrameIndex;
+		uint64 InputTime;
+	};
+	TArray<FFrameDebugInfo> FrameDebugInfos;
+#endif
+
 	virtual uint32 Run() override
 	{
 		while (bRun)
@@ -180,6 +190,22 @@ struct FRHIFrameOffsetThread : public FRunnable
 				LastFlipFrame.FlipTimeInSeconds = LastFlipFrame.FlipTimeInSeconds + TargetFrameTimeInSeconds - SlackInSeconds;
 				LastFlipFrame.VBlankTimeInSeconds = LastFlipFrame.VBlankTimeInSeconds + TargetFrameTimeInSeconds - SlackInSeconds;
 				LastFlipFrame.PresentIndex++;
+
+#if !UE_BUILD_SHIPPING && PLATFORM_SUPPORTS_FLIP_TRACKING
+				for (int32 DebugInfoIndex = FrameDebugInfos.Num() - 1; DebugInfoIndex >= 0; --DebugInfoIndex)
+				{
+					auto const& DebugInfo = FrameDebugInfos[DebugInfoIndex];
+					if (NewFlipFrame.PresentIndex == DebugInfo.PresentIndex)
+					{
+						GInputLatencyTime = (NewFlipFrame.VBlankTimeInSeconds / FPlatformTime::GetSecondsPerCycle64()) - DebugInfo.InputTime;
+					}
+
+					if (DebugInfo.PresentIndex <= NewFlipFrame.PresentIndex)
+					{
+						FrameDebugInfos.RemoveAtSwap(DebugInfoIndex);
+					}
+				}
+#endif
 			}
 
 			WaitEvent->Trigger();
@@ -211,7 +237,14 @@ public:
 	static FRHIFlipDetails WaitForFlip(double Timeout)
 	{
 		check(Singleton.WaitEvent);
-		Singleton.WaitEvent->Wait((uint32)(Timeout * 1000.0));
+		if (Timeout >= 0)
+		{
+			Singleton.WaitEvent->Wait((uint32)(Timeout * 1000.0));
+		}
+		else
+		{
+			Singleton.WaitEvent->Wait();
+		}
 
 		FScopeLock Lock(&Singleton.CS);
 		return Singleton.LastFlipFrame;
@@ -250,6 +283,21 @@ public:
 		Singleton.GetOrInitializeWaitEvent()->Trigger();
 	}
 
+	static void SetFrameDebugInfo(uint64 PresentIndex, uint64 FrameIndex, uint64 InputTime)
+	{
+#if !UE_BUILD_SHIPPING && PLATFORM_SUPPORTS_FLIP_TRACKING
+		FScopeLock Lock(&Singleton.CS);
+		if (Thread)
+		{
+			FFrameDebugInfo DebugInfo;
+			DebugInfo.PresentIndex = PresentIndex;
+			DebugInfo.FrameIndex = FrameIndex;
+			DebugInfo.InputTime = InputTime;
+			Singleton.FrameDebugInfos.Add(DebugInfo);
+		}
+#endif
+	}
+
 private:
 	FEvent* GetOrInitializeWaitEvent()
 	{
@@ -422,6 +470,13 @@ RHI_API void RHICompleteGraphEventOnFlip(uint64 PresentIndex, FGraphEventRef Eve
 	FRHIFrameFlipTrackingRunnable::CompleteGraphEventOnFlip(PresentIndex, Event);
 }
 
+RHI_API void RHISetFrameDebugInfo(uint64 PresentIndex, uint64 FrameIndex, uint64 InputTime)
+{
+#if USE_FRAME_OFFSET_THREAD
+	FRHIFrameOffsetThread::SetFrameDebugInfo(PresentIndex, FrameIndex, InputTime);
+#endif
+}
+
 RHI_API void RHIInitializeFlipTracking()
 {
 #if USE_FRAME_OFFSET_THREAD
diff --git a/Engine/Source/Runtime/RHI/Public/RHICommandList.h b/Engine/Source/Runtime/RHI/Public/RHICommandList.h
--- a/Engine/Source/Runtime/RHI/Public/RHICommandList.h
+++ b/Engine/Source/Runtime/RHI/Public/RHICommandList.h
@@ -77,6 +77,9 @@ extern RHI_API uint32 GWorkingRHIThreadTime;
 extern RHI_API uint32 GWorkingRHIThreadStallTime;
 extern RHI_API uint32 GWorkingRHIThreadStartCycles;
 
+/** How many cycles the from sampling input to the frame being flipped. */
+extern RHI_API uint64 GInputLatencyTime;
+
 /**
 * Whether the RHI commands are being run in a thread other than the render thread
 */
diff --git a/Engine/Source/Runtime/RHI/Public/RHIUtilities.h b/Engine/Source/Runtime/RHI/Public/RHIUtilities.h
--- a/Engine/Source/Runtime/RHI/Public/RHIUtilities.h
+++ b/Engine/Source/Runtime/RHI/Public/RHIUtilities.h
@@ -943,6 +943,9 @@ extern RHI_API void RHIGetPresentThresholds(float& OutTopPercent, float& OutBott
 /** Signals the completion of the specified task graph event when the given frame has flipped. */
 extern RHI_API void RHICompleteGraphEventOnFlip(uint64 PresentIndex, FGraphEventRef Event);
 
+/** Sets the FrameIndex and InputTime for the current frame. */
+extern RHI_API void RHISetFrameDebugInfo(uint64 PresentIndex, uint64 FrameIndex, uint64 InputTime);
+
 extern RHI_API void RHIInitializeFlipTracking();
 extern RHI_API void RHIShutdownFlipTracking();
 
