Back out changelist 5188921. Possible cause for EdTest ensure

diff --git a/Engine/Source/Runtime/Engine/Private/AssetManager.cpp b/Engine/Source/Runtime/Engine/Private/AssetManager.cpp
--- a/Engine/Source/Runtime/Engine/Private/AssetManager.cpp
+++ b/Engine/Source/Runtime/Engine/Private/AssetManager.cpp
@@ -1131,30 +1131,8 @@ TSharedPtr<FStreamableHandle> UAssetManager::ChangeBundleStateForPrimaryAssets(c
 					ExistingHandles.Add(NameData->PendingState.Handle);
 					continue;
 				}
-
-				bool bIsSuperSet = false;
-				if (NewBundleState.Num() > NameData->PendingState.BundleNames.Num())
-				{
-					bIsSuperSet = true;
-					for (int i = 0; i < NameData->PendingState.BundleNames.Num(); i++)
-					{
-						if (!NewBundleState.Contains(NameData->PendingState.BundleNames[i]))
-						{
-							bIsSuperSet = false;
-							break;
-						}
-					}
-				}
-
-				if (bIsSuperSet)
-				{
-					ExistingHandles.Add(NameData->PendingState.Handle);
-				}
-				else
-				{
-					// Clear pending state
-					NameData->PendingState.Reset(true);
-				}
+				// Clear pending state
+				NameData->PendingState.Reset(true);
 			}
 			else if (NameData->CurrentState.IsValid() && NameData->CurrentState.BundleNames == NewBundleState)
 			{
@@ -1207,8 +1185,6 @@ TSharedPtr<FStreamableHandle> UAssetManager::ChangeBundleStateForPrimaryAssets(c
 				DebugName += TEXT(")");
 			}
 
-			UE_LOG(LogAssetManager, Verbose, TEXT("%s bundle state change: %s"), *PrimaryAssetId.ToString(), *DebugName);
-
 			NewHandle = LoadAssetList(PathsToLoad.Array(), FStreamableDelegate(), Priority, DebugName);
 
 			if (!NewHandle.IsValid())
