Fix warning due to bad resolve.

diff --git a/Engine/Source/Runtime/Engine/Private/UnrealEngine.cpp b/Engine/Source/Runtime/Engine/Private/UnrealEngine.cpp
--- a/Engine/Source/Runtime/Engine/Private/UnrealEngine.cpp
+++ b/Engine/Source/Runtime/Engine/Private/UnrealEngine.cpp
@@ -12763,7 +12763,7 @@ void UEngine::UpdateTransitionType(UWorld *CurrentWorld)
 	else if(TransitionType == TT_None || TransitionType == TT_Paused)
 	{
 		// Display a paused screen if the game is paused.
-		TransitionType = (CurrentWorld->GetWorldSettings()->Pauser != NULL) ? TT_Paused : TT_None;
+		TransitionType = (CurrentWorld->GetWorldSettings()->GetPauserPlayerState() != NULL) ? TT_Paused : TT_None;
 	}
 }
 
