Add missing instance offset for InstancedStaticMesh. Was missing on vulkan causing draw order of objects to be broken.

diff --git a/Engine/Source/Runtime/Engine/Private/InstancedStaticMesh.cpp b/Engine/Source/Runtime/Engine/Private/InstancedStaticMesh.cpp
--- a/Engine/Source/Runtime/Engine/Private/InstancedStaticMesh.cpp
+++ b/Engine/Source/Runtime/Engine/Private/InstancedStaticMesh.cpp
@@ -2458,9 +2458,7 @@ void FInstancedStaticMeshVertexFactoryShaderParameters::GetElementShaderBindings
 		ShaderBindings.Add(VertexFetch_InstanceOriginBufferParameter, InstancedVertexFactory->GetInstanceOriginSRV());
 		ShaderBindings.Add(VertexFetch_InstanceTransformBufferParameter, InstancedVertexFactory->GetInstanceTransformSRV());
 		ShaderBindings.Add(VertexFetch_InstanceLightmapBufferParameter, InstancedVertexFactory->GetInstanceLightmapSRV());
-
-		const bool bZeroInstanceOffset = RHISupportsAbsoluteVertexID(GMaxRHIShaderPlatform);
-		ShaderBindings.Add(InstanceOffset, bZeroInstanceOffset ? 0 : InstanceOffsetValue);
+		ShaderBindings.Add(InstanceOffset, InstanceOffsetValue);
 	}
 
 	if (GRHISupportsInstancing && InstanceOffsetValue > 0 && VertexStreams.Num() > 0)
