Restoring FOverrideSelectionColorMaterialRenderProxy which was removed in CL 4636584, but is required to be able to override colors on materials which modify vertex positions. This fixes view mode overrides (wirefame, lighting only, edges) for materials, which modify vertex positions.
#rnx

diff --git a/Engine/Source/Runtime/Engine/Private/Materials/HLSLMaterialTranslator.h b/Engine/Source/Runtime/Engine/Private/Materials/HLSLMaterialTranslator.h
--- a/Engine/Source/Runtime/Engine/Private/Materials/HLSLMaterialTranslator.h
+++ b/Engine/Source/Runtime/Engine/Private/Materials/HLSLMaterialTranslator.h
@@ -1687,7 +1687,9 @@ protected:
 					FMaterialUniformExpressionVectorParameter* VectorParameterA = (FMaterialUniformExpressionVectorParameter*)TestExpression;
 					FMaterialUniformExpressionVectorParameter* VectorParameterB = (FMaterialUniformExpressionVectorParameter*)UniformExpression;
 
-					if (!VectorParameterA->GetParameterInfo().Name.IsNone() && VectorParameterA->GetParameterInfo() == VectorParameterB->GetParameterInfo())
+					// Note: Skipping NAME_SelectionColor here as this behavior is relied on for editor materials
+					if (!VectorParameterA->GetParameterInfo().Name.IsNone() && VectorParameterA->GetParameterInfo() == VectorParameterB->GetParameterInfo()
+						&& VectorParameterA->GetParameterInfo().Name != NAME_SelectionColor)
 					{
 						delete UniformExpression;
 						return Errorf(TEXT("Invalid vector parameter '%s' found. Identical parameters must have the same value."), *(VectorParameterA->GetParameterInfo().Name.ToString()));
@@ -5819,6 +5821,30 @@ protected:
 
 	// The compiler can run in a different state and this affects caching of sub expression, Expressions are different (e.g. View.PrevWorldViewOrigin) when using previous frame's values
 	virtual bool IsCurrentlyCompilingForPreviousFrame() const { return bCompilingPreviousFrame; }
+
+	virtual bool IsDevelopmentFeatureEnabled(const FName& FeatureName) const override
+	{
+		if (FeatureName == NAME_SelectionColor)
+		{
+			// This is an editor-only feature (see FDefaultMaterialInstance::GetVectorValue).
+
+			// Determine if we're sure the editor will never run using the target shader platform.
+			// The list below may not be comprehensive enough, but it definitely includes platforms which won't use selection color for sure.
+			const bool bEditorMayUseTargetShaderPlatform = IsPCPlatform(Platform);
+			static const auto CVar = IConsoleManager::Get().FindTConsoleVariableDataInt(TEXT("r.CompileShadersForDevelopment"));
+			const bool bCompileShadersForDevelopment = (CVar && CVar->GetValueOnAnyThread() != 0);
+
+			return
+				// Does the material explicitly forbid development features?
+				Material->GetAllowDevelopmentShaderCompile()
+				// Can the editor run using the current shader platform?
+				&& bEditorMayUseTargetShaderPlatform
+				// Are shader development features globally disabled?
+				&& bCompileShadersForDevelopment;
+		}
+
+		return true;
+	}
 };
 
 #endif
diff --git a/Engine/Source/Runtime/Engine/Private/Materials/Material.cpp b/Engine/Source/Runtime/Engine/Private/Materials/Material.cpp
--- a/Engine/Source/Runtime/Engine/Private/Materials/Material.cpp
+++ b/Engine/Source/Runtime/Engine/Private/Materials/Material.cpp
@@ -113,6 +113,15 @@ int32 FMaterialResource::CompilePropertyAndSetMaterialProperty(EMaterialProperty
 
 	EShaderFrequency ShaderFrequency = Compiler->GetCurrentShaderFrequency();
 	
+	int32 SelectionColorIndex = INDEX_NONE;
+
+	if (ShaderFrequency == SF_Pixel &&
+		GetMaterialDomain() != MD_Volume &&
+		Compiler->IsDevelopmentFeatureEnabled(NAME_SelectionColor))
+	{
+		SelectionColorIndex = Compiler->ComponentMask(Compiler->VectorParameter(NAME_SelectionColor, FLinearColor::Black), 1, 1, 1, 0);
+	}
+
 	//Compile the material instance if we have one.
 	UMaterialInterface* MaterialInterface = MaterialInstance ? static_cast<UMaterialInterface*>(MaterialInstance) : Material;
 
@@ -121,7 +130,14 @@ int32 FMaterialResource::CompilePropertyAndSetMaterialProperty(EMaterialProperty
 	switch(Property)
 	{
 		case MP_EmissiveColor:
-			Ret = MaterialInterface->CompileProperty(Compiler, MP_EmissiveColor);
+			if (SelectionColorIndex != INDEX_NONE)
+			{
+				Ret = Compiler->Add(MaterialInterface->CompileProperty(Compiler, MP_EmissiveColor, MFCF_ForceCast), SelectionColorIndex);
+			}
+			else
+			{
+				Ret = MaterialInterface->CompileProperty(Compiler, MP_EmissiveColor);
+			}
 			break;
 
 		case MP_DiffuseColor: 
diff --git a/Engine/Source/Runtime/Engine/Private/Materials/MaterialShared.cpp b/Engine/Source/Runtime/Engine/Private/Materials/MaterialShared.cpp
--- a/Engine/Source/Runtime/Engine/Private/Materials/MaterialShared.cpp
+++ b/Engine/Source/Runtime/Engine/Private/Materials/MaterialShared.cpp
@@ -2208,6 +2208,37 @@ bool FColoredMaterialRenderProxy::GetTextureValue(const FMaterialParameterInfo&
 	return Parent->GetTextureValue(ParameterInfo,OutValue,Context);
 }
 
+/*-----------------------------------------------------------------------------
+	FOverrideSelectionColorMaterialRenderProxy
+-----------------------------------------------------------------------------*/
+const FMaterial& FOverrideSelectionColorMaterialRenderProxy::GetMaterialWithFallback(ERHIFeatureLevel::Type InFeatureLevel, const FMaterialRenderProxy*& OutFallbackMaterialRenderProxy) const
+{
+	return Parent->GetMaterialWithFallback(InFeatureLevel, OutFallbackMaterialRenderProxy);
+}
+
+bool FOverrideSelectionColorMaterialRenderProxy::GetVectorValue(const FMaterialParameterInfo& ParameterInfo, FLinearColor* OutValue, const FMaterialRenderContext& Context) const
+{
+	if (ParameterInfo.Name == NAME_SelectionColor)
+	{
+		*OutValue = SelectionColor;
+		return true;
+	}
+	else
+	{
+		return Parent->GetVectorValue(ParameterInfo, OutValue, Context);
+	}
+}
+
+bool FOverrideSelectionColorMaterialRenderProxy::GetScalarValue(const FMaterialParameterInfo& ParameterInfo, float* OutValue, const FMaterialRenderContext& Context) const
+{
+	return Parent->GetScalarValue(ParameterInfo, OutValue, Context);
+}
+
+bool FOverrideSelectionColorMaterialRenderProxy::GetTextureValue(const FMaterialParameterInfo& ParameterInfo, const UTexture** OutValue, const FMaterialRenderContext& Context) const
+{
+	return Parent->GetTextureValue(ParameterInfo, OutValue, Context);
+}
+
 /*-----------------------------------------------------------------------------
 	FLightingDensityMaterialRenderProxy
 -----------------------------------------------------------------------------*/
diff --git a/Engine/Source/Runtime/Engine/Private/PrimitiveDrawingUtils.cpp b/Engine/Source/Runtime/Engine/Private/PrimitiveDrawingUtils.cpp
--- a/Engine/Source/Runtime/Engine/Private/PrimitiveDrawingUtils.cpp
+++ b/Engine/Source/Runtime/Engine/Private/PrimitiveDrawingUtils.cpp
@@ -1386,13 +1386,13 @@ void ApplyViewModeOverrides(
 #if !(UE_BUILD_SHIPPING || UE_BUILD_TEST)
 
 	// If debug viewmodes are not allowed, skip all of the debug viewmode handling.
-	// If material has tesselation then also skip it, as we can't just replace it with a simple material.
-	if (!AllowDebugViewmodes()
-		|| Mesh.MaterialRenderProxy->GetMaterial(FeatureLevel)->MaterialModifiesMeshPosition_RenderThread())
+	if (!AllowDebugViewmodes())
 	{
 		return;
 	}
 
+	const bool bMaterialModifiesMeshPosition = Mesh.MaterialRenderProxy->GetMaterial(FeatureLevel)->MaterialModifiesMeshPosition_RenderThread();
+
 	if (EngineShowFlags.Wireframe)
 	{
 		// In wireframe mode, draw the edges of the mesh with the specified wireframe color, or
@@ -1407,14 +1407,29 @@ void ApplyViewModeOverrides(
 			BaseColor = PrimitiveSceneProxy->GetLevelColor();
 		}
 
-		auto WireframeMaterialInstance = new FColoredMaterialRenderProxy(
-			GEngine->WireframeMaterial->GetRenderProxy(),
-			GetSelectionColor( BaseColor, bSelected, PrimitiveSceneProxy->IsHovered(), /*bUseOverlayIntensity=*/false)
+		if (bMaterialModifiesMeshPosition)
+		{
+			// If the material is mesh-modifying, we cannot rely on substitution.
+			auto WireframeMaterialInstance = new FOverrideSelectionColorMaterialRenderProxy(
+				Mesh.MaterialRenderProxy,
+				GetSelectionColor(BaseColor, bSelected, PrimitiveSceneProxy->IsHovered(), /*bUseOverlayIntensity=*/false)
 			);
 
-		Mesh.bWireframe = true;
-		Mesh.MaterialRenderProxy = WireframeMaterialInstance;
-		Collector.RegisterOneFrameMaterialProxy(WireframeMaterialInstance);
+			Mesh.bWireframe = true;
+			Mesh.MaterialRenderProxy = WireframeMaterialInstance;
+			Collector.RegisterOneFrameMaterialProxy(WireframeMaterialInstance);
+		}
+		else
+		{
+			auto WireframeMaterialInstance = new FColoredMaterialRenderProxy(
+				GEngine->WireframeMaterial->GetRenderProxy(),
+				GetSelectionColor(BaseColor, bSelected, PrimitiveSceneProxy->IsHovered(), /*bUseOverlayIntensity=*/false)
+			);
+
+			Mesh.bWireframe = true;
+			Mesh.MaterialRenderProxy = WireframeMaterialInstance;
+			Collector.RegisterOneFrameMaterialProxy(WireframeMaterialInstance);
+		}
 	}
 	else if (EngineShowFlags.LODColoration)
 	{
@@ -1478,7 +1493,18 @@ void ApplyViewModeOverrides(
 				}
 			}
 
-			if (bTextureMapped == false)
+			if (bMaterialModifiesMeshPosition)
+			{
+				// If the material is mesh-modifying, we cannot rely on substitution.
+				auto LightingOnlyMaterialInstance = new FOverrideSelectionColorMaterialRenderProxy(
+					Mesh.MaterialRenderProxy, 
+					GEngine->LightingOnlyBrightness
+				);
+
+				Mesh.MaterialRenderProxy = LightingOnlyMaterialInstance;
+				Collector.RegisterOneFrameMaterialProxy(LightingOnlyMaterialInstance);
+			}
+			else if (bTextureMapped == false)
 			{
 				FMaterialRenderProxy* RenderProxy = GEngine->LevelColorationLitMaterial->GetRenderProxy();
 				auto LightingOnlyMaterialInstance = new FColoredMaterialRenderProxy(
@@ -1569,17 +1595,35 @@ void ApplyViewModeOverrides(
 			// Avoid infinite recursion
 			MeshEdgeElement.bCanApplyViewModeOverrides = false;
 
+			
 			// Draw the mesh's edges in blue, on top of the base geometry.
-			auto WireframeMaterialInstance = new FColoredMaterialRenderProxy(
-				GEngine->WireframeMaterial->GetRenderProxy(),
-				PrimitiveSceneProxy->GetWireframeColor()
-			);
+			if (bMaterialModifiesMeshPosition)
+			{
+				// If the material is mesh-modifying, we cannot rely on substitution
+				auto WireframeMaterialInstance = new FOverrideSelectionColorMaterialRenderProxy(
+					MeshEdgeElement.MaterialRenderProxy,
+					PrimitiveSceneProxy->GetWireframeColor()
+				);
 
-			MeshEdgeElement.bWireframe = true;
-			MeshEdgeElement.MaterialRenderProxy = WireframeMaterialInstance;
-			Collector.RegisterOneFrameMaterialProxy(WireframeMaterialInstance);
+				MeshEdgeElement.bWireframe = true;
+				MeshEdgeElement.MaterialRenderProxy = WireframeMaterialInstance;
+				Collector.RegisterOneFrameMaterialProxy(WireframeMaterialInstance);
+
+				Collector.AddMesh(ViewIndex, MeshEdgeElement);
+			}
+			else
+			{
+				auto WireframeMaterialInstance = new FColoredMaterialRenderProxy(
+					GEngine->WireframeMaterial->GetRenderProxy(),
+					PrimitiveSceneProxy->GetWireframeColor()
+				);
 
-			Collector.AddMesh(ViewIndex, MeshEdgeElement);
+				MeshEdgeElement.bWireframe = true;
+				MeshEdgeElement.MaterialRenderProxy = WireframeMaterialInstance;
+				Collector.RegisterOneFrameMaterialProxy(WireframeMaterialInstance);
+
+				Collector.AddMesh(ViewIndex, MeshEdgeElement);
+			}
 		}
 	}
 #endif
diff --git a/Engine/Source/Runtime/Engine/Public/MaterialCompiler.h b/Engine/Source/Runtime/Engine/Public/MaterialCompiler.h
--- a/Engine/Source/Runtime/Engine/Public/MaterialCompiler.h
+++ b/Engine/Source/Runtime/Engine/Public/MaterialCompiler.h
@@ -294,6 +294,7 @@ public:
 	// The compiler can run in a different state and this affects caching of sub expression, Expressions are different (e.g. View.PrevWorldViewOrigin) when using previous frame's values
 	// If possible we should re-factor this to avoid having to deal with compiler state
 	virtual bool IsCurrentlyCompilingForPreviousFrame() const { return false; }
+	virtual bool IsDevelopmentFeatureEnabled(const FName& FeatureName) const { return true; }
 };
 
 /** 
@@ -551,6 +552,11 @@ public:
 		return Compiler->EyeAdaptation();
 	}
 
+	virtual bool IsDevelopmentFeatureEnabled(const FName& FeatureName) const override
+	{
+		return Compiler->IsDevelopmentFeatureEnabled(FeatureName);
+	}
+
 protected:
 		
 	FMaterialCompiler* Compiler;
diff --git a/Engine/Source/Runtime/Engine/Public/MaterialShared.h b/Engine/Source/Runtime/Engine/Public/MaterialShared.h
--- a/Engine/Source/Runtime/Engine/Public/MaterialShared.h
+++ b/Engine/Source/Runtime/Engine/Public/MaterialShared.h
@@ -1792,6 +1792,31 @@ public:
 	ENGINE_API virtual bool GetTextureValue(const FMaterialParameterInfo& ParameterInfo,const UTexture** OutValue, const FMaterialRenderContext& Context) const;
 };
 
+/**
+ * A material render proxy which overrides the selection color
+ */
+class FOverrideSelectionColorMaterialRenderProxy : public FMaterialRenderProxy
+{
+public:
+
+	const FMaterialRenderProxy* const Parent;
+	const FLinearColor SelectionColor;
+
+	/** Initialization constructor. */
+	FOverrideSelectionColorMaterialRenderProxy(const FMaterialRenderProxy* InParent, const FLinearColor& InSelectionColor) :
+		Parent(InParent),
+		SelectionColor(InSelectionColor)
+	{
+	}
+
+	// FMaterialRenderProxy interface.
+	ENGINE_API virtual const FMaterial& GetMaterialWithFallback(ERHIFeatureLevel::Type InFeatureLevel, const FMaterialRenderProxy*& OutFallbackMaterialRenderProxy) const;
+	ENGINE_API virtual bool GetVectorValue(const FMaterialParameterInfo& ParameterInfo, FLinearColor* OutValue, const FMaterialRenderContext& Context) const;
+	ENGINE_API virtual bool GetScalarValue(const FMaterialParameterInfo& ParameterInfo, float* OutValue, const FMaterialRenderContext& Context) const;
+	ENGINE_API virtual bool GetTextureValue(const FMaterialParameterInfo& ParameterInfo, const UTexture** OutValue, const FMaterialRenderContext& Context) const;
+};
+
+
 /**
  * An material render proxy which overrides the material's Color and Lightmap resolution vector parameter.
  */
diff --git a/Engine/Source/Runtime/Renderer/Private/MeshPassProcessor.cpp b/Engine/Source/Runtime/Renderer/Private/MeshPassProcessor.cpp
--- a/Engine/Source/Runtime/Renderer/Private/MeshPassProcessor.cpp
+++ b/Engine/Source/Runtime/Renderer/Private/MeshPassProcessor.cpp
@@ -300,8 +300,9 @@ void FMeshDrawShaderBindings::SetOnRayTracingStructure(FRHICommandList& RHICmdLi
 
 FGraphicsMinimalPipelineStateId FGraphicsMinimalPipelineStateId::GetPersistentId(const FGraphicsMinimalPipelineStateInitializer& InPipelineState)
 {
-	FSetElementId TableId = PersistentIdTable.FindId(InPipelineState);
+	checkSlow(IsInRenderingThread());
 
+	FSetElementId TableId = PersistentIdTable.FindId(InPipelineState);
 	if (TableId.IsValidId())
 	{
 		++PersistentIdTable[TableId].RefNum;
