Fixed bad type names when exposed to Python

#rnx

diff --git a/Engine/Plugins/Runtime/AnimationBudgetAllocator/Source/AnimationBudgetAllocator/Private/AnimationBudgetBlueprintLibrary.h b/Engine/Plugins/Runtime/AnimationBudgetAllocator/Source/AnimationBudgetAllocator/Private/AnimationBudgetBlueprintLibrary.h
--- a/Engine/Plugins/Runtime/AnimationBudgetAllocator/Source/AnimationBudgetAllocator/Private/AnimationBudgetBlueprintLibrary.h
+++ b/Engine/Plugins/Runtime/AnimationBudgetAllocator/Source/AnimationBudgetAllocator/Private/AnimationBudgetBlueprintLibrary.h
@@ -9,7 +9,7 @@
 /**
  * Function library to expose the budget allocator to Blueprints
  */
-UCLASS(meta = (ScriptName = "Animation Budget"))
+UCLASS(meta = (ScriptName = "AnimationBudgetLibrary"))
 class UAnimationBudgetBlueprintLibrary : public UBlueprintFunctionLibrary
 {
 	GENERATED_BODY()
diff --git a/Engine/Source/Runtime/Engine/Classes/Kismet/KismetMathLibrary.h b/Engine/Source/Runtime/Engine/Classes/Kismet/KismetMathLibrary.h
--- a/Engine/Source/Runtime/Engine/Classes/Kismet/KismetMathLibrary.h
+++ b/Engine/Source/Runtime/Engine/Classes/Kismet/KismetMathLibrary.h
@@ -1054,7 +1054,7 @@ class ENGINE_API UKismetMathLibrary : public UBlueprintFunctionLibrary
 	 * @param OutDir Reference passed in to store unit direction vector.
 	 * @param OutLength Reference passed in to store length of the vector.
 	 */
-	UFUNCTION(BlueprintPure, meta = (ScriptMethod = "To Direction And Length"), Category = "Math|Vector2D")
+	UFUNCTION(BlueprintPure, meta = (DisplayName = "To Direction And Length", ScriptMethod = "ToDirectionAndLength"), Category = "Math|Vector2D")
 	static void ToDirectionAndLength2D(FVector2D A, FVector2D &OutDir, float &OutLength);
 
 	/**
@@ -1062,7 +1062,7 @@ class ENGINE_API UKismetMathLibrary : public UBlueprintFunctionLibrary
 	 *
 	 * @return New FVector2D from this vector that is rounded.
 	 */
-	UFUNCTION(BlueprintPure, meta = (ScriptMethod = "To Rounded (Vector2D)"), Category = "Math|Vector2D")
+	UFUNCTION(BlueprintPure, meta = (DisplayName = "To Rounded (Vector2D)", ScriptMethod = "ToRounded"), Category = "Math|Vector2D")
 	static FVector2D ToRounded2D(FVector2D A);
 
 	/**
