Fixed OOM when rendering SplineMeshComponent with huge Tangent vector

diff --git a/Engine/Source/Editor/ComponentVisualizers/Private/SplineMeshComponentVisualizer.cpp b/Engine/Source/Editor/ComponentVisualizers/Private/SplineMeshComponentVisualizer.cpp
--- a/Engine/Source/Editor/ComponentVisualizers/Private/SplineMeshComponentVisualizer.cpp
+++ b/Engine/Source/Editor/ComponentVisualizers/Private/SplineMeshComponentVisualizer.cpp
@@ -56,8 +56,13 @@ void FSplineMeshComponentVisualizer::DrawVisualization(const UActorComponent* Co
 			const FVector TangentWorldDirection = SplineMeshComp->GetComponentTransform().TransformVector(Spline.Points[PointIndex].LeaveTangent);
 
 			PDI->SetHitProxy(NULL);
-			DrawDashedLine(PDI, KeyPos, KeyPos + TangentWorldDirection, Color, 5, SDPG_Foreground);
-			DrawDashedLine(PDI, KeyPos, KeyPos - TangentWorldDirection, Color, 5, SDPG_Foreground);
+			const float DefaultDashSize = 5.f;
+			const int32 MaxDashDrawCount = 65536;
+			const float LineSize = TangentWorldDirection.Size();
+			int32 DrawCount = FMath::CeilToInt(LineSize / (2.f * DefaultDashSize));
+			float DashSize = (DrawCount > MaxDashDrawCount) ? LineSize / (2.f * MaxDashDrawCount) : DefaultDashSize;
+			DrawDashedLine(PDI, KeyPos, KeyPos + TangentWorldDirection, Color, DashSize, SDPG_Foreground);
+			DrawDashedLine(PDI, KeyPos, KeyPos - TangentWorldDirection, Color, DashSize, SDPG_Foreground);
 			PDI->SetHitProxy(new HSplineMeshTangentHandleProxy(Component, PointIndex, false));
 			PDI->DrawPoint(KeyPos + TangentWorldDirection, Color, TangentHandleSize, SDPG_Foreground);
 			PDI->SetHitProxy(new HSplineMeshTangentHandleProxy(Component, PointIndex, true));
diff --git a/Engine/Source/Runtime/Engine/Private/Components/SplineMeshComponent.cpp b/Engine/Source/Runtime/Engine/Private/Components/SplineMeshComponent.cpp
--- a/Engine/Source/Runtime/Engine/Private/Components/SplineMeshComponent.cpp
+++ b/Engine/Source/Runtime/Engine/Private/Components/SplineMeshComponent.cpp
@@ -270,7 +270,7 @@ FVector USplineMeshComponent::GetEndTangent() const
 
 void USplineMeshComponent::SetEndTangent(FVector EndTangent, bool bUpdateMesh)
 {
-	SplineParams.EndTangent = EndTangent;
+	SplineParams.EndTangent = ClampVector(EndTangent, FVector(-WORLD_MAX), FVector(WORLD_MAX));
 	bMeshDirty = true;
 	if (bUpdateMesh)
 	{
@@ -283,7 +283,7 @@ void USplineMeshComponent::SetStartAndEnd(FVector StartPos, FVector StartTangent
 	SplineParams.StartPos = StartPos;
 	SplineParams.StartTangent = StartTangent;
 	SplineParams.EndPos = EndPos;
-	SplineParams.EndTangent = EndTangent;
+	SetEndTangent(EndTangent, false);
 	bMeshDirty = true;
 	if (bUpdateMesh)
 	{
@@ -1134,7 +1134,7 @@ void USplineMeshComponent::ApplyComponentInstanceData(FSplineMeshInstanceData* S
 			SplineParams.StartPos = SplineMeshInstanceData->StartPos;
 			SplineParams.EndPos = SplineMeshInstanceData->EndPos;
 			SplineParams.StartTangent = SplineMeshInstanceData->StartTangent;
-			SplineParams.EndTangent = SplineMeshInstanceData->EndTangent;
+			SetEndTangent(SplineMeshInstanceData->EndTangent, false);
 			UpdateRenderStateAndCollision();
 		}
 	}
@@ -1206,6 +1206,8 @@ void USplineMeshComponent::PostEditChangeProperty(FPropertyChangedEvent& Propert
 		// If the spline params were changed the actual geometry is, so flag the owning HLOD cluster as dirty
 		if (MemberPropertyThatChanged->GetNameCPP() == TEXT("SplineParams"))
 		{
+			SetEndTangent(SplineParams.EndTangent, false);
+
 			IHierarchicalLODUtilitiesModule& Module = FModuleManager::LoadModuleChecked<IHierarchicalLODUtilitiesModule>("HierarchicalLODUtilities");
 			IHierarchicalLODUtilities* Utilities = Module.GetUtilities();
 			Utilities->HandleActorModified(GetOwner());
