When getting out of streaming distance range, assume the farthest LOD, so we always have a LOD selected.

diff --git a/Engine/Source/Runtime/Engine/Private/WorldComposition.cpp b/Engine/Source/Runtime/Engine/Private/WorldComposition.cpp
--- a/Engine/Source/Runtime/Engine/Private/WorldComposition.cpp
+++ b/Engine/Source/Runtime/Engine/Private/WorldComposition.cpp
@@ -589,6 +589,10 @@ void UWorldComposition::GetDistanceVisibleLevels(
 			int32 NumAvailableLOD = FMath::Min(Tile.Info.LODList.Num(), Tile.LODPackageNames.Num());
 			// Find LOD
 			// INDEX_NONE for original non-LOD level
+
+			int32 FarthestTileStreamingDistance = 0;
+			int32 FarthestLODIndex = INDEX_NONE;
+
 			for (int32 LODIdx = INDEX_NONE; LODIdx < NumAvailableLOD; ++LODIdx)
 			{
 				if (bIsVisible && LODIdx > VisibleLevel.LODIndex)
@@ -607,7 +611,20 @@ void UWorldComposition::GetDistanceVisibleLevels(
 						bIsVisible = true;
 						break;
 					}
-				}
+
+					if (TileStreamingDistance > FarthestTileStreamingDistance)
+					{
+						FarthestTileStreamingDistance = TileStreamingDistance;
+						FarthestLODIndex = LODIdx;
+					}
+				}				
+			}
+
+			// If we're not in any streaming distance range, it's likely we're just too far, so assume the farthest LOD, but we can't assume that the last LOD is the farthest one
+			if (!bIsVisible)
+			{
+				VisibleLevel.LODIndex = FarthestLODIndex;
+				bIsVisible = true;
 			}
 		}
 
