SceneComponent now broadcasts OnVisibilityChanged and OnHiddenInGameChanged events when bVisible or bHiddenInGame is toggled from the editor details panel

diff --git a/Engine/Source/Runtime/Engine/Private/Components/SceneComponent.cpp b/Engine/Source/Runtime/Engine/Private/Components/SceneComponent.cpp
--- a/Engine/Source/Runtime/Engine/Private/Components/SceneComponent.cpp
+++ b/Engine/Source/Runtime/Engine/Private/Components/SceneComponent.cpp
@@ -508,6 +508,14 @@ void USceneComponent::PostEditChangeProperty(FPropertyChangedEvent& PropertyChan
 	{
 		UpdateAttachedIsEditorOnly(this);
 	}
+	if (PropertyName == GET_MEMBER_NAME_CHECKED(USceneComponent, bVisible))
+	{
+		OnVisibilityChanged();
+	}
+	else if (PropertyName == GET_MEMBER_NAME_CHECKED(USceneComponent, bHiddenInGame))
+	{
+		OnHiddenInGameChanged();
+	}
 
 	Super::PostEditChangeProperty(PropertyChangedEvent);
 
