SceneCaptureComponent now responds to OnVisibilityChanged events and sets the ProxyMesh Visibility to match the component's visibility

diff --git a/Engine/Source/Runtime/Engine/Classes/Components/SceneCaptureComponent.h b/Engine/Source/Runtime/Engine/Classes/Components/SceneCaptureComponent.h
--- a/Engine/Source/Runtime/Engine/Classes/Components/SceneCaptureComponent.h
+++ b/Engine/Source/Runtime/Engine/Classes/Components/SceneCaptureComponent.h
@@ -181,6 +181,9 @@ public:
 
 	static void UpdateDeferredCaptures(FSceneInterface* Scene);
 
+	/** Respond to visibility changes by updating Proxy Mesh visibility to match*/
+	virtual void OnVisibilityChanged() override;
+
 protected:
 	/** Update the show flags from our show flags settings (ideally, you'd be able to set this more directly, but currently unable to make FEngineShowFlags a UStruct to use it as a UProperty...) */
 	void UpdateShowFlags();
diff --git a/Engine/Source/Runtime/Engine/Private/Components/SceneCaptureComponent.cpp b/Engine/Source/Runtime/Engine/Private/Components/SceneCaptureComponent.cpp
--- a/Engine/Source/Runtime/Engine/Private/Components/SceneCaptureComponent.cpp
+++ b/Engine/Source/Runtime/Engine/Private/Components/SceneCaptureComponent.cpp
@@ -141,7 +141,8 @@ void USceneCaptureComponent::OnRegister()
 			ProxyMeshComponent->SetIsVisualizationComponent(true);
 			ProxyMeshComponent->SetStaticMesh(CaptureMesh);
 			ProxyMeshComponent->SetCollisionProfileName(UCollisionProfile::NoCollision_ProfileName);
-			ProxyMeshComponent->bHiddenInGame = true;
+			ProxyMeshComponent->bVisible = bVisible; // Match the visibility of the component
+			ProxyMeshComponent->bHiddenInGame = true; // Hidden in game should always be true as this is a visualization component
 			ProxyMeshComponent->CastShadow = false;
 			ProxyMeshComponent->CreationMethod = CreationMethod;
 			ProxyMeshComponent->RegisterComponentWithWorld(GetWorld());
@@ -401,6 +402,15 @@ void USceneCaptureComponent::UpdateDeferredCaptures(FSceneInterface* Scene)
 	SceneCapturesToUpdateMap.Remove(World);
 }
 
+void USceneCaptureComponent::OnVisibilityChanged()
+{
+	if (ProxyMeshComponent != nullptr)
+	{
+		ProxyMeshComponent->SetVisibility(bVisible);
+	}
+	Super::OnVisibilityChanged();
+}
+
 void USceneCaptureComponent::OnUnregister()
 {
 	for (int32 ViewIndex = 0; ViewIndex < ViewStates.Num(); ViewIndex++)
