Reallocate uniformbuffer when size has changed.

diff --git a/Engine/Source/Runtime/Engine/Private/Materials/ParameterCollection.cpp b/Engine/Source/Runtime/Engine/Private/Materials/ParameterCollection.cpp
--- a/Engine/Source/Runtime/Engine/Private/Materials/ParameterCollection.cpp
+++ b/Engine/Source/Runtime/Engine/Private/Materials/ParameterCollection.cpp
@@ -696,11 +696,13 @@ void FMaterialParameterCollectionInstanceResource::UpdateContents(const FGuid& I
 
 	if (InId != FGuid() && Data.Num() > 0)
 	{
+		uint32 Size = Data.GetTypeSize() * Data.Num();
+		bool bSizeChanged = Size != UniformBufferLayout.ConstantBufferSize;
 		UniformBufferLayout.ConstantBufferSize = Data.GetTypeSize() * Data.Num();
 		UniformBufferLayout.ComputeHash();
 		check(UniformBufferLayout.Resources.Num() == 0);
 
-		if (IsValidRef(UniformBuffer))
+		if (!bSizeChanged && IsValidRef(UniformBuffer))
 		{
 			check(UniformBuffer->GetLayout() == UniformBufferLayout);
 			RHIUpdateUniformBuffer(UniformBuffer, Data.GetData());
