UE4.22 - Do not cache static meshes with external textures; re-enable immutable samplers on mesh pass processor
#rnx

diff --git a/Engine/Source/Runtime/Engine/Private/PrimitiveSceneProxy.cpp b/Engine/Source/Runtime/Engine/Private/PrimitiveSceneProxy.cpp
--- a/Engine/Source/Runtime/Engine/Private/PrimitiveSceneProxy.cpp
+++ b/Engine/Source/Runtime/Engine/Private/PrimitiveSceneProxy.cpp
@@ -41,6 +41,29 @@ bool SupportsCachingMeshDrawCommands(const FVertexFactory* RESTRICT VertexFactor
 		&& !PrimitiveSceneProxy->CastsVolumetricTranslucentShadow();
 }
 
+bool SupportsCachingMeshDrawCommands(const FVertexFactory* RESTRICT VertexFactory, const FPrimitiveSceneProxy* RESTRICT PrimitiveSceneProxy, const FMaterialRenderProxy* MaterialRenderProxy, ERHIFeatureLevel::Type FeatureLevel)
+{
+	if (SupportsCachingMeshDrawCommands(VertexFactory, PrimitiveSceneProxy))
+	{
+		// External textures get mapped to immutable samplers (which are part of the PSO); the mesh must go through the dynamic path, as the media player might not have
+		// valid textures/samplers the first few calls; once they're available the PSO needs to get invalidated and recreated with the immutable samplers.
+		const FMaterial* Material = MaterialRenderProxy->GetMaterial(FeatureLevel);
+		const FMaterialShaderMap* ShaderMap = Material->GetRenderingThreadShaderMap();
+		if (ShaderMap)
+		{
+			const FUniformExpressionSet& ExpressionSet = ShaderMap->GetUniformExpressionSet();
+			if (ExpressionSet.HasExternalTextureExpressions())
+			{
+				return false;
+			}
+		}
+
+		return true;
+	}
+
+	return false;
+}
+
 FPrimitiveSceneProxy::FPrimitiveSceneProxy(const UPrimitiveComponent* InComponent, FName InResourceName)
 :
 #if !(UE_BUILD_SHIPPING || UE_BUILD_TEST)
diff --git a/Engine/Source/Runtime/Engine/Public/MaterialShared.h b/Engine/Source/Runtime/Engine/Public/MaterialShared.h
--- a/Engine/Source/Runtime/Engine/Public/MaterialShared.h
+++ b/Engine/Source/Runtime/Engine/Public/MaterialShared.h
@@ -372,6 +372,11 @@ public:
 
 	ENGINE_API void FillUniformBuffer(const FMaterialRenderContext& MaterialRenderContext, void* TempBuffer) const;
 
+	inline bool HasExternalTextureExpressions() const
+	{
+		return UniformExternalTextureExpressions.Num() > 0;
+	}
+
 	uint32 GetAllocatedSize() const
 	{
 		return UniformVectorExpressions.GetAllocatedSize()
diff --git a/Engine/Source/Runtime/Engine/Public/PrimitiveSceneProxy.h b/Engine/Source/Runtime/Engine/Public/PrimitiveSceneProxy.h
--- a/Engine/Source/Runtime/Engine/Public/PrimitiveSceneProxy.h
+++ b/Engine/Source/Runtime/Engine/Public/PrimitiveSceneProxy.h
@@ -1055,4 +1055,10 @@ protected:
 /**
  * Returns if specified mesh command can be cached, or needs to be recreated every frame.
  */
-ENGINE_API extern bool SupportsCachingMeshDrawCommands(const FVertexFactory* RESTRICT VertexFactory, const FPrimitiveSceneProxy* RESTRICT PrimitiveSceneProxy);
\ No newline at end of file
+ENGINE_API extern bool SupportsCachingMeshDrawCommands(const FVertexFactory* RESTRICT VertexFactory, const FPrimitiveSceneProxy* RESTRICT PrimitiveSceneProxy);
+
+/**
+ * Returns if specified mesh command can be cached, or needs to be recreated every frame; this is a slightly slower version
+ * used for materials with external textures that need invalidating their PSOs.
+ */
+ENGINE_API extern bool SupportsCachingMeshDrawCommands(const FVertexFactory* RESTRICT VertexFactory, const FPrimitiveSceneProxy* RESTRICT PrimitiveSceneProxy, const class FMaterialRenderProxy* MaterialRenderProxy, ERHIFeatureLevel::Type FeatureLevel);
diff --git a/Engine/Source/Runtime/Renderer/Private/PrimitiveSceneInfo.cpp b/Engine/Source/Runtime/Renderer/Private/PrimitiveSceneInfo.cpp
--- a/Engine/Source/Runtime/Renderer/Private/PrimitiveSceneInfo.cpp
+++ b/Engine/Source/Runtime/Renderer/Private/PrimitiveSceneInfo.cpp
@@ -62,9 +62,10 @@ public:
 				CurrentHitProxy ? CurrentHitProxy->Id : FHitProxyId()
 				);
 
-			StaticMesh->PreparePrimitiveUniformBuffer(PrimitiveSceneProxy, PrimitiveSceneInfo->Scene->GetFeatureLevel());
+			const ERHIFeatureLevel::Type FeatureLevel = PrimitiveSceneInfo->Scene->GetFeatureLevel();
+			StaticMesh->PreparePrimitiveUniformBuffer(PrimitiveSceneProxy, FeatureLevel);
 
-			const bool bSupportsCachingMeshDrawCommands = SupportsCachingMeshDrawCommands(StaticMesh->VertexFactory, PrimitiveSceneProxy);
+			const bool bSupportsCachingMeshDrawCommands = SupportsCachingMeshDrawCommands(StaticMesh->VertexFactory, PrimitiveSceneProxy, Mesh.MaterialRenderProxy, FeatureLevel);
 
 			FStaticMeshBatchRelevance* StaticMeshRelevance = new(PrimitiveSceneInfo->StaticMeshRelevances) FStaticMeshBatchRelevance(
 				*StaticMesh, 
diff --git a/Engine/Source/Runtime/Renderer/Public/MeshPassProcessor.inl b/Engine/Source/Runtime/Renderer/Public/MeshPassProcessor.inl
--- a/Engine/Source/Runtime/Renderer/Public/MeshPassProcessor.inl
+++ b/Engine/Source/Runtime/Renderer/Public/MeshPassProcessor.inl
@@ -24,14 +24,13 @@ void FMeshPassProcessor::BuildMeshDrawCommands(
 	const FVertexFactory* RESTRICT VertexFactory = MeshBatch.VertexFactory;
 	const FPrimitiveSceneInfo* RESTRICT PrimitiveSceneInfo = PrimitiveSceneProxy ? PrimitiveSceneProxy->GetPrimitiveSceneInfo() : nullptr;
 
-	checkf(MaterialRenderProxy.ImmutableSamplerState.ImmutableSamplers[0] == nullptr, TEXT("Immutable samplers not yet supported in Mesh Draw Command pipeline"));
-
 	FMeshDrawCommand SharedMeshDrawCommand;
 
 	SharedMeshDrawCommand.SetStencilRef(DrawRenderState.GetStencilRef());
 
 	FGraphicsMinimalPipelineStateInitializer PipelineState;
 	PipelineState.PrimitiveType = (EPrimitiveType)MeshBatch.Type;
+	PipelineState.ImmutableSamplerState = MaterialRenderProxy.ImmutableSamplerState;
 
 	const bool bPositionOnly = (MeshPassFeatures & EMeshPassFeatures::PositionOnly) != EMeshPassFeatures::Default;
 
