Back out changelist 5070596, which had other odd side effects not found in original testing.

diff --git a/Engine/Source/Editor/UnrealEd/Private/Factories/EditorFactories.cpp b/Engine/Source/Editor/UnrealEd/Private/Factories/EditorFactories.cpp
--- a/Engine/Source/Editor/UnrealEd/Private/Factories/EditorFactories.cpp
+++ b/Engine/Source/Editor/UnrealEd/Private/Factories/EditorFactories.cpp
@@ -650,7 +650,7 @@ UObject* UMaterialParameterCollectionFactoryNew::FactoryCreateNew(UClass* Class,
 		for (TObjectIterator<UWorld> It; It; ++It)
 		{
 			UWorld* CurrentWorld = *It;
-			CurrentWorld->AddParameterCollectionInstance(MaterialParameterCollection, true /* bUpdateScene */);
+			CurrentWorld->AddParameterCollectionInstance(MaterialParameterCollection, true);
 		}
 	}
 
diff --git a/Engine/Source/Runtime/Engine/Classes/Engine/World.h b/Engine/Source/Runtime/Engine/Classes/Engine/World.h
--- a/Engine/Source/Runtime/Engine/Classes/Engine/World.h
+++ b/Engine/Source/Runtime/Engine/Classes/Engine/World.h
@@ -1138,7 +1138,7 @@ private:
 	class UGameInstance*						OwningGameInstance;
 
 	/** Parameter collection instances that hold parameter overrides for this world. */
-	UPROPERTY(NonPIEDuplicateTransient)
+	UPROPERTY(Transient)
 	TArray<class UMaterialParameterCollectionInstance*> ParameterCollectionInstances;
 
 	/** 
@@ -2416,7 +2416,7 @@ public:
 	void SetupParameterCollectionInstances();
 
 	/** Adds a new instance of the given collection, or overwrites an existing instance if there is one. */
-	void AddParameterCollectionInstance(class UMaterialParameterCollection* Collection, bool bUpdateScene, bool bOverwriteIfExists = true);
+	void AddParameterCollectionInstance(class UMaterialParameterCollection* Collection, bool bUpdateScene);
 
 	/** Gets this world's instance for a given collection. */
 	UMaterialParameterCollectionInstance* GetParameterCollectionInstance(const UMaterialParameterCollection* Collection);
diff --git a/Engine/Source/Runtime/Engine/Classes/Materials/MaterialParameterCollectionInstance.h b/Engine/Source/Runtime/Engine/Classes/Materials/MaterialParameterCollectionInstance.h
--- a/Engine/Source/Runtime/Engine/Classes/Materials/MaterialParameterCollectionInstance.h
+++ b/Engine/Source/Runtime/Engine/Classes/Materials/MaterialParameterCollectionInstance.h
@@ -31,9 +31,6 @@ class ENGINE_API UMaterialParameterCollectionInstance : public UObject
 	/** Initializes the instance with the collection it is based off of and the world it is owned by. */
 	void SetCollection(UMaterialParameterCollection* InCollection, UWorld* InWorld);
 
-	/** Assigns new parent world instance. */
-	void SetWorld(UWorld* InWorld);
-
 	/** Sets parameter value overrides, returns false if the parameter was not found. */
 	bool SetScalarParameterValue(FName ParameterName, float ParameterValue);
 	bool SetVectorParameterValue(FName ParameterName, const FLinearColor& ParameterValue);
@@ -69,11 +66,9 @@ protected:
 	TWeakObjectPtr<UWorld> World;
 
 	/** Overrides for scalar parameter values. */
-	UPROPERTY()
 	TMap<FName, float> ScalarParameterValues;
 
 	/** Overrides for vector parameter values. */
-	UPROPERTY()
 	TMap<FName, FLinearColor> VectorParameterValues;
 
 	/** Instance resource which stores the rendering thread representation of this instance. */
diff --git a/Engine/Source/Runtime/Engine/Private/Materials/ParameterCollection.cpp b/Engine/Source/Runtime/Engine/Private/Materials/ParameterCollection.cpp
--- a/Engine/Source/Runtime/Engine/Private/Materials/ParameterCollection.cpp
+++ b/Engine/Source/Runtime/Engine/Private/Materials/ParameterCollection.cpp
@@ -53,7 +53,7 @@ void UMaterialParameterCollection::PostLoad()
 	for (TObjectIterator<UWorld> It; It; ++It)
 	{
 		UWorld* CurrentWorld = *It;
-		CurrentWorld->AddParameterCollectionInstance(this, true /* bUpdateScene */);
+		CurrentWorld->AddParameterCollectionInstance(this, true);
 	}
 
 	UpdateDefaultResource();
@@ -186,7 +186,7 @@ void UMaterialParameterCollection::PostEditChangeProperty(FPropertyChangedEvent&
 		for (TObjectIterator<UWorld> It; It; ++It)
 		{
 			UWorld* CurrentWorld = *It;
-			CurrentWorld->AddParameterCollectionInstance(this, false /* bUpdateScene */);
+			CurrentWorld->AddParameterCollectionInstance(this, false);
 		}
 
 		// Create a material update context so we can safely update materials using this parameter collection.
@@ -466,11 +466,6 @@ void UMaterialParameterCollectionInstance::PostInitProperties()
 	}
 }
 
-void UMaterialParameterCollectionInstance::SetWorld(UWorld* InWorld)
-{
-	World = InWorld;
-}
-
 void UMaterialParameterCollectionInstance::SetCollection(UMaterialParameterCollection* InCollection, UWorld* InWorld)
 {
 	Collection = InCollection;
diff --git a/Engine/Source/Runtime/Engine/Private/World.cpp b/Engine/Source/Runtime/Engine/Private/World.cpp
--- a/Engine/Source/Runtime/Engine/Private/World.cpp
+++ b/Engine/Source/Runtime/Engine/Private/World.cpp
@@ -1033,22 +1033,16 @@ void UWorld::SetupParameterCollectionInstances()
 {
 	QUICK_SCOPE_CYCLE_COUNTER(Stat_World_SetupParameterCollectionInstances);
 
-	// Newly serialized collection instances need to be parented to this world instance.
-	for (int32 InstanceIndex = 0; InstanceIndex < ParameterCollectionInstances.Num(); InstanceIndex++)
-	{
-		ParameterCollectionInstances[InstanceIndex]->SetWorld(this);
-	}
-
-	// Create an instance for each parameter collection in memory that didn't already exist.
+	// Create an instance for each parameter collection in memory
 	for (UMaterialParameterCollection* CurrentCollection : TObjectRange<UMaterialParameterCollection>())
 	{
-		AddParameterCollectionInstance(CurrentCollection, false /* bUpdateScene */, false /* bOverwriteIfExists */);
+		AddParameterCollectionInstance(CurrentCollection, false);
 	}
 
 	UpdateParameterCollectionInstances(false);
 }
 
-void UWorld::AddParameterCollectionInstance(UMaterialParameterCollection* Collection, bool bUpdateScene, bool bOverwriteIfExists)
+void UWorld::AddParameterCollectionInstance(UMaterialParameterCollection* Collection, bool bUpdateScene)
 {
 	int32 ExistingIndex = INDEX_NONE;
 
@@ -1061,30 +1055,23 @@ void UWorld::AddParameterCollectionInstance(UMaterialParameterCollection* Collec
 		}
 	}
 
-	const bool bFoundExisting = ExistingIndex != INDEX_NONE;
-	const bool bMakeNewInstance = !bFoundExisting || (bFoundExisting && bOverwriteIfExists);
+	UMaterialParameterCollectionInstance* NewInstance = NewObject<UMaterialParameterCollectionInstance>();
+	NewInstance->SetCollection(Collection, this);
 
-	if (bMakeNewInstance)
+	if (ExistingIndex != INDEX_NONE)
 	{
-		UMaterialParameterCollectionInstance* NewInstance = NewObject<UMaterialParameterCollectionInstance>();
-		NewInstance->SetCollection(Collection, this);
-
-		if (bFoundExisting)
-		{
-			ParameterCollectionInstances[ExistingIndex] = NewInstance;
-		}
-		else
-		{
-			ParameterCollectionInstances.Add(NewInstance);
-		}
-
-		// Ensure the new instance creates initial render thread resources
-		// This needs to happen right away, so they can be picked up by any cached shader bindings
-		if (NewInstance)
-		{
-			NewInstance->DeferredUpdateRenderState(false);
-		}
+		// Overwrite an existing instance
+		ParameterCollectionInstances[ExistingIndex] = NewInstance;
 	}
+	else
+	{
+		// Add a new instance
+		ParameterCollectionInstances.Add(NewInstance);
+	}
+
+	// Ensure the new instance creates initial render thread resources
+	// This needs to happen right away, so they can be picked up by any cached shader bindings
+	NewInstance->DeferredUpdateRenderState(false);
 
 	if (bUpdateScene)
 	{
