Removed recreation of UDS default instance as it is no longer needed

diff --git a/Engine/Source/Runtime/Engine/Private/UserDefinedStruct.cpp b/Engine/Source/Runtime/Engine/Private/UserDefinedStruct.cpp
--- a/Engine/Source/Runtime/Engine/Private/UserDefinedStruct.cpp
+++ b/Engine/Source/Runtime/Engine/Private/UserDefinedStruct.cpp
@@ -154,8 +154,10 @@ void UUserDefinedStruct::Serialize(FStructuredArchive::FRecord Record)
 			if (EditorData != nullptr)
 			{
 				UnderlyingArchive.Preload(EditorData);
-				if (!(UnderlyingArchive.GetPortFlags() & PPF_Duplicate))
+				if (!(UnderlyingArchive.GetPortFlags() & PPF_Duplicate) && !DefaultStructInstance.IsValid())
 				{
+					// There is likely no need for this, but historically we were unconditionally recreating the
+					// default instance, so make sure that it is created here:
 					FStructureEditorUtils::RecreateDefaultInstanceInEditorData(this);
 				}
 			}
