Undo //UE4/Release-4.22/Engine/Source/Runtime/Engine/Private/... changelist 5178602

diff --git a/Engine/Source/Runtime/Engine/Private/UserDefinedStruct.cpp b/Engine/Source/Runtime/Engine/Private/UserDefinedStruct.cpp
--- a/Engine/Source/Runtime/Engine/Private/UserDefinedStruct.cpp
+++ b/Engine/Source/Runtime/Engine/Private/UserDefinedStruct.cpp
@@ -154,10 +154,8 @@ void UUserDefinedStruct::Serialize(FStructuredArchive::FRecord Record)
 			if (EditorData != nullptr)
 			{
 				UnderlyingArchive.Preload(EditorData);
-				if (!(UnderlyingArchive.GetPortFlags() & PPF_Duplicate) && !DefaultStructInstance.IsValid())
+				if (!(UnderlyingArchive.GetPortFlags() & PPF_Duplicate))
 				{
-					// There is likely no need for this, but historically we were unconditionally recreating the
-					// default instance, so make sure that it is created here:
 					FStructureEditorUtils::RecreateDefaultInstanceInEditorData(this);
 				}
 			}
