PR #5546: UE-69951: World nullptr check in CheckForErrors (Contributed by [REDACTED_PERSON])

diff --git a/Engine/Source/Runtime/Engine/Private/WorldSettings.cpp b/Engine/Source/Runtime/Engine/Private/WorldSettings.cpp
--- a/Engine/Source/Runtime/Engine/Private/WorldSettings.cpp
+++ b/Engine/Source/Runtime/Engine/Private/WorldSettings.cpp
@@ -491,6 +491,12 @@ void AWorldSettings::CheckForErrors()
 	Super::CheckForErrors();
 
 	UWorld* World = GetWorld();
+	// World is nullptr if save is done from a derived AWorldSettings blueprint
+	if (World == nullptr)
+	{
+		return;
+	}
+
 	if ( World->GetWorldSettings() != this )
 	{
 		FMessageLog("MapCheck").Warning()
