Fixed instancing not taking into account instance scale while calculating LOD, it could result in having large instance clipped. Tested using LOD Ditherering and normal one.
Remove part of the previous tentative fixed i had made a few months ago.

diff --git a/Engine/Source/Runtime/Engine/Classes/Components/InstancedStaticMeshComponent.h b/Engine/Source/Runtime/Engine/Classes/Components/InstancedStaticMeshComponent.h
--- a/Engine/Source/Runtime/Engine/Classes/Components/InstancedStaticMeshComponent.h
+++ b/Engine/Source/Runtime/Engine/Classes/Components/InstancedStaticMeshComponent.h
@@ -300,6 +300,7 @@ public:
 
 	virtual void PropagateLightingScenarioChange() override;
 
+	void GetInstancesMinMaxScale(FVector& MinScale, FVector& MaxScale) const;
 private:
 
 	/** Sets up new instance data to sensible defaults, creates physics counterparts if possible. */
diff --git a/Engine/Source/Runtime/Engine/Private/HierarchicalInstancedStaticMesh.cpp b/Engine/Source/Runtime/Engine/Private/HierarchicalInstancedStaticMesh.cpp
--- a/Engine/Source/Runtime/Engine/Private/HierarchicalInstancedStaticMesh.cpp
+++ b/Engine/Source/Runtime/Engine/Private/HierarchicalInstancedStaticMesh.cpp
@@ -1135,7 +1135,7 @@ static FORCEINLINE_DEBUGGABLE bool CullNode(const FFoliageCullInstanceParams& Pa
 	return false;
 }
 
-inline void CalcLOD(int32& InOutMinLOD, int32& InOutMaxLOD, const FVector& BoundMin, const FVector& BoundMax, const FVector& ViewOriginInLocalZero, const FVector& ViewOriginInLocalOne, const float (&LODPlanesMin)[MAX_STATIC_MESH_LODS], const float (&LODPlanesMax)[MAX_STATIC_MESH_LODS], float LODDistanceScaleFactor)
+inline void CalcLOD(int32& InOutMinLOD, int32& InOutMaxLOD, const FVector& BoundMin, const FVector& BoundMax, const FVector& ViewOriginInLocalZero, const FVector& ViewOriginInLocalOne, const float (&LODPlanesMin)[MAX_STATIC_MESH_LODS], const float (&LODPlanesMax)[MAX_STATIC_MESH_LODS])
 {
 	if (InOutMinLOD != InOutMaxLOD)
 	{
@@ -1146,11 +1146,11 @@ inline void CalcLOD(int32& InOutMinLOD, int32& InOutMaxLOD, const FVector& Bound
 		const float NearDot = FMath::Min(DistCenterZero, DistCenterOne) - HalfWidth;
 		const float FarDot = FMath::Max(DistCenterZero, DistCenterOne) + HalfWidth;
 
-		while (InOutMaxLOD > InOutMinLOD && NearDot > LODPlanesMax[InOutMinLOD] * LODDistanceScaleFactor)
+		while (InOutMaxLOD > InOutMinLOD && NearDot > LODPlanesMax[InOutMinLOD])
 		{
 			InOutMinLOD++;
 		}
-		while (InOutMaxLOD > InOutMinLOD && FarDot < LODPlanesMin[InOutMaxLOD - 1] * LODDistanceScaleFactor)
+		while (InOutMaxLOD > InOutMinLOD && FarDot < LODPlanesMin[InOutMaxLOD - 1])
 		{
 			InOutMaxLOD--;
 		}
@@ -1183,14 +1183,9 @@ void FHierarchicalStaticMeshSceneProxy::Traverse(const FFoliageCullInstanceParam
 		}
 	}
 
-	float DistanceScale = 1.0f;
-
 	if (MinLOD != MaxLOD)
 	{
-		FVector ScaleAverage = Node.MinInstanceScale + ((Node.MaxInstanceScale - Node.MinInstanceScale) / 2.0f);
-		//DistanceScale = FMath::Max(FMath::Max3(ScaleAverage.X, ScaleAverage.Y, ScaleAverage.Z), 0.001f);
-
-		CalcLOD(MinLOD, MaxLOD, Node.BoundMin, Node.BoundMax, Params.ViewOriginInLocalZero, Params.ViewOriginInLocalOne, Params.LODPlanesMin, Params.LODPlanesMax, DistanceScale);
+		CalcLOD(MinLOD, MaxLOD, Node.BoundMin, Node.BoundMax, Params.ViewOriginInLocalZero, Params.ViewOriginInLocalOne, Params.LODPlanesMin, Params.LODPlanesMax);
 
 		if (MinLOD >= Params.LODs)
 		{
@@ -1211,7 +1206,7 @@ void FHierarchicalStaticMeshSceneProxy::Traverse(const FFoliageCullInstanceParam
 
 	bool bShouldGroup = Node.FirstChild < 0
 		|| ((Node.LastInstance - Node.FirstInstance + 1) < Params.MinInstancesToSplit[MinLOD]
-			&& CanGroup(Node.BoundMin, Node.BoundMax, Params.ViewOriginInLocalZero, Params.ViewOriginInLocalOne, Params.LODPlanesMax[Params.LODs - 1] * DistanceScale));
+			&& CanGroup(Node.BoundMin, Node.BoundMax, Params.ViewOriginInLocalZero, Params.ViewOriginInLocalOne, Params.LODPlanesMax[Params.LODs - 1]));
 	bool bSplit = (!bFullyContained || MinLOD < MaxLOD || Index < Params.FirstOcclusionNode)
 		&& !bShouldGroup;
 
@@ -1600,12 +1595,15 @@ void FHierarchicalStaticMeshSceneProxy::GetDynamicMeshElements(const TArray<cons
 				float LODScale = CVarFoliageLODDistanceScale.GetValueOnRenderThread();
 				float LODRandom = CVarRandomLODRange.GetValueOnRenderThread();
 				float MaxDrawDistanceScale = GetCachedScalabilityCVars().ViewDistanceScale;
-				float SphereRadius = RenderData->Bounds.SphereRadius;
+
+				FVector AverageScale = (InstanceParams.Tree[0].MinInstanceScale + (InstanceParams.Tree[0].MaxInstanceScale - InstanceParams.Tree[0].MinInstanceScale) / 2.0f);
+				FBoxSphereBounds ScaledBounds = RenderData->Bounds.TransformBy(FTransform(FRotator::ZeroRotator, FVector::ZeroVector, AverageScale));
+				float SphereRadius = ScaledBounds.SphereRadius;
 
 				float FinalCull = MAX_flt;
 				if (MinSize > 0.0)
 				{
-					FinalCull = ComputeBoundsDrawDistance(MinSize, SphereRadius * FMath::Max(FMath::Min3(ClusterTree[0].MinInstanceScale.X, ClusterTree[0].MinInstanceScale.Y, ClusterTree[0].MinInstanceScale.Z), 0.001f), View->ViewMatrices.GetProjectionMatrix()) * LODScale;
+					FinalCull = ComputeBoundsDrawDistance(MinSize, SphereRadius, View->ViewMatrices.GetProjectionMatrix()) * LODScale;
 				}
 				if (UserData_AllInstances.EndCullDistance > 0.0f)
 				{
@@ -1803,13 +1801,13 @@ void FHierarchicalStaticMeshSceneProxy::GetDynamicMeshElements(const TArray<cons
 						// calculate runs
 						int32 MinLOD = 0;
 						int32 MaxLOD = NumLODs;
-						CalcLOD(MinLOD, MaxLOD, UnbuiltBounds[0].Min, UnbuiltBounds[0].Max, ViewOriginInLocalZero, ViewOriginInLocalOne, LODPlanesMin, LODPlanesMax, 1.0f);
+						CalcLOD(MinLOD, MaxLOD, UnbuiltBounds[0].Min, UnbuiltBounds[0].Max, ViewOriginInLocalZero, ViewOriginInLocalOne, LODPlanesMin, LODPlanesMax);
 						int32 FirstIndexInRun = 0;
 						for (int32 Index = 1; Index < UnbuiltInstanceCount; ++Index)
 						{
 							int32 TempMinLOD = 0;
 							int32 TempMaxLOD = NumLODs;
-							CalcLOD(TempMinLOD, TempMaxLOD, UnbuiltBounds[Index].Min, UnbuiltBounds[Index].Max, ViewOriginInLocalZero, ViewOriginInLocalOne, LODPlanesMin, LODPlanesMax, 1.0f);
+							CalcLOD(TempMinLOD, TempMaxLOD, UnbuiltBounds[Index].Min, UnbuiltBounds[Index].Max, ViewOriginInLocalZero, ViewOriginInLocalOne, LODPlanesMin, LODPlanesMax);
 							if (TempMinLOD != MinLOD)
 							{
 								if (MinLOD < NumLODs)
@@ -3264,3 +3262,4 @@ static FAutoConsoleCommand RebuildFoliageTreesCmd(
 	TEXT("Rebuild the trees for non-grass foliage."),
 	FConsoleCommandWithArgsDelegate::CreateStatic(&RebuildFoliageTrees)
 	);
+
diff --git a/Engine/Source/Runtime/Engine/Private/InstancedStaticMesh.cpp b/Engine/Source/Runtime/Engine/Private/InstancedStaticMesh.cpp
--- a/Engine/Source/Runtime/Engine/Private/InstancedStaticMesh.cpp
+++ b/Engine/Source/Runtime/Engine/Private/InstancedStaticMesh.cpp
@@ -782,6 +782,55 @@ int32 FInstancedStaticMeshSceneProxy::GetNumMeshBatches() const
 	}
 }
 
+void FInstancedStaticMeshSceneProxy::SetupProxy(UInstancedStaticMeshComponent* InComponent)
+{
+#if WITH_EDITOR
+	if (bHasSelectedInstances)
+	{
+		// if we have selected indices, mark scene proxy as selected.
+		SetSelection_GameThread(true);
+	}
+#endif
+	// Make sure all the materials are okay to be rendered as an instanced mesh.
+	for (int32 LODIndex = 0; LODIndex < LODs.Num(); LODIndex++)
+	{
+		FStaticMeshSceneProxy::FLODInfo& LODInfo = LODs[LODIndex];
+		for (int32 SectionIndex = 0; SectionIndex < LODInfo.Sections.Num(); SectionIndex++)
+		{
+			FStaticMeshSceneProxy::FLODInfo::FSectionInfo& Section = LODInfo.Sections[SectionIndex];
+			if (!Section.Material->CheckMaterialUsage_Concurrent(MATUSAGE_InstancedStaticMeshes))
+			{
+				Section.Material = UMaterial::GetDefaultMaterial(MD_Surface);
+			}
+		}
+	}
+
+	const bool bInstanced = GRHISupportsInstancing;
+
+	// Copy the parameters for LOD - all instances
+	UserData_AllInstances.MeshRenderData = InComponent->GetStaticMesh()->RenderData.Get();
+	UserData_AllInstances.StartCullDistance = InComponent->InstanceStartCullDistance;
+	UserData_AllInstances.EndCullDistance = InComponent->InstanceEndCullDistance;
+	UserData_AllInstances.MinLOD = ClampedMinLOD;
+	UserData_AllInstances.bRenderSelected = true;
+	UserData_AllInstances.bRenderUnselected = true;
+	UserData_AllInstances.RenderData = bInstanced ? nullptr : &InstancedRenderData;
+
+	FVector MinScale(0);
+	FVector MaxScale(0);
+	InComponent->GetInstancesMinMaxScale(MinScale, MaxScale);
+
+	UserData_AllInstances.AverageInstancesScale = MinScale + (MaxScale - MinScale) / 2.0f;
+
+	// selected only
+	UserData_SelectedInstances = UserData_AllInstances;
+	UserData_SelectedInstances.bRenderUnselected = false;
+
+	// unselected only
+	UserData_DeselectedInstances = UserData_AllInstances;
+	UserData_DeselectedInstances.bRenderSelected = false;
+}
+
 void FInstancedStaticMeshSceneProxy::SetupInstancedMeshBatch(int32 LODIndex, int32 BatchIndex, FMeshBatch& OutMeshBatch) const
 {
 	const bool bInstanced = GRHISupportsInstancing;
@@ -2100,6 +2149,21 @@ static bool ComponentRequestsCPUAccess(UInstancedStaticMeshComponent* InComponen
 	return false;
 }
 
+void UInstancedStaticMeshComponent::GetInstancesMinMaxScale(FVector& MinScale, FVector& MaxScale) const
+{
+	MinScale = FVector(MAX_flt);
+	MaxScale = FVector(-MAX_flt);
+
+	for (int32 i = 0; i < PerInstanceSMData.Num(); ++i)
+	{
+		const FInstancedStaticMeshInstanceData& InstanceData = PerInstanceSMData[i];
+		FVector ScaleVector = InstanceData.Transform.GetScaleVector();
+
+		MinScale = MinScale.ComponentMin(ScaleVector);
+		MaxScale = MaxScale.ComponentMax(ScaleVector);
+	}
+}
+
 void UInstancedStaticMeshComponent::InitPerInstanceRenderData(bool InitializeFromCurrentData, FStaticMeshInstanceData* InSharedInstanceBufferData, bool InRequireCPUAccess)
 {
 	if (PerInstanceRenderData.IsValid())
@@ -2484,7 +2548,8 @@ void FInstancedStaticMeshVertexFactoryShaderParameters::GetElementShaderBindings
 				FirstLOD = FMath::Max(FirstLOD, DebugMin);
 			}
 
-			float SphereRadius = InstancingUserData->MeshRenderData->Bounds.SphereRadius;
+			FBoxSphereBounds ScaledBounds = InstancingUserData->MeshRenderData->Bounds.TransformBy(FTransform(FRotator::ZeroRotator, FVector::ZeroVector, InstancingUserData->AverageInstancesScale));
+			float SphereRadius = ScaledBounds.SphereRadius;
 			float MinSize = View->ViewMatrices.IsPerspectiveProjection() ? CVarFoliageMinimumScreenSize.GetValueOnRenderThread() : 0.0f;
 			float LODScale = CVarFoliageLODDistanceScale.GetValueOnRenderThread();
 			float LODRandom = CVarRandomLODRange.GetValueOnRenderThread();
diff --git a/Engine/Source/Runtime/Engine/Private/InstancedStaticMesh.h b/Engine/Source/Runtime/Engine/Private/InstancedStaticMesh.h
--- a/Engine/Source/Runtime/Engine/Private/InstancedStaticMesh.h
+++ b/Engine/Source/Runtime/Engine/Private/InstancedStaticMesh.h
@@ -152,6 +152,7 @@ struct FInstancingUserData
 
 	bool bRenderSelected;
 	bool bRenderUnselected;
+	FVector AverageInstancesScale;
 };
 
 struct FInstancedStaticMeshDataType
@@ -624,48 +625,7 @@ protected:
 
 private:
 
-	void SetupProxy(UInstancedStaticMeshComponent* InComponent)
-	{
-#if WITH_EDITOR
-		if( bHasSelectedInstances )
-		{
-			// if we have selected indices, mark scene proxy as selected.
-			SetSelection_GameThread(true);
-		}
-#endif
-		// Make sure all the materials are okay to be rendered as an instanced mesh.
-		for (int32 LODIndex = 0; LODIndex < LODs.Num(); LODIndex++)
-		{
-			FStaticMeshSceneProxy::FLODInfo& LODInfo = LODs[LODIndex];
-			for (int32 SectionIndex = 0; SectionIndex < LODInfo.Sections.Num(); SectionIndex++)
-			{
-				FStaticMeshSceneProxy::FLODInfo::FSectionInfo& Section = LODInfo.Sections[SectionIndex];
-				if (!Section.Material->CheckMaterialUsage_Concurrent(MATUSAGE_InstancedStaticMeshes))
-				{
-					Section.Material = UMaterial::GetDefaultMaterial(MD_Surface);
-				}
-			}
-		}
-
-		const bool bInstanced = GRHISupportsInstancing;
-
-		// Copy the parameters for LOD - all instances
-		UserData_AllInstances.MeshRenderData = InComponent->GetStaticMesh()->RenderData.Get();
-		UserData_AllInstances.StartCullDistance = InComponent->InstanceStartCullDistance;
-		UserData_AllInstances.EndCullDistance = InComponent->InstanceEndCullDistance;
-		UserData_AllInstances.MinLOD = ClampedMinLOD;
-		UserData_AllInstances.bRenderSelected = true;
-		UserData_AllInstances.bRenderUnselected = true;
-		UserData_AllInstances.RenderData = bInstanced ? nullptr : &InstancedRenderData;
-
-		// selected only
-		UserData_SelectedInstances = UserData_AllInstances;
-		UserData_SelectedInstances.bRenderUnselected = false;
-
-		// unselected only
-		UserData_DeselectedInstances = UserData_AllInstances;
-		UserData_DeselectedInstances.bRenderSelected = false;
-	}
+	void SetupProxy(UInstancedStaticMeshComponent* InComponent);
 };
 
 #if WITH_EDITOR
