Stop ExecuteConsoleCommand logging a warning if given a null world context

diff --git a/Engine/Source/Runtime/Engine/Private/KismetSystemLibrary.cpp b/Engine/Source/Runtime/Engine/Private/KismetSystemLibrary.cpp
--- a/Engine/Source/Runtime/Engine/Private/KismetSystemLibrary.cpp
+++ b/Engine/Source/Runtime/Engine/Private/KismetSystemLibrary.cpp
@@ -301,14 +301,14 @@ void UKismetSystemLibrary::SetWindowTitle(const FText& Title)
 void UKismetSystemLibrary::ExecuteConsoleCommand(UObject* WorldContextObject, const FString& Command, APlayerController* Player)
 {
 	// First, try routing through the primary player
-	APlayerController* TargetPC = Player ? Player : UGameplayStatics::GetPlayerController(WorldContextObject, 0);
+	UWorld* World = GEngine->GetWorldFromContextObject(WorldContextObject, EGetWorldErrorMode::ReturnNull);
+	APlayerController* TargetPC = Player || !World ? Player : World->GetFirstPlayerController();
 	if (TargetPC)
 	{
 		TargetPC->ConsoleCommand(Command, true);
 	}
-	else if (GEngine)
+	else
 	{
-		UWorld* World = GEngine->GetWorldFromContextObject(WorldContextObject, EGetWorldErrorMode::ReturnNull);
 		GEngine->Exec(World, *Command);
 	}
 }
