Prevent audio dropping out during garbage collection.
Also caused audio to lose sync with visuals or other gameplay.

Merged in from another stream.

diff --git a/Engine/Source/Runtime/Engine/Private/AudioDecompress.cpp b/Engine/Source/Runtime/Engine/Private/AudioDecompress.cpp
--- a/Engine/Source/Runtime/Engine/Private/AudioDecompress.cpp
+++ b/Engine/Source/Runtime/Engine/Private/AudioDecompress.cpp
@@ -702,7 +702,7 @@ void FAsyncAudioDecompressWorker::DoWork()
 
 static TAutoConsoleVariable<int32> CVarShouldUseBackgroundPoolFor_FAsyncRealtimeAudioTask(
 	TEXT("AudioThread.UseBackgroundThreadPool"),
-	1,
+	0,
 	TEXT("If true, use the background thread pool for realtime audio decompression."));
 
 bool ShouldUseBackgroundPoolFor_FAsyncRealtimeAudioTask()
