HLOD min draw distance is no longer static. Need to refresh it in instance task

diff --git a/Engine/Source/Runtime/Engine/Private/Streaming/TextureInstanceState.cpp b/Engine/Source/Runtime/Engine/Private/Streaming/TextureInstanceState.cpp
--- a/Engine/Source/Runtime/Engine/Private/Streaming/TextureInstanceState.cpp
+++ b/Engine/Source/Runtime/Engine/Private/Streaming/TextureInstanceState.cpp
@@ -19,7 +19,7 @@ int32 FRenderAssetInstanceState::AddBounds(const UPrimitiveComponent* Component)
 	return AddBounds(Component->Bounds, PackedRelativeBox_Identity, Component, Component->LastRenderTimeOnScreen, Component->Bounds.Origin, 0, 0, FLT_MAX);
 }
 
-int32 FRenderAssetInstanceState::AddBounds(const FBoxSphereBounds& Bounds, uint32 PackedRelativeBox, const UPrimitiveComponent* InComponent, float LastRenderTime, const FVector4& RangeOrigin, float MinDistance, float MinRange, float MaxRange)
+int32 FRenderAssetInstanceState::AddBounds(const FBoxSphereBounds& Bounds, uint32 PackedRelativeBox, const UPrimitiveComponent* InComponent, float LastRenderTime, const FVector4& RangeOrigin, float MinDistanceSq, float MinRangeSq, float MaxRangeSq)
 {
 	check(InComponent);
 
@@ -46,7 +46,7 @@ int32 FRenderAssetInstanceState::AddBounds(const FBoxSphereBounds& Bounds, uint3
 		FreeBoundIndices.Push(BoundsIndex + 1);
 	}
 
-	Bounds4[BoundsIndex / 4].Set(BoundsIndex % 4, Bounds, PackedRelativeBox, LastRenderTime, RangeOrigin, MinDistance, MinRange, MaxRange);
+	Bounds4[BoundsIndex / 4].Set(BoundsIndex % 4, Bounds, PackedRelativeBox, LastRenderTime, RangeOrigin, MinDistanceSq, MinRangeSq, MaxRangeSq);
 	Bounds4Components[BoundsIndex] = InComponent;
 
 	return BoundsIndex;
@@ -355,7 +355,7 @@ EAddComponentResult FRenderAssetInstanceState::AddComponent(const UPrimitiveComp
 		});
 
 		int32* ComponentLink = ComponentMap.Find(Component);
-		float MinDistance = 0, MinRange = 0, MaxRange = FLT_MAX;
+		float MinDistanceSq = 0, MinRangeSq = 0, MaxRangeSq = FLT_MAX;
 
 		// Loop for each bound.
 		for (int32 InfoIndex = 0; InfoIndex < RenderAssetInstanceInfos.Num();)
@@ -368,8 +368,8 @@ EAddComponentResult FRenderAssetInstanceState::AddComponent(const UPrimitiveComp
 				++NumOfBoundReferences;
 			}
 
-			GetDistanceAndRange(Component, Info.Bounds, MinDistance, MinRange, MaxRange);
-			const int32 BoundsIndex = AddBounds(Info.Bounds, PackedRelativeBox_Identity, Component, Component->LastRenderTimeOnScreen, Component->Bounds.Origin, MinDistance, MinRange, MaxRange);
+			GetDistanceAndRange(Component, Info.Bounds, MinDistanceSq, MinRangeSq, MaxRangeSq);
+			const int32 BoundsIndex = AddBounds(Info.Bounds, PackedRelativeBox_Identity, Component, Component->LastRenderTimeOnScreen, Component->Bounds.Origin, MinDistanceSq, MinRangeSq, MaxRangeSq);
 			AddRenderAssetElements(Component, TArrayView<FStreamingRenderAssetPrimitiveInfo>(RenderAssetInstanceInfos.GetData() + InfoIndex, NumOfBoundReferences), BoundsIndex, ComponentLink);
 
 			InfoIndex += NumOfBoundReferences;
@@ -539,12 +539,26 @@ bool FRenderAssetInstanceState::ConditionalUpdateBounds(int32 BoundIndex)
 }
 
 
-void FRenderAssetInstanceState::UpdateLastRenderTime(int32 BoundIndex)
+void FRenderAssetInstanceState::UpdateLastRenderTimeAndMaxDrawDistance(int32 BoundIndex)
 {
 	const UPrimitiveComponent* Component = ensure(Bounds4Components.IsValidIndex(BoundIndex)) ? Bounds4Components[BoundIndex] : nullptr;
 	if (Component)
 	{
-		Bounds4[BoundIndex / 4].UpdateLastRenderTime(BoundIndex % 4, Component->LastRenderTimeOnScreen);
+		const int32 Bounds4Idx = BoundIndex / 4;
+		const int32 SubIdx = BoundIndex & 3;
+		Bounds4[Bounds4Idx].UpdateLastRenderTime(SubIdx, Component->LastRenderTimeOnScreen);
+		// The min draw distances of HLODs can change dynamically (see Tick, PauseDitherTransition,
+		// and StartDitherTransition methods of ALODActor)
+		const UPrimitiveComponent* LODParent = Component->GetLODParentPrimitive();
+		if (LODParent)
+		{
+			const float MaxRangeSq = FRenderAssetInstanceView::GetMaxDrawDistSqWithLODParent(
+				Component->Bounds.Origin,
+				LODParent->Bounds.Origin,
+				LODParent->MinDrawDistance,
+				LODParent->Bounds.SphereRadius);
+			Bounds4[Bounds4Idx].UpdateMaxDrawDistanceSquared(SubIdx, MaxRangeSq);
+		}
 	}
 }
 
diff --git a/Engine/Source/Runtime/Engine/Private/Streaming/TextureInstanceState.h b/Engine/Source/Runtime/Engine/Private/Streaming/TextureInstanceState.h
--- a/Engine/Source/Runtime/Engine/Private/Streaming/TextureInstanceState.h
+++ b/Engine/Source/Runtime/Engine/Private/Streaming/TextureInstanceState.h
@@ -44,7 +44,7 @@ public:
 	void UpdateBounds(const UPrimitiveComponent* Component);
 	bool UpdateBounds(int32 BoundIndex);
 	bool ConditionalUpdateBounds(int32 BoundIndex);
-	void UpdateLastRenderTime(int32 BoundIndex);
+	void UpdateLastRenderTimeAndMaxDrawDistance(int32 BoundIndex);
 
 	uint32 GetAllocatedSize() const;
 
@@ -66,7 +66,7 @@ private:
 	// Returns the next elements using the same component.
 	void RemoveElement(int32 ElementIndex, int32& NextComponentLink, int32& BoundsIndex, const UStreamableRenderAsset*& Asset);
 
-	int32 AddBounds(const FBoxSphereBounds& Bounds, uint32 PackedRelativeBox, const UPrimitiveComponent* Component, float LastRenderTime, const FVector4& RangeOrigin, float MinDistance, float MinRange, float MaxRange);
+	int32 AddBounds(const FBoxSphereBounds& Bounds, uint32 PackedRelativeBox, const UPrimitiveComponent* Component, float LastRenderTime, const FVector4& RangeOrigin, float MinDistanceSq, float MinRangeSq, float MaxRangeSq);
 	FORCEINLINE int32 AddBounds(const UPrimitiveComponent* Component);
 	void RemoveBounds(int32 Index);
 
diff --git a/Engine/Source/Runtime/Engine/Private/Streaming/TextureInstanceTask.cpp b/Engine/Source/Runtime/Engine/Private/Streaming/TextureInstanceTask.cpp
--- a/Engine/Source/Runtime/Engine/Private/Streaming/TextureInstanceTask.cpp
+++ b/Engine/Source/Runtime/Engine/Private/Streaming/TextureInstanceTask.cpp
@@ -29,7 +29,7 @@ void FRefreshVisibility::operator()(bool bAsync)
 {
 	for (int32 Index = BeginIndex; Index < EndIndex; ++Index)
 	{
-		State->UpdateLastRenderTime(Index);
+		State->UpdateLastRenderTimeAndMaxDrawDistance(Index);
 	}
 }
 
diff --git a/Engine/Source/Runtime/Engine/Private/Streaming/TextureInstanceView.cpp b/Engine/Source/Runtime/Engine/Private/Streaming/TextureInstanceView.cpp
--- a/Engine/Source/Runtime/Engine/Private/Streaming/TextureInstanceView.cpp
+++ b/Engine/Source/Runtime/Engine/Private/Streaming/TextureInstanceView.cpp
@@ -11,7 +11,7 @@
 #include "Components/PrimitiveComponent.h"
 #include "ContentStreaming.h"
 
-void FRenderAssetInstanceView::FBounds4::Set(int32 Index, const FBoxSphereBounds& Bounds, uint32 InPackedRelativeBox, float InLastRenderTime, const FVector& RangeOrigin, float InMinDistance, float InMinRange, float InMaxRange)
+void FRenderAssetInstanceView::FBounds4::Set(int32 Index, const FBoxSphereBounds& Bounds, uint32 InPackedRelativeBox, float InLastRenderTime, const FVector& RangeOrigin, float InMinDistanceSq, float InMinRangeSq, float InMaxRangeSq)
 {
 	check(Index >= 0 && Index < 4);
 
@@ -26,9 +26,9 @@ void FRenderAssetInstanceView::FBounds4::Set(int32 Index, const FBoxSphereBounds
 	ExtentZ.Component(Index) = Bounds.BoxExtent.Z;
 	Radius.Component(Index) = Bounds.SphereRadius;
 	PackedRelativeBox[Index] = InPackedRelativeBox;
-	MinDistanceSq.Component(Index) = InMinDistance * InMinDistance;
-	MinRangeSq.Component(Index) = InMinRange * InMinRange;
-	MaxRangeSq.Component(Index) = InMaxRange != FLT_MAX ? (InMaxRange * InMaxRange) : FLT_MAX;
+	MinDistanceSq.Component(Index) = InMinDistanceSq;
+	MinRangeSq.Component(Index) = InMinRangeSq;
+	MaxRangeSq.Component(Index) = InMaxRangeSq;
 	LastRenderTime.Component(Index) = InLastRenderTime;
 }
 
@@ -43,8 +43,8 @@ void FRenderAssetInstanceView::FBounds4::UnpackBounds(int32 Index, const UPrimit
 		UnpackRelativeBox(Component->Bounds, PackedRelativeBox[Index], SubBounds);
 
 		// Update the visibility range once we have the bounds.
-		float MinDistance = 0, MinRange = 0, MaxRange = FLT_MAX;
-		FRenderAssetInstanceView::GetDistanceAndRange(Component, SubBounds, MinDistance, MinRange, MaxRange);
+		float MinDistance2 = 0, MinRange2 = 0, MaxRange2 = FLT_MAX;
+		FRenderAssetInstanceView::GetDistanceAndRange(Component, SubBounds, MinDistance2, MinRange2, MaxRange2);
 
 		OriginX.Component(Index) = SubBounds.Origin.X;
 		OriginY.Component(Index) = SubBounds.Origin.Y;
@@ -57,9 +57,9 @@ void FRenderAssetInstanceView::FBounds4::UnpackBounds(int32 Index, const UPrimit
 		ExtentZ.Component(Index) = SubBounds.BoxExtent.Z;
 		Radius.Component(Index) = SubBounds.SphereRadius;
 		PackedRelativeBox[Index] = PackedRelativeBox_Identity;
-		MinDistanceSq.Component(Index) = MinDistance * MinDistance;
-		MinRangeSq.Component(Index) = MinRange * MinRange;
-		MaxRangeSq.Component(Index) = MaxRange != FLT_MAX ? (MaxRange * MaxRange) : FLT_MAX;
+		MinDistanceSq.Component(Index) = MinDistance2;
+		MinRangeSq.Component(Index) = MinRange2;
+		MaxRangeSq.Component(Index) = MaxRange2;
 	}
 }
 
@@ -223,11 +223,17 @@ TRefCountPtr<FRenderAssetInstanceView> FRenderAssetInstanceView::CreateViewWithU
 	return NewView;
 }
 
-void FRenderAssetInstanceView::GetDistanceAndRange(const UPrimitiveComponent* Component, const FBoxSphereBounds& RenderAssetInstanceBounds, float& MinDistance, float& MinRange, float& MaxRange)
+float FRenderAssetInstanceView::GetMaxDrawDistSqWithLODParent(const FVector& Origin, const FVector& ParentOrigin, float ParentMinDrawDist, float ParentBoundingSphereRadius)
+{
+	const float Result = ParentMinDrawDist + ParentBoundingSphereRadius + (Origin - ParentOrigin).Size();
+	return Result * Result;
+}
+
+void FRenderAssetInstanceView::GetDistanceAndRange(const UPrimitiveComponent* Component, const FBoxSphereBounds& RenderAssetInstanceBounds, float& MinDistanceSq, float& MinRangeSq, float& MaxRangeSq)
 {
 	check(Component && Component->IsRegistered());
 
-	// In the engine, the MinDistance is computed from the component bound center to the viewpoint.
+	// In the engine, the MinDistance is computed from the component bound center to the viewpoint (except for HLODs).
 	// The streaming computes the distance as the distance from viewpoint to the edge of the texture bound box.
 	// The implementation also handles MinDistance by bounding the distance to it so that if the viewpoint gets closer the screen size will stop increasing at some point.
 	// The fact that the primitive will disappear is not so relevant as this will be handled by the visibility logic, normally streaming one less mip than requested.
@@ -235,10 +241,12 @@ void FRenderAssetInstanceView::GetDistanceAndRange(const UPrimitiveComponent* Co
 
 	const UPrimitiveComponent* LODParent = Component->GetLODParentPrimitive();
 
-	MinDistance = FMath::Max<float>(0, Component->MinDrawDistance - (RenderAssetInstanceBounds.Origin - Component->Bounds.Origin).Size() - RenderAssetInstanceBounds.SphereRadius);
-	MinRange = FMath::Max<float>(0, Component->MinDrawDistance);
+	MinDistanceSq = FMath::Max<float>(0, Component->MinDrawDistance - (RenderAssetInstanceBounds.Origin - Component->Bounds.Origin).Size() - RenderAssetInstanceBounds.SphereRadius);
+	MinDistanceSq *= MinDistanceSq;
+	MinRangeSq = FMath::Max<float>(0, Component->MinDrawDistance);
+	MinRangeSq *= MinRangeSq;
 	// Max distance when HLOD becomes visible.
-	MaxRange = LODParent ? (LODParent->MinDrawDistance + (Component->Bounds.Origin - LODParent->Bounds.Origin).Size()) : FLT_MAX;
+	MaxRangeSq = LODParent ? GetMaxDrawDistSqWithLODParent(Component->Bounds.Origin, LODParent->Bounds.Origin, LODParent->MinDrawDistance, LODParent->Bounds.SphereRadius) : FLT_MAX;
 }
 
 void FRenderAssetInstanceView::SwapData(FRenderAssetInstanceView* Lfs, FRenderAssetInstanceView* Rhs)
diff --git a/Engine/Source/Runtime/Engine/Private/Streaming/TextureInstanceView.h b/Engine/Source/Runtime/Engine/Private/Streaming/TextureInstanceView.h
--- a/Engine/Source/Runtime/Engine/Private/Streaming/TextureInstanceView.h
+++ b/Engine/Source/Runtime/Engine/Private/Streaming/TextureInstanceView.h
@@ -69,10 +69,11 @@ public:
 		FVector4 LastRenderTime; //(FApp::GetCurrentTime() - Component->LastRenderTime);
 
 
-		void Set(int32 Index, const FBoxSphereBounds& Bounds, uint32 InPackedRelativeBox, float LastRenderTime, const FVector& RangeOrigin, float MinDistance, float MinRange, float MaxRange);
+		void Set(int32 Index, const FBoxSphereBounds& Bounds, uint32 InPackedRelativeBox, float LastRenderTime, const FVector& RangeOrigin, float MinDistanceSq, float MinRangeSq, float MaxRangeSq);
 		void UnpackBounds(int32 Index, const UPrimitiveComponent* Component);
 		void FullUpdate(int32 Index, const FBoxSphereBounds& Bounds, float LastRenderTime);
 		FORCEINLINE void UpdateLastRenderTime(int32 Index, float LastRenderTime);
+		FORCEINLINE void UpdateMaxDrawDistanceSquared(int32 Index, float InMaxRangeSq);
 
 		// Clears entry between 0 and 4
 		FORCEINLINE void Clear(int32 Index);
@@ -198,7 +199,9 @@ public:
 
 	float GetMaxTexelFactor() const { return MaxTexelFactor; }
 
-	static void GetDistanceAndRange(const UPrimitiveComponent* Component, const FBoxSphereBounds& RenderAssetInstanceBounds, float& MinDistance, float& MinRange, float& MaxRange);
+	static float GetMaxDrawDistSqWithLODParent(const FVector& Origin, const FVector& ParentOrigin, float ParentMinDrawDist, float ParentBoundingSphereRadius);
+
+	static void GetDistanceAndRange(const UPrimitiveComponent* Component, const FBoxSphereBounds& RenderAssetInstanceBounds, float& MinDistanceSq, float& MinRangeSq, float& MaxRangeSq);
 
 protected:
 
diff --git a/Engine/Source/Runtime/Engine/Private/Streaming/TextureInstanceView.inl b/Engine/Source/Runtime/Engine/Private/Streaming/TextureInstanceView.inl
--- a/Engine/Source/Runtime/Engine/Private/Streaming/TextureInstanceView.inl
+++ b/Engine/Source/Runtime/Engine/Private/Streaming/TextureInstanceView.inl
@@ -62,6 +62,12 @@ void FRenderAssetInstanceView::FBounds4::UpdateLastRenderTime(int32 Index, float
 	LastRenderTime.Component(Index) = InLastRenderTime;
 }
 
+void FRenderAssetInstanceView::FBounds4::UpdateMaxDrawDistanceSquared(int32 Index, float InMaxRangeSq)
+{
+	check(Index >= 0 && Index < 4);
+	MaxRangeSq.Component(Index) = InMaxRangeSq;
+}
+
 FRenderAssetInstanceView::FElement::FElement()
 :	Component(nullptr)
 ,	RenderAsset(nullptr)
