Fixup CallRepNotifies to properly handle bitfield properties.

diff --git a/Engine/Source/Runtime/Engine/Private/RepLayout.cpp b/Engine/Source/Runtime/Engine/Private/RepLayout.cpp
--- a/Engine/Source/Runtime/Engine/Private/RepLayout.cpp
+++ b/Engine/Source/Runtime/Engine/Private/RepLayout.cpp
@@ -3076,6 +3076,7 @@ void FRepLayout::CallRepNotifies(FReceivingRepState* RepState, UObject* Object)
 	check(LayoutState == ERepLayoutState::Normal);
 
 	FRepShadowDataBuffer ShadowData(RepState->StaticBuffer.GetData());
+	FRepObjectDataBuffer ObjectData(Object);
 
 	for (UProperty* RepProperty : RepState->RepNotifies)
 	{
@@ -3105,13 +3106,24 @@ void FRepLayout::CallRepNotifies(FReceivingRepState* RepState, UObject* Object)
 			});
 
 			check(Parent);
+			
+			FRepShadowDataBuffer PropertyData = ShadowData + (*Parent);
 
-			Object->ProcessEvent(RepNotifyFunc, ShadowData + (*Parent));
+			// This could be cached off as a Parent flag, to avoid touching the Commands array.
+			if (ERepLayoutCmdType::PropertyBool == Cmds[Parent->CmdStart].Type)
+			{
+				bool BoolPropertyValue = !!static_cast<const UBoolProperty*>(Parent->Property)->GetPropertyValue(PropertyData);
+				Object->ProcessEvent(RepNotifyFunc, &BoolPropertyValue);
+			}
+			else
+			{
+				Object->ProcessEvent(RepNotifyFunc, PropertyData);
+			}
 			
 			// now store the complete value in the shadow buffer
 			if (!EnumHasAnyFlags(Parent->Flags, ERepParentFlags::IsNetSerialize | ERepParentFlags::IsCustomDelta))
 			{
-				RepProperty->CopyCompleteValue(ShadowData + (*Parent), RepProperty->ContainerPtrToValuePtr<uint8>(Object));
+				RepProperty->CopyCompleteValue(ShadowData + (*Parent), ObjectData + (*Parent));
 			}
 		}
 	}
