Added a function to the content encryption config which moves all the assets out of a given group and puts them into the unnamed "loose group".

diff --git a/Engine/Source/Runtime/Engine/Classes/Engine/ContentEncryptionConfig.h b/Engine/Source/Runtime/Engine/Classes/Engine/ContentEncryptionConfig.h
--- a/Engine/Source/Runtime/Engine/Classes/Engine/ContentEncryptionConfig.h
+++ b/Engine/Source/Runtime/Engine/Classes/Engine/ContentEncryptionConfig.h
@@ -40,6 +40,18 @@ public:
 		return ReleasedKeys;
 	}
 
+	void DissolveGroups(const TSet<FName>& InGroupsToDissolve)
+	{
+		for (FName GroupName : InGroupsToDissolve)
+		{
+			if (PackageGroups.Contains(GroupName))
+			{
+				PackageGroups.FindOrAdd(NAME_None).PackageNames.Append(PackageGroups.Find(GroupName)->PackageNames);
+				PackageGroups.Remove(GroupName);
+			}
+		}
+	}
+
 private:
 
 	TGroupMap PackageGroups;
