Fix unessesary/incorrect usage of GRHICommandList inside a lambda.

The immediate commandlist is passed down as part of the lambda and so this should be used instead.

diff --git a/Engine/Source/Runtime/Engine/Private/UnrealEngine.cpp b/Engine/Source/Runtime/Engine/Private/UnrealEngine.cpp
--- a/Engine/Source/Runtime/Engine/Private/UnrealEngine.cpp
+++ b/Engine/Source/Runtime/Engine/Private/UnrealEngine.cpp
@@ -10192,10 +10192,10 @@ void FFrameEndSync::Sync( bool bAllowOneFrameThreadLag )
 #if !UE_BUILD_SHIPPING && PLATFORM_SUPPORTS_FLIP_TRACKING
 	// Set the FrameDebugInfo on platforms that have accurate frame tracking.
 	ENQUEUE_RENDER_COMMAND(FrameDebugInfo)(
-		[CurrentFrameCounter = GFrameCounter, CurrentInputTime = GInputTime](FRHICommandList&)
+		[CurrentFrameCounter = GFrameCounter, CurrentInputTime = GInputTime](FRHICommandListImmediate& RHICmdList)
 	{
-		GRHICommandList.GetImmediateCommandList().EnqueueLambda(
-			[CurrentFrameCounter, CurrentInputTime](FRHICommandList&)
+		RHICmdList.EnqueueLambda(
+			[CurrentFrameCounter, CurrentInputTime](FRHICommandListImmediate&)
 		{
 			// Set the FrameCount and InputTime for input latency stats and flip debugging.
 			RHISetFrameDebugInfo(GRHIPresentCounter - 1, CurrentFrameCounter - 1, CurrentInputTime);
diff --git a/Engine/Source/Runtime/RenderCore/Private/RenderingThread.cpp b/Engine/Source/Runtime/RenderCore/Private/RenderingThread.cpp
--- a/Engine/Source/Runtime/RenderCore/Private/RenderingThread.cpp
+++ b/Engine/Source/Runtime/RenderCore/Private/RenderingThread.cpp
@@ -1221,9 +1221,9 @@ void FlushRenderingCommands(bool bFlushDeferredDeletes)
 	}
 
 	ENQUEUE_RENDER_COMMAND(FlushPendingDeleteRHIResourcesCmd)(
-		[bFlushDeferredDeletes](FRHICommandList&)
+		[bFlushDeferredDeletes](FRHICommandListImmediate& RHICmdList)
 	{
-		GRHICommandList.GetImmediateCommandList().ImmediateFlush(
+		RHICmdList.ImmediateFlush(
 			bFlushDeferredDeletes ?
 			EImmediateFlushType::FlushRHIThreadFlushResourcesFlushDeferredDeletes :
 			EImmediateFlushType::FlushRHIThreadFlushResources);
