Fix assumption that editor and ShaderCompileWorker executables are in the same directory. It's possible to build an editor using a unique build environment so that it can be built with custom macros defined, etc...which puts the executable in (ProjectDir)/Binaries/Win64.

Now stores the full path to the shader source directories in the SCW input data archive, and normalizes it again on read.

diff --git a/Engine/Source/Programs/ShaderCompileWorker/Private/ShaderCompileWorker.cpp b/Engine/Source/Programs/ShaderCompileWorker/Private/ShaderCompileWorker.cpp
--- a/Engine/Source/Programs/ShaderCompileWorker/Private/ShaderCompileWorker.cpp
+++ b/Engine/Source/Programs/ShaderCompileWorker/Private/ShaderCompileWorker.cpp
@@ -367,8 +367,9 @@ private:
 			InputFile << DirectoryMappings;
 
 			ResetAllShaderSourceDirectoryMappings();
-			for (const auto& MappingEntry : DirectoryMappings)
+			for (TPair<FString, FString>& MappingEntry : DirectoryMappings)
 			{
+				FPaths::NormalizeDirectoryName(MappingEntry.Value);
 				AddShaderSourceDirectoryMapping(MappingEntry.Key, MappingEntry.Value);
 			}
 		}
diff --git a/Engine/Source/Runtime/Engine/Private/ShaderCompiler/ShaderCompiler.cpp b/Engine/Source/Runtime/Engine/Private/ShaderCompiler/ShaderCompiler.cpp
--- a/Engine/Source/Runtime/Engine/Private/ShaderCompiler/ShaderCompiler.cpp
+++ b/Engine/Source/Runtime/Engine/Private/ShaderCompiler/ShaderCompiler.cpp
@@ -394,7 +394,12 @@ bool FShaderCompileUtilities::DoWriteTasks(const TArray<FShaderCommonCompileJob*
 
 	TransferFile << FormatVersionMap;
 
+	// Convert all the source directory paths to absolute, since SCW might be in a different directory to the editor executable
 	TMap<FString, FString> ShaderSourceDirectoryMappings = AllShaderSourceDirectoryMappings();
+	for(TPair<FString, FString>& Pair : ShaderSourceDirectoryMappings)
+	{
+		Pair.Value = FPaths::ConvertRelativePathToFull(Pair.Value);
+	}
 	TransferFile << ShaderSourceDirectoryMappings;
 
 	TArray<FShaderCompileJob*> QueuedSingleJobs;
@@ -1442,11 +1447,11 @@ FShaderCompilingManager::FShaderCompilingManager() :
 	NumOutstandingJobs(0),
 	NumExternalJobs(0),
 #if PLATFORM_MAC
-	ShaderCompileWorkerName(TEXT("../../../Engine/Binaries/Mac/ShaderCompileWorker")),
+	ShaderCompileWorkerName(FPaths::EngineDir() / TEXT("Binaries/Mac/ShaderCompileWorker")),
 #elif PLATFORM_LINUX
-	ShaderCompileWorkerName(TEXT("../../../Engine/Binaries/Linux/ShaderCompileWorker")),
+	ShaderCompileWorkerName(FPaths::EngineDir() / TEXT("Binaries/Linux/ShaderCompileWorker")),
 #else
-	ShaderCompileWorkerName(TEXT("../../../Engine/Binaries/Win64/ShaderCompileWorker.exe")),
+	ShaderCompileWorkerName(FPaths::EngineDir() / TEXT("Binaries/Win64/ShaderCompileWorker.exe")),
 #endif
 	SuppressedShaderPlatforms(0)
 {
