Fix for crash in ProcessPendingAsyncFileResults

diff --git a/Engine/Source/Runtime/Engine/Private/AudioStreaming.cpp b/Engine/Source/Runtime/Engine/Private/AudioStreaming.cpp
--- a/Engine/Source/Runtime/Engine/Private/AudioStreaming.cpp
+++ b/Engine/Source/Runtime/Engine/Private/AudioStreaming.cpp
@@ -507,8 +507,12 @@ void FAudioStreamingManager::ProcessPendingAsyncFileResults()
 		const int32 LoadedAudioChunkIndex = AudioChunkLoadResult->LoadedAudioChunkIndex;
 
 		check(AudioChunkLoadResult->StreamingWaveData != nullptr);
-		check(LoadedAudioChunkIndex != INDEX_NONE);
-		check(LoadedAudioChunkIndex < AudioChunkLoadResult->StreamingWaveData->LoadedChunks.Num());
+
+		if (LoadedAudioChunkIndex == INDEX_NONE || LoadedAudioChunkIndex >= AudioChunkLoadResult->StreamingWaveData->LoadedChunks.Num())
+		{
+			UE_LOG(LogAudio, Warning, TEXT("Loaded streaming chunk index %d is invalid. Current number of loaded chunks is: %d"), LoadedAudioChunkIndex, AudioChunkLoadResult->StreamingWaveData->LoadedChunks.Num());
+			continue;
+		}
 
 		FLoadedAudioChunk* ChunkStorage = &AudioChunkLoadResult->StreamingWaveData->LoadedChunks[LoadedAudioChunkIndex];
 
