Adding missing delete.

diff --git a/Engine/Source/Runtime/Engine/Private/AudioStreaming.cpp b/Engine/Source/Runtime/Engine/Private/AudioStreaming.cpp
--- a/Engine/Source/Runtime/Engine/Private/AudioStreaming.cpp
+++ b/Engine/Source/Runtime/Engine/Private/AudioStreaming.cpp
@@ -511,6 +511,9 @@ void FAudioStreamingManager::ProcessPendingAsyncFileResults()
 		if (LoadedAudioChunkIndex == INDEX_NONE || LoadedAudioChunkIndex >= AudioChunkLoadResult->StreamingWaveData->LoadedChunks.Num())
 		{
 			UE_LOG(LogAudio, Warning, TEXT("Loaded streaming chunk index %d is invalid. Current number of loaded chunks is: %d"), LoadedAudioChunkIndex, AudioChunkLoadResult->StreamingWaveData->LoadedChunks.Num());
+
+			delete AudioChunkLoadResult;
+			AudioChunkLoadResult = nullptr;
 			continue;
 		}
 
