Added missing ENGINE_API declaration

diff --git a/Engine/Source/Runtime/Engine/Classes/Engine/NetConnection.h b/Engine/Source/Runtime/Engine/Classes/Engine/NetConnection.h
--- a/Engine/Source/Runtime/Engine/Classes/Engine/NetConnection.h
+++ b/Engine/Source/Runtime/Engine/Classes/Engine/NetConnection.h
@@ -1064,7 +1064,7 @@ public:
 
 protected:
 
-	void SetPendingCloseDueToSocketSendFailure();
+	ENGINE_API void SetPendingCloseDueToSocketSendFailure();
 
 	void CleanupDormantActorState();
 
