REL - Added more debug information for missing shadertypes.

diff --git a/Engine/Source/Runtime/Engine/Private/Materials/MaterialShared.cpp b/Engine/Source/Runtime/Engine/Private/Materials/MaterialShared.cpp
--- a/Engine/Source/Runtime/Engine/Private/Materials/MaterialShared.cpp
+++ b/Engine/Source/Runtime/Engine/Private/Materials/MaterialShared.cpp
@@ -1883,7 +1883,7 @@ FShader* FMaterial::GetShader(FMeshMaterialShaderType* ShaderType, FVertexFactor
 
 		// Assert with detailed information if the shader wasn't found for rendering.  
 		// This is usually the result of an incorrect ShouldCache function.
-		UE_LOG(LogMaterial, Fatal,
+		UE_LOG(LogMaterial, Error,
 			TEXT("Couldn't find Shader %s for Material Resource %s!\n")
 			TEXT("		RenderMeshShaderMap %d, RenderThreadShaderMap %d\n")
 			TEXT("		GameMeshShaderMap %d, GameThreadShaderMap %d, bShaderWasFoundInGameShaderMap %d\n")
@@ -1897,6 +1897,20 @@ FShader* FMaterial::GetShader(FMeshMaterialShaderType* ShaderType, FVertexFactor
 			bMaterialShouldCache, bVFShouldCache, bShaderShouldCache,
 			*MaterialUsage
 			);
+
+		if (MeshShaderMap)
+		{
+			TMap<FShaderId, FShader*> List;
+			MeshShaderMap->GetShaderList(List);
+
+			for (const auto& ShaderPair : List)
+			{
+				FString TypeName = ShaderPair.Value->GetType()->GetName();
+				UE_LOG(LogMaterial, Error, TEXT("ShaderType found in MaterialMap: %s"), *TypeName);
+			}
+		}
+
+		UE_LOG(LogMaterial, Fatal, TEXT("Fatal Error Material not found"));
 	}
 
 	return Shader;
