Fixed default filtering mode for volume texture to be wrap instead of clamp.

diff --git a/Engine/Source/Runtime/Engine/Private/VolumeTexture.cpp b/Engine/Source/Runtime/Engine/Private/VolumeTexture.cpp
--- a/Engine/Source/Runtime/Engine/Private/VolumeTexture.cpp
+++ b/Engine/Source/Runtime/Engine/Private/VolumeTexture.cpp
@@ -466,9 +466,9 @@ public:
 		FSamplerStateInitializerRHI SamplerStateInitializer
 		(
 			SamplerFilter,
-			AM_Clamp,
-			AM_Clamp,
-			AM_Clamp
+			AM_Wrap,
+			AM_Wrap,
+			AM_Wrap
 		);
 		SamplerStateRHI = RHICreateSamplerState(SamplerStateInitializer);
 	}
