Allow reloading asset redirectors at runtime when they are updated via hotfix.

diff --git a/Engine/Source/Runtime/Engine/Classes/Engine/AssetManagerSettings.h b/Engine/Source/Runtime/Engine/Classes/Engine/AssetManagerSettings.h
--- a/Engine/Source/Runtime/Engine/Classes/Engine/AssetManagerSettings.h
+++ b/Engine/Source/Runtime/Engine/Classes/Engine/AssetManagerSettings.h
@@ -125,6 +125,8 @@ public:
 	UPROPERTY(config, EditAnywhere, AdvancedDisplay, Category = "Asset Registry", DisplayName = "Metadata Tags For Asset Registry")
 	TSet<FName> MetaDataTagsForAssetRegistry;
 
+	virtual void PostReloadConfig(class UProperty* PropertyThatWasLoaded) override;
+
 #if WITH_EDITOR
 	virtual void PostInitProperties() override;
 	virtual void PreEditChange(UProperty* PropertyAboutToChange) override;
diff --git a/Engine/Source/Runtime/Engine/Private/AssetManagerTypes.cpp b/Engine/Source/Runtime/Engine/Private/AssetManagerTypes.cpp
--- a/Engine/Source/Runtime/Engine/Private/AssetManagerTypes.cpp
+++ b/Engine/Source/Runtime/Engine/Private/AssetManagerTypes.cpp
@@ -107,6 +107,13 @@ void FPrimaryAssetRules::PropagateCookRules(const FPrimaryAssetRules& ParentRule
 	}
 }
 
+void UAssetManagerSettings::PostReloadConfig(UProperty* PropertyThatWasLoaded)
+{
+	Super::PostReloadConfig(PropertyThatWasLoaded);
+
+	UAssetManager::Get().LoadRedirectorMaps();
+}
+
 #if WITH_EDITOR
 void UAssetManagerSettings::PostInitProperties()
 {
