Fixed crash converting a sequencer posessed actor to spawnable when it uses the rigid body anim node

diff --git a/Engine/Source/Runtime/Engine/Private/Components/MovementComponent.cpp b/Engine/Source/Runtime/Engine/Private/Components/MovementComponent.cpp
--- a/Engine/Source/Runtime/Engine/Private/Components/MovementComponent.cpp
+++ b/Engine/Source/Runtime/Engine/Private/Components/MovementComponent.cpp
@@ -366,7 +366,8 @@ bool UMovementComponent::ShouldSkipUpdate(float DeltaTime) const
 
 float UMovementComponent::GetGravityZ() const
 {
-	return GetPhysicsVolume()->GetGravityZ();
+	APhysicsVolume* PhysicsVolume = GetPhysicsVolume();
+	return PhysicsVolume ? PhysicsVolume->GetGravityZ() : UPhysicsSettings::Get()->DefaultGravityZ;
 }
 
 void UMovementComponent::HandleImpact(const FHitResult& Hit, float TimeSlice, const FVector& MoveDelta)
