- Added a flag to disable out-of-world checks per component.

#tests PIE

diff --git a/Engine/Source/Runtime/Engine/Classes/Components/InterpToMovementComponent.h b/Engine/Source/Runtime/Engine/Classes/Components/InterpToMovementComponent.h
--- a/Engine/Source/Runtime/Engine/Classes/Components/InterpToMovementComponent.h
+++ b/Engine/Source/Runtime/Engine/Classes/Components/InterpToMovementComponent.h
@@ -80,6 +80,10 @@ class ENGINE_API UInterpToMovementComponent : public UMovementComponent
 	UPROPERTY(EditAnywhere, BlueprintReadWrite, Category=Behaviour)
 	EInterpToBehaviourType BehaviourType;
 
+	/* Do we want this comp to perform CheckStillInWorld checks? */
+	UPROPERTY(EditAnywhere, BlueprintReadWrite, Category = Behaviour)
+	bool bCheckIfStillInWorld = true;
+
 	/**
 	 * If true, forces sub-stepping to break up movement into discrete smaller steps to improve accuracy of the trajectory.
 	 * Objects that move in a straight line typically do *not* need to set this, as movement always uses continuous collision detection (sweeps) so collision is not missed.
diff --git a/Engine/Source/Runtime/Engine/Private/Components/InterpToMovementComponent.cpp b/Engine/Source/Runtime/Engine/Private/Components/InterpToMovementComponent.cpp
--- a/Engine/Source/Runtime/Engine/Private/Components/InterpToMovementComponent.cpp
+++ b/Engine/Source/Runtime/Engine/Private/Components/InterpToMovementComponent.cpp
@@ -380,6 +380,10 @@ bool UInterpToMovementComponent::CheckStillInWorld()
 	{
 		return false;
 	}
+	if (!bCheckIfStillInWorld)
+	{
+		return true;
+	}
 	if (ActorOwner->GetActorLocation().Z < WorldSettings->KillZ)
 	{
 		UDamageType const* DmgType = WorldSettings->KillZDamageType ? WorldSettings->KillZDamageType->GetDefaultObject<UDamageType>() : GetDefault<UDamageType>();
