Re-enable audio background pool when not using Audio Mixer.
Procedural sounds no longer run in background even when cvar to use background is enabled (includes MediaSoundComponent).

Merged.

diff --git a/Engine/Source/Runtime/Engine/Private/AudioDecompress.cpp b/Engine/Source/Runtime/Engine/Private/AudioDecompress.cpp
--- a/Engine/Source/Runtime/Engine/Private/AudioDecompress.cpp
+++ b/Engine/Source/Runtime/Engine/Private/AudioDecompress.cpp
@@ -702,7 +702,7 @@ void FAsyncAudioDecompressWorker::DoWork()
 
 static TAutoConsoleVariable<int32> CVarShouldUseBackgroundPoolFor_FAsyncRealtimeAudioTask(
 	TEXT("AudioThread.UseBackgroundThreadPool"),
-	0,
+	1,
 	TEXT("If true, use the background thread pool for realtime audio decompression."));
 
 bool ShouldUseBackgroundPoolFor_FAsyncRealtimeAudioTask()
diff --git a/Engine/Source/Runtime/Engine/Public/AudioDecompress.h b/Engine/Source/Runtime/Engine/Public/AudioDecompress.h
--- a/Engine/Source/Runtime/Engine/Public/AudioDecompress.h
+++ b/Engine/Source/Runtime/Engine/Public/AudioDecompress.h
@@ -470,7 +470,8 @@ public:
 	void StartBackgroundTask()
 	{
 		FScopeLock Lock(&CritSect);
-		Task->StartBackgroundTask(ShouldUseBackgroundPoolFor_FAsyncRealtimeAudioTask() ? GBackgroundPriorityThreadPool : GThreadPool);
+		const bool bUseBackground = ShouldUseBackgroundPoolFor_FAsyncRealtimeAudioTask() && (Task->GetTask().GetTaskType() != ERealtimeAudioTaskType::Procedural);
+		Task->StartBackgroundTask(bUseBackground ? GBackgroundPriorityThreadPool : GThreadPool);
 	}
 
 	FAsyncRealtimeAudioTaskWorker<T>& GetTask()
