Prevented render commands from being issues on FXSystemInterface, once the destroy command has been enqueued.

diff --git a/Engine/Plugins/FX/Niagara/Source/Niagara/Private/NiagaraEmitterInstanceBatcher.cpp b/Engine/Plugins/FX/Niagara/Source/Niagara/Private/NiagaraEmitterInstanceBatcher.cpp
--- a/Engine/Plugins/FX/Niagara/Source/Niagara/Private/NiagaraEmitterInstanceBatcher.cpp
+++ b/Engine/Plugins/FX/Niagara/Source/Niagara/Private/NiagaraEmitterInstanceBatcher.cpp
@@ -51,12 +51,14 @@ NiagaraEmitterInstanceBatcher::~NiagaraEmitterInstanceBatcher()
 
 void NiagaraEmitterInstanceBatcher::Queue(FNiagaraComputeExecutionContext *ExecContext)
 {
-	//UE_LOG(LogNiagara, Warning, TEXT("Submitted!"));
-	//SimulationQueue[CurQueueIndex]->Add(InContext);
-	TArray<FNiagaraComputeExecutionContext*>* Queue = &SimulationQueue[0];
-	uint32 QueueIndex = CurQueueIndex;
-	ENQUEUE_RENDER_COMMAND(QueueNiagaraDispatch)(
-		[Queue, QueueIndex, ExecContext](FRHICommandListImmediate& RHICmdList)
+	if (!IsPendingKill())
+	{
+		//UE_LOG(LogNiagara, Warning, TEXT("Submitted!"));
+		//SimulationQueue[CurQueueIndex]->Add(InContext);
+		TArray<FNiagaraComputeExecutionContext*>* Queue = &SimulationQueue[0];
+		uint32 QueueIndex = CurQueueIndex;
+		ENQUEUE_RENDER_COMMAND(QueueNiagaraDispatch)(
+			[Queue, QueueIndex, ExecContext](FRHICommandListImmediate& RHICmdList)
 		{
 			const uint32 QueueIndexMask = (1 << QueueIndex);
 			//Don't queue the same context for execution multiple times. TODO: possibly try to combine/accumulate the tick info if we happen to have > 1 before it's executed.
@@ -66,13 +68,16 @@ void NiagaraEmitterInstanceBatcher::Queue(FNiagaraComputeExecutionContext *ExecC
 				ExecContext->PendingExecutionQueueMask |= QueueIndexMask;
 			}
 		});
+	}
 }
 
 void NiagaraEmitterInstanceBatcher::Remove(FNiagaraComputeExecutionContext *ExecContext)
 {
-	TArray<FNiagaraComputeExecutionContext*>* Queue = &SimulationQueue[0];
-	ENQUEUE_RENDER_COMMAND(RemoveNiagaraDispatch)(
-		[Queue, ExecContext](FRHICommandListImmediate& RHICmdList)
+	if (!IsPendingKill())
+	{
+		TArray<FNiagaraComputeExecutionContext*>* Queue = &SimulationQueue[0];
+		ENQUEUE_RENDER_COMMAND(RemoveNiagaraDispatch)(
+			[Queue, ExecContext](FRHICommandListImmediate& RHICmdList)
 		{
 			for (int32 i = 0; i < SIMULATION_QUEUE_COUNT; i++)
 			{
@@ -80,6 +85,7 @@ void NiagaraEmitterInstanceBatcher::Remove(FNiagaraComputeExecutionContext *Exec
 			}
 			ExecContext->PendingExecutionQueueMask = 0;
 		});
+	}
 }
 
 void NiagaraEmitterInstanceBatcher::ExecuteAll(FRHICommandList &RHICmdList, FUniformBufferRHIParamRef ViewUniformBuffer)
diff --git a/Engine/Source/Runtime/Engine/Private/Particles/FXSystem.cpp b/Engine/Source/Runtime/Engine/Private/Particles/FXSystem.cpp
--- a/Engine/Source/Runtime/Engine/Private/Particles/FXSystem.cpp
+++ b/Engine/Source/Runtime/Engine/Private/Particles/FXSystem.cpp
@@ -48,6 +48,10 @@ FFXSystemInterface* FFXSystemInterface::Create(ERHIFeatureLevel::Type InFeatureL
 
 void FFXSystemInterface::Destroy( FFXSystemInterface* FXSystem )
 {
+	check(FXSystem && !FXSystem->bIsPendingKill);
+
+	// Notify that the delete command is on its way. Preventing any future render commands from accessing the FFXSystemInterface.
+	FXSystem->bIsPendingKill = true;
 	ENQUEUE_RENDER_COMMAND(FDestroyFXSystemCommand)(
 		[FXSystem](FRHICommandList& RHICmdList)
 		{
@@ -188,6 +192,15 @@ FFXSystem::FFXSystem(ERHIFeatureLevel::Type InFeatureLevel, EShaderPlatform InSh
 
 FFXSystem::~FFXSystem()
 {
+	for (FVectorFieldInstance* VectorFieldInstance : VectorFields)
+	{
+		if (VectorFieldInstance)
+		{
+			delete VectorFieldInstance;
+		}
+	}
+	VectorFields.Empty();
+
 	DestroyGPUSimulation();
 }
 
@@ -250,7 +263,7 @@ void FFXSystem::AddVectorField( UVectorFieldComponent* VectorFieldComponent )
 		check( VectorFieldComponent->VectorFieldInstance == NULL );
 		checkSlow( VectorFieldComponent->FXSystem && VectorFieldComponent->FXSystem->GetInterface(Name) == this );
 
-		if ( VectorFieldComponent->VectorField )
+		if ( VectorFieldComponent->VectorField && !IsPendingKill() )
 		{
 			FVectorFieldInstance* Instance = new FVectorFieldInstance();
 			VectorFieldComponent->VectorField->InitInstance(Instance, /*bPreviewInstance=*/ false);
@@ -281,7 +294,8 @@ void FFXSystem::RemoveVectorField( UVectorFieldComponent* VectorFieldComponent )
 		FVectorFieldInstance* Instance = VectorFieldComponent->VectorFieldInstance;
 		VectorFieldComponent->VectorFieldInstance = NULL;
 
-		if ( Instance )
+		// If pending kill the VectorFieldInstance will be freed in the destructor.
+		if ( Instance  && !IsPendingKill() )
 		{
 			FFXSystem* FXSystem = this;
 			ENQUEUE_RENDER_COMMAND(FRemoveVectorFieldCommand)(
@@ -305,7 +319,7 @@ void FFXSystem::UpdateVectorField( UVectorFieldComponent* VectorFieldComponent )
 
 		FVectorFieldInstance* Instance = VectorFieldComponent->VectorFieldInstance;
 
-		if ( Instance )
+		if ( Instance && !IsPendingKill() )
 		{
 			struct FUpdateVectorFieldParams
 			{
diff --git a/Engine/Source/Runtime/Engine/Private/Particles/ParticleGpuSimulation.cpp b/Engine/Source/Runtime/Engine/Private/Particles/ParticleGpuSimulation.cpp
--- a/Engine/Source/Runtime/Engine/Private/Particles/ParticleGpuSimulation.cpp
+++ b/Engine/Source/Runtime/Engine/Private/Particles/ParticleGpuSimulation.cpp
@@ -4535,9 +4535,11 @@ void FFXSystem::ReleaseGPUResources()
 
 void FFXSystem::AddGPUSimulation(FParticleSimulationGPU* Simulation)
 {
-	FFXSystem* FXSystem = this;
-	ENQUEUE_RENDER_COMMAND(FAddGPUSimulationCommand)(
-		[FXSystem, Simulation](FRHICommandListImmediate& RHICmdList)
+	if (!IsPendingKill())
+	{
+		FFXSystem* FXSystem = this;
+		ENQUEUE_RENDER_COMMAND(FAddGPUSimulationCommand)(
+			[FXSystem, Simulation](FRHICommandListImmediate& RHICmdList)
 		{
 			if (Simulation->SimulationIndex == INDEX_NONE)
 			{
@@ -4547,13 +4549,16 @@ void FFXSystem::AddGPUSimulation(FParticleSimulationGPU* Simulation)
 			}
 			check(FXSystem->GPUSimulations[Simulation->SimulationIndex] == Simulation);
 		});
+	}
 }
 
 void FFXSystem::RemoveGPUSimulation(FParticleSimulationGPU* Simulation)
 {
-	FFXSystem* FXSystem = this;
-	ENQUEUE_RENDER_COMMAND(FRemoveGPUSimulationCommand)(
-		[FXSystem, Simulation](FRHICommandListImmediate& RHICmdList)
+	if (!IsPendingKill())
+	{
+		FFXSystem* FXSystem = this;
+		ENQUEUE_RENDER_COMMAND(FRemoveGPUSimulationCommand)(
+			[FXSystem, Simulation](FRHICommandListImmediate& RHICmdList)
 		{
 			if (Simulation->SimulationIndex != INDEX_NONE)
 			{
@@ -4562,6 +4567,7 @@ void FFXSystem::RemoveGPUSimulation(FParticleSimulationGPU* Simulation)
 			}
 			Simulation->SimulationIndex = INDEX_NONE;
 		});
+	}
 }
 
 int32 FFXSystem::AddSortedGPUSimulation(FParticleSimulationGPU* Simulation, const FVector& ViewOrigin)
@@ -5103,17 +5109,20 @@ void FFXSystem::UpdateMultiGPUResources(FRHICommandListImmediate& RHICmdList)
 
 void FFXSystem::VisualizeGPUParticles(FCanvas* Canvas)
 {
-	FFXSystem* FXSystem = this;
-	int32 VisualizationMode = FXConsoleVariables::VisualizeGPUSimulation;
-	FRenderTarget* RenderTarget = Canvas->GetRenderTarget();
-	ERHIFeatureLevel::Type InFeatureLevel = FeatureLevel;
-	ENQUEUE_RENDER_COMMAND(FVisualizeGPUParticlesCommand)(
-		[FXSystem, VisualizationMode, RenderTarget, InFeatureLevel](FRHICommandList& RHICmdList)
-	{
-		FParticleSimulationResources* Resources = FXSystem->GetParticleSimulationResources();
-		FParticleStateTextures& CurrentStateTextures = Resources->GetVisualizeStateTextures();
-		VisualizeGPUSimulation(RHICmdList, InFeatureLevel, VisualizationMode, RenderTarget, CurrentStateTextures, GParticleCurveTexture.GetCurveTexture());
-	});
+	if (!IsPendingKill())
+	{
+		FFXSystem* FXSystem = this;
+		int32 VisualizationMode = FXConsoleVariables::VisualizeGPUSimulation;
+		FRenderTarget* RenderTarget = Canvas->GetRenderTarget();
+		ERHIFeatureLevel::Type InFeatureLevel = FeatureLevel;
+		ENQUEUE_RENDER_COMMAND(FVisualizeGPUParticlesCommand)(
+			[FXSystem, VisualizationMode, RenderTarget, InFeatureLevel](FRHICommandList& RHICmdList)
+		{
+			FParticleSimulationResources* Resources = FXSystem->GetParticleSimulationResources();
+			FParticleStateTextures& CurrentStateTextures = Resources->GetVisualizeStateTextures();
+			VisualizeGPUSimulation(RHICmdList, InFeatureLevel, VisualizationMode, RenderTarget, CurrentStateTextures, GParticleCurveTexture.GetCurveTexture());
+		});
+	}
 }
 
 /*-----------------------------------------------------------------------------
diff --git a/Engine/Source/Runtime/Engine/Public/FXSystem.h b/Engine/Source/Runtime/Engine/Public/FXSystem.h
--- a/Engine/Source/Runtime/Engine/Public/FXSystem.h
+++ b/Engine/Source/Runtime/Engine/Public/FXSystem.h
@@ -221,6 +221,8 @@ public:
 	 */
 	void PostRenderOpaque(FRHICommandListImmediate& RHICmdList) { PostRenderOpaque(RHICmdList, nullptr, nullptr, FUniformBufferRHIParamRef()); }
 
+	bool IsPendingKill() const { return bIsPendingKill; }
+
 protected:
 
 	/** By making the destructor protected, an instance must be destroyed via FFXSystemInterface::Destroy. */
@@ -228,6 +230,8 @@ protected:
 
 private:
 
+	bool bIsPendingKill = false;
+
 	static TMap<FName, FCreateCustomFXSystemDelegate> CreateCustomFXDelegates;
 };
 
