Better error message for when shader uses Primitive.Member directly

diff --git a/Engine/Source/Runtime/Engine/Private/LocalVertexFactory.cpp b/Engine/Source/Runtime/Engine/Private/LocalVertexFactory.cpp
--- a/Engine/Source/Runtime/Engine/Private/LocalVertexFactory.cpp
+++ b/Engine/Source/Runtime/Engine/Private/LocalVertexFactory.cpp
@@ -203,7 +203,7 @@ void FLocalVertexFactory::ValidateCompiledResult(const FVertexFactoryType* Type,
 		&& UseGPUScene(Platform, GetMaxSupportedFeatureLevel(Platform)) 
 		&& ParameterMap.ContainsParameterAllocation(FPrimitiveUniformShaderParameters::StaticStructMetadata.GetShaderVariableName()))
 	{
-		OutErrors.AddUnique(*FString::Printf(TEXT("Shader attempted to bind the Primitive uniform buffer even though Vertex Factory %s computes a PrimitiveId per-instance.  This will break auto-instancing.  Shaders should use GetPrimitiveData(PrimitiveId).Member instead of Primitive.Member."), Type->GetName()));
+		OutErrors.AddUnique(*FString::Printf(TEXT("Shader attempted to bind the Primitive uniform buffer even though Vertex Factory %s computes a PrimitiveId per-instance.  This will break auto-instancing.  Shaders should use GetPrimitiveData(Parameters.PrimitiveId).Member instead of Primitive.Member."), Type->GetName()));
 	}
 }
 
