FBodyInstance::SetMassSpaceLocal used incorrect physics lock

diff --git a/Engine/Source/Runtime/Engine/Private/PhysicsEngine/BodyInstance.cpp b/Engine/Source/Runtime/Engine/Private/PhysicsEngine/BodyInstance.cpp
--- a/Engine/Source/Runtime/Engine/Private/PhysicsEngine/BodyInstance.cpp
+++ b/Engine/Source/Runtime/Engine/Private/PhysicsEngine/BodyInstance.cpp
@@ -2153,7 +2153,7 @@ FTransform FBodyInstance::GetMassSpaceLocal() const
 void FBodyInstance::SetMassSpaceLocal(const FTransform& NewMassSpaceLocalTM)
 {
 	//TODO: UE doesn't store this so any changes to mass properties will not remember about this properly
-	FPhysicsCommand::ExecuteRead(ActorHandle, [&](const FPhysicsActorHandle& Actor)
+	FPhysicsCommand::ExecuteWrite(ActorHandle, [&](const FPhysicsActorHandle& Actor)
 	{
 		FPhysicsInterface::SetComLocalPose_AssumesLocked(Actor, NewMassSpaceLocalTM);
 	});
