Fixed warning on deprecated Pauser variable being Replicated but not in GetLifetimeReplicatedProps
#rnx

diff --git a/Engine/Source/Runtime/Engine/Classes/GameFramework/WorldSettings.h b/Engine/Source/Runtime/Engine/Classes/GameFramework/WorldSettings.h
--- a/Engine/Source/Runtime/Engine/Classes/GameFramework/WorldSettings.h
+++ b/Engine/Source/Runtime/Engine/Classes/GameFramework/WorldSettings.h
@@ -683,7 +683,7 @@ public:
 
 	// If paused, FName of person pausing the game.
 	UE_DEPRECATED(4.23, "This property is deprecated. Please use Get/SetPauserPlayerState().")
-	UPROPERTY(transient, replicated)
+	UPROPERTY(transient)
 	class APlayerState* Pauser;
 
 	/** valid only during replication - information about the player(s) being replicated to
