#ue4 - Improve jump consistency over varied framerates on the frame when the jump force ends and we transition to applying gravity.

We have a portion of the time at constant speed then a portion with constant acceleration due to gravity, and those need to be separated.

diff --git a/Engine/Source/Runtime/Engine/Private/Components/CharacterMovementComponent.cpp b/Engine/Source/Runtime/Engine/Private/Components/CharacterMovementComponent.cpp
--- a/Engine/Source/Runtime/Engine/Private/Components/CharacterMovementComponent.cpp
+++ b/Engine/Source/Runtime/Engine/Private/Components/CharacterMovementComponent.cpp
@@ -4139,6 +4139,7 @@ void UCharacterMovementComponent::PhysFalling(float deltaTime, int32 Iterations)
 		float GravityTime = timeTick;
 
 		// If jump is providing force, gravity may be affected.
+		bool bEndingJumpForce = false;
 		if (CharacterOwner->JumpForceTimeRemaining > 0.0f)
 		{
 			// Consume some of the force time. Only the remaining time (if any) is affected by gravity when bApplyGravityWhileJumping=false.
@@ -4150,6 +4151,7 @@ void UCharacterMovementComponent::PhysFalling(float deltaTime, int32 Iterations)
 			if (CharacterOwner->JumpForceTimeRemaining <= 0.0f)
 			{
 				CharacterOwner->ResetJumpState();
+				bEndingJumpForce = true;
 			}
 		}
 
@@ -4167,9 +4169,20 @@ void UCharacterMovementComponent::PhysFalling(float deltaTime, int32 Iterations)
 		}
 
 
+		// Compute change in position (using midpoint integration method).
+		FVector Adjusted = 0.5f*(OldVelocity + Velocity) * timeTick;
+		
+		// Special handling if ending the jump force where we didn't apply gravity during the jump.
+		if (bEndingJumpForce && !bApplyGravityWhileJumping)
+		{
+			// We had a portion of the time at constant speed then a portion with acceleration due to gravity.
+			// Account for that here with a more correct change in position.
+			const float NonGravityTime = FMath::Max(0.f, timeTick - GravityTime);
+			Adjusted = (OldVelocity * NonGravityTime) + (0.5f*(OldVelocity + Velocity) * GravityTime);
+		}
+
 		// Move
 		FHitResult Hit(1.f);
-		FVector Adjusted = 0.5f*(OldVelocity + Velocity) * timeTick;
 		SafeMoveUpdatedComponent( Adjusted, PawnRotation, true, Hit);
 		
 		if (!HasValidData())
