Fix it so switching worlds via the scope correctly sets GPlayInEditorID, related to UE-59417
Add a scope switcher to ShooterGameInstance Tick as it's executed from FCoreTicker, outside the normal world tick. It needs to be outside worldtick because it handles map transfer and other things that are not safe to do during world tick.

diff --git a/Engine/Source/Editor/UnrealEd/Classes/Editor/EditorEngine.h b/Engine/Source/Editor/UnrealEd/Classes/Editor/EditorEngine.h
--- a/Engine/Source/Editor/UnrealEd/Classes/Editor/EditorEngine.h
+++ b/Engine/Source/Editor/UnrealEd/Classes/Editor/EditorEngine.h
@@ -2705,7 +2705,7 @@ private:
 	/**
 	 * Called via a delegate to toggle between the editor and pie world
 	 */
-	void OnSwitchWorldsForPIE( bool bSwitchToPieWorld );
+	void OnSwitchWorldsForPIE( bool bSwitchToPieWorld, UWorld* OverrideWorld = nullptr );
 
 	/**
 	 * Gives focus to the server or first PIE client viewport
diff --git a/Engine/Source/Editor/UnrealEd/Private/Editor.cpp b/Engine/Source/Editor/UnrealEd/Private/Editor.cpp
--- a/Engine/Source/Editor/UnrealEd/Private/Editor.cpp
+++ b/Engine/Source/Editor/UnrealEd/Private/Editor.cpp
@@ -736,6 +736,11 @@ UWorld* SetPlayInEditorWorld( UWorld* PlayInEditorWorld )
 	GIsPlayInEditorWorld = true;
 	GWorld = PlayInEditorWorld;
 
+	if (FWorldContext* WorldContext = GEngine->GetWorldContextFromWorld(PlayInEditorWorld))
+	{
+		GPlayInEditorID = WorldContext->PIEInstance;
+	}
+
 	return SavedWorld;
 }
 
@@ -750,6 +755,7 @@ void RestoreEditorWorld( UWorld* EditorWorld )
 	check(GIsPlayInEditorWorld);
 	GIsPlayInEditorWorld = false;
 	GWorld = EditorWorld;
+	GPlayInEditorID = -1;
 }
 
 /**
diff --git a/Engine/Source/Editor/UnrealEd/Private/PlayLevel.cpp b/Engine/Source/Editor/UnrealEd/Private/PlayLevel.cpp
--- a/Engine/Source/Editor/UnrealEd/Private/PlayLevel.cpp
+++ b/Engine/Source/Editor/UnrealEd/Private/PlayLevel.cpp
@@ -3644,15 +3644,15 @@ int32 UEditorEngine::OnSwitchWorldForSlatePieWindow( int32 WorldID )
 	return RestoreID;
 }
 
-void UEditorEngine::OnSwitchWorldsForPIE( bool bSwitchToPieWorld )
+void UEditorEngine::OnSwitchWorldsForPIE( bool bSwitchToPieWorld, UWorld* OverrideWorld )
 {
 	if( bSwitchToPieWorld )
 	{
-		SetPlayInEditorWorld( PlayWorld );
+		SetPlayInEditorWorld( OverrideWorld ? OverrideWorld : PlayWorld );
 	}
 	else
 	{
-		RestoreEditorWorld( EditorWorld );
+		RestoreEditorWorld( OverrideWorld ? OverrideWorld : EditorWorld );
 	}
 }
 
diff --git a/Engine/Source/Runtime/Engine/Private/UnrealEngine.cpp b/Engine/Source/Runtime/Engine/Private/UnrealEngine.cpp
--- a/Engine/Source/Runtime/Engine/Private/UnrealEngine.cpp
+++ b/Engine/Source/Runtime/Engine/Private/UnrealEngine.cpp
@@ -10269,7 +10269,24 @@ FOnSwitchWorldForPIE FScopedConditionalWorldSwitcher::SwitchWorldForPIEDelegate;
 
 FScopedConditionalWorldSwitcher::FScopedConditionalWorldSwitcher( FViewportClient* InViewportClient )
 	: ViewportClient( InViewportClient )
-	, OldWorld( NULL )
+	, OldWorld( nullptr )
+{
+	ConditionalSwitchWorld( ViewportClient, nullptr );
+}
+
+FScopedConditionalWorldSwitcher::FScopedConditionalWorldSwitcher(UWorld* InWorld)
+	: ViewportClient( nullptr )
+	, OldWorld( nullptr )
+{
+	if (InWorld)
+	{
+		ViewportClient = InWorld->GetGameViewport();
+	}
+
+	ConditionalSwitchWorld( ViewportClient, InWorld );
+}
+
+void FScopedConditionalWorldSwitcher::ConditionalSwitchWorld( FViewportClient* InViewportClient, UWorld* InWorld )
 {
 	if( GIsEditor )
 	{
@@ -10278,32 +10295,45 @@ FScopedConditionalWorldSwitcher::FScopedConditionalWorldSwitcher( FViewportClien
 			OldWorld = GWorld; 
 			const bool bSwitchToPIEWorld = true;
 			// Delegate must be valid
-			SwitchWorldForPIEDelegate.ExecuteIfBound( bSwitchToPIEWorld );
-		} 
+			SwitchWorldForPIEDelegate.ExecuteIfBound( bSwitchToPIEWorld, nullptr );
+		}
 		else if( ViewportClient )
 		{
 			// Tell the viewport client to set the correct world and store what the world used to be
 			OldWorld = ViewportClient->ConditionalSetWorld();
 		}
+		else if ( InWorld && !GIsPlayInEditorWorld )
+		{
+			OldWorld = GWorld;
+			const bool bSwitchToPIEWorld = true;
+			// No viewport so set the world directly
+			SwitchWorldForPIEDelegate.ExecuteIfBound( bSwitchToPIEWorld, InWorld );
+		}
 	}
 }
 
 FScopedConditionalWorldSwitcher::~FScopedConditionalWorldSwitcher()
 {
-	// Only switch in the editor and if we made a swtich (OldWorld not null)
+	// Only switch in the editor and if we made a switch (OldWorld not null)
 	if( GIsEditor && OldWorld )
 	{
 		if( ViewportClient && ViewportClient == GEngine->GameViewport && GIsPlayInEditorWorld )
 		{
 			const bool bSwitchToPIEWorld = false;
 			// Delegate must be valid
-			SwitchWorldForPIEDelegate.ExecuteIfBound( bSwitchToPIEWorld );
-		} 
+			SwitchWorldForPIEDelegate.ExecuteIfBound( bSwitchToPIEWorld, nullptr );
+		}
 		else if( ViewportClient )
 		{
 			// Tell the viewport client to restore the old world
 			ViewportClient->ConditionalRestoreWorld( OldWorld );
 		}
+		else if( GIsPlayInEditorWorld )
+		{
+			// No viewport so always restore to old world
+			const bool bSwitchToPIEWorld = false;
+			SwitchWorldForPIEDelegate.ExecuteIfBound( bSwitchToPIEWorld, OldWorld );
+		}
 	}
 }
 
diff --git a/Engine/Source/Runtime/Engine/Public/UnrealEngine.h b/Engine/Source/Runtime/Engine/Public/UnrealEngine.h
--- a/Engine/Source/Runtime/Engine/Public/UnrealEngine.h
+++ b/Engine/Source/Runtime/Engine/Public/UnrealEngine.h
@@ -299,25 +299,35 @@ private:
 
 
 /** Delegate for switching between PIE and Editor worlds without having access to them */
-DECLARE_DELEGATE_OneParam( FOnSwitchWorldForPIE, bool );
+DECLARE_DELEGATE_TwoParams( FOnSwitchWorldForPIE, bool, UWorld* );
 
 #if WITH_EDITOR
 /**
- * When created, tells the viewport client to set the appropriate GWorld
- * When destroyed, tells the viewport client to reset the GWorld back to what it was
+ * When created, switches global context to a PIE world
+ * When destroyed, resets the GWorld back to what it was before
  */
 class ENGINE_API FScopedConditionalWorldSwitcher
 {
 public:
+	/** Use the viewport to figure out what world to set temporarily */
 	FScopedConditionalWorldSwitcher( class FViewportClient* InViewportClient );
+
+	/** Explicitly set to the specific world */
+	FScopedConditionalWorldSwitcher( UWorld* InWorld );
+
+	/** Resets back to initial world */
 	~FScopedConditionalWorldSwitcher();
 
 	/** Delegate to call to switch worlds for PIE viewports.  Is not called when simulating (non-gameviewportclient) */
 	static FOnSwitchWorldForPIE SwitchWorldForPIEDelegate;
 private:
+	/** Called by constructors to do actual switch */
+	void ConditionalSwitchWorld( class FViewportClient* InViewportClient, UWorld* InWorld );
+
 	/** Viewport client used to set the world */
 	FViewportClient* ViewportClient;
-	/** World to reset when this is destoyed.  Can be null if nothing needs to be reset */	
+
+	/** World to reset when this is destroyed.  Can be null if nothing needs to be reset */
 	UWorld* OldWorld;
 };
 #else
@@ -325,7 +335,8 @@ private:
 class ENGINE_API FScopedConditionalWorldSwitcher
 {
 public:
-	FScopedConditionalWorldSwitcher( class FViewportClient* InViewportClient ){}
+	FScopedConditionalWorldSwitcher( class FViewportClient* InViewportClient ) {}
+	FScopedConditionalWorldSwitcher( UWorld* InWorld ) {}
 	~FScopedConditionalWorldSwitcher() {}
 };
 
