Fixed the case where StaticMesh LightMapCoordinateIndex could go negative if mesh has stripped LOD data. Removed workoraund for negative LightMapCoordinateIndex.

diff --git a/Engine/Source/Runtime/Engine/Private/Rendering/StaticMeshVertexBuffer.cpp b/Engine/Source/Runtime/Engine/Private/Rendering/StaticMeshVertexBuffer.cpp
--- a/Engine/Source/Runtime/Engine/Private/Rendering/StaticMeshVertexBuffer.cpp
+++ b/Engine/Source/Runtime/Engine/Private/Rendering/StaticMeshVertexBuffer.cpp
@@ -596,18 +596,6 @@ void FStaticMeshVertexBuffer::BindTexCoordVertexBuffer(const FVertexFactory* Ver
 void FStaticMeshVertexBuffer::BindLightMapVertexBuffer(const FVertexFactory* VertexFactory, FStaticMeshDataType& Data, int LightMapCoordinateIndex) const
 {
 	LightMapCoordinateIndex = LightMapCoordinateIndex < (int32)GetNumTexCoords() ? LightMapCoordinateIndex : (int32)GetNumTexCoords() - 1;
-
-	// Temp patch.
-	if (LightMapCoordinateIndex < 0)
-	{
-		static bool bOnce = false;
-		if (!bOnce)
-		{
-			UE_LOG(LogStaticMesh, Warning, TEXT("LightMapCoordinateIndex (%d)< 0, GetNumTexCoords: %d!"), LightMapCoordinateIndex, (int32)GetNumTexCoords());
-			bOnce = false;
-		}
-		return;
-	}
 	check(LightMapCoordinateIndex >= 0);
 
 	Data.LightMapCoordinateIndex = LightMapCoordinateIndex;
diff --git a/Engine/Source/Runtime/Engine/Private/StaticMesh.cpp b/Engine/Source/Runtime/Engine/Private/StaticMesh.cpp
--- a/Engine/Source/Runtime/Engine/Private/StaticMesh.cpp
+++ b/Engine/Source/Runtime/Engine/Private/StaticMesh.cpp
@@ -5251,7 +5251,11 @@ void UStaticMesh::EnforceLightmapRestrictions()
 	{
 		for (int32 LODIndex = 0; LODIndex < RenderData->LODResources.Num(); ++LODIndex)
 		{
-			NumUVs = FMath::Min(RenderData->LODResources[LODIndex].GetNumTexCoords(),NumUVs);
+			const FStaticMeshLODResources& LODResource = RenderData->LODResources[LODIndex];
+			if (LODResource.GetNumVertices() > 0) // skip LOD that was stripped (eg. MinLOD)
+			{
+				NumUVs = FMath::Min(LODResource.GetNumTexCoords(), NumUVs);
+			}
 		}
 	}
 	else
@@ -5259,6 +5263,9 @@ void UStaticMesh::EnforceLightmapRestrictions()
 		NumUVs = 1;
 	}
 
+	// do not allow LightMapCoordinateIndex go negative
+	check(NumUVs > 0);
+
 	// Clamp LightMapCoordinateIndex to be valid for all lightmap uvs
 	LightMapCoordinateIndex = FMath::Clamp(LightMapCoordinateIndex, 0, NumUVs - 1);
 }
