Trim FTexture2DMipMap in cooked builds

diff --git a/Engine/Source/Runtime/Core/Private/HAL/LowLevelMemTracker.cpp b/Engine/Source/Runtime/Core/Private/HAL/LowLevelMemTracker.cpp
--- a/Engine/Source/Runtime/Core/Private/HAL/LowLevelMemTracker.cpp
+++ b/Engine/Source/Runtime/Core/Private/HAL/LowLevelMemTracker.cpp
@@ -78,6 +78,7 @@ DECLARE_LLM_MEMORY_STAT(TEXT("Stats"), STAT_StatsLLM, STATGROUP_LLMFULL);
 DECLARE_LLM_MEMORY_STAT(TEXT("Shaders"), STAT_ShadersLLM, STATGROUP_LLMFULL);
 DECLARE_LLM_MEMORY_STAT(TEXT("PSO"), STAT_PSOLLM, STATGROUP_LLMFULL);
 DECLARE_LLM_MEMORY_STAT(TEXT("Textures"), STAT_TexturesLLM, STATGROUP_LLMFULL);
+DECLARE_LLM_MEMORY_STAT(TEXT("TextureMetaData"), STAT_TextureMetaDataLLM, STATGROUP_LLMFULL);
 DECLARE_LLM_MEMORY_STAT(TEXT("Render Targets"), STAT_RenderTargetsLLM, STATGROUP_LLMFULL);
 DECLARE_LLM_MEMORY_STAT(TEXT("RHIMisc"), STAT_RHIMiscLLM, STATGROUP_LLMFULL);
 DECLARE_LLM_MEMORY_STAT(TEXT("PhysX TriMesh"), STAT_PhysXTriMeshLLM, STATGROUP_LLMFULL);
diff --git a/Engine/Source/Runtime/Core/Public/HAL/LowLevelMemTracker.h b/Engine/Source/Runtime/Core/Public/HAL/LowLevelMemTracker.h
--- a/Engine/Source/Runtime/Core/Public/HAL/LowLevelMemTracker.h
+++ b/Engine/Source/Runtime/Core/Public/HAL/LowLevelMemTracker.h
@@ -206,6 +206,7 @@ enum class ELLMTagSet : uint8
 	macro(Shaders,								"Shaders",						GET_STATFNAME(STAT_ShadersLLM),								GET_STATFNAME(STAT_EngineSummaryLLM),			-1)\
 	macro(PSO,									"PSO",							GET_STATFNAME(STAT_PSOLLM),									GET_STATFNAME(STAT_EngineSummaryLLM),			-1)\
 	macro(Textures,								"Textures",						GET_STATFNAME(STAT_TexturesLLM),							GET_STATFNAME(STAT_TexturesSummaryLLM),			-1)\
+	macro(TextureMetaData,						"TextureMetaData",				GET_STATFNAME(STAT_TextureMetaDataLLM),						GET_STATFNAME(STAT_TexturesSummaryLLM),			-1)\
 	macro(RenderTargets,						"RenderTargets",				GET_STATFNAME(STAT_RenderTargetsLLM),						GET_STATFNAME(STAT_EngineSummaryLLM),			-1)\
 	macro(RHIMisc,								"RHIMisc",						GET_STATFNAME(STAT_RHIMiscLLM),								GET_STATFNAME(STAT_EngineSummaryLLM),			-1)\
 	macro(AsyncLoading,							"AsyncLoading",					GET_STATFNAME(STAT_AsyncLoadingLLM),						GET_STATFNAME(STAT_EngineSummaryLLM),			-1)\
diff --git a/Engine/Source/Runtime/Engine/Classes/Engine/Texture.h b/Engine/Source/Runtime/Engine/Classes/Engine/Texture.h
--- a/Engine/Source/Runtime/Engine/Classes/Engine/Texture.h
+++ b/Engine/Source/Runtime/Engine/Classes/Engine/Texture.h
@@ -319,6 +319,11 @@ struct FTexturePlatformData
 	/** Mip data. */
 	TIndirectArray<struct FTexture2DMipMap> Mips;
 
+#if TEXTURE2DMIPMAP_USE_COMPACT_BULKDATA
+	/** Cached UPackage file name where the owning texture is loaded from */
+	FString CachedPackageFileName;
+#endif
+
 #if WITH_EDITORONLY_DATA
 	/** The key associated with this derived data. */
 	FString DerivedDataKey;
diff --git a/Engine/Source/Runtime/Engine/Classes/Engine/Texture2D.h b/Engine/Source/Runtime/Engine/Classes/Engine/Texture2D.h
--- a/Engine/Source/Runtime/Engine/Classes/Engine/Texture2D.h
+++ b/Engine/Source/Runtime/Engine/Classes/Engine/Texture2D.h
@@ -115,7 +115,7 @@ public:
 	virtual EMaterialValueType GetMaterialType() const override { return MCT_Texture2D; }
 	virtual void UpdateResource() override;
 	virtual float GetAverageBrightness(bool bIgnoreTrueBlack, bool bUseGrayscale) override;
-	virtual FTexturePlatformData** GetRunningPlatformData() override { return &PlatformData; }
+	virtual FTexturePlatformData** GetRunningPlatformData() final override { return &PlatformData; }
 #if WITH_EDITOR
 	virtual TMap<FString,FTexturePlatformData*>* GetCookedPlatformData() override { return &CookedPlatformData; }
 #endif
@@ -127,18 +127,7 @@ public:
 	virtual int32 CalcNumOptionalMips() const final override;
 	virtual int32 CalcCumulativeLODSize(int32 NumLODs) const final override { return CalcTextureMemorySize(NumLODs); }
 
-	virtual bool GetMipDataFilename(const int32 MipIndex, FString& OutBulkDataFilename) const final override
-	{
-		if (PlatformData)
-		{
-			if (MipIndex < PlatformData->Mips.Num() && MipIndex >= 0)
-			{
-				OutBulkDataFilename = PlatformData->Mips[MipIndex].BulkData.GetFilename();
-				return true;
-			}
-		}
-		return false;
-	}
+	virtual bool GetMipDataFilename(const int32 MipIndex, FString& OutBulkDataFilename) const final override;
 
 	/**
 	* Returns whether the texture is ready for streaming aka whether it has had InitRHI called on it.
diff --git a/Engine/Source/Runtime/Engine/Private/Streaming/Texture2DStreamIn_IO.cpp b/Engine/Source/Runtime/Engine/Private/Streaming/Texture2DStreamIn_IO.cpp
--- a/Engine/Source/Runtime/Engine/Private/Streaming/Texture2DStreamIn_IO.cpp
+++ b/Engine/Source/Runtime/Engine/Private/Streaming/Texture2DStreamIn_IO.cpp
@@ -59,7 +59,11 @@ void FTexture2DStreamIn_IO::SetIOFilename(const FContext& Context)
 
 		if (MipIndex == PendingFirstMip)
 		{
+#if !TEXTURE2DMIPMAP_USE_COMPACT_BULKDATA
 			IOFilename = MipMap.BulkData.GetFilename();
+#else
+			verify(Context.Texture->GetMipDataFilename(MipIndex, IOFilename));
+#endif
 
 			if (GEventDrivenLoaderEnabled)
 			{
@@ -72,12 +76,14 @@ void FTexture2DStreamIn_IO::SetIOFilename(const FContext& Context)
 				}
 			}
 		}
+#if !TEXTURE2DMIPMAP_USE_COMPACT_BULKDATA
 		else if (IOFilename != MipMap.BulkData.GetFilename())
 		{
 			UE_LOG(LogTexture, Error, TEXT("All of the streaming mips must be stored in the same file %s %s."), *IOFilename, *MipMap.BulkData.GetFilename());
 			IOFilename.Reset();
 			break;
 		}
+#endif
 	}
 
 	if (IOFilename.IsEmpty())
diff --git a/Engine/Source/Runtime/Engine/Private/Texture2D.cpp b/Engine/Source/Runtime/Engine/Private/Texture2D.cpp
--- a/Engine/Source/Runtime/Engine/Private/Texture2D.cpp
+++ b/Engine/Source/Runtime/Engine/Private/Texture2D.cpp
@@ -16,6 +16,7 @@
 #include "UObject/UObjectIterator.h"
 #include "UObject/Package.h"
 #include "UObject/LinkerLoad.h"
+#include "UObject/CoreRedirects.h"
 #include "RenderUtils.h"
 #include "ContentStreaming.h"
 #include "EngineUtils.h"
@@ -179,6 +180,153 @@ private:
 	FTexture2DMipMap
 -----------------------------------------------------------------------------*/
 
+FTexture2DMipMap::FCompactByteBulkData::FCompactByteBulkData()
+{
+	Reset();
+}
+
+void FTexture2DMipMap::FCompactByteBulkData::Reset()
+{
+	OffsetInFile = 0xffffffff;
+	BulkDataSize = 0;
+	BulkDataFlags = 0;
+	TexelData = nullptr;
+}
+
+FTexture2DMipMap::FCompactByteBulkData::~FCompactByteBulkData()
+{
+	FMemory::Free(TexelData);
+	Reset();
+}
+
+FTexture2DMipMap::FCompactByteBulkData::FCompactByteBulkData(const FCompactByteBulkData& Other)
+{
+	*this = Other;
+}
+
+FTexture2DMipMap::FCompactByteBulkData::FCompactByteBulkData(FCompactByteBulkData&& Other)
+{
+	*this = MoveTemp(Other);
+}
+
+typename FTexture2DMipMap::FCompactByteBulkData& FTexture2DMipMap::FCompactByteBulkData::operator=(const FCompactByteBulkData& Other)
+{
+	if (TexelData != Other.TexelData)
+	{
+		OffsetInFile = Other.OffsetInFile;
+		BulkDataFlags = Other.BulkDataFlags;
+		Realloc(Other.BulkDataSize);
+		FMemory::Memcpy(TexelData, Other.TexelData, BulkDataSize);
+	}
+	return *this;
+}
+
+typename FTexture2DMipMap::FCompactByteBulkData& FTexture2DMipMap::FCompactByteBulkData::operator=(FCompactByteBulkData&& Other)
+{
+	if (TexelData != Other.TexelData)
+	{
+		FMemory::Free(TexelData);
+		OffsetInFile = Other.OffsetInFile;
+		BulkDataSize = Other.BulkDataSize;
+		BulkDataFlags = Other.BulkDataFlags;
+		TexelData = Other.TexelData;
+		Other.Reset();
+	}
+	return *this;
+}
+
+void FTexture2DMipMap::FCompactByteBulkData::Serialize(FArchive& Ar, UObject* Owner, int32 MipIdx)
+{
+	check(Ar.IsLoading());
+	FMemory::Free(TexelData);
+	TexelData = nullptr;
+
+	FByteBulkData TmpBulkData;
+	TmpBulkData.Serialize(Ar, Owner, MipIdx);
+
+	const int64 Tmp = TmpBulkData.GetBulkDataOffsetInFile();
+	check(Tmp >= 0 && Tmp <= 0xffffffffll);
+	OffsetInFile = static_cast<uint32>(Tmp);
+	BulkDataSize = static_cast<uint32>(TmpBulkData.GetBulkDataSize());
+	BulkDataFlags = TmpBulkData.GetBulkDataFlags();
+
+	if (IsInlined())
+	{
+		TmpBulkData.GetCopy((void**)&TexelData);
+	}
+
+	if (!MipIdx && !IsInlined())
+	{
+		UTexture2D* OwningTexture2D = Cast<UTexture2D>(Owner);
+		if (OwningTexture2D)
+		{
+			FTexturePlatformData** PlatformDataPtr = OwningTexture2D->GetRunningPlatformData();
+			check(PlatformDataPtr && *PlatformDataPtr);
+			FTexturePlatformData* PlatformData = *PlatformDataPtr;
+#if TEXTURE2DMIPMAP_USE_COMPACT_BULKDATA
+			PlatformData->CachedPackageFileName = TmpBulkData.GetFilename();
+#endif
+		}
+	}
+}
+
+void FTexture2DMipMap::FCompactByteBulkData::Unlock() const
+{
+	if (!!(BulkDataFlags & BULKDATA_SingleUse))
+	{
+		FMemory::Free(TexelData);
+		const_cast<FCompactByteBulkData*>(this)->TexelData = nullptr;
+	}
+}
+
+void* FTexture2DMipMap::FCompactByteBulkData::Realloc(int32 NumBytes)
+{
+	FMemory::Free(TexelData);
+	if (NumBytes > 0)
+	{
+		TexelData = (uint8*)FMemory::Malloc(NumBytes);
+		BulkDataSize = NumBytes;
+	}
+	else
+	{
+		TexelData = nullptr;
+		BulkDataSize = 0;
+	}
+	return TexelData;
+}
+
+void FTexture2DMipMap::FCompactByteBulkData::GetCopy(void** Dest, bool bDiscardInternalCopy)
+{
+	if (!IsInlined())
+	{
+		UE_LOG(LogTexture, Fatal, TEXT("FCompactByteBulkData doesn't support GetCopy if data isn't inlined"));
+	}
+
+	if (!BulkDataSize)
+	{
+		check(!TexelData);
+		return;
+	}
+
+	if (!*Dest)
+	{
+		if (bDiscardInternalCopy)
+		{
+			*Dest = TexelData;
+			TexelData = nullptr;
+			return;
+		}
+		*Dest = FMemory::Malloc(BulkDataSize);
+	}
+
+	FMemory::Memcpy(*Dest, TexelData, BulkDataSize);
+	if (bDiscardInternalCopy)
+	{
+		FMemory::Free(TexelData);
+		TexelData = nullptr;
+	}
+}
+
 void FTexture2DMipMap::Serialize(FArchive& Ar, UObject* Owner, int32 MipIdx)
 {
 	bool bCooked = Ar.IsCooking();
@@ -230,6 +378,8 @@ bool UTexture2D::GetResourceMemSettings(int32 FirstMipIdx, int32& OutSizeX, int3
 
 void UTexture2D::Serialize(FArchive& Ar)
 {
+	LLM_SCOPE(ELLMTag::TextureMetaData);
+
 	Super::Serialize(Ar);
 
 	FStripDataFlags StripDataFlags(Ar);
@@ -296,6 +446,26 @@ int32 UTexture2D::CalcNumOptionalMips() const
 	}
 	return 0;
 }
+
+bool UTexture2D::GetMipDataFilename(const int32 MipIndex, FString& OutBulkDataFilename) const
+{
+	if (PlatformData)
+	{
+		if (MipIndex < PlatformData->Mips.Num() && MipIndex >= 0)
+		{
+#if !TEXTURE2DMIPMAP_USE_COMPACT_BULKDATA
+			OutBulkDataFilename = PlatformData->Mips[MipIndex].BulkData.GetFilename();
+#else
+			OutBulkDataFilename = PlatformData->CachedPackageFileName;
+			const int32 NumOptionalMips = CalcNumOptionalMips();
+			OutBulkDataFilename = FPaths::ChangeExtension(OutBulkDataFilename, MipIndex < NumOptionalMips ? TEXT(".uptnl") : TEXT(".ubulk"));
+#endif
+			return true;
+		}
+	}
+	return false;
+}
+
 int32 UTexture2D::GetNumResidentMips() const
 {
 	FTexture2DResource* Texture2DResource = (FTexture2DResource*)Resource;
diff --git a/Engine/Source/Runtime/Engine/Private/TextureDerivedData.cpp b/Engine/Source/Runtime/Engine/Private/TextureDerivedData.cpp
--- a/Engine/Source/Runtime/Engine/Private/TextureDerivedData.cpp
+++ b/Engine/Source/Runtime/Engine/Private/TextureDerivedData.cpp
@@ -830,7 +830,7 @@ bool FTexturePlatformData::TryLoadMips(int32 FirstMipToLoad, void** OutMipData)
 			{
 				OutMipData[MipIndex - FirstMipToLoad] = FMemory::Malloc(Mip.BulkData.GetBulkDataSize());
 
-#if PLATFORM_SUPPORTS_TEXTURE_STREAMING
+#if PLATFORM_SUPPORTS_TEXTURE_STREAMING && !TEXTURE2DMIPMAP_USE_COMPACT_BULKDATA
 				// We want to make sure that any non-streamed mips are coming from the texture asset file, and not from an external bulk file.
 				// But because "r.TextureStreaming" is driven by the project setting as well as the command line option "-NoTextureStreaming", 
 				// is it possible for streaming mips to be loaded in non streaming ways.
diff --git a/Engine/Source/Runtime/Engine/Public/TextureResource.h b/Engine/Source/Runtime/Engine/Public/TextureResource.h
--- a/Engine/Source/Runtime/Engine/Public/TextureResource.h
+++ b/Engine/Source/Runtime/Engine/Public/TextureResource.h
@@ -27,11 +27,59 @@ class UTexture2D;
 /** Maximum number of slices in texture source art. */
 #define MAX_TEXTURE_SOURCE_SLICES 6
 
+#define TEXTURE2DMIPMAP_USE_COMPACT_BULKDATA !(WITH_EDITORONLY_DATA)
+
 /**
  * A 2D texture mip-map.
  */
 struct FTexture2DMipMap
 {
+	class FCompactByteBulkData
+	{
+	public:
+		FCompactByteBulkData();
+
+		~FCompactByteBulkData();
+
+		FCompactByteBulkData(const FCompactByteBulkData&);
+		FCompactByteBulkData(FCompactByteBulkData&&);
+
+		FCompactByteBulkData& operator=(const FCompactByteBulkData& Other);
+		FCompactByteBulkData& operator=(FCompactByteBulkData&& Other);
+
+		void Serialize(FArchive& Ar, UObject* Owner, int32 MipIdx);
+
+		uint32 GetBulkDataOffsetInFile() const { return OffsetInFile; }
+		uint32 GetBulkDataSize() const { return BulkDataSize; }
+		uint32 GetBulkDataFlags() const { return BulkDataFlags; }
+		uint32 GetElementCount() const { return BulkDataSize; }
+		uint32 GetElementSize() const { return sizeof(uint8); }
+		void SetBulkDataFlags(uint32 Flags) { BulkDataFlags |= Flags; }
+		void ClearBulkDataFlags(uint32 FlagsToClear) { BulkDataFlags &= ~FlagsToClear; }
+		bool CanLoadFromDisk() const { return !IsInlined(); }
+		bool IsAvailableForUse() const { return !(BulkDataFlags & BULKDATA_Unused); }
+		bool IsBulkDataLoaded() const { return IsInlined(); }
+		bool IsStoredCompressedOnDisk() const { return !!(BulkDataFlags & BULKDATA_SerializeCompressed); }
+		const void* LockReadOnly() const { return TexelData; }
+		void* Lock(EBulkDataLockFlags LockFlags) { return TexelData; }
+		void Unlock() const;
+		void* Realloc(int32 NumBytes);
+		void GetCopy(void** Dest, bool bDiscardInternalCopy = true);
+
+	private:
+		/** Byte offset of bulk data in file. */
+		uint32 OffsetInFile;
+		/** Size of bulk data in bytes. */
+		uint32 BulkDataSize;
+		/** Bulk data flags serialized. */
+		uint32 BulkDataFlags;
+		/** Address to texel data for inlined mips or nullptr otherwise. */
+		uint8* TexelData;
+
+		void Reset();
+		bool IsInlined() const { return !(BulkDataFlags & BULKDATA_Force_NOT_InlinePayload); }
+	};
+
 	/** Width of the mip-map. */
 	int32 SizeX;
 	/** Height of the mip-map. */
@@ -39,7 +87,7 @@ struct FTexture2DMipMap
 	/** Depth of the mip-map. */
 	int32 SizeZ;
 	/** Bulk data if stored in the package. */
-	FByteBulkData BulkData;
+	typename TChooseClass<TEXTURE2DMIPMAP_USE_COMPACT_BULKDATA, FCompactByteBulkData, FByteBulkData>::Result BulkData;
 
 	/** Default constructor. */
 	FTexture2DMipMap()
