Mark FTexture2DMipMap::FCompactByteBulkData::GetFilename as deleted to better convey intention

diff --git a/Engine/Source/Runtime/Engine/Public/TextureResource.h b/Engine/Source/Runtime/Engine/Public/TextureResource.h
--- a/Engine/Source/Runtime/Engine/Public/TextureResource.h
+++ b/Engine/Source/Runtime/Engine/Public/TextureResource.h
@@ -66,6 +66,12 @@ struct FTexture2DMipMap
 		void* Realloc(int32 NumBytes);
 		void GetCopy(void** Dest, bool bDiscardInternalCopy = true);
 
+		/**
+		 * FCompactByteBulkData doesn't support GetFilename. Use FTexturePlatformData::CachedPackageFileName
+		 * or UTexture2D::GetMipDataFilename instead
+		 */
+		const FString& GetFilename() const = delete;
+
 	private:
 		/** Byte offset of bulk data in file. */
 		uint32 OffsetInFile;
