Sound Debugging;

Adding an extra check to see if we're filtering sounds using AudioSoloSoundCue or AudioSoloSoundWave when printing for the Stat SoundCues or Stat SoundWaves cmd. This way we can filter and only see what we want.

-- Code --
UnrealEngine:
RenderStatSoundWaves() - If we have a solo sound name filter and the current wave instance doesn't contain any of the filter text, then we'll print the wave.

RenderStatSoundCues() - If we have a solo sound name filter and the current wave instance doesn't contain any of the filter text, then we'll print the sound cue.

diff --git a/Engine/Source/Runtime/Engine/Private/UnrealEngine.cpp b/Engine/Source/Runtime/Engine/Private/UnrealEngine.cpp
--- a/Engine/Source/Runtime/Engine/Private/UnrealEngine.cpp
+++ b/Engine/Source/Runtime/Engine/Private/UnrealEngine.cpp
@@ -15629,7 +15629,18 @@ int32 UEngine::RenderStatSoundWaves(UWorld* World, FViewport* Viewport, FCanvas*
 				{
 					if (WaveInstanceInfo.ActualVolume >= 0.01f)
 					{
-						WaveInstances.Emplace(&WaveInstanceInfo, &StatSoundInfo);
+						bool bShouldPrint = true;
+						FString WaveInstanceName = WaveInstanceInfo.WaveInstanceName.ToString();
+						const FString& DebugSoloSoundName = GEngine->GetAudioDeviceManager()->GetDebugSoloSoundWave();
+						if (DebugSoloSoundName != TEXT("") && !WaveInstanceName.Contains(DebugSoloSoundName))
+						{
+							bShouldPrint = false;
+						}
+
+						if (bShouldPrint == true)
+						{
+							WaveInstances.Emplace(&WaveInstanceInfo, &StatSoundInfo);
+						}
 					}
 				}
 			}
@@ -15702,10 +15713,21 @@ int32 UEngine::RenderStatSoundCues(UWorld* World, FViewport* Viewport, FCanvas*
 			{
 				if (WaveInstanceInfo.ActualVolume >= 0.01f)
 				{
-					const FString TheString = FString::Printf(TEXT("%4i. %s %s"), ActiveSoundCount++, *StatSoundInfo.SoundName, *StatSoundInfo.SoundClassName.ToString());
-					Canvas->DrawShadowedString(X, Y, *TheString, GetSmallFont(), FColor::White);
-					Y += 12;
-					break;
+					bool bShouldPrint = true;
+					FString WaveInstanceName = WaveInstanceInfo.WaveInstanceName.ToString();
+					const FString& DebugSoloSoundCue = GEngine->GetAudioDeviceManager()->GetDebugSoloSoundCue();
+					if (DebugSoloSoundCue != TEXT("") && !WaveInstanceName.Contains(DebugSoloSoundCue))
+					{
+						bShouldPrint = false;
+					}
+
+					if (bShouldPrint == true)
+					{
+						const FString TheString = FString::Printf(TEXT("%4i. %s %s"), ActiveSoundCount++, *StatSoundInfo.SoundName, *StatSoundInfo.SoundClassName.ToString());
+						Canvas->DrawShadowedString(X, Y, *TheString, GetSmallFont(), FColor::White);
+						Y += 12;
+						break;
+					}
 				}
 			}
 		}
