Actor, ActorComponent, SceneComponent, and PrimitiveComponent dynamic multicast delegates that are not frequently bound now make use of the sparse delegate system.
This reduces the base AActor by 240 bytes, UActorComponent by 32 bytes, USceneComponent by an additional 16 bytes (48 total), and UPrimitiveComponent by an additional 192 bytes (240 total).

diff --git a/Engine/Source/Runtime/Engine/Classes/Components/ActorComponent.h b/Engine/Source/Runtime/Engine/Classes/Components/ActorComponent.h
--- a/Engine/Source/Runtime/Engine/Classes/Components/ActorComponent.h
+++ b/Engine/Source/Runtime/Engine/Classes/Components/ActorComponent.h
@@ -67,9 +67,10 @@ CONSTEXPR inline EUpdateTransformFlags operator ~(EUpdateTransformFlags Value)
 /** Converts legacy bool into the SkipPhysicsUpdate bitflag */
 FORCEINLINE EUpdateTransformFlags SkipPhysicsToEnum(bool bSkipPhysics){ return bSkipPhysics ? EUpdateTransformFlags::SkipPhysicsUpdate : EUpdateTransformFlags::None; }
 
+class UActorComponent;
 
-DECLARE_DYNAMIC_MULTICAST_DELEGATE_TwoParams(FActorComponentActivatedSignature, UActorComponent*, Component, bool, bReset);
-DECLARE_DYNAMIC_MULTICAST_DELEGATE_OneParam(FActorComponentDeactivateSignature, UActorComponent*, Component);
+DECLARE_DYNAMIC_MULTICAST_SPARSE_DELEGATE_TwoParams(FActorComponentActivatedSignature, UActorComponent, OnComponentActivated, UActorComponent*, Component, bool, bReset);
+DECLARE_DYNAMIC_MULTICAST_SPARSE_DELEGATE_OneParam(FActorComponentDeactivateSignature, UActorComponent, OnComponentDeactivated, UActorComponent*, Component);
 
 DECLARE_MULTICAST_DELEGATE_OneParam(FActorComponentGlobalCreatePhysicsSignature, UActorComponent*);
 DECLARE_MULTICAST_DELEGATE_OneParam(FActorComponentGlobalDestroyPhysicsSignature, UActorComponent*);
@@ -250,10 +251,6 @@ public:
 	UPROPERTY()
 	EComponentCreationMethod CreationMethod;
 
-private:
-	UPROPERTY()
-	TArray<FSimpleMemberReference> UCSModifiedProperties;
-
 public:
 	/** Tracks whether the component has been added to one of the world's end of frame update lists */
 	uint32 GetMarkedForEndOfFrameUpdateState() const { return MarkedForEndOfFrameUpdateState; }
@@ -321,6 +318,11 @@ public:
 	UPROPERTY(BlueprintAssignable, Category = "Components|Activation")
 	FActorComponentDeactivateSignature OnComponentDeactivated;
 
+private:
+	UPROPERTY()
+	TArray<FSimpleMemberReference> UCSModifiedProperties;
+
+public:
 	/**
 	 * Activates the SceneComponent, should be overridden by native child classes.
 	 * @param bReset - Whether the activation should happen even if ShouldActivate returns false.
diff --git a/Engine/Source/Runtime/Engine/Classes/Components/PrimitiveComponent.h b/Engine/Source/Runtime/Engine/Classes/Components/PrimitiveComponent.h
--- a/Engine/Source/Runtime/Engine/Classes/Components/PrimitiveComponent.h
+++ b/Engine/Source/Runtime/Engine/Classes/Components/PrimitiveComponent.h
@@ -142,26 +142,26 @@ struct FRendererStencilMaskEvaluation
  * Delegate for notification of blocking collision against a specific component.  
  * NormalImpulse will be filled in for physics-simulating bodies, but will be zero for swept-component blocking collisions. 
  */
-DECLARE_DYNAMIC_MULTICAST_DELEGATE_FiveParams( FComponentHitSignature, UPrimitiveComponent*, HitComponent, AActor*, OtherActor, UPrimitiveComponent*, OtherComp, FVector, NormalImpulse, const FHitResult&, Hit );
+DECLARE_DYNAMIC_MULTICAST_SPARSE_DELEGATE_FiveParams( FComponentHitSignature, UPrimitiveComponent, OnComponentHit, UPrimitiveComponent*, HitComponent, AActor*, OtherActor, UPrimitiveComponent*, OtherComp, FVector, NormalImpulse, const FHitResult&, Hit );
 /** Delegate for notification of start of overlap with a specific component */
-DECLARE_DYNAMIC_MULTICAST_DELEGATE_SixParams( FComponentBeginOverlapSignature, UPrimitiveComponent*, OverlappedComponent, AActor*, OtherActor, UPrimitiveComponent*, OtherComp, int32, OtherBodyIndex, bool, bFromSweep, const FHitResult &, SweepResult);
+DECLARE_DYNAMIC_MULTICAST_SPARSE_DELEGATE_SixParams( FComponentBeginOverlapSignature, UPrimitiveComponent, OnComponentBeginOverlap, UPrimitiveComponent*, OverlappedComponent, AActor*, OtherActor, UPrimitiveComponent*, OtherComp, int32, OtherBodyIndex, bool, bFromSweep, const FHitResult &, SweepResult);
 /** Delegate for notification of end of overlap with a specific component */
-DECLARE_DYNAMIC_MULTICAST_DELEGATE_FourParams( FComponentEndOverlapSignature, UPrimitiveComponent*, OverlappedComponent, AActor*, OtherActor, UPrimitiveComponent*, OtherComp, int32, OtherBodyIndex);
+DECLARE_DYNAMIC_MULTICAST_SPARSE_DELEGATE_FourParams( FComponentEndOverlapSignature, UPrimitiveComponent, OnComponentEndOverlap, UPrimitiveComponent*, OverlappedComponent, AActor*, OtherActor, UPrimitiveComponent*, OtherComp, int32, OtherBodyIndex);
 /** Delegate for notification when a wake event is fired by physics*/
-DECLARE_DYNAMIC_MULTICAST_DELEGATE_TwoParams(FComponentWakeSignature, UPrimitiveComponent*, WakingComponent, FName, BoneName);
+DECLARE_DYNAMIC_MULTICAST_SPARSE_DELEGATE_TwoParams(FComponentWakeSignature, UPrimitiveComponent, OnComponentWake, UPrimitiveComponent*, WakingComponent, FName, BoneName);
 /** Delegate for notification when a sleep event is fired by physics*/
-DECLARE_DYNAMIC_MULTICAST_DELEGATE_TwoParams(FComponentSleepSignature, UPrimitiveComponent*, SleepingComponent, FName, BoneName);
+DECLARE_DYNAMIC_MULTICAST_SPARSE_DELEGATE_TwoParams(FComponentSleepSignature, UPrimitiveComponent, OnComponentSleep, UPrimitiveComponent*, SleepingComponent, FName, BoneName);
 /** Delegate for notification when collision settings change. */
-DECLARE_DYNAMIC_MULTICAST_DELEGATE_OneParam(FComponentCollisionSettingsChangedSignature, UPrimitiveComponent*, ChangedComponent);
-
-DECLARE_DYNAMIC_MULTICAST_DELEGATE_OneParam( FComponentBeginCursorOverSignature, UPrimitiveComponent*, TouchedComponent );
-DECLARE_DYNAMIC_MULTICAST_DELEGATE_OneParam( FComponentEndCursorOverSignature, UPrimitiveComponent*, TouchedComponent );
-DECLARE_DYNAMIC_MULTICAST_DELEGATE_TwoParams( FComponentOnClickedSignature, UPrimitiveComponent*, TouchedComponent , FKey, ButtonPressed);
-DECLARE_DYNAMIC_MULTICAST_DELEGATE_TwoParams( FComponentOnReleasedSignature, UPrimitiveComponent*, TouchedComponent, FKey, ButtonReleased);
-DECLARE_DYNAMIC_MULTICAST_DELEGATE_TwoParams( FComponentOnInputTouchBeginSignature, ETouchIndex::Type, FingerIndex, UPrimitiveComponent*, TouchedComponent );
-DECLARE_DYNAMIC_MULTICAST_DELEGATE_TwoParams( FComponentOnInputTouchEndSignature, ETouchIndex::Type, FingerIndex, UPrimitiveComponent*, TouchedComponent );
-DECLARE_DYNAMIC_MULTICAST_DELEGATE_TwoParams( FComponentBeginTouchOverSignature, ETouchIndex::Type, FingerIndex, UPrimitiveComponent*, TouchedComponent );
-DECLARE_DYNAMIC_MULTICAST_DELEGATE_TwoParams( FComponentEndTouchOverSignature, ETouchIndex::Type, FingerIndex, UPrimitiveComponent*, TouchedComponent );
+DECLARE_DYNAMIC_MULTICAST_SPARSE_DELEGATE_OneParam(FComponentCollisionSettingsChangedSignature, UPrimitiveComponent, OnComponentCollisionSettingsChangedEvent, UPrimitiveComponent*, ChangedComponent);
+
+DECLARE_DYNAMIC_MULTICAST_SPARSE_DELEGATE_OneParam( FComponentBeginCursorOverSignature, UPrimitiveComponent, OnBeginCursorOver, UPrimitiveComponent*, TouchedComponent );
+DECLARE_DYNAMIC_MULTICAST_SPARSE_DELEGATE_OneParam( FComponentEndCursorOverSignature, UPrimitiveComponent, OnEndCursorOver, UPrimitiveComponent*, TouchedComponent );
+DECLARE_DYNAMIC_MULTICAST_SPARSE_DELEGATE_TwoParams( FComponentOnClickedSignature, UPrimitiveComponent, OnClicked, UPrimitiveComponent*, TouchedComponent , FKey, ButtonPressed);
+DECLARE_DYNAMIC_MULTICAST_SPARSE_DELEGATE_TwoParams( FComponentOnReleasedSignature, UPrimitiveComponent, OnReleased, UPrimitiveComponent*, TouchedComponent, FKey, ButtonReleased);
+DECLARE_DYNAMIC_MULTICAST_SPARSE_DELEGATE_TwoParams( FComponentOnInputTouchBeginSignature, UPrimitiveComponent, OnInputTouchBegin, ETouchIndex::Type, FingerIndex, UPrimitiveComponent*, TouchedComponent );
+DECLARE_DYNAMIC_MULTICAST_SPARSE_DELEGATE_TwoParams( FComponentOnInputTouchEndSignature, UPrimitiveComponent, OnInputTouchEnd, ETouchIndex::Type, FingerIndex, UPrimitiveComponent*, TouchedComponent );
+DECLARE_DYNAMIC_MULTICAST_SPARSE_DELEGATE_TwoParams( FComponentBeginTouchOverSignature, UPrimitiveComponent, OnInputTouchEnter, ETouchIndex::Type, FingerIndex, UPrimitiveComponent*, TouchedComponent );
+DECLARE_DYNAMIC_MULTICAST_SPARSE_DELEGATE_TwoParams( FComponentEndTouchOverSignature, UPrimitiveComponent, OnInputTouchLeave, ETouchIndex::Type, FingerIndex, UPrimitiveComponent*, TouchedComponent );
 
 /**
  * PrimitiveComponents are SceneComponents that contain or generate some sort of geometry, generally to be rendered or used as collision data.
diff --git a/Engine/Source/Runtime/Engine/Classes/Components/SceneComponent.h b/Engine/Source/Runtime/Engine/Classes/Components/SceneComponent.h
--- a/Engine/Source/Runtime/Engine/Classes/Components/SceneComponent.h
+++ b/Engine/Source/Runtime/Engine/Classes/Components/SceneComponent.h
@@ -94,7 +94,7 @@ FORCEINLINE EMoveComponentFlags operator&(EMoveComponentFlags Arg1,EMoveComponen
 FORCEINLINE void operator&=(EMoveComponentFlags& Dest,EMoveComponentFlags Arg)					{ Dest = EMoveComponentFlags(Dest & Arg); }
 FORCEINLINE void operator|=(EMoveComponentFlags& Dest,EMoveComponentFlags Arg)					{ Dest = EMoveComponentFlags(Dest | Arg); }
 
-DECLARE_DYNAMIC_MULTICAST_DELEGATE_OneParam(FPhysicsVolumeChanged, class APhysicsVolume*, NewVolume);
+DECLARE_DYNAMIC_MULTICAST_SPARSE_DELEGATE_OneParam(FPhysicsVolumeChanged, USceneComponent, PhysicsVolumeChangedDelegate, class APhysicsVolume*, NewVolume);
 DECLARE_EVENT_ThreeParams(USceneComponent, FTransformUpdated, USceneComponent* /*UpdatedComponent*/, EUpdateTransformFlags /*UpdateTransformFlags*/, ETeleportType /*Teleport*/);
 
 /**
@@ -1662,4 +1662,4 @@ FORCEINLINE_DEBUGGABLE bool USceneComponent::UpdateOverlaps(TArray<FOverlapInfo>
 FORCEINLINE_DEBUGGABLE bool USceneComponent::GetShouldUpdatePhysicsVolume() const
 {
 	return bShouldUpdatePhysicsVolume;
-}
\ No newline at end of file
+}
diff --git a/Engine/Source/Runtime/Engine/Classes/GameFramework/Actor.h b/Engine/Source/Runtime/Engine/Classes/GameFramework/Actor.h
--- a/Engine/Source/Runtime/Engine/Classes/GameFramework/Actor.h
+++ b/Engine/Source/Runtime/Engine/Classes/GameFramework/Actor.h
@@ -36,24 +36,24 @@ struct FNetworkObjectInfo;
 ENGINE_API DECLARE_LOG_CATEGORY_EXTERN(LogActor, Log, Warning);
 
 // Delegate signatures
-DECLARE_DYNAMIC_MULTICAST_DELEGATE_FiveParams( FTakeAnyDamageSignature, AActor*, DamagedActor, float, Damage, const class UDamageType*, DamageType, class AController*, InstigatedBy, AActor*, DamageCauser );
-DECLARE_DYNAMIC_MULTICAST_DELEGATE_NineParams( FTakePointDamageSignature, AActor*, DamagedActor, float, Damage, class AController*, InstigatedBy, FVector, HitLocation, class UPrimitiveComponent*, FHitComponent, FName, BoneName, FVector, ShotFromDirection, const class UDamageType*, DamageType, AActor*, DamageCauser );
-DECLARE_DYNAMIC_MULTICAST_DELEGATE_SevenParams( FTakeRadialDamageSignature, AActor*, DamagedActor, float, Damage, const class UDamageType*, DamageType, FVector, Origin, FHitResult, HitInfo, class AController*, InstigatedBy, AActor*, DamageCauser );
-DECLARE_DYNAMIC_MULTICAST_DELEGATE_TwoParams( FActorBeginOverlapSignature, AActor*, OverlappedActor, AActor*, OtherActor );
-DECLARE_DYNAMIC_MULTICAST_DELEGATE_TwoParams( FActorEndOverlapSignature, AActor*, OverlappedActor, AActor*, OtherActor );
-DECLARE_DYNAMIC_MULTICAST_DELEGATE_FourParams( FActorHitSignature, AActor*, SelfActor, AActor*, OtherActor, FVector, NormalImpulse, const FHitResult&, Hit );
-
-DECLARE_DYNAMIC_MULTICAST_DELEGATE_OneParam( FActorBeginCursorOverSignature, AActor*, TouchedActor );
-DECLARE_DYNAMIC_MULTICAST_DELEGATE_OneParam( FActorEndCursorOverSignature, AActor*, TouchedActor );
-DECLARE_DYNAMIC_MULTICAST_DELEGATE_TwoParams( FActorOnClickedSignature, AActor*, TouchedActor , FKey, ButtonPressed );
-DECLARE_DYNAMIC_MULTICAST_DELEGATE_TwoParams( FActorOnReleasedSignature, AActor*, TouchedActor , FKey, ButtonReleased );
-DECLARE_DYNAMIC_MULTICAST_DELEGATE_TwoParams( FActorOnInputTouchBeginSignature, ETouchIndex::Type, FingerIndex, AActor*, TouchedActor );
-DECLARE_DYNAMIC_MULTICAST_DELEGATE_TwoParams( FActorOnInputTouchEndSignature, ETouchIndex::Type, FingerIndex, AActor*, TouchedActor );
-DECLARE_DYNAMIC_MULTICAST_DELEGATE_TwoParams( FActorBeginTouchOverSignature, ETouchIndex::Type, FingerIndex, AActor*, TouchedActor );
-DECLARE_DYNAMIC_MULTICAST_DELEGATE_TwoParams( FActorEndTouchOverSignature, ETouchIndex::Type, FingerIndex, AActor*, TouchedActor );
-
-DECLARE_DYNAMIC_MULTICAST_DELEGATE_OneParam(FActorDestroyedSignature, AActor*, DestroyedActor );
-DECLARE_DYNAMIC_MULTICAST_DELEGATE_TwoParams(FActorEndPlaySignature, AActor*, Actor , EEndPlayReason::Type, EndPlayReason);
+DECLARE_DYNAMIC_MULTICAST_SPARSE_DELEGATE_FiveParams( FTakeAnyDamageSignature, AActor, OnTakeAnyDamage, AActor*, DamagedActor, float, Damage, const class UDamageType*, DamageType, class AController*, InstigatedBy, AActor*, DamageCauser );
+DECLARE_DYNAMIC_MULTICAST_SPARSE_DELEGATE_NineParams( FTakePointDamageSignature, AActor, OnTakePointDamage, AActor*, DamagedActor, float, Damage, class AController*, InstigatedBy, FVector, HitLocation, class UPrimitiveComponent*, FHitComponent, FName, BoneName, FVector, ShotFromDirection, const class UDamageType*, DamageType, AActor*, DamageCauser );
+DECLARE_DYNAMIC_MULTICAST_SPARSE_DELEGATE_SevenParams( FTakeRadialDamageSignature, AActor, OnTakeRadialDamage, AActor*, DamagedActor, float, Damage, const class UDamageType*, DamageType, FVector, Origin, FHitResult, HitInfo, class AController*, InstigatedBy, AActor*, DamageCauser );
+DECLARE_DYNAMIC_MULTICAST_SPARSE_DELEGATE_TwoParams( FActorBeginOverlapSignature, AActor, OnActorBeginOverlap, AActor*, OverlappedActor, AActor*, OtherActor );
+DECLARE_DYNAMIC_MULTICAST_SPARSE_DELEGATE_TwoParams( FActorEndOverlapSignature, AActor, OnActorEndOverlap, AActor*, OverlappedActor, AActor*, OtherActor );
+DECLARE_DYNAMIC_MULTICAST_SPARSE_DELEGATE_FourParams( FActorHitSignature, AActor, OnActorHit, AActor*, SelfActor, AActor*, OtherActor, FVector, NormalImpulse, const FHitResult&, Hit );
+
+DECLARE_DYNAMIC_MULTICAST_SPARSE_DELEGATE_OneParam( FActorBeginCursorOverSignature, AActor, OnBeginCursorOver, AActor*, TouchedActor );
+DECLARE_DYNAMIC_MULTICAST_SPARSE_DELEGATE_OneParam( FActorEndCursorOverSignature, AActor, OnEndCursorOver, AActor*, TouchedActor );
+DECLARE_DYNAMIC_MULTICAST_SPARSE_DELEGATE_TwoParams( FActorOnClickedSignature, AActor, OnClicked, AActor*, TouchedActor , FKey, ButtonPressed );
+DECLARE_DYNAMIC_MULTICAST_SPARSE_DELEGATE_TwoParams( FActorOnReleasedSignature, AActor, OnReleased, AActor*, TouchedActor , FKey, ButtonReleased );
+DECLARE_DYNAMIC_MULTICAST_SPARSE_DELEGATE_TwoParams( FActorOnInputTouchBeginSignature, AActor, OnInputTouchBegin, ETouchIndex::Type, FingerIndex, AActor*, TouchedActor );
+DECLARE_DYNAMIC_MULTICAST_SPARSE_DELEGATE_TwoParams( FActorOnInputTouchEndSignature, AActor, OnInputTouchEnd, ETouchIndex::Type, FingerIndex, AActor*, TouchedActor );
+DECLARE_DYNAMIC_MULTICAST_SPARSE_DELEGATE_TwoParams( FActorBeginTouchOverSignature, AActor, OnInputTouchEnter, ETouchIndex::Type, FingerIndex, AActor*, TouchedActor );
+DECLARE_DYNAMIC_MULTICAST_SPARSE_DELEGATE_TwoParams( FActorEndTouchOverSignature, AActor, OnInputTouchLeave, ETouchIndex::Type, FingerIndex, AActor*, TouchedActor );
+
+DECLARE_DYNAMIC_MULTICAST_SPARSE_DELEGATE_OneParam(FActorDestroyedSignature, AActor, OnDestroyed, AActor*, DestroyedActor );
+DECLARE_DYNAMIC_MULTICAST_SPARSE_DELEGATE_TwoParams(FActorEndPlaySignature, AActor, OnEndPlay, AActor*, Actor , EEndPlayReason::Type, EndPlayReason);
 
 DECLARE_DELEGATE_SixParams(FMakeNoiseDelegate, AActor*, float /*Loudness*/, class APawn*, const FVector&, float /*MaxRange*/, FName /*Tag*/);
 
