Resolved interactions between cull distance volume, desired max draw distances and never cull flag.

diff --git a/Engine/Source/Runtime/Engine/Classes/Engine/World.h b/Engine/Source/Runtime/Engine/Classes/Engine/World.h
--- a/Engine/Source/Runtime/Engine/Classes/Engine/World.h
+++ b/Engine/Source/Runtime/Engine/Classes/Engine/World.h
@@ -2315,8 +2315,9 @@ public:
 	 * Updates cull distance volumes for a specified component or a specified actor or all actors
          * @param ComponentToUpdate If specified just that Component will be updated
 	 * @param ActorToUpdate If specified (and ComponentToUpdate is not specified), all Components owned by this Actor will be updated
+	 * @return True if the passed in actors or components were within a volume
 	 */
-	void UpdateCullDistanceVolumes(AActor* ActorToUpdate = nullptr, UPrimitiveComponent* ComponentToUpdate = nullptr);
+	bool UpdateCullDistanceVolumes(AActor* ActorToUpdate = nullptr, UPrimitiveComponent* ComponentToUpdate = nullptr);
 
 	/**
 	 * Cleans up components, streaming data and assorted other intermediate data.
diff --git a/Engine/Source/Runtime/Engine/Private/Components/PrimitiveComponent.cpp b/Engine/Source/Runtime/Engine/Private/Components/PrimitiveComponent.cpp
--- a/Engine/Source/Runtime/Engine/Private/Components/PrimitiveComponent.cpp
+++ b/Engine/Source/Runtime/Engine/Private/Components/PrimitiveComponent.cpp
@@ -819,13 +819,16 @@ void UPrimitiveComponent::PostEditChangeProperty(FPropertyChangedEvent& Property
 		const FName PropertyName = PropertyThatChanged->GetFName();
 
 		// CachedMaxDrawDistance needs to be set as if you have no cull distance volumes affecting this primitive component the cached value wouldn't get updated
-		if (PropertyName == GET_MEMBER_NAME_CHECKED(UPrimitiveComponent, LDMaxDrawDistance)
-			|| PropertyName == GET_MEMBER_NAME_CHECKED(UPrimitiveComponent, bAllowCullDistanceVolume)
-			|| PropertyName == GET_MEMBER_NAME_CHECKED(UPrimitiveComponent, bNeverDistanceCull))
+		if (PropertyName == GET_MEMBER_NAME_CHECKED(UPrimitiveComponent, LDMaxDrawDistance))
 		{
 			NewCachedMaxDrawDistance = LDMaxDrawDistance;
 			bCullDistanceInvalidated = true;
 		}
+		else if (PropertyName == GET_MEMBER_NAME_CHECKED(UPrimitiveComponent, bAllowCullDistanceVolume)
+			|| PropertyName == GET_MEMBER_NAME_CHECKED(UPrimitiveComponent, bNeverDistanceCull))
+		{
+			bCullDistanceInvalidated = true;
+		}
 
 		// we need to reregister the primitive if the min draw distance changed to propagate the change to the rendering thread
 		if (PropertyName == GET_MEMBER_NAME_CHECKED(UPrimitiveComponent, MinDrawDistance))
@@ -843,11 +846,6 @@ void UPrimitiveComponent::PostEditChangeProperty(FPropertyChangedEvent& Property
 
 	if (bCullDistanceInvalidated)
 	{
-		// Make sure cached cull distance is up-to-date.
-		if (LDMaxDrawDistance > 0)
-		{
-			NewCachedMaxDrawDistance = FMath::Min(LDMaxDrawDistance, CachedMaxDrawDistance);
-		}
 		// Directly use LD cull distance if cull distance volumes are disabled.
 		if (!bAllowCullDistanceVolume)
 		{
@@ -855,7 +853,22 @@ void UPrimitiveComponent::PostEditChangeProperty(FPropertyChangedEvent& Property
 		}
 		else if (UWorld* World = GetWorld())
 		{
-			World->UpdateCullDistanceVolumes(nullptr, this);
+			const bool bUpdatedDrawDistances = World->UpdateCullDistanceVolumes(nullptr, this);
+
+			// If volumes invalidated the distance, handle the desired distance against other sources
+			if (bUpdatedDrawDistances)
+			{
+				if (LDMaxDrawDistance <= 0)
+				{
+					// Volume is the only controlling source, use directly
+					NewCachedMaxDrawDistance = CachedMaxDrawDistance;
+				}
+				else
+				{
+					// Select the minimum desired value
+					NewCachedMaxDrawDistance = FMath::Min(CachedMaxDrawDistance, LDMaxDrawDistance);
+				}
+			}
 		}
 
 		// Reattach to propagate cull distance change.
diff --git a/Engine/Source/Runtime/Engine/Private/World.cpp b/Engine/Source/Runtime/Engine/Private/World.cpp
--- a/Engine/Source/Runtime/Engine/Private/World.cpp
+++ b/Engine/Source/Runtime/Engine/Private/World.cpp
@@ -1737,10 +1737,11 @@ void UWorld::UpdateWorldComponents(bool bRerunConstructionScripts, bool bCurrent
 }
 
 
-void UWorld::UpdateCullDistanceVolumes(AActor* ActorToUpdate, UPrimitiveComponent* ComponentToUpdate)
+bool UWorld::UpdateCullDistanceVolumes(AActor* ActorToUpdate, UPrimitiveComponent* ComponentToUpdate)
 {
 	// Map that will store new max draw distance for every primitive
 	TMap<UPrimitiveComponent*,float> CompToNewMaxDrawMap;
+	bool bUpdatedDrawDistances = false;
 
 	// Keep track of time spent.
 	double Duration = 0.0;
@@ -1811,6 +1812,8 @@ void UWorld::UpdateCullDistanceVolumes(AActor* ActorToUpdate, UPrimitiveComponen
 				{
 					CullDistanceVolume->GetPrimitiveMaxDrawDistances(CompToNewMaxDrawMap);
 				}
+
+				bUpdatedDrawDistances = true;
 			}
 
 			// Finally, go over all primitives, and see if they need to change.
@@ -1829,6 +1832,8 @@ void UWorld::UpdateCullDistanceVolumes(AActor* ActorToUpdate, UPrimitiveComponen
 	{
 		UE_LOG(LogWorld, Log, TEXT("Updating cull distance volumes took %5.2f seconds"),Duration);
 	}
+
+	return bUpdatedDrawDistances;
 }
 
 
