Fixed lightmap cluster initialization. During cooked build streaming cache mesh draw commands could pickup an initialized lightmap cluster uniform buffer, as first we call AddPrimitive (which caches mesh draw commands) and only after that we call InitializeClusterRenderingResources (which initializes lightmap cluster UB). The fix is to first create a dummy uniform buffer, so it can be cached inside cached mesh draw commands. Then after all lightmap resources are ready this uniform buffer is updated with proper values.

diff --git a/Engine/Source/Runtime/Engine/Private/LightMap.cpp b/Engine/Source/Runtime/Engine/Private/LightMap.cpp
--- a/Engine/Source/Runtime/Engine/Private/LightMap.cpp
+++ b/Engine/Source/Runtime/Engine/Private/LightMap.cpp
@@ -2630,7 +2630,7 @@ bool FQuantizedLightmapData::HasNonZeroData() const
 	return false;
 }
 
-void FLightmapResourceCluster::SetFeatureLevel(ERHIFeatureLevel::Type InFeatureLevel)
+void FLightmapResourceCluster::UpdateUniformBuffer(ERHIFeatureLevel::Type InFeatureLevel)
 {
 	FLightmapResourceCluster* Cluster = this;
 
@@ -2638,14 +2638,18 @@ void FLightmapResourceCluster::SetFeatureLevel(ERHIFeatureLevel::Type InFeatureL
 		[Cluster, InFeatureLevel](FRHICommandList& RHICmdList)
 	{
 		Cluster->FeatureLevel = InFeatureLevel;
+
+		FLightmapResourceClusterShaderParameters Parameters;
+		GetLightmapClusterResourceParameters(InFeatureLevel, Cluster->Input, Parameters);
+
+		RHIUpdateUniformBuffer(Cluster->UniformBuffer, &Parameters);
 	});
 }
 
 void FLightmapResourceCluster::InitRHI()
 {
 	FLightmapResourceClusterShaderParameters Parameters;
-	check(FeatureLevel != ERHIFeatureLevel::Num);
-	GetLightmapClusterResourceParameters(FeatureLevel, Input, Parameters);
+	GetLightmapClusterResourceParameters(GMaxRHIFeatureLevel, FLightmapClusterResourceInput(), Parameters);
 
 	UniformBuffer = FLightmapResourceClusterShaderParameters::CreateUniformBuffer(Parameters, UniformBuffer_MultiFrame);
 }
\ No newline at end of file
diff --git a/Engine/Source/Runtime/Engine/Private/MapBuildData.cpp b/Engine/Source/Runtime/Engine/Private/MapBuildData.cpp
--- a/Engine/Source/Runtime/Engine/Private/MapBuildData.cpp
+++ b/Engine/Source/Runtime/Engine/Private/MapBuildData.cpp
@@ -784,6 +784,13 @@ void UMapBuildDataRegistry::SetupLightmapResourceClusters()
 			LightmapResourceClusters[ClusterIndex].Input = ClusterInput;
 			Data.ResourceCluster = &LightmapResourceClusters[ClusterIndex];
 		}
+
+		// Init empty cluster uniform buffers so they can be referenced by cached mesh draw commands.
+		// Can't create final uniform buffers as feature level is unknown at this point.
+		for (FLightmapResourceCluster& Cluster : LightmapResourceClusters)
+		{
+			BeginInitResource(&Cluster);
+		}
 	}
 }
 
@@ -801,10 +808,10 @@ void UMapBuildDataRegistry::InitializeClusterRenderingResources(ERHIFeatureLevel
 	// If we have any mesh build data, we must have at least one resource cluster, otherwise clusters have not been setup properly.
 	check(LightmapResourceClusters.Num() > 0 || MeshBuildData.Num() == 0);
 
+	// At this point all lightmap cluster resources are initialized and we can update cluster uniform buffers.
 	for (FLightmapResourceCluster& Cluster : LightmapResourceClusters)
 	{
-		Cluster.SetFeatureLevel(InFeatureLevel);
-		BeginInitResource(&Cluster);
+		Cluster.UpdateUniformBuffer(InFeatureLevel);
 	}
 }
 
diff --git a/Engine/Source/Runtime/Engine/Public/LightMap.h b/Engine/Source/Runtime/Engine/Public/LightMap.h
--- a/Engine/Source/Runtime/Engine/Public/LightMap.h
+++ b/Engine/Source/Runtime/Engine/Public/LightMap.h
@@ -613,8 +613,6 @@ class FLightmapResourceCluster : public FRenderResource
 {
 public:
 
-	ENGINE_API void SetFeatureLevel(ERHIFeatureLevel::Type InFeatureLevel);
-
 	ENGINE_API virtual void InitRHI();
 
 	virtual void ReleaseRHI()
@@ -622,6 +620,8 @@ public:
 		UniformBuffer = nullptr;
 	}
 
+	ENGINE_API void UpdateUniformBuffer(ERHIFeatureLevel::Type InFeatureLevel);
+
 	FLightmapClusterResourceInput Input;
 
 	FUniformBufferRHIRef UniformBuffer;
diff --git a/Engine/Source/Runtime/Renderer/Private/LightMapRendering.cpp b/Engine/Source/Runtime/Renderer/Private/LightMapRendering.cpp
--- a/Engine/Source/Runtime/Renderer/Private/LightMapRendering.cpp
+++ b/Engine/Source/Runtime/Renderer/Private/LightMapRendering.cpp
@@ -42,7 +42,12 @@ void SetupLCIUniformBuffers(const FPrimitiveSceneProxy* PrimitiveSceneProxy, con
 	if (LCI)
 	{
 		PrecomputedLightingBuffer = LCI->GetPrecomputedLightingBuffer();
-		LightmapResourceClusterBuffer = LCI->GetResourceCluster() ? LCI->GetResourceCluster()->UniformBuffer : nullptr;
+	}
+
+	if (LCI && LCI->GetResourceCluster())
+	{
+		check(LCI->GetResourceCluster()->UniformBuffer);
+		LightmapResourceClusterBuffer = LCI->GetResourceCluster()->UniformBuffer;
 	}
 
 	if (!PrecomputedLightingBuffer)
