[NO NON-PII COMMIT MESSAGE]

diff --git a/Engine/Source/Runtime/Engine/Private/UnrealEngine.cpp b/Engine/Source/Runtime/Engine/Private/UnrealEngine.cpp
--- a/Engine/Source/Runtime/Engine/Private/UnrealEngine.cpp
+++ b/Engine/Source/Runtime/Engine/Private/UnrealEngine.cpp
@@ -8703,7 +8703,6 @@ FGuid UEngine::GetPackageGuid(FName PackageName, bool bForPIE)
 	UPackage* PackageToReset = nullptr;
 	FLinkerLoad* Linker = LoadPackageLinker(nullptr, *PackageName.ToString(), LoadFlags, nullptr, nullptr, nullptr, [&PackageToReset, &Result](FLinkerLoad* InLinker)
 	{
-		check(InLinker);
 		if (InLinker != nullptr && InLinker->LinkerRoot != nullptr)
 		{
 			Result = InLinker->LinkerRoot->GetGuid();
@@ -8711,8 +8710,10 @@ FGuid UEngine::GetPackageGuid(FName PackageName, bool bForPIE)
 		}
 	});
 
-	ResetLoaders(PackageToReset);
-	Linker = nullptr;
+	if (PackageToReset)
+	{
+		ResetLoaders(PackageToReset);
+	}
 
 	return Result;
 }
