Added a new bool flag array as part of the FAnimExtractContext to be optionally used for disabling retrieval of bone transforms. The PoseDriver can then fill this array based on the bones it is interested in.

diff --git a/Engine/Source/Runtime/AnimGraphRuntime/Private/AnimNodes/AnimNode_PoseDriver.cpp b/Engine/Source/Runtime/AnimGraphRuntime/Private/AnimNodes/AnimNode_PoseDriver.cpp
--- a/Engine/Source/Runtime/AnimGraphRuntime/Private/AnimNodes/AnimNode_PoseDriver.cpp
+++ b/Engine/Source/Runtime/AnimGraphRuntime/Private/AnimNodes/AnimNode_PoseDriver.cpp
@@ -124,6 +124,12 @@ void FAnimNode_PoseDriver::CacheBones_AnyThread(const FAnimationCacheBonesContex
 			}
 		}
 	}
+
+	PoseExtractContext.BonesRequired.SetNumZeroed(BoneBlendWeights.Num());
+	for (int32 BoneIndex = 0; BoneIndex < BoneBlendWeights.Num(); BoneIndex++)
+	{
+		PoseExtractContext.BonesRequired[BoneIndex] = BoneBlendWeights[BoneIndex] > SMALL_NUMBER;
+	}
 }
 
 void FAnimNode_PoseDriver::UpdateAssetPlayer(const FAnimationUpdateContext& Context)
diff --git a/Engine/Source/Runtime/Engine/Classes/Animation/AnimationAsset.h b/Engine/Source/Runtime/Engine/Classes/Animation/AnimationAsset.h
--- a/Engine/Source/Runtime/Engine/Classes/Animation/AnimationAsset.h
+++ b/Engine/Source/Runtime/Engine/Classes/Animation/AnimationAsset.h
@@ -186,6 +186,13 @@ struct FAnimExtractContext
 	 */
 	TArray<FPoseCurve> PoseCurves;
 
+	/**
+	 * The BonesRequired array is a list of bool flags to determine
+	 * if a bone is required to be retrieved. This is currently used
+	 * by several animation nodes to optimize evaluation time.
+	 */
+	TArray<bool> BonesRequired;
+
 	FAnimExtractContext()
 		: bExtractRootMotion(false)
 		, CurrentTime(0.f)
@@ -204,6 +211,14 @@ struct FAnimExtractContext
 	{
 	}
 
+	bool IsBoneRequired(int32 BoneIndex) const
+	{
+		if (BoneIndex >= BonesRequired.Num())
+		{
+			return true;
+		}
+		return BonesRequired[BoneIndex];
+	}
 };
 
 //Represent a current play position in an animation
diff --git a/Engine/Source/Runtime/Engine/Private/Animation/PoseAsset.cpp b/Engine/Source/Runtime/Engine/Private/Animation/PoseAsset.cpp
--- a/Engine/Source/Runtime/Engine/Private/Animation/PoseAsset.cpp
+++ b/Engine/Source/Runtime/Engine/Private/Animation/PoseAsset.cpp
@@ -607,6 +607,11 @@ bool UPoseAsset::GetAnimationPose(struct FCompactPose& OutPose, FBlendedCurve& O
 				const FCompactPoseBoneIndex CompactIndex = BoneIndices[TrackIndex].CompactBoneIndex;
 				if (CompactIndex != INDEX_NONE)
 				{
+					if (!ExtractionContext.IsBoneRequired(CompactIndex.GetInt()))
+					{
+						continue;
+					}
+
 					TArray<FTransform> BlendingTransform;
 					TArray<float> BlendingWeights;
 					float TotalLocalWeight = 0.f;
@@ -685,6 +690,10 @@ bool UPoseAsset::GetAnimationPose(struct FCompactPose& OutPose, FBlendedCurve& O
 				const FBoneIndices& LocalBoneIndices = BoneIndices[TrackIndex];
 				if (LocalBoneIndices.CompactBoneIndex != INDEX_NONE)
 				{
+					if (!ExtractionContext.IsBoneRequired(LocalBoneIndices.CompactBoneIndex.GetInt()))
+					{
+						continue;
+					}
 					FAnimationRuntime::RetargetBoneTransform(MySkeleton, RetargetSource, BlendedBoneTransform[TrackIndex], LocalBoneIndices.SkeletonBoneIndex, LocalBoneIndices.CompactBoneIndex, RequiredBones, bAdditivePose);
 					OutPose[LocalBoneIndices.CompactBoneIndex] = BlendedBoneTransform[TrackIndex];
 					OutPose.NormalizeRotations();
