#ue4 - Add option to InterpMovementComponent to disable swept collision movement and to set teleport type.

diff --git a/Engine/Source/Runtime/Engine/Classes/Components/InterpToMovementComponent.h b/Engine/Source/Runtime/Engine/Classes/Components/InterpToMovementComponent.h
--- a/Engine/Source/Runtime/Engine/Classes/Components/InterpToMovementComponent.h
+++ b/Engine/Source/Runtime/Engine/Classes/Components/InterpToMovementComponent.h
@@ -76,6 +76,14 @@ class ENGINE_API UInterpToMovementComponent : public UMovementComponent
 	UPROPERTY(EditAnywhere, BlueprintReadWrite, Category=Behaviour)
 	uint32 bPauseOnImpact:1;
 
+	/** If true, will sweep for blocking collision during movement. If false, it will simply teleport to the next position and ignore collision. */
+	UPROPERTY(EditAnywhere, BlueprintReadWrite, Category = Behaviour)
+	bool bSweep = true;
+
+	/** Physics teleport type. */
+	UPROPERTY(EditAnywhere, BlueprintReadWrite, Category = Behaviour)
+	ETeleportType TeleportType = ETeleportType::None;
+
 	/* Movement behaviour of the component */
 	UPROPERTY(EditAnywhere, BlueprintReadWrite, Category=Behaviour)
 	EInterpToBehaviourType BehaviourType;
diff --git a/Engine/Source/Runtime/Engine/Private/Components/InterpToMovementComponent.cpp b/Engine/Source/Runtime/Engine/Private/Components/InterpToMovementComponent.cpp
--- a/Engine/Source/Runtime/Engine/Private/Components/InterpToMovementComponent.cpp
+++ b/Engine/Source/Runtime/Engine/Private/Components/InterpToMovementComponent.cpp
@@ -106,13 +106,13 @@ void UInterpToMovementComponent::TickComponent(float DeltaTime, enum ELevelTick
  		if ((bPauseOnImpact == false ) && (BehaviourType != EInterpToBehaviourType::OneShot))
  		{
  			// If we can bounce, we are allowed to move out of penetrations, so use SafeMoveUpdatedComponent which does that automatically.
- 			SafeMoveUpdatedComponent(MoveDelta, CurrentRotation, true, Hit);
+ 			SafeMoveUpdatedComponent(MoveDelta, CurrentRotation, bSweep, Hit, TeleportType);
  		}
  		else
 		{
 			// If we can't bounce, then we shouldn't adjust if initially penetrating, because that should be a blocking hit that causes a hit event and stop simulation.
 			TGuardValue<EMoveComponentFlags> ScopedFlagRestore(MoveComponentFlags, MoveComponentFlags | MOVECOMP_NeverIgnoreBlockingOverlaps);
-			MoveUpdatedComponent(MoveDelta, CurrentRotation, true, &Hit);
+			MoveUpdatedComponent(MoveDelta, CurrentRotation, bSweep, &Hit, TeleportType);
 		}
 		//DrawDebugPoint(GetWorld(), UpdatedComponent->GetComponentLocation(), 16, FColor::White,true,5.0f);
 		// If we hit a trigger that destroyed us, abort.
