Possible fix for hang in PGO

diff --git a/Engine/Source/Runtime/Engine/Private/AudioStreaming.cpp b/Engine/Source/Runtime/Engine/Private/AudioStreaming.cpp
--- a/Engine/Source/Runtime/Engine/Private/AudioStreaming.cpp
+++ b/Engine/Source/Runtime/Engine/Private/AudioStreaming.cpp
@@ -877,11 +877,10 @@ const uint8* FAudioStreamingManager::GetLoadedChunk(const USoundWave* SoundWave,
 		return nullptr;
 	}
 
-	// Get a critical section lock, but don't block others. This avoids rare conditions of deadlocking in this streaming manager with UpdateStreamingResources.
-	// But allows us to make sure we can get the requested loaded chunk. Most of the time, this function is called in an async task.
-	while (!CriticalSection.TryLock())
+	// If we fail at getting the critical section here, early out. 
+	if (!CriticalSection.TryLock())
 	{
-		FPlatformProcess::Sleep(0.001f);
+		return nullptr;
 	}
 
 	const FStreamingWaveData* WaveData = StreamingSoundWaves.FindRef(SoundWave);
