Material Parameter Collections now correctly recreate their uniform buffers when needed, which can only be done when newly loaded (game and editor), or inside a FGlobalComponentRecreateRenderStateContext (editor layout changes).  This change reverts cl 5175660 which incorrectly recreated the uniform buffer outside of a FGlobalComponentRecreateRenderStateContext, causing existing cached mesh draw commands to have a dangling pointer.
When a Material Parameter Collection layout changes in the editor, all materials referencing the parameter collection are now recompiled.  This reverts the behavior added by cl 2612044 which incorrectly only recompiled materials referencing a parameter which was changed.

diff --git a/Engine/Source/Runtime/Engine/Classes/Engine/World.h b/Engine/Source/Runtime/Engine/Classes/Engine/World.h
--- a/Engine/Source/Runtime/Engine/Classes/Engine/World.h
+++ b/Engine/Source/Runtime/Engine/Classes/Engine/World.h
@@ -2423,7 +2423,7 @@ public:
 	UMaterialParameterCollectionInstance* GetParameterCollectionInstance(const UMaterialParameterCollection* Collection);
 
 	/** Updates this world's scene with the list of instances, and optionally updates each instance's uniform buffer. */
-	void UpdateParameterCollectionInstances(bool bUpdateInstanceUniformBuffers);
+	void UpdateParameterCollectionInstances(bool bUpdateInstanceUniformBuffers, bool bRecreateUniformBuffer);
 
 	/** Gets the canvas object for rendering to a render target.  Will allocate one if needed. */
 	UCanvas* GetCanvasForRenderingToTarget();
diff --git a/Engine/Source/Runtime/Engine/Classes/Materials/MaterialParameterCollection.h b/Engine/Source/Runtime/Engine/Classes/Materials/MaterialParameterCollection.h
--- a/Engine/Source/Runtime/Engine/Classes/Materials/MaterialParameterCollection.h
+++ b/Engine/Source/Runtime/Engine/Classes/Materials/MaterialParameterCollection.h
@@ -137,7 +137,7 @@ private:
 	/** Gets default values into data to be set on the uniform buffer. */
 	void GetDefaultParameterData(TArray<FVector4>& ParameterData) const;
 
-	void UpdateDefaultResource();
+	void UpdateDefaultResource(bool bRecreateUniformBuffer);
 };
 
 
diff --git a/Engine/Source/Runtime/Engine/Classes/Materials/MaterialParameterCollectionInstance.h b/Engine/Source/Runtime/Engine/Classes/Materials/MaterialParameterCollectionInstance.h
--- a/Engine/Source/Runtime/Engine/Classes/Materials/MaterialParameterCollectionInstance.h
+++ b/Engine/Source/Runtime/Engine/Classes/Materials/MaterialParameterCollectionInstance.h
@@ -49,9 +49,9 @@ class ENGINE_API UMaterialParameterCollectionInstance : public UObject
 		return Collection;
 	}
 
-	void UpdateRenderState();
+	void UpdateRenderState(bool bRecreateUniformBuffer);
 
-	void DeferredUpdateRenderState(bool bUpdateScene = false);
+	void DeferredUpdateRenderState(bool bRecreateUniformBuffer);
 
 	/** Tracks whether this instance has ever issued a missing parameter warning, to reduce log spam. */
 	bool bLoggedMissingParameterWarning;
diff --git a/Engine/Source/Runtime/Engine/Private/LevelTick.cpp b/Engine/Source/Runtime/Engine/Private/LevelTick.cpp
--- a/Engine/Source/Runtime/Engine/Private/LevelTick.cpp
+++ b/Engine/Source/Runtime/Engine/Private/LevelTick.cpp
@@ -1051,7 +1051,7 @@ void UWorld::SendAllEndOfFrameUpdates()
 	{
 		if (ParameterCollectionInstance)
 		{
-			ParameterCollectionInstance->DeferredUpdateRenderState();
+			ParameterCollectionInstance->DeferredUpdateRenderState(false);
 		}
 	}
 	bMaterialParameterCollectionInstanceNeedsDeferredUpdate = false;
diff --git a/Engine/Source/Runtime/Engine/Private/Materials/ParameterCollection.cpp b/Engine/Source/Runtime/Engine/Private/Materials/ParameterCollection.cpp
--- a/Engine/Source/Runtime/Engine/Private/Materials/ParameterCollection.cpp
+++ b/Engine/Source/Runtime/Engine/Private/Materials/ParameterCollection.cpp
@@ -56,7 +56,7 @@ void UMaterialParameterCollection::PostLoad()
 		CurrentWorld->AddParameterCollectionInstance(this, true);
 	}
 
-	UpdateDefaultResource();
+	UpdateDefaultResource(true);
 }
 
 void UMaterialParameterCollection::BeginDestroy()
@@ -177,18 +177,12 @@ void UMaterialParameterCollection::PostEditChangeProperty(FPropertyChangedEvent&
 		}
 
 		// Generate a new Id so that unloaded materials that reference this collection will update correctly on load
+		// Now that we changed the guid, we must recompile all materials which reference this collection
 		StateId = FGuid::NewGuid();
 
 		// Update the uniform buffer layout
 		CreateBufferStruct();
 
-		// Recreate each instance of this collection
-		for (TObjectIterator<UWorld> It; It; ++It)
-		{
-			UWorld* CurrentWorld = *It;
-			CurrentWorld->AddParameterCollectionInstance(this, false);
-		}
-
 		// Create a material update context so we can safely update materials using this parameter collection.
 		{
 			FMaterialUpdateContext UpdateContext;
@@ -228,17 +222,28 @@ void UMaterialParameterCollection::PostEditChangeProperty(FPropertyChangedEvent&
 					CurrentMaterial->MarkPackageDirty();
 				}
 			}
+
+			// Recreate all uniform buffers based off of this collection
+			for (TObjectIterator<UWorld> It; It; ++It)
+			{
+				UWorld* CurrentWorld = *It;
+				CurrentWorld->UpdateParameterCollectionInstances(true, true);
+			}
+
+			UpdateDefaultResource(true);
 		}
 	}
-
-	// Update each world's scene with the new instance, and update each instance's uniform buffer to reflect the changes made by the user
-	for (TObjectIterator<UWorld> It; It; ++It)
+	else
 	{
-		UWorld* CurrentWorld = *It;
-		CurrentWorld->UpdateParameterCollectionInstances(true);
-	}
+		// We didn't need to recreate the uniform buffer, just update its contents
+		for (TObjectIterator<UWorld> It; It; ++It)
+		{
+			UWorld* CurrentWorld = *It;
+			CurrentWorld->UpdateParameterCollectionInstances(true, false);
+		}
 
-	UpdateDefaultResource();
+		UpdateDefaultResource(false);
+	}
 
 	Super::PostEditChangeProperty(PropertyChangedEvent);
 }
@@ -432,12 +437,12 @@ void UMaterialParameterCollection::GetDefaultParameterData(TArray<FVector4>& Par
 	}
 }
 
-void UMaterialParameterCollection::UpdateDefaultResource()
+void UMaterialParameterCollection::UpdateDefaultResource(bool bRecreateUniformBuffer)
 {
 	// Propagate the new values to the rendering thread
 	TArray<FVector4> ParameterData;
 	GetDefaultParameterData(ParameterData);
-	DefaultResource->GameThread_UpdateContents(StateId, ParameterData, GetFName());
+	DefaultResource->GameThread_UpdateContents(StateId, ParameterData, GetFName(), bRecreateUniformBuffer);
 
 	FGuid Id = StateId;
 	FMaterialParameterCollectionInstanceResource* Resource = DefaultResource;
@@ -470,8 +475,6 @@ void UMaterialParameterCollectionInstance::SetCollection(UMaterialParameterColle
 {
 	Collection = InCollection;
 	World = InWorld;
-
-	UpdateRenderState();
 }
 
 bool UMaterialParameterCollectionInstance::SetScalarParameterValue(FName ParameterName, float ParameterValue)
@@ -498,8 +501,7 @@ bool UMaterialParameterCollectionInstance::SetScalarParameterValue(FName Paramet
 
 		if (bUpdateUniformBuffer)
 		{
-			//@todo - only update uniform buffers max once per frame
-			UpdateRenderState();
+			UpdateRenderState(false);
 		}
 
 		return true;
@@ -532,8 +534,7 @@ bool UMaterialParameterCollectionInstance::SetVectorParameterValue(FName Paramet
 
 		if (bUpdateUniformBuffer)
 		{
-			//@todo - only update uniform buffers max once per frame
-			UpdateRenderState();
+			UpdateRenderState(false);
 		}
 
 		return true;
@@ -570,7 +571,7 @@ bool UMaterialParameterCollectionInstance::GetVectorParameterValue(FName Paramet
 	return false;
 }
 
-void UMaterialParameterCollectionInstance::UpdateRenderState()
+void UMaterialParameterCollectionInstance::UpdateRenderState(bool bRecreateUniformBuffer)
 {
 	// Don't need material parameters on the server
 	if (!World.IsValid() || World->GetNetMode() == NM_DedicatedServer)
@@ -581,26 +582,22 @@ void UMaterialParameterCollectionInstance::UpdateRenderState()
 	bNeedsRenderStateUpdate = true;
 	World->SetMaterialParameterCollectionInstanceNeedsUpdate();
 
-	if (!GDeferUpdateRenderStates)
+	if (!GDeferUpdateRenderStates || bRecreateUniformBuffer)
 	{
-		DeferredUpdateRenderState();
+		DeferredUpdateRenderState(bRecreateUniformBuffer);
 	}
 }
 
-void UMaterialParameterCollectionInstance::DeferredUpdateRenderState(bool bUpdateScene)
+void UMaterialParameterCollectionInstance::DeferredUpdateRenderState(bool bRecreateUniformBuffer)
 {
+	checkf(bNeedsRenderStateUpdate || !bRecreateUniformBuffer, TEXT("DeferredUpdateRenderState was told to recreate the uniform buffer, but there's nothing to update"));
+
 	if (bNeedsRenderStateUpdate && World.IsValid())
 	{
 		// Propagate the new values to the rendering thread
 		TArray<FVector4> ParameterData;
 		GetParameterData(ParameterData);
-		Resource->GameThread_UpdateContents(Collection ? Collection->StateId : FGuid(), ParameterData, GetFName());
-
-		if (bUpdateScene)
-		{
-			// Update the world's scene with the new uniform buffer pointer
-			World->UpdateParameterCollectionInstances(false);
-		}
+		Resource->GameThread_UpdateContents(Collection ? Collection->StateId : FGuid(), ParameterData, GetFName(), bRecreateUniformBuffer);
 	}
 
 	bNeedsRenderStateUpdate = false;
@@ -650,13 +647,13 @@ void UMaterialParameterCollectionInstance::FinishDestroy()
 	Super::FinishDestroy();
 }
 
-void FMaterialParameterCollectionInstanceResource::GameThread_UpdateContents(const FGuid& InGuid, const TArray<FVector4>& Data, const FName& InOwnerName)
+void FMaterialParameterCollectionInstanceResource::GameThread_UpdateContents(const FGuid& InGuid, const TArray<FVector4>& Data, const FName& InOwnerName, bool bRecreateUniformBuffer)
 {
 	FMaterialParameterCollectionInstanceResource* Resource = this;
 	ENQUEUE_RENDER_COMMAND(UpdateCollectionCommand)(
-		[InGuid, Data, InOwnerName, Resource](FRHICommandListImmediate& RHICmdList)
+		[InGuid, Data, InOwnerName, Resource, bRecreateUniformBuffer](FRHICommandListImmediate& RHICmdList)
 		{
-			Resource->UpdateContents(InGuid, Data, InOwnerName);
+			Resource->UpdateContents(InGuid, Data, InOwnerName, bRecreateUniformBuffer);
 		}
 	);
 }
@@ -684,26 +681,26 @@ FMaterialParameterCollectionInstanceResource::~FMaterialParameterCollectionInsta
 	UniformBuffer.SafeRelease();
 }
 
-void FMaterialParameterCollectionInstanceResource::UpdateContents(const FGuid& InId, const TArray<FVector4>& Data, const FName& InOwnerName)
+void FMaterialParameterCollectionInstanceResource::UpdateContents(const FGuid& InId, const TArray<FVector4>& Data, const FName& InOwnerName, bool bRecreateUniformBuffer)
 {
 	Id = InId;
 	OwnerName = InOwnerName;
 
 	if (InId != FGuid() && Data.Num() > 0)
 	{
-		uint32 Size = Data.GetTypeSize() * Data.Num();
-		bool bSizeChanged = Size != UniformBufferLayout.ConstantBufferSize;
-		UniformBufferLayout.ConstantBufferSize = Data.GetTypeSize() * Data.Num();
-		UniformBufferLayout.ComputeHash();
+		const uint32 NewSize = Data.GetTypeSize() * Data.Num();
 		check(UniformBufferLayout.Resources.Num() == 0);
 
-		if (!bSizeChanged && IsValidRef(UniformBuffer))
+		if (!bRecreateUniformBuffer && IsValidRef(UniformBuffer))
 		{
+			check(NewSize == UniformBufferLayout.ConstantBufferSize);
 			check(UniformBuffer->GetLayout() == UniformBufferLayout);
 			RHIUpdateUniformBuffer(UniformBuffer, Data.GetData());
 		}
 		else
 		{
+			UniformBufferLayout.ConstantBufferSize = NewSize;
+			UniformBufferLayout.ComputeHash();
 			UniformBuffer = RHICreateUniformBuffer(Data.GetData(), UniformBufferLayout, UniformBuffer_MultiFrame);
 		}
 	}
diff --git a/Engine/Source/Runtime/Engine/Private/World.cpp b/Engine/Source/Runtime/Engine/Private/World.cpp
--- a/Engine/Source/Runtime/Engine/Private/World.cpp
+++ b/Engine/Source/Runtime/Engine/Private/World.cpp
@@ -1039,7 +1039,7 @@ void UWorld::SetupParameterCollectionInstances()
 		AddParameterCollectionInstance(CurrentCollection, false);
 	}
 
-	UpdateParameterCollectionInstances(false);
+	UpdateParameterCollectionInstances(false, false);
 }
 
 void UWorld::AddParameterCollectionInstance(UMaterialParameterCollection* Collection, bool bUpdateScene)
@@ -1071,14 +1071,14 @@ void UWorld::AddParameterCollectionInstance(UMaterialParameterCollection* Collec
 
 	// Ensure the new instance creates initial render thread resources
 	// This needs to happen right away, so they can be picked up by any cached shader bindings
-	NewInstance->DeferredUpdateRenderState(false);
+	NewInstance->UpdateRenderState(true);
 
 	if (bUpdateScene)
 	{
 		// Update the scene's list of instances, needs to happen to prevent a race condition with GC 
 		// (rendering thread still uses the FMaterialParameterCollectionInstanceResource when GC deletes the UMaterialParameterCollectionInstance)
 		// However, if UpdateParameterCollectionInstances is going to be called after many AddParameterCollectionInstance's, this can be skipped for now.
-		UpdateParameterCollectionInstances(false);
+		UpdateParameterCollectionInstances(false, false);
 	}
 }
 
@@ -1097,7 +1097,7 @@ UMaterialParameterCollectionInstance* UWorld::GetParameterCollectionInstance(con
 	return NULL;
 }
 
-void UWorld::UpdateParameterCollectionInstances(bool bUpdateInstanceUniformBuffers)
+void UWorld::UpdateParameterCollectionInstances(bool bUpdateInstanceUniformBuffers, bool bRecreateUniformBuffer)
 {
 	if (Scene)
 	{
@@ -1109,7 +1109,11 @@ void UWorld::UpdateParameterCollectionInstances(bool bUpdateInstanceUniformBuffe
 
 			if (bUpdateInstanceUniformBuffers)
 			{
-				Instance->UpdateRenderState();
+				Instance->UpdateRenderState(bRecreateUniformBuffer);
+			}
+			else
+			{
+				checkf(!bRecreateUniformBuffer, TEXT("Recreate Uniform Buffer was requested but cannot be executed because bUpdateInstanceUniformBuffers was false"));
 			}
 
 			InstanceResources.Add(Instance->GetResource());
diff --git a/Engine/Source/Runtime/Engine/Public/ParameterCollection.h b/Engine/Source/Runtime/Engine/Public/ParameterCollection.h
--- a/Engine/Source/Runtime/Engine/Public/ParameterCollection.h
+++ b/Engine/Source/Runtime/Engine/Public/ParameterCollection.h
@@ -25,8 +25,12 @@ class FMaterialParameterCollectionInstanceResource
 {
 public:
 	
-	/** Update the contents of the uniform buffer, called from the game thread. */
-	void GameThread_UpdateContents(const FGuid& InId, const TArray<FVector4>& Data, const FName& InOwnerName);
+	/** 
+	 * Update the contents of the uniform buffer, called from the game thread. 
+	 * bRecreateUniformBuffer must be false if any components have their rendering state created which might reference this!
+	 * The calling code is responsible for issuing a FGlobalComponentRecreateRenderStateContext before modifying.  Otherwise cached mesh draw commands will be left with a dangling pointer.
+	 */
+	void GameThread_UpdateContents(const FGuid& InId, const TArray<FVector4>& Data, const FName& InOwnerName, bool bRecreateUniformBuffer);
 
 	/** Destroy, called from the game thread. */
 	void GameThread_Destroy();
@@ -61,7 +65,7 @@ private:
 
 	FRHIUniformBufferLayout UniformBufferLayout;
 
-	void UpdateContents(const FGuid& InId, const TArray<FVector4>& Data, const FName& InOwnerName);
+	void UpdateContents(const FGuid& InId, const TArray<FVector4>& Data, const FName& InOwnerName, bool bRecreateUniformBuffer);
 };
 
 // Default instance resources used when rendering a material using a parameter collection but there's no FScene present to get a FMaterialParameterCollectionInstanceResource
diff --git a/Engine/Source/Runtime/Renderer/Private/RendererScene.cpp b/Engine/Source/Runtime/Renderer/Private/RendererScene.cpp
--- a/Engine/Source/Runtime/Renderer/Private/RendererScene.cpp
+++ b/Engine/Source/Runtime/Renderer/Private/RendererScene.cpp
@@ -1024,7 +1024,7 @@ FScene::FScene(UWorld* InWorld, bool bInRequiresHitProxies, bool bInIsEditorScen
 	}
 #endif
 
-	World->UpdateParameterCollectionInstances(false);
+	World->UpdateParameterCollectionInstances(false, false);
 
 	FPersistentUniformBuffers* PersistentUniformBuffers = &UniformBuffers;
 	ENQUEUE_RENDER_COMMAND(InitializeUniformBuffers)(
