Fixed possible crash when rendering debug view mode.
Added a refresh logic for debug view mode shaders.

diff --git a/Engine/Source/Editor/UnrealEd/Private/EditorViewportClient.cpp b/Engine/Source/Editor/UnrealEd/Private/EditorViewportClient.cpp
--- a/Engine/Source/Editor/UnrealEd/Private/EditorViewportClient.cpp
+++ b/Engine/Source/Editor/UnrealEd/Private/EditorViewportClient.cpp
@@ -51,6 +51,9 @@
 #include "ViewportWorldInteraction.h"
 #include "Editor/EditorPerformanceSettings.h"
 #include "ImageWriteQueue.h"
+#include "DebugViewModeHelpers.h"
+#include "Misc/ScopedSlowTask.h"
+#include "UnrealEngine.h"
 
 #define LOCTEXT_NAMESPACE "EditorViewportClient"
 
@@ -3625,6 +3628,17 @@ void FEditorViewportClient::Draw(FViewport* InViewport, FCanvas* Canvas)
 
 	ViewFamily.EngineShowFlags = EngineShowFlags;
 
+	if (GIsEditor && World && ViewFamily.GetDebugViewShaderMode() != DVSM_None && HasMissingDebugViewModeShaders(true))
+	{
+		FScopedSlowTask CompileShaderTask(3.f, LOCTEXT("CompileMissingViewModeShaders", "Compiling Missing ViewMode Shaders")); // { Get Used Materials, Sync Pending Shader, Wait for Compilation }
+		// CompileShaderTask.MakeDialog(true);
+		TSet<UMaterialInterface*> Materials;
+		if (GetUsedMaterialsInWorld(World, Materials, CompileShaderTask))
+		{
+			CompileDebugViewModeShaders(ViewFamily.GetDebugViewShaderMode(), GetCachedScalabilityCVars().MaterialQualityLevel, ViewFamily.GetFeatureLevel(), false, false, Materials, CompileShaderTask);
+		}
+	}
+
 	if( ModeTools->GetActiveMode( FBuiltinEditorModes::EM_InterpEdit ) == 0 || !AllowsCinematicControl() )
 	{
 		if( !EngineShowFlags.Game )
diff --git a/Engine/Source/Runtime/Engine/Private/DebugViewModeMaterialProxy.cpp b/Engine/Source/Runtime/Engine/Private/DebugViewModeMaterialProxy.cpp
--- a/Engine/Source/Runtime/Engine/Private/DebugViewModeMaterialProxy.cpp
+++ b/Engine/Source/Runtime/Engine/Private/DebugViewModeMaterialProxy.cpp
@@ -10,15 +10,30 @@ DebugViewModeMaterialProxy.cpp : Contains definitions the debug view mode materi
 #include "EngineModule.h"
 #include "RendererInterface.h"
 
-ENGINE_API bool GetDebugViewMaterial(const UMaterialInterface* InMaterialInterface, EDebugViewShaderMode InDebugViewMode, const FMaterialRenderProxy*& OutMaterialRenderProxy, const FMaterial*& OutMaterial)
+ENGINE_API bool GetDebugViewMaterial(const UMaterialInterface* InMaterialInterface, EDebugViewShaderMode InDebugViewMode, ERHIFeatureLevel::Type InFeatureLevel, const FMaterialRenderProxy*& OutMaterialRenderProxy, const FMaterial*& OutMaterial)
 {
 #if WITH_EDITORONLY_DATA
-	return FDebugViewModeMaterialProxy::GetShader(InMaterialInterface, InDebugViewMode, OutMaterialRenderProxy, OutMaterial);
+	return FDebugViewModeMaterialProxy::GetShader(InMaterialInterface, InDebugViewMode, InFeatureLevel, OutMaterialRenderProxy, OutMaterial);
 #else
 	return false;
 #endif
 }
 
+ENGINE_API bool HasMissingDebugViewModeShaders(bool bClearFlag)
+{
+#if WITH_EDITORONLY_DATA
+	if (FDebugViewModeMaterialProxy::MissingShadersChanged())
+	{
+		if (bClearFlag)
+		{
+			FDebugViewModeMaterialProxy::ClearMissingShadersFlag();
+		}
+		return true;
+	}
+#endif
+	return false;
+}
+
 ENGINE_API void ClearDebugViewMaterials(UMaterialInterface* InMaterialInterface)
 {
 #if WITH_EDITORONLY_DATA
@@ -28,9 +43,10 @@ ENGINE_API void ClearDebugViewMaterials(UMaterialInterface* InMaterialInterface)
 
 #if WITH_EDITORONLY_DATA
 
-TMap<FDebugViewModeMaterialProxy::FMaterialUsagePair, FDebugViewModeMaterialProxy*> FDebugViewModeMaterialProxy::DebugMaterialShaderMap;
-
 volatile bool FDebugViewModeMaterialProxy::bReentrantCall = false;
+TMap<FDebugViewModeMaterialProxy::FMaterialKey, FDebugViewModeMaterialProxy*> FDebugViewModeMaterialProxy::DebugMaterialShaderMap;
+TSet<FDebugViewModeMaterialProxy::FMaterialKey> FDebugViewModeMaterialProxy::MissingShaders;
+bool FDebugViewModeMaterialProxy::bMissingShadersChanged = false;
 
 void FDebugViewModeMaterialProxy::AddShader(
 	UMaterialInterface* InMaterialInterface, 
@@ -50,24 +66,42 @@ void FDebugViewModeMaterialProxy::AddShader(
 		InMaterialInterface = UMaterial::GetDefaultMaterial(MD_Surface);
 	}
 
-	FMaterialUsagePair ShaderMapKey(InMaterialInterface, InDebugViewMode);
+	FMaterialKey ShaderMapKey(InMaterialInterface, InDebugViewMode, InFeatureLevel);
 	if (!DebugMaterialShaderMap.Contains(ShaderMapKey))
 	{
 		DebugMaterialShaderMap.Add(ShaderMapKey, new FDebugViewModeMaterialProxy(InMaterialInterface, InQualityLevel, InFeatureLevel, bSynchronousCompilation, InDebugViewMode));
 	}
+	MissingShaders.Remove(ShaderMapKey);
 }
 
-bool FDebugViewModeMaterialProxy::GetShader(const UMaterialInterface* InMaterialInterface, EDebugViewShaderMode InDebugViewMode, const FMaterialRenderProxy*& OutMaterialRenderProxy, const FMaterial*& OutMaterial)
+bool FDebugViewModeMaterialProxy::GetShader(const UMaterialInterface* InMaterialInterface, EDebugViewShaderMode InDebugViewMode, ERHIFeatureLevel::Type InFeatureLevel, const FMaterialRenderProxy*& OutMaterialRenderProxy, const FMaterial*& OutMaterial)
 {
-	FDebugViewModeMaterialProxy** BoundMaterial = DebugMaterialShaderMap.Find(FMaterialUsagePair(InMaterialInterface, InDebugViewMode));
-	if (BoundMaterial && *BoundMaterial && (*BoundMaterial)->IsValid() && (*BoundMaterial)->GetRenderingThreadShaderMap())
+	FMaterialKey MaterialKey(InMaterialInterface, InDebugViewMode, InFeatureLevel);
+	FDebugViewModeMaterialProxy** BoundMaterial = DebugMaterialShaderMap.Find(MaterialKey);
+	if (BoundMaterial && *BoundMaterial)
 	{
-		OutMaterialRenderProxy = *BoundMaterial;
-		OutMaterial = *BoundMaterial;
-		return true;
+		if ((*BoundMaterial)->IsValid() && (*BoundMaterial)->GetRenderingThreadShaderMap())
+		{
+			OutMaterialRenderProxy = *BoundMaterial;
+			OutMaterial = *BoundMaterial;
+			return true;
+		}
+		else // This material is not usable for debug view modes.
+		{
+			return false;
+		}
 	}
 	else
 	{
+		if (!MissingShaders.Contains(MaterialKey))
+		{
+			MissingShaders.Add(MaterialKey);
+
+			// Note that a new shader key is missing, so that we can trigger recompilation.
+			// Because it is not garantied that this can be fixed (see ValidateAllShaders()), 
+			// we only kept track of new entries, and don't necessarily try to fix all of them.
+			bMissingShadersChanged = true;
+		}
 		return false;
 	}
 }
@@ -83,17 +117,18 @@ void FDebugViewModeMaterialProxy::ClearAllShaders(UMaterialInterface* InMaterial
 
 	if (!InMaterialInterface)
 	{
-		for (TMap<FMaterialUsagePair, FDebugViewModeMaterialProxy*>::TIterator It(DebugMaterialShaderMap); It; ++It)
+		for (TMap<FMaterialKey, FDebugViewModeMaterialProxy*>::TIterator It(DebugMaterialShaderMap); It; ++It)
 		{
 			FDebugViewModeMaterialProxy* Proxy = It.Value();
 			MaterialsToDelete.Add(Proxy);
 			It.Value() = nullptr;
 		}
 		DebugMaterialShaderMap.Empty();
+		MissingShaders.Empty();
 	}
 	else
 	{
-		for (TMap<FMaterialUsagePair, FDebugViewModeMaterialProxy*>::TIterator It(DebugMaterialShaderMap); It; ++It)
+		for (TMap<FMaterialKey, FDebugViewModeMaterialProxy*>::TIterator It(DebugMaterialShaderMap); It; ++It)
 		{
 			if (It.Key().MaterialInterface == InMaterialInterface)
 			{
@@ -104,6 +139,14 @@ void FDebugViewModeMaterialProxy::ClearAllShaders(UMaterialInterface* InMaterial
 				It.RemoveCurrent();
 			}
 		}
+
+		for (TSet<FMaterialKey>::TIterator It(MissingShaders); It; ++It)
+		{
+			if (It->MaterialInterface == InMaterialInterface)
+			{
+				It.RemoveCurrent();
+			}
+		}
 	}
 
 	if (MaterialsToDelete.Num())
@@ -131,7 +174,7 @@ void FDebugViewModeMaterialProxy::ValidateAllShaders(TSet<UMaterialInterface*>&
 
 	TArray<FDebugViewModeMaterialProxy*> MaterialsToUpdate;
 
-	for (TMap<FMaterialUsagePair, FDebugViewModeMaterialProxy*>::TIterator It(DebugMaterialShaderMap); It; ++It)
+	for (TMap<FMaterialKey, FDebugViewModeMaterialProxy*>::TIterator It(DebugMaterialShaderMap); It; ++It)
 	{
 		const UMaterialInterface* OriginalMaterialInterface = It.Key().MaterialInterface;
 		FDebugViewModeMaterialProxy* DebugMaterial = It.Value();
diff --git a/Engine/Source/Runtime/Engine/Private/DebugViewModeMaterialProxy.h b/Engine/Source/Runtime/Engine/Private/DebugViewModeMaterialProxy.h
--- a/Engine/Source/Runtime/Engine/Private/DebugViewModeMaterialProxy.h
+++ b/Engine/Source/Runtime/Engine/Private/DebugViewModeMaterialProxy.h
@@ -173,10 +173,12 @@ public:
 		bool bSynchronousCompilation, 
 		EDebugViewShaderMode InDebugViewMode
 	);
-	static bool GetShader(const UMaterialInterface* InMaterialInterface, EDebugViewShaderMode InDebugViewMode, const FMaterialRenderProxy*& OutMaterialRenderProxy, const FMaterial*& OutMaterial);
+	static bool GetShader(const UMaterialInterface* InMaterialInterface, EDebugViewShaderMode InDebugViewMode, ERHIFeatureLevel::Type InFeatureLevel, const FMaterialRenderProxy*& OutMaterialRenderProxy, const FMaterial*& OutMaterial);
 	static void ClearAllShaders(UMaterialInterface* InMaterialInterface);
 	static bool HasAnyShaders() { return DebugMaterialShaderMap.Num() > 0; }
 	static void ValidateAllShaders(TSet<UMaterialInterface*>& Materials);
+	static bool MissingShadersChanged() { return bMissingShadersChanged; }
+	static void ClearMissingShadersFlag() { bMissingShadersChanged = false; }
 
 private:
 
@@ -215,27 +217,34 @@ private:
 	bool bIsDefaultMaterial;
 	bool bSynchronousCompilation;
 
-	struct FMaterialUsagePair
+	struct FMaterialKey
 	{
-		FMaterialUsagePair(const UMaterialInterface* InMaterialInterface, EDebugViewShaderMode InDebugViewMode) : MaterialInterface(InMaterialInterface), DebugViewMode(InDebugViewMode) {}
+		FMaterialKey(const UMaterialInterface* InMaterialInterface, EDebugViewShaderMode InDebugViewMode, ERHIFeatureLevel::Type InFeatureLevel) : MaterialInterface(InMaterialInterface), DebugViewMode((uint32)InDebugViewMode), FeatureLevel((uint32)InFeatureLevel) {}
 
 		const UMaterialInterface* MaterialInterface;
-		EDebugViewShaderMode DebugViewMode;
+		uint32 DebugViewMode;
+		uint32 FeatureLevel;
 
-		friend bool operator==(const FMaterialUsagePair& Lhs, const FMaterialUsagePair& Rhs)
+		friend bool operator==(const FMaterialKey& Lhs, const FMaterialKey& Rhs)
 		{
-			return Lhs.MaterialInterface == Rhs.MaterialInterface && Lhs.DebugViewMode == Rhs.DebugViewMode;
+			return Lhs.MaterialInterface == Rhs.MaterialInterface
+				&& Lhs.DebugViewMode == Rhs.DebugViewMode
+				&& Lhs.FeatureLevel == Rhs.FeatureLevel;
 		}
 
-		friend uint32 GetTypeHash( const FMaterialUsagePair& Pair )
+		friend uint32 GetTypeHash(const FMaterialKey& Key)
 		{
-			return GetTypeHash(Pair.MaterialInterface) ^ GetTypeHash((uint32)Pair.DebugViewMode);
-		}
+			const uint32 DebugViewModeShift = 6;
+			static_assert((2 ^ DebugViewModeShift) < DVSM_MAX, "Bit shift too small!");
 
+			return GetTypeHash(Key.MaterialInterface) ^ GetTypeHash((Key.FeatureLevel << DebugViewModeShift) ^ Key.DebugViewMode);
+		}
 	};
 
 	static volatile bool bReentrantCall;
-	static TMap<FMaterialUsagePair, FDebugViewModeMaterialProxy*> DebugMaterialShaderMap;
+	static TMap<FMaterialKey, FDebugViewModeMaterialProxy*> DebugMaterialShaderMap;
+	static TSet<FMaterialKey> MissingShaders;
+	static bool bMissingShadersChanged;
 };
 
 #endif
diff --git a/Engine/Source/Runtime/Engine/Public/DebugViewModeHelpers.h b/Engine/Source/Runtime/Engine/Public/DebugViewModeHelpers.h
--- a/Engine/Source/Runtime/Engine/Public/DebugViewModeHelpers.h
+++ b/Engine/Source/Runtime/Engine/Public/DebugViewModeHelpers.h
@@ -46,6 +46,7 @@ FORCEINLINE bool AllowDebugViewShaderMode(EDebugViewShaderMode ShaderMode) { ret
 ENGINE_API int32 GetNumActorsInWorld(UWorld* InWorld);
 ENGINE_API bool GetUsedMaterialsInWorld(UWorld* InWorld, OUT TSet<UMaterialInterface*>& OutMaterials, struct FSlowTask& Task);
 ENGINE_API bool CompileDebugViewModeShaders(EDebugViewShaderMode Mode, EMaterialQualityLevel::Type QualityLevel, ERHIFeatureLevel::Type FeatureLevel, bool bFullRebuild, bool bWaitForPreviousShaders, TSet<UMaterialInterface*>& Materials, FSlowTask& ProgressTask);
+ENGINE_API bool HasMissingDebugViewModeShaders(bool bClearFlag);
 
 ENGINE_API bool WaitForShaderCompilation(const FText& Message, FSlowTask& ProgressTask);
 
diff --git a/Engine/Source/Runtime/Renderer/Private/DebugViewModeRendering.cpp b/Engine/Source/Runtime/Renderer/Private/DebugViewModeRendering.cpp
--- a/Engine/Source/Runtime/Renderer/Private/DebugViewModeRendering.cpp
+++ b/Engine/Source/Runtime/Renderer/Private/DebugViewModeRendering.cpp
@@ -49,7 +49,8 @@ IMPLEMENT_MATERIAL_SHADER_TYPE(,FDebugViewModeVS,TEXT("/Engine/Private/DebugView
 IMPLEMENT_MATERIAL_SHADER_TYPE(,FDebugViewModeHS,TEXT("/Engine/Private/DebugViewModeVertexShader.usf"),TEXT("MainHull"),SF_Hull);	
 IMPLEMENT_MATERIAL_SHADER_TYPE(,FDebugViewModeDS,TEXT("/Engine/Private/DebugViewModeVertexShader.usf"),TEXT("MainDomain"),SF_Domain);
 
-ENGINE_API bool GetDebugViewMaterial(const UMaterialInterface* InMaterialInterface, EDebugViewShaderMode InDebugViewMode, const FMaterialRenderProxy*& OutMaterialRenderProxy, const FMaterial*& OutMaterial);
+ENGINE_API bool GetDebugViewMaterial(const UMaterialInterface* InMaterialInterface, EDebugViewShaderMode InDebugViewMode, ERHIFeatureLevel::Type InFeatureLevel,const FMaterialRenderProxy*& OutMaterialRenderProxy, const FMaterial*& OutMaterial);
+ENGINE_API bool HasMissingDebugViewModeShaders(bool bClearFlag);
 
 void FDeferredShadingSceneRenderer::DoDebugViewModePostProcessing(FRHICommandListImmediate& RHICmdList, const FViewInfo& View, TRefCountPtr<IPooledRenderTarget>& VelocityRT)
 {
@@ -304,7 +305,7 @@ void FDebugViewModeMeshProcessor::AddMeshBatch(const FMeshBatch& RESTRICT MeshBa
 	{
 		ResolvedMaterial = UMaterial::GetDefaultMaterial(MD_Surface);
 	}
-	if (!GetDebugViewMaterial(ResolvedMaterial, DebugViewMode, MaterialRenderProxy,  Material))
+	if (!GetDebugViewMaterial(ResolvedMaterial, DebugViewMode, FeatureLevel, MaterialRenderProxy, Material))
 	{
 		return;
 	}
@@ -346,7 +347,8 @@ void FDebugViewModeMeshProcessor::AddMeshBatch(const FMeshBatch& RESTRICT MeshBa
 		ViewModeParam, 
 		ViewModeParamName);
 
-	if (DebugViewModeInterface->bNeedsInstructionCount)
+	// Shadermap can be null while shaders are compiling.
+	if (DebugViewModeInterface->bNeedsInstructionCount && BatchMaterial->GetRenderingThreadShaderMap())
 	{
 		ShaderElementData.NumVSInstructions = BatchMaterial->GetShader<TBasePassVS<TUniformLightMapPolicy<LMP_NO_LIGHTMAP>, false>>(VertexFactoryType)->GetNumInstructions();
 		ShaderElementData.NumPSInstructions = BatchMaterial->GetShader<TBasePassPS<TUniformLightMapPolicy<LMP_NO_LIGHTMAP>, false>>(VertexFactoryType)->GetNumInstructions();
