Fix formatting in ensure string

diff --git a/Engine/Source/Runtime/Engine/Private/Particles/ParticleModules_Location.cpp b/Engine/Source/Runtime/Engine/Private/Particles/ParticleModules_Location.cpp
--- a/Engine/Source/Runtime/Engine/Private/Particles/ParticleModules_Location.cpp
+++ b/Engine/Source/Runtime/Engine/Private/Particles/ParticleModules_Location.cpp
@@ -2050,7 +2050,7 @@ bool UParticleModuleLocationBoneSocket::GetSocketInfoForSourceIndex(FModuleLocat
 	{
 		case EBoneSocketSourceIndexMode::SourceLocations:
 		{
-			if (ensureMsgf(SourceLocations.IsValidIndex(SourceIndex), TEXT("Invalid index of %s for %s"), SourceIndex, *GetPathName()))
+			if (ensureMsgf(SourceLocations.IsValidIndex(SourceIndex), TEXT("Invalid index of %d for %s"), SourceIndex, *GetPathName()))
 			{
 				OutSocket = SourceComponent->SkeletalMesh->FindSocket(SourceLocations[SourceIndex].BoneSocketName);
 				OutOffset = SourceLocations[SourceIndex].Offset + UniversalOffset;
@@ -2063,8 +2063,8 @@ bool UParticleModuleLocationBoneSocket::GetSocketInfoForSourceIndex(FModuleLocat
 		break;
 		case EBoneSocketSourceIndexMode::PreSelectedIndices:
 		{
-			if (ensureMsgf(InstancePayload, TEXT("Invalid instance payload parameter on GetSocketInfoForSourceIndex for %s"), SourceIndex, *GetPathName()) &&
-				ensureMsgf(SourceIndex >= 0 && SourceIndex < InstancePayload->PreSelectedBoneSocketIndices.Num(), TEXT("Invalid index of %s for %s"), SourceIndex, *GetPathName()))
+			if (ensureMsgf(InstancePayload, TEXT("Invalid instance payload parameter on GetSocketInfoForSourceIndex %d for %s"), SourceIndex, *GetPathName()) &&
+				ensureMsgf(SourceIndex >= 0 && SourceIndex < InstancePayload->PreSelectedBoneSocketIndices.Num(), TEXT("Invalid index of ds for %s"), SourceIndex, *GetPathName()))
 			{
 				OutSocket = SourceComponent->SkeletalMesh->GetSocketByIndex(InstancePayload->PreSelectedBoneSocketIndices[SourceIndex]);
 				OutOffset = UniversalOffset;
