[LLM] adding some missing tags to clean up EngineMisc LLM category.

diff --git a/Engine/Source/Runtime/Core/Private/Stats/Stats2.cpp b/Engine/Source/Runtime/Core/Private/Stats/Stats2.cpp
--- a/Engine/Source/Runtime/Core/Private/Stats/Stats2.cpp
+++ b/Engine/Source/Runtime/Core/Private/Stats/Stats2.cpp
@@ -1221,6 +1221,8 @@ void FThreadStats::FlushRegularStats( bool bHasBrokenCallstacks, bool bForceFlus
 
 void FThreadStats::FlushRawStats( bool bHasBrokenCallstacks /*= false*/, bool bForceFlush /*= false*/ )
 {
+	LLM_SCOPE(ELLMTag::Stats);
+
 	if (bReentranceGuard)
 	{
 		return;
diff --git a/Engine/Source/Runtime/CoreUObject/Private/UObject/GarbageCollection.cpp b/Engine/Source/Runtime/CoreUObject/Private/UObject/GarbageCollection.cpp
--- a/Engine/Source/Runtime/CoreUObject/Private/UObject/GarbageCollection.cpp
+++ b/Engine/Source/Runtime/CoreUObject/Private/UObject/GarbageCollection.cpp
@@ -1468,6 +1468,7 @@ void CollectGarbageInternal(EObjectFlags KeepFlags, bool bPerformFullPurge)
 	SCOPED_NAMED_EVENT(CollectGarbageInternal, FColor::Red);
 	CSV_EVENT_GLOBAL(TEXT("GC"));
 	CSV_SCOPED_TIMING_STAT_EXCLUSIVE(GarbageCollection);
+	LLM_SCOPE(ELLMTag::GC);
 
 	FGCCSyncObject::Get().ResetGCIsWaiting();
 
diff --git a/Engine/Source/Runtime/Engine/Private/Components/SkeletalMeshComponent.cpp b/Engine/Source/Runtime/Engine/Private/Components/SkeletalMeshComponent.cpp
--- a/Engine/Source/Runtime/Engine/Private/Components/SkeletalMeshComponent.cpp
+++ b/Engine/Source/Runtime/Engine/Private/Components/SkeletalMeshComponent.cpp
@@ -2162,6 +2162,7 @@ void USkeletalMeshComponent::SwapEvaluationContextBuffers()
 
 void USkeletalMeshComponent::DispatchParallelEvaluationTasks(FActorComponentTickFunction* TickFunction)
 {
+	LLM_SCOPE(ELLMTag::SkeletalMesh);
 	SwapEvaluationContextBuffers();
 
 	// start parallel work
diff --git a/Engine/Source/Runtime/Engine/Private/Materials/MaterialShared.cpp b/Engine/Source/Runtime/Engine/Private/Materials/MaterialShared.cpp
--- a/Engine/Source/Runtime/Engine/Private/Materials/MaterialShared.cpp
+++ b/Engine/Source/Runtime/Engine/Private/Materials/MaterialShared.cpp
@@ -2182,6 +2182,8 @@ void FMaterialRenderProxy::ReleaseDynamicRHI()
 
 void FMaterialRenderProxy::UpdateDeferredCachedUniformExpressions()
 {
+	LLM_SCOPE(ELLMTag::Materials);
+
 	check(IsInRenderingThread());
 
 	QUICK_SCOPE_CYCLE_COUNTER(STAT_UpdateDeferredCachedUniformExpressions);
diff --git a/Engine/Source/Runtime/Engine/Private/Particles/ParticleComponents.cpp b/Engine/Source/Runtime/Engine/Private/Particles/ParticleComponents.cpp
--- a/Engine/Source/Runtime/Engine/Private/Particles/ParticleComponents.cpp
+++ b/Engine/Source/Runtime/Engine/Private/Particles/ParticleComponents.cpp
@@ -4373,6 +4373,7 @@ void UParticleSystemComponent::ClearDynamicData()
 void UParticleSystemComponent::UpdateDynamicData()
 {
 	//SCOPE_CYCLE_COUNTER(STAT_ParticleSystemComponent_UpdateDynamicData);
+	LLM_SCOPE(ELLMTag::Particles);
 
 	ForceAsyncWorkCompletion(ENSURE_AND_STALL);
 	if (SceneProxy)
diff --git a/Engine/Source/Runtime/Engine/Private/Particles/ParticleSystemManager.cpp b/Engine/Source/Runtime/Engine/Private/Particles/ParticleSystemManager.cpp
--- a/Engine/Source/Runtime/Engine/Private/Particles/ParticleSystemManager.cpp
+++ b/Engine/Source/Runtime/Engine/Private/Particles/ParticleSystemManager.cpp
@@ -457,6 +457,7 @@ FORCEINLINE void FParticleSystemWorldManager::FlushAsyncTicks(const FGraphEventR
 {
 	if (AsyncTickBatch.Num() > 0)
 	{
+		LLM_SCOPE(ELLMTag::Particles);
 		FGraphEventRef AsyncTask = TGraphTask<FParticleManagerAsyncTask>::CreateTask(nullptr, ENamedThreads::GameThread).ConstructAndDispatchWhenReady(this, AsyncTickBatch);
 
 #if PSC_MAN_TG_WAIT_FOR_ASYNC
diff --git a/Engine/Source/Runtime/Engine/Private/PhysicsEngine/Experimental/PhysScene_ImmediatePhysX.cpp b/Engine/Source/Runtime/Engine/Private/PhysicsEngine/Experimental/PhysScene_ImmediatePhysX.cpp
--- a/Engine/Source/Runtime/Engine/Private/PhysicsEngine/Experimental/PhysScene_ImmediatePhysX.cpp
+++ b/Engine/Source/Runtime/Engine/Private/PhysicsEngine/Experimental/PhysScene_ImmediatePhysX.cpp
@@ -276,6 +276,8 @@ struct FPhysXCPUDispatcher : public PxCpuDispatcher
 
 	virtual void submitTask(PxBaseTask& Task) override
 	{
+		LLM_SCOPE(ELLMTag::PhysX);
+
 		if (IsInGameThread())
 		{
 			//Game thread enqueues on task graph
diff --git a/Engine/Source/Runtime/Engine/Private/PhysicsEngine/PhysXSupport.h b/Engine/Source/Runtime/Engine/Private/PhysicsEngine/PhysXSupport.h
--- a/Engine/Source/Runtime/Engine/Private/PhysicsEngine/PhysXSupport.h
+++ b/Engine/Source/Runtime/Engine/Private/PhysicsEngine/PhysXSupport.h
@@ -13,6 +13,7 @@
 #include "Containers/Queue.h"
 #include "Physics/PhysicsFiltering.h"
 #include "PhysXPublic.h"
+#include "HAL/LowLevelMemTracker.h"
 
 class UBodySetup;
 class UPhysicalMaterial;
@@ -294,6 +295,7 @@ public:
 
 	virtual void* allocate(size_t size, const char* typeName, const char* filename, int line) override
 	{
+		LLM_SCOPE(ELLMTag::PhysX);
 #if PHYSX_MEMORY_STATS
 		INC_DWORD_STAT_BY(STAT_MemoryPhysXTotalAllocationSize, size);
 
diff --git a/Engine/Source/Runtime/Engine/Private/SkeletalRenderGPUSkin.cpp b/Engine/Source/Runtime/Engine/Private/SkeletalRenderGPUSkin.cpp
--- a/Engine/Source/Runtime/Engine/Private/SkeletalRenderGPUSkin.cpp
+++ b/Engine/Source/Runtime/Engine/Private/SkeletalRenderGPUSkin.cpp
@@ -13,6 +13,7 @@
 #include "ClearQuad.h"
 #include "ShaderParameterUtils.h"
 #include "SkeletalMeshTypes.h"
+#include "HAL/LowLevelMemTracker.h"
 
 DEFINE_LOG_CATEGORY_STATIC(LogSkeletalGPUSkinMesh, Warning, All);
 
@@ -248,6 +249,8 @@ void FSkeletalMeshObjectGPUSkin::ReleaseMorphResources()
 
 void FSkeletalMeshObjectGPUSkin::Update(int32 LODIndex,USkinnedMeshComponent* InMeshComponent,const TArray<FActiveMorphTarget>& ActiveMorphTargets, const TArray<float>& MorphTargetWeights, EPreviousBoneTransformUpdateMode PreviousBoneTransformUpdateMode)
 {
+	LLM_SCOPE(ELLMTag::SkeletalMesh);
+
 	// make sure morph data has been initialized for each LOD
 	if(InMeshComponent && !bMorphResourcesInitialized && ActiveMorphTargets.Num() > 0 )
 	{
diff --git a/Engine/Source/Runtime/RHI/Private/RHICommandList.cpp b/Engine/Source/Runtime/RHI/Private/RHICommandList.cpp
--- a/Engine/Source/Runtime/RHI/Private/RHICommandList.cpp
+++ b/Engine/Source/Runtime/RHI/Private/RHICommandList.cpp
@@ -425,6 +425,8 @@ public:
 
 	void DoTask(ENamedThreads::Type CurrentThread, const FGraphEventRef& MyCompletionGraphEvent)
 	{
+		LLM_SCOPE(ELLMTag::RHIMisc);
+
 		check(bRHIThread || IsInRenderingThread());
 		FGraphEventArray Prereq;
 		if (RHIThreadTask.GetReference())
