Framework - Allow getting the game instance on the FLocalPlayerContext.

diff --git a/Engine/Source/Runtime/Engine/Classes/Engine/LocalPlayer.h b/Engine/Source/Runtime/Engine/Classes/Engine/LocalPlayer.h
--- a/Engine/Source/Runtime/Engine/Classes/Engine/LocalPlayer.h
+++ b/Engine/Source/Runtime/Engine/Classes/Engine/LocalPlayer.h
@@ -53,6 +53,9 @@ struct ENGINE_API FLocalPlayerContext
 	/** Returns the world context. */
 	UWorld* GetWorld() const;
 
+	/** Returns the game instance */
+	UGameInstance* GetGameInstance() const;
+
 	/** Returns the local player. */
 	class ULocalPlayer* GetLocalPlayer() const;
 
diff --git a/Engine/Source/Runtime/Engine/Private/LocalPlayer.cpp b/Engine/Source/Runtime/Engine/Private/LocalPlayer.cpp
--- a/Engine/Source/Runtime/Engine/Private/LocalPlayer.cpp
+++ b/Engine/Source/Runtime/Engine/Private/LocalPlayer.cpp
@@ -119,6 +119,16 @@ UWorld* FLocalPlayerContext::GetWorld() const
 	return GetLocalPlayer()->GetWorld();
 }
 
+UGameInstance* FLocalPlayerContext::GetGameInstance() const
+{
+	if (UWorld* WorldPtr = GetWorld())
+	{
+		return WorldPtr->GetGameInstance();
+	}
+
+	return nullptr;
+}
+
 ULocalPlayer* FLocalPlayerContext::GetLocalPlayer() const
 {
 	ULocalPlayer* LocalPlayerPtr = LocalPlayer.Get();
