Add a delegate in the networking profiler that can be registered for which fires when a network profiler is finished being written to

diff --git a/Engine/Source/Runtime/Engine/Private/NetworkProfiler.cpp b/Engine/Source/Runtime/Engine/Private/NetworkProfiler.cpp
--- a/Engine/Source/Runtime/Engine/Private/NetworkProfiler.cpp
+++ b/Engine/Source/Runtime/Engine/Private/NetworkProfiler.cpp
@@ -567,7 +567,6 @@ void FNetworkProfiler::TrackEvent( const FString& EventName, const FString& Even
  */
 void FNetworkProfiler::TrackSessionChange( bool bShouldContinueTracking, const FURL& InURL )
 {
-#if ALLOW_DEBUG_FILES
 	if ( bIsTrackingEnabled )
 	{
 		UE_LOG( LogNet, Log, TEXT( "Network Profiler: TrackSessionChange.  InURL: %s" ), *InURL.ToString() );
@@ -592,6 +591,11 @@ void FNetworkProfiler::TrackSessionChange( bool bShouldContinueTracking, const F
 			// Close file writer so we can rename the file to its final destination.
 			FileWriter->Close();
 
+			if (OnNetworkProfileFinished().IsBound())
+			{
+				OnNetworkProfileFinished().Broadcast(FileWriter->GetArchiveName());
+			}
+
 			// Clean up.
 			delete FileWriter;
 			FileWriter = nullptr;
@@ -622,7 +626,6 @@ void FNetworkProfiler::TrackSessionChange( bool bShouldContinueTracking, const F
 			(*FileWriter) << CurrentHeader;
 		}
 	}
-#endif	//#if ALLOW_DEBUG_FILES
 }
 
 void FNetworkProfiler::TrackSendAck( uint16 NumBits, UNetConnection* Connection )
diff --git a/Engine/Source/Runtime/Engine/Public/Net/NetworkProfiler.h b/Engine/Source/Runtime/Engine/Public/Net/NetworkProfiler.h
--- a/Engine/Source/Runtime/Engine/Public/Net/NetworkProfiler.h
+++ b/Engine/Source/Runtime/Engine/Public/Net/NetworkProfiler.h
@@ -19,6 +19,8 @@ struct FURL;
 
 #define NETWORK_PROFILER( x ) if ( GNetworkProfiler.IsTrackingEnabled() ) { x; }
 
+DECLARE_MULTICAST_DELEGATE_OneParam(FOnNetworkProfileFinished, const FString& /*Filename */);
+
 /*=============================================================================
 	Network profiler header.
 =============================================================================*/
@@ -94,6 +96,9 @@ private:
 	/** Last known address																			*/
 	TSharedPtr<const FInternetAddr>			LastAddress;
 
+	/** Delegate that's fired when tracking stops on the current network profile */
+	FOnNetworkProfileFinished				OnNetworkProfileFinishedDelegate;
+
 	/** All the data required for writing sent bunches to the profiler stream						*/
 	struct FSendBunchInfo
 	{
@@ -393,6 +398,9 @@ public:
 	bool Exec( UWorld * InWorld, const TCHAR* Cmd, FOutputDevice & Ar );
 
 	bool FORCEINLINE IsTrackingEnabled() const { return bIsTrackingEnabled; }
+
+	/** Return the network profile finished delegate */
+	FOnNetworkProfileFinished& OnNetworkProfileFinished() { return OnNetworkProfileFinishedDelegate; }
 };
 
 /** Global network profiler instance. */
