Fixed a crash w.r.t. missing joint that is linked to curve data - possibly caused by reimport or removing joints.

diff --git a/Engine/Source/Runtime/Engine/Private/Animation/Skeleton.cpp b/Engine/Source/Runtime/Engine/Private/Animation/Skeleton.cpp
--- a/Engine/Source/Runtime/Engine/Private/Animation/Skeleton.cpp
+++ b/Engine/Source/Runtime/Engine/Private/Animation/Skeleton.cpp
@@ -1048,6 +1048,13 @@ void USkeleton::HandleSkeletonHierarchyChange()
 	RefreshAllRetargetSources();
 #endif
 
+	// refresh curve meta data that contains joint info
+	FSmartNameMapping* CurveMappingTable = SmartNames.GetContainerInternal(USkeleton::AnimCurveMappingName);
+	if (CurveMappingTable)
+	{
+		CurveMappingTable->InitializeCurveMetaData(this);
+	}
+
 	OnSkeletonHierarchyChanged.Broadcast();
 }
 
@@ -1647,6 +1654,13 @@ void USkeleton::HandleVirtualBoneChanges()
 		}
 	}
 
+	// refresh curve meta data that contains joint info
+	FSmartNameMapping* CurveMappingTable = SmartNames.GetContainerInternal(USkeleton::AnimCurveMappingName);
+	if (CurveMappingTable)
+	{
+		CurveMappingTable->InitializeCurveMetaData(this);
+	}
+
 	for (TObjectIterator<USkinnedMeshComponent> It; It; ++It)
 	{
 		USkinnedMeshComponent* MeshComponent = *It;
diff --git a/Engine/Source/Runtime/Engine/Private/BoneContainer.cpp b/Engine/Source/Runtime/Engine/Private/BoneContainer.cpp
--- a/Engine/Source/Runtime/Engine/Private/BoneContainer.cpp
+++ b/Engine/Source/Runtime/Engine/Private/BoneContainer.cpp
@@ -260,7 +260,7 @@ void FBoneContainer::CacheRequiredAnimCurveUids(const FCurveEvaluationOption& Cu
 									{
 										const FBoneReference& BoneReference = CurveMetaData->LinkedBones[LinkedBoneIndex];
 										// when you enter first time, sometimes it does not have all info yet
-										if (BoneReference.BoneName != NAME_None)
+										if (BoneReference.BoneIndex != INDEX_NONE && BoneReference.BoneName != NAME_None)
 										{
 											// this linked bone alkways use skeleton index
 											ensure(BoneReference.bUseSkeletonIndex);
