Fix curve compression setting being resettable to none

diff --git a/Engine/Source/Runtime/Engine/Classes/Animation/AnimSequence.h b/Engine/Source/Runtime/Engine/Classes/Animation/AnimSequence.h
--- a/Engine/Source/Runtime/Engine/Classes/Animation/AnimSequence.h
+++ b/Engine/Source/Runtime/Engine/Classes/Animation/AnimSequence.h
@@ -1063,8 +1063,11 @@ public:
 	 * @return		The approximate size of compressed animation data.
 	 */
 	int32 GetApproxCompressedSize() const;
-
-	/** 
+#if WITH_EDITORONLY_DATA
+	// Initialize curve compression settings, does nothing if scheme already valid
+	void InitCurveCompressionScheme();
+#endif
+	/**
 	 *  Utility function for lossless compression of a FRawAnimSequenceTrack 
 	 *  @return true if keys were removed.
 	 **/
diff --git a/Engine/Source/Runtime/Engine/Private/Animation/AnimSequence.cpp b/Engine/Source/Runtime/Engine/Private/Animation/AnimSequence.cpp
--- a/Engine/Source/Runtime/Engine/Private/Animation/AnimSequence.cpp
+++ b/Engine/Source/Runtime/Engine/Private/Animation/AnimSequence.cpp
@@ -301,6 +301,8 @@ UAnimSequence::UAnimSequence(const FObjectInitializer& ObjectInitializer)
 	ImportFileFramerate = 0.0f;
 	ImportResampleFramerate = 0;
 	bAllowFrameStripping = true;
+
+	InitCurveCompressionScheme();
 #endif
 }
 
@@ -399,6 +401,22 @@ static void LoadOldCompressedTrack(FArchive& Ar, FCompressedTrack& Dst, int32 By
 	Ar << Dst.Ranges[0] << Dst.Ranges[1] << Dst.Ranges[2];
 }
 
+#if WITH_EDITORONLY_DATA
+void UAnimSequence::InitCurveCompressionScheme()
+{
+	// Do this is serialize as if the default animation curve compression asset isn't loaded it will
+// fire a warning if we try and load it in post load
+	if ((CurveCompressionSettings == nullptr || !CurveCompressionSettings->AreSettingsValid())
+#if WITH_HOT_RELOAD
+		&& (GetClass()->HasAnyClassFlags(CLASS_CompiledFromBlueprint) || !HasAnyFlags(RF_ClassDefaultObject) || !GIsHotReload) // Don't do this to native CDOs during Hot-Reload
+#endif
+		)
+	{
+		CurveCompressionSettings = FAnimationUtils::GetDefaultAnimationCurveCompressionSettings();
+	}
+}
+#endif
+
 void UAnimSequence::Serialize(FArchive& Ar)
 {
 	LLM_SCOPE(ELLMTag::Animation);
@@ -519,16 +537,7 @@ void UAnimSequence::Serialize(FArchive& Ar)
 		SourceFileTimestamp_DEPRECATED = TEXT("");
 	}
 
-	// Do this is serialize as if the default animation curve compression asset isn't loaded it will
-	// fire a warning if we try and load it in post load
-	if ((CurveCompressionSettings == nullptr || !CurveCompressionSettings->AreSettingsValid()) 
-#if WITH_HOT_RELOAD
-		&& (GetClass()->HasAnyClassFlags(CLASS_CompiledFromBlueprint) || !HasAnyFlags(RF_ClassDefaultObject) || !GIsHotReload) // Don't do this to native CDOs during Hot-Reload
-#endif
-		)
-	{
-		CurveCompressionSettings = FAnimationUtils::GetDefaultAnimationCurveCompressionSettings();
-	}
+	InitCurveCompressionScheme();
 
 #endif // WITH_EDITORONLY_DATA
 
