Don't recurse in components of other attached actors when applying the component instance data.

diff --git a/Engine/Source/Runtime/Engine/Private/ComponentInstanceDataCache.cpp b/Engine/Source/Runtime/Engine/Private/ComponentInstanceDataCache.cpp
--- a/Engine/Source/Runtime/Engine/Private/ComponentInstanceDataCache.cpp
+++ b/Engine/Source/Runtime/Engine/Private/ComponentInstanceDataCache.cpp
@@ -530,7 +530,7 @@ void FComponentInstanceDataCache::ApplyToActor(AActor* Actor, const ECacheApplyP
 		TInlineComponentArray<UActorComponent*> Components;
 		Components.Reserve(Actor->GetComponents().Num());
 
-		auto AddComponentHierarchy = [&Components](USceneComponent* Component)
+		auto AddComponentHierarchy = [Actor, &Components](USceneComponent* Component)
 		{
 			int32 FirstProcessIndex = Components.Num();
 
@@ -554,7 +554,12 @@ void FComponentInstanceDataCache::ApplyToActor(AActor* Actor, const ECacheApplyP
 					{
 						if (USceneComponent* ChildComponent = SceneComponent->GetChildComponent(ChildIndex))
 						{
-							Components.Add(ChildComponent);
+							// We don't want to recurse in to child actors (or any other attached actor) when applying the instance cache, 
+							// components within a child actor are handled by applying the instance data to the child actor component
+							if (ChildComponent->GetOwner() == Actor)
+							{
+								Components.Add(ChildComponent);
+							}
 						}
 					}
 				}
