Removing character interpolation samples from replays in favor of rep movement

diff --git a/Engine/Source/Runtime/Engine/Classes/Engine/DemoNetDriver.h b/Engine/Source/Runtime/Engine/Classes/Engine/DemoNetDriver.h
--- a/Engine/Source/Runtime/Engine/Classes/Engine/DemoNetDriver.h
+++ b/Engine/Source/Runtime/Engine/Classes/Engine/DemoNetDriver.h
@@ -109,6 +109,7 @@ enum ENetworkVersionHistory
 	HISTORY_SAVE_FULL_ENGINE_VERSION		= 11,			// Now saving the entire FEngineVersion including branch name
 	HISTORY_HEADER_GUID						= 12,			// Save guid to demo header
 	HISTORY_CHARACTER_MOVEMENT				= 13,			// Change to using replicated movement and not interpolation
+	HISTORY_CHARACTER_MOVEMENT_NOINTERP		= 14,			// No longer recording interpolated movement samples
 	
 	// -----<new versions can be added before this line>-------------------------------------------------
 	HISTORY_PLUS_ONE,
diff --git a/Engine/Source/Runtime/Engine/Private/Character.cpp b/Engine/Source/Runtime/Engine/Private/Character.cpp
--- a/Engine/Source/Runtime/Engine/Private/Character.cpp
+++ b/Engine/Source/Runtime/Engine/Private/Character.cpp
@@ -1471,73 +1471,18 @@ void ACharacter::PreReplicationForReplay(IRepChangedPropertyTracker & ChangedPro
 {
 	Super::PreReplicationForReplay(ChangedPropertyTracker);
 
-	// If this is a replay, we save out certain values we need to runtime to do smooth interpolation
-	// We'll be able to look ahead in the replay to have these ahead of time for smoother playback
-	FCharacterReplaySample ReplaySample;
-
 	const UWorld* World = GetWorld();
-
-	// If this is a client-recorded replay, use the mesh location and rotation, since these will always
-	// be smoothed - unlike the actor position and rotation.
-	const USkeletalMeshComponent* const MeshComponent = GetMesh();
-	if (MeshComponent && World && World->IsRecordingClientReplay())
+	if (World)
 	{
-		FNetworkPredictionData_Client_Character const* const ClientNetworkPredicationData = CharacterMovement->GetPredictionData_Client_Character();
-		if ((Role == ROLE_SimulatedProxy) && ClientNetworkPredicationData)
-		{
-			ReplaySample.Location = GetActorLocation() + ClientNetworkPredicationData->MeshTranslationOffset;
-			if (CharacterMovement->NetworkSmoothingMode == ENetworkSmoothingMode::Exponential)
-			{
-				ReplaySample.Rotation = GetActorRotation() + ClientNetworkPredicationData->MeshRotationOffset.Rotator();
-			}
-			else
-			{
-				ReplaySample.Rotation = ClientNetworkPredicationData->MeshRotationOffset.Rotator();
-			}
-		}
-		else
-		{
-			// Remove the base transform from the mesh's transform, since on playback the base transform
-			// will be stored in the mesh's RelativeLocation and RelativeRotation.
-			const FTransform BaseTransform(GetBaseRotationOffset(), GetBaseTranslationOffset());
-			const FTransform MeshRootTransform = BaseTransform.Inverse() * MeshComponent->GetComponentTransform();
-
-			ReplaySample.Location = MeshRootTransform.GetLocation();
-			ReplaySample.Rotation = MeshRootTransform.GetRotation().Rotator();
-		}
-
 		// On client replays, our view pitch will be set to 0 as by default we do not replicate
 		// pitch for owners, just for simulated. So instead push our rotation into the sampler
-		if (Controller != nullptr && Role == ROLE_AutonomousProxy && GetNetMode() == NM_Client)
+		if (World->IsRecordingClientReplay() && Controller != nullptr && Role == ROLE_AutonomousProxy && GetNetMode() == NM_Client)
 		{
 			SetRemoteViewPitch(Controller->GetControlRotation().Pitch);
 		}
-	}
-	else
-	{
-		ReplaySample.Location = GetActorLocation();
-		ReplaySample.Rotation = GetActorRotation();
-	}
-
-	ReplaySample.Velocity = GetVelocity();
-	ReplaySample.Acceleration = CharacterMovement->GetCurrentAcceleration();
-	ReplaySample.RemoteViewPitch = RemoteViewPitch;
-
-	if (World)
-	{
-		if (World->DemoNetDriver)
-		{
-			ReplaySample.Time = World->DemoNetDriver->DemoCurrentTime;
-		}
 
 		ReplayLastTransformUpdateTimeStamp = World->GetTimeSeconds();
 	}
-
-	const uint32 ReserveBits = 1024;
-	FBitWriter Writer(ReserveBits, true);
-	Writer << ReplaySample;
-
-	ChangedPropertyTracker.SetExternalData(Writer.GetData(), Writer.GetNumBits());
 }
 
 void ACharacter::GetLifetimeReplicatedProps( TArray< FLifetimeProperty > & OutLifetimeProps ) const
diff --git a/Engine/Source/Runtime/Engine/Private/Components/CharacterMovementComponent.cpp b/Engine/Source/Runtime/Engine/Private/Components/CharacterMovementComponent.cpp
--- a/Engine/Source/Runtime/Engine/Private/Components/CharacterMovementComponent.cpp
+++ b/Engine/Source/Runtime/Engine/Private/Components/CharacterMovementComponent.cpp
@@ -602,16 +602,38 @@ void UCharacterMovementComponent::OnRegister()
 
 	// Force linear smoothing for replays.
 	const UWorld* MyWorld = GetWorld();
-	const bool IsReplay = (MyWorld && MyWorld->IsPlayingReplay());
-	if (IsReplay)
+	const bool bIsReplay = (MyWorld && MyWorld->IsPlayingReplay());
+	if (bIsReplay)
 	{
-		if ((CharacterMovementCVars::ReplayUseInterpolation == 1) || (MyWorld && MyWorld->DemoNetDriver && (MyWorld->DemoNetDriver->GetPlaybackDemoVersion() < HISTORY_CHARACTER_MOVEMENT)))
+		// At least one of these conditions will be true
+		const bool bHasInterpolationData = MyWorld && MyWorld->DemoNetDriver && (MyWorld->DemoNetDriver->GetPlaybackDemoVersion() < HISTORY_CHARACTER_MOVEMENT_NOINTERP);
+		const bool bHasRepMovement = MyWorld && MyWorld->DemoNetDriver && (MyWorld->DemoNetDriver->GetPlaybackDemoVersion() >= HISTORY_CHARACTER_MOVEMENT);
+
+		if (CharacterMovementCVars::ReplayUseInterpolation == 1)
 		{
-			NetworkSmoothingMode = ENetworkSmoothingMode::Replay;
+			if (bHasInterpolationData)
+			{
+				NetworkSmoothingMode = ENetworkSmoothingMode::Replay;
+			}
+			else
+			{
+				UE_LOG(LogCharacterMovement, Warning, TEXT("p.ReplayUseInterpolation is enabled, but the replay was not recorded with interpolation data."));
+				ensure(bHasRepMovement);
+				NetworkSmoothingMode = ENetworkSmoothingMode::Linear;
+			}
 		}
 		else
 		{
-			NetworkSmoothingMode = ENetworkSmoothingMode::Linear;
+			if (bHasRepMovement)
+			{
+				NetworkSmoothingMode = ENetworkSmoothingMode::Linear;
+			}
+			else
+			{
+				UE_LOG(LogCharacterMovement, Warning, TEXT("p.ReplayUseInterpolation is disabled, but the replay was not recorded with rep movement data."));
+				ensure(bHasInterpolationData);
+				NetworkSmoothingMode = ENetworkSmoothingMode::Replay;
+			}
 		}
 	}
 	else if (NetMode == NM_ListenServer)
