Copy CL 5319247 - Texture streaming cannot be disabled if FCompactBulkData is used because reload from disk via compact bulk data isn't supported. Added a option so that project or platform can opt-in using compact bulk data

diff --git a/Engine/Source/Runtime/Engine/Private/ContentStreaming.cpp b/Engine/Source/Runtime/Engine/Private/ContentStreaming.cpp
--- a/Engine/Source/Runtime/Engine/Private/ContentStreaming.cpp
+++ b/Engine/Source/Runtime/Engine/Private/ContentStreaming.cpp
@@ -744,7 +744,7 @@ FStreamingManagerCollection::FStreamingManagerCollection()
 {
 #if PLATFORM_SUPPORTS_TEXTURE_STREAMING
 	// Disable texture streaming if that was requested (needs to happen before the call to ProcessNewlyLoadedUObjects, as that can load textures)
-	if( FParse::Param( FCommandLine::Get(), TEXT( "NoTextureStreaming" ) ) )
+	if( !TEXTURE2DMIPMAP_USE_COMPACT_BULKDATA && FParse::Param( FCommandLine::Get(), TEXT( "NoTextureStreaming" ) ) )
 	{
 		CVarSetTextureStreaming.AsVariable()->Set(0, ECVF_SetByCommandline);
 	}
@@ -1126,10 +1126,26 @@ void FStreamingManagerCollection::AddOrRemoveTextureStreamingManagerIfNeeded(boo
 #if PLATFORM_SUPPORTS_TEXTURE_STREAMING
 	{
 		bUseTextureStreaming = CVarSetTextureStreaming.GetValueOnGameThread() != 0;
+		if (!bUseTextureStreaming && TEXTURE2DMIPMAP_USE_COMPACT_BULKDATA)
+		{
+			static bool bWarned;
+			if (!bWarned)
+			{
+				UE_LOG(LogContentStreaming, Warning, TEXT("Texture streaming cannot be disabled when FTexture2DMipMap::FCompactBulkData is used"));
+				bWarned = true;
+			}
+			CVarSetTextureStreaming.AsVariable()->Set(1, ECVF_SetByCode);
+			bUseTextureStreaming = true;
+		}
 	}
 
 	if( !GRHISupportsTextureStreaming || IsRunningDedicatedServer() )
 	{
+		if (TEXTURE2DMIPMAP_USE_COMPACT_BULKDATA)
+		{
+			UE_LOG(LogContentStreaming, Fatal, TEXT("RHI or target doesn't support texture streaming but cannot disable because TEXTURE2DMIPMAP_USE_COMPACT_BULKDATA is 1"));
+		}
+
 		bUseTextureStreaming = false;
 
 		// some code relies on r.TextureStreaming so we're going to disable it here to reflect the hardware capabilities and system needs
diff --git a/Engine/Source/Runtime/Engine/Public/TextureResource.h b/Engine/Source/Runtime/Engine/Public/TextureResource.h
--- a/Engine/Source/Runtime/Engine/Public/TextureResource.h
+++ b/Engine/Source/Runtime/Engine/Public/TextureResource.h
@@ -27,7 +27,10 @@ class UTexture2D;
 /** Maximum number of slices in texture source art. */
 #define MAX_TEXTURE_SOURCE_SLICES 6
 
-#define TEXTURE2DMIPMAP_USE_COMPACT_BULKDATA !(WITH_EDITORONLY_DATA)
+#ifndef FORCE_ENABLE_TEXTURE_STREAMING
+#define FORCE_ENABLE_TEXTURE_STREAMING 0
+#endif
+#define TEXTURE2DMIPMAP_USE_COMPACT_BULKDATA (!(WITH_EDITORONLY_DATA) && !(UE_SERVER) && FORCE_ENABLE_TEXTURE_STREAMING && PLATFORM_SUPPORTS_TEXTURE_STREAMING)
 
 /**
  * A 2D texture mip-map.
@@ -61,7 +64,7 @@ struct FTexture2DMipMap
 		bool IsBulkDataLoaded() const { return IsInlined(); }
 		bool IsStoredCompressedOnDisk() const { return !!(BulkDataFlags & BULKDATA_SerializeCompressed); }
 		const void* LockReadOnly() const { return TexelData; }
-		void* Lock(EBulkDataLockFlags LockFlags) { return TexelData; }
+		void* Lock(uint32 LockFlags) { return TexelData; }
 		void Unlock() const;
 		void* Realloc(int32 NumBytes);
 		void GetCopy(void** Dest, bool bDiscardInternalCopy = true);
