Workaround corrupted particle indices or unallocated particle data.

diff --git a/Engine/Source/Runtime/Engine/Private/Particles/ParticleEmitterInstances.cpp b/Engine/Source/Runtime/Engine/Private/Particles/ParticleEmitterInstances.cpp
--- a/Engine/Source/Runtime/Engine/Private/Particles/ParticleEmitterInstances.cpp
+++ b/Engine/Source/Runtime/Engine/Private/Particles/ParticleEmitterInstances.cpp
@@ -753,16 +753,16 @@ bool FParticleEmitterInstance::Resize(int32 NewMaxActiveParticles, bool bSetMaxA
 		{
 			SCOPE_CYCLE_COUNTER(STAT_ParticleMemTime);
 
-		ParticleData = (uint8*) FMemory::Realloc(ParticleData, ParticleStride * NewMaxActiveParticles);
-		check(ParticleData);
+			ParticleData = (uint8*) FMemory::Realloc(ParticleData, ParticleStride * NewMaxActiveParticles);
+			check(ParticleData);
 
-		// Allocate memory for indices.
-		if (ParticleIndices == NULL)
-		{
-			// Make sure that we clear all when it is the first alloc
-			MaxActiveParticles = 0;
-		}
-		ParticleIndices	= (uint16*) FMemory::Realloc(ParticleIndices, sizeof(uint16) * (NewMaxActiveParticles + 1));
+			// Allocate memory for indices.
+			if (ParticleIndices == NULL)
+			{
+				// Make sure that we clear all when it is the first alloc
+				MaxActiveParticles = 0;
+			}
+			ParticleIndices	= (uint16*) FMemory::Realloc(ParticleIndices, sizeof(uint16) * (NewMaxActiveParticles + 1));
 		}
 
 		// Fill in default 1:1 mapping.
@@ -2103,6 +2103,49 @@ float FParticleEmitterInstance::Spawn(float DeltaTime)
 	return SpawnFraction;
 }
 
+/**
+* Fixup particle indices to only have valid entries.
+*/
+void FParticleEmitterInstance::FixupParticleIndices()
+{
+	// Something is wrong and particle data are be invalid. Try to fix-up things.
+	TBitArray<> UsedIndices(false, MaxActiveParticles);
+
+	for (int32 i = 0; i < ActiveParticles; ++i)
+	{
+		const uint16 UsedIndex = ParticleIndices[i];
+		if (UsedIndex < MaxActiveParticles && UsedIndices[UsedIndex] == 0)
+		{
+			UsedIndices[UsedIndex] = 1;
+		}
+		else
+		{
+			if (i != ActiveParticles - 1)
+			{
+				// Remove this bad or duplicated index
+				ParticleIndices[i] = ParticleIndices[ActiveParticles - 1];
+			}
+			// Decrease particle count.
+			--ActiveParticles;
+
+			// Retry the new index.
+			--i;
+		}
+	}
+
+	for (int32 i = ActiveParticles; i < MaxActiveParticles; ++i)
+	{
+		const int32 FreeIndex = UsedIndices.FindAndSetFirstZeroBit();
+		if (ensure(FreeIndex != INDEX_NONE))
+		{
+			ParticleIndices[i] =  (uint16)FreeIndex;
+		}
+		else // Can't really handle that.
+		{
+			ParticleIndices[i] = (uint16)i;
+		}
+	}
+}
 /**
  * Spawn the indicated number of particles.
  *
@@ -2117,8 +2160,12 @@ void FParticleEmitterInstance::SpawnParticles( int32 Count, float StartTime, flo
 {
 	UParticleLODLevel* LODLevel = GetCurrentLODLevelChecked();
 
-	check( ActiveParticles + Count <= MaxActiveParticles );
-	check( LODLevel->EventGenerator != NULL || EventPayload == NULL );
+	check(ActiveParticles <= MaxActiveParticles);
+	check(LODLevel->EventGenerator != NULL || EventPayload == NULL);
+
+	// Ensure we don't access particle beyond what is allocated.
+	ensure( ActiveParticles + Count <= MaxActiveParticles );
+	Count = FMath::Min<int32>(Count, MaxActiveParticles - ActiveParticles);
 
 	if (EventPayload && EventPayload->bBurstEventsPresent && Count > 0)
 	{
@@ -2133,8 +2180,30 @@ void FParticleEmitterInstance::SpawnParticles( int32 Count, float StartTime, flo
 		const float InterpIncrement = (Count > 0 && Increment > 0.0f) ? (1.0f / (float)Count) : 0.0f;
 		for (int32 i = 0; i < Count; i++)
 		{
-			check(ActiveParticles <= MaxActiveParticles);
-			DECLARE_PARTICLE_PTR(Particle, ParticleData + ParticleStride * ParticleIndices[ActiveParticles]);
+			// Workaround to released data.
+			if (!ensure(ParticleData && ParticleIndices))
+			{
+				static bool bErrorReported = false;
+				if (!bErrorReported)
+				{
+					UE_LOG(LogParticles, Error, TEXT("Detected null particles. Template : %s, Component %s"), *GetNameSafe(Component && Component->IsValidLowLevel() ? Component->Template : nullptr), *GetFullNameSafe(Component));
+					bErrorReported = true;
+				}
+				ActiveParticles = 0;
+				Count = 0;
+				continue;
+			}
+
+			// Workaround to corrupted indices.
+			uint16 NextFreeIndex = ParticleIndices[ActiveParticles];
+			if (!ensure(NextFreeIndex < MaxActiveParticles))
+			{
+				UE_LOG(LogParticles, Error, TEXT("Detected corrupted particle indices. Template : %s, Component %s"), *GetNameSafe(Component && Component->IsValidLowLevel() ? Component->Template : nullptr), *GetFullNameSafe(Component));
+				FixupParticleIndices();
+				NextFreeIndex = ParticleIndices[ActiveParticles];
+			}
+
+			DECLARE_PARTICLE_PTR(Particle, ParticleData + ParticleStride * NextFreeIndex);
 			const uint32 CurrentParticleIndex = ActiveParticles++;
 
 			if (bLegacySpawnBehavior)
diff --git a/Engine/Source/Runtime/Engine/Public/ParticleEmitterInstances.h b/Engine/Source/Runtime/Engine/Public/ParticleEmitterInstances.h
--- a/Engine/Source/Runtime/Engine/Public/ParticleEmitterInstances.h
+++ b/Engine/Source/Runtime/Engine/Public/ParticleEmitterInstances.h
@@ -757,6 +757,12 @@ protected:
 	 */
 	UMaterialInterface* GetCurrentMaterial();
 
+
+	/**
+	 * Fixup particle indices to only have valid entries.
+	 */
+	void FixupParticleIndices();
+
 };
 
 #if STATS
