Fast TArray Delta Serialization

diff --git a/Engine/Plugins/Runtime/GameplayAbilities/Source/GameplayAbilities/Private/GameplayEffect.cpp b/Engine/Plugins/Runtime/GameplayAbilities/Source/GameplayAbilities/Private/GameplayEffect.cpp
--- a/Engine/Plugins/Runtime/GameplayAbilities/Source/GameplayAbilities/Private/GameplayEffect.cpp
+++ b/Engine/Plugins/Runtime/GameplayAbilities/Source/GameplayAbilities/Private/GameplayEffect.cpp
@@ -1880,6 +1880,7 @@ FActiveGameplayEffectsContainer::FActiveGameplayEffectsContainer()
 	, PendingGameplayEffectHead(nullptr)
 {
 	PendingGameplayEffectNext = &PendingGameplayEffectHead;
+	bUseDeltaStructSerialization = true;
 }
 
 FActiveGameplayEffectsContainer::~FActiveGameplayEffectsContainer()
diff --git a/Engine/Source/Runtime/Core/Private/Misc/NetworkVersion.cpp b/Engine/Source/Runtime/Core/Private/Misc/NetworkVersion.cpp
--- a/Engine/Source/Runtime/Core/Private/Misc/NetworkVersion.cpp
+++ b/Engine/Source/Runtime/Core/Private/Misc/NetworkVersion.cpp
@@ -18,7 +18,7 @@ FString FNetworkVersion::ProjectVersion = TEXT("1.0.0");
 bool FNetworkVersion::bHasCachedNetworkChecksum			= false;
 uint32 FNetworkVersion::CachedNetworkChecksum			= 0;
 
-uint32 FNetworkVersion::EngineNetworkProtocolVersion	= HISTORY_NETEXPORT_SERIALIZE_FIX;
+uint32 FNetworkVersion::EngineNetworkProtocolVersion	= HISTORY_ENGINENETVERSION_LATEST;
 uint32 FNetworkVersion::GameNetworkProtocolVersion		= 0;
 
 uint32 FNetworkVersion::EngineCompatibleNetworkProtocolVersion		= HISTORY_REPLAY_BACKWARDS_COMPAT;
diff --git a/Engine/Source/Runtime/Core/Public/Misc/NetworkVersion.h b/Engine/Source/Runtime/Core/Public/Misc/NetworkVersion.h
--- a/Engine/Source/Runtime/Core/Public/Misc/NetworkVersion.h
+++ b/Engine/Source/Runtime/Core/Public/Misc/NetworkVersion.h
@@ -42,6 +42,12 @@ enum EEngineNetworkVersionHistory
 	HISTORY_ACKS_INCLUDED_IN_HEADER = 8,			// Bump version since acks are now sent as part of the header
 	HISTORY_NETEXPORT_SERIALIZATION = 9,			// Bump version due to serialization change to FNetFieldExport
 	HISTORY_NETEXPORT_SERIALIZE_FIX = 10,			// Bump version to fix net field export name serialization 
+	HISTORY_FAST_ARRAY_DELTA_STRUCT = 11,			// Bump version to allow fast array serialization, delta struct serialization.
+
+	// New history items go above here.
+
+	HISTORY_ENGINENETVERSION_PLUS_ONE,
+	HISTORY_ENGINENETVERSION_LATEST = HISTORY_ENGINENETVERSION_PLUS_ONE - 1,
 };
 
 struct CORE_API FNetworkVersion
diff --git a/Engine/Source/Runtime/CoreUObject/Private/UObject/CoreNet.cpp b/Engine/Source/Runtime/CoreUObject/Private/UObject/CoreNet.cpp
--- a/Engine/Source/Runtime/CoreUObject/Private/UObject/CoreNet.cpp
+++ b/Engine/Source/Runtime/CoreUObject/Private/UObject/CoreNet.cpp
@@ -10,9 +10,6 @@
 
 DEFINE_LOG_CATEGORY_STATIC(LogCoreNet, Log, All);
 
-DEFINE_STAT(STAT_NetSerializeFastArray);
-DEFINE_STAT(STAT_NetSerializeFastArray_BuildMap);
-
 /*-----------------------------------------------------------------------------
 	FClassNetCache implementation.
 -----------------------------------------------------------------------------*/
@@ -538,3 +535,28 @@ const TCHAR* LexToString(const EChannelCloseReason Value)
 
 	return TEXT("Unknown");
 }
+
+void INetSerializeCB::NetSerializeStruct(
+	class UScriptStruct* Struct,
+	class FBitArchive& Ar,
+	class UPackageMap* Map,
+	void* Data,
+	bool& bHasUnmapped)
+{
+	FNetDeltaSerializeInfo Params;
+	Params.Struct = Struct;
+	Params.Map = Map;
+	Params.Data = Data;
+
+	if (Ar.IsSaving())
+	{
+		Params.Writer = static_cast<FBitWriter*>(&Ar);
+	}
+	else
+	{
+		Params.Reader = static_cast<FBitReader*>(&Ar);
+	}
+
+	NetSerializeStruct(Params);
+	bHasUnmapped = Params.bOutHasMoreUnmapped;
+}
\ No newline at end of file
diff --git a/Engine/Source/Runtime/CoreUObject/Public/UObject/CoreNet.h b/Engine/Source/Runtime/CoreUObject/Public/UObject/CoreNet.h
--- a/Engine/Source/Runtime/CoreUObject/Public/UObject/CoreNet.h
+++ b/Engine/Source/Runtime/CoreUObject/Public/UObject/CoreNet.h
@@ -355,13 +355,74 @@ public:
 private:
 };
 
+struct FNetDeltaSerializeInfo;
 
-class INetSerializeCB
+/**
+ * An interface for handling serialization of Structs for networking.
+ *
+ * See notes in NetSerialization.h
+ */
+class COREUOBJECT_API INetSerializeCB
 {
+protected:
+
+	using FGuidReferencesMap = TMap<int32, class FGuidReferences>;
+
 public:
+
 	INetSerializeCB() { }
 
-	virtual void NetSerializeStruct( UScriptStruct* Struct, FBitArchive& Ar, UPackageMap* Map, void* Data, bool& bHasUnmapped ) = 0;
+	virtual ~INetSerializeCB() {}
+
+	/**
+	 * Serializes an entire struct to / from the given archive.
+	 * It is up to callers to manage Guid References created during reads.
+	 *
+	 * @param Params		NetDeltaSerialization Params to use.
+	 *						Object must be valid.
+	 *						Data must be valid.
+	 *						Connection must be valid.
+	 *						Map must be valid.
+	 *						Struct must point to the UScriptStruct of Data.
+	 *						Either Reader or Writer (but not both) must be valid.
+	 *						bOutHasMoreUnmapped will be used to return whether or not we have we have unmapped guids.
+	 *						Only used when reading.
+	 */
+	virtual void NetSerializeStruct(FNetDeltaSerializeInfo& Params) = 0;
+
+	UE_DEPRECATED(4.23, "Please use the version of NetSerializeStruct that accepts an FNetDeltaSerializeInfo reference")
+	virtual void NetSerializeStruct(
+		class UScriptStruct* Struct,
+		class FBitArchive& Ar,
+		class UPackageMap* Map,
+		void* Data,
+		bool& bHasUnmapped);
+
+	/**
+	 * Gathers any guid references for a FastArraySerializer.
+	 * @see GuidReferences.h for more info.
+	 */
+	virtual void GatherGuidReferencesForFastArray(struct FFastArrayDeltaSerializeParams& Params) = 0;
+
+	/**
+	 * Moves a previously mapped guid to an unmapped state for a FastArraySerializer.
+	 * @see GuidReferences.h for more info.
+	 *
+	 * @return True if the guid was found and unmapped.
+	 */
+	virtual bool MoveGuidToUnmappedForFastArray(struct FFastArrayDeltaSerializeParams& Params) = 0;
+
+	/**
+	 * Updates any unmapped guid references for a FastArraySerializer.
+	 * @see GuidReferences.h for more info.
+	 */
+	virtual void UpdateUnmappedGuidsForFastArray(struct FFastArrayDeltaSerializeParams& Params) = 0;
+
+	/**
+	 * Similar to NetSerializeStruct, except serializes an entire FastArraySerializer at once
+	 * instead of element by element.
+	 */
+	virtual bool NetDeltaSerializeForFastArray(struct FFastArrayDeltaSerializeParams& Params) = 0;
 };
 
 
@@ -390,62 +451,81 @@ public:
  */
 struct FNetDeltaSerializeInfo
 {
-	FNetDeltaSerializeInfo()
-	{
-		Writer		= NULL;
-		Reader		= NULL;
-
-		NewState	= NULL;
-		OldState	= NULL;
-		Map			= NULL;
-		Data		= NULL;
-
-		Struct		= NULL;
-
-		NetSerializeCB = NULL;
-
-		bUpdateUnmappedObjects		= false;
-		bOutSomeObjectsWereMapped	= false;
-		bCalledPreNetReceive		= false;
-		bOutHasMoreUnmapped			= false;
-		bGuidListsChanged			= false;
-		bIsWritingOnClient			= false;
-		Object						= nullptr;
-		GatherGuidReferences		= nullptr;
-		TrackedGuidMemoryBytes		= nullptr;
-		MoveGuidToUnmapped			= nullptr;
-	}
+	/** Used when writing */
+	FBitWriter* Writer = nullptr;
+
+	/** Used when reading */
+	FBitReader* Reader = nullptr;
+
+	/** SharedPtr to new base state created by NetDeltaSerialize. Used when writing.*/
+	TSharedPtr<INetDeltaBaseState>* NewState = nullptr;
+
+	/** Pointer to the previous base state. Used when writing. */
+	INetDeltaBaseState* OldState = nullptr;
+
+	/** PackageMap that can be used to serialize objects and track Guid References. Used primarily when reading. */
+	class UPackageMap* Map = nullptr;
+
+	/** Connection that we're currently serializing data for. */
+	class UNetConnection* Connection = nullptr;
+
+	/** Pointer to the struct that we're serializing.*/
+	void* Data = nullptr;
+
+	/** Type of struct that we're serializing. */
+	class UStruct* Struct = nullptr;
+
+	/** Pointer to a NetSerializeCB implementation that can be used when serializing. */
+	INetSerializeCB* NetSerializeCB = nullptr;
 
-	// Used when writing
-	FBitWriter*						Writer;
+	/** If true, we are updating unmapped objects */
+	bool bUpdateUnmappedObjects = false;
 
-	// Used for when reading
-	FBitReader*						Reader;
+	/** If true, then we successfully mapped some unmapped objects. */
+	bool bOutSomeObjectsWereMapped = false;
 
-	TSharedPtr<INetDeltaBaseState>*	NewState;		// SharedPtr to new base state created by NetDeltaSerialize.
-	INetDeltaBaseState*				OldState;				// Pointer to the previous base state.
-	UPackageMap*					Map;
-	void*							Data;
+	/** Whether or not PreNetReceive has been called on the owning object. */
+	bool bCalledPreNetReceive = false;
 
-	// Only used for fast TArray replication
-	UStruct*						Struct;
+	/** Whether or not there are still some outstanding unmapped objects referenced by the struct. */
+	bool bOutHasMoreUnmapped = false;
 
-	INetSerializeCB*				NetSerializeCB;
+	/** Whether or not we changed Guid / Object references. Used when reading. */
+	bool bGuidListsChanged = false;
+
+	/** Whether or not we're sending / writing data from the client. */
+	bool bIsWritingOnClient = false;
+
+	//~ TODO: This feels hacky, and a better alternative might be something like connection specific
+	//~ capabilities.
+
+	/** Whether or not we support FFastArraySerializer::FastArrayDeltaSerialize_DeltaSerializeStructs */
+	bool bSupportsFastArrayDeltaStructSerialization = false;
+
+	/** The object that owns the struct we're serializing. */
+	UObject* Object = nullptr;
+
+	/**
+	 * When non-null, this indicates that we're gathering Guid References.
+	 * Any Guids the struct is referencing should be added.
+	 * This may contain gathered Guids from other structs, so do not clear this set.
+	 */
+	TSet<FNetworkGUID>* GatherGuidReferences = nullptr;
+
+	/**
+	 * When we're gathering guid references, ny memory used to track Guids can be added to this.
+	 * This may be tracking Guid memory from other structs, so do not reset this.
+	 * Note, this is not guaranteed to be valid when GatherGuidReferences is.
+	 */
+	int32* TrackedGuidMemoryBytes = nullptr;
 
-	bool							bUpdateUnmappedObjects;		// If true, we are wanting to update unmapped objects
-	bool							bOutSomeObjectsWereMapped;
-	bool							bCalledPreNetReceive;
-	bool							bOutHasMoreUnmapped;
-	bool							bGuidListsChanged;
-	bool							bIsWritingOnClient;
-	UObject*						Object;
+	/** When non-null, this indicates the given Guid has become unmapped and any references to it should be updated. */
+	const FNetworkGUID* MoveGuidToUnmapped = nullptr;
 
-	TSet< FNetworkGUID >*			GatherGuidReferences;
-	int32*							TrackedGuidMemoryBytes;
-	const FNetworkGUID*				MoveGuidToUnmapped;
+	int32 PropertyRepIndex = INDEX_NONE;
 
 	// Debugging variables
-	FString							DebugName;
+	FString DebugName;
 };
 
 
diff --git a/Engine/Source/Runtime/CoreUObject/Public/UObject/UObjectGlobals.h b/Engine/Source/Runtime/CoreUObject/Public/UObject/UObjectGlobals.h
--- a/Engine/Source/Runtime/CoreUObject/Public/UObject/UObjectGlobals.h
+++ b/Engine/Source/Runtime/CoreUObject/Public/UObject/UObjectGlobals.h
@@ -35,13 +35,6 @@ DECLARE_CYCLE_STAT_EXTERN(TEXT("~UObject"),STAT_DestroyObject,STATGROUP_Object,
 DECLARE_DWORD_COUNTER_STAT_EXTERN(TEXT("FindObject"),STAT_FindObject,STATGROUP_ObjectVerbose, );
 DECLARE_DWORD_COUNTER_STAT_EXTERN(TEXT("FindObjectFast"),STAT_FindObjectFast,STATGROUP_ObjectVerbose, );
 
-/**
- * Network stats counters
- */
-
-DECLARE_CYCLE_STAT_EXTERN(TEXT("NetSerializeFast Array"),STAT_NetSerializeFastArray,STATGROUP_ServerCPU, COREUOBJECT_API);
-DECLARE_CYCLE_STAT_EXTERN(TEXT("NetSerializeFast Array BuildMap"),STAT_NetSerializeFastArray_BuildMap,STATGROUP_ServerCPU, COREUOBJECT_API);
-
 #define	INVALID_OBJECT	(UObject*)-1
 
 
diff --git a/Engine/Source/Runtime/Engine/Classes/Engine/NetSerialization.h b/Engine/Source/Runtime/Engine/Classes/Engine/NetSerialization.h
--- a/Engine/Source/Runtime/Engine/Classes/Engine/NetSerialization.h
+++ b/Engine/Source/Runtime/Engine/Classes/Engine/NetSerialization.h
@@ -16,12 +16,17 @@
 #include "UObject/CoreNet.h"
 #include "EngineLogs.h"
 #include "Containers/ArrayView.h"
+#include "Net/GuidReferences.h"
 #include "NetSerialization.generated.h"
 
 class Error;
 struct FFastArraySerializer;
 struct FFastArraySerializerItem;
 
+DECLARE_CYCLE_STAT_EXTERN(TEXT("NetSerializeFast Array"), STAT_NetSerializeFastArray, STATGROUP_ServerCPU, ENGINE_API);
+DECLARE_CYCLE_STAT_EXTERN(TEXT("NetSerializeFast Array BuildMap"), STAT_NetSerializeFastArray_BuildMap, STATGROUP_ServerCPU, ENGINE_API);
+DECLARE_CYCLE_STAT_EXTERN(TEXT("NetSerializeFast Array Delta Struct"), STAT_NetSerializeFastArray_DeltaStruct, STATGROUP_ServerCPU, ENGINE_API);
+
 /**
  *	===================== NetSerialize and NetDeltaSerialize customization. =====================
  *
@@ -264,7 +269,10 @@ class FNetFastTArrayBaseState : public INetDeltaBaseState
 {
 public:
 
-	FNetFastTArrayBaseState() : ArrayReplicationKey(INDEX_NONE) { }
+	FNetFastTArrayBaseState():
+		ArrayReplicationKey(INDEX_NONE),
+		ChangelistHistory(INDEX_NONE)
+	{}
 
 	virtual bool IsStateEqual(INetDeltaBaseState* OtherState)
 	{
@@ -280,8 +288,12 @@ public:
 		return true;
 	}
 
+	/** Maps an element's Replication ID to Index. */
 	TMap<int32, int32> IDToCLMap;
+
 	int32 ArrayReplicationKey;
+
+	int32 ChangelistHistory;
 };
 
 
@@ -359,10 +371,28 @@ struct FFastArraySerializerItem
 /** Struct for holding guid references */
 struct FFastArraySerializerGuidReferences
 {
-	TSet< FNetworkGUID >		UnmappedGUIDs;		// List of guids that were unmapped so we can quickly check
-	TSet< FNetworkGUID >		MappedDynamicGUIDs;	// List of guids that were mapped so we can move them to unmapped when necessary (i.e. actor channel closes)
-	TArray< uint8 >				Buffer;				// Buffer of data to re-serialize when the guids are mapped
-	int32						NumBufferBits;		// Number of bits in the buffer
+	/** List of guids that were unmapped so we can quickly check */
+	TSet<FNetworkGUID> UnmappedGUIDs;
+
+	/** List of guids that were mapped so we can move them to unmapped when necessary (i.e. actor channel closes) */
+	TSet<FNetworkGUID> MappedDynamicGUIDs;
+
+	/** Buffer of data to re-serialize when the guids are mapped */
+	TArray<uint8> Buffer;
+
+	/** Number of bits in the buffer */
+	int32 NumBufferBits;
+};
+
+/**
+ * Struct used only in FFastArraySerializer::FastArrayDeltaSerialize, however, declaring it within the templated function
+ * causes crashes on some clang compilers
+ */
+struct FFastArraySerializer_FastArrayDeltaSerialize_FIdxIDPair
+{
+	FFastArraySerializer_FastArrayDeltaSerialize_FIdxIDPair(int32 _idx, int32 _id) : Idx(_idx), ID(_id) { }
+	int32	Idx;
+	int32	ID;
 };
 
 /** Base struct for wrapping the array used in Fast TArray Replication */
@@ -375,14 +405,28 @@ struct FFastArraySerializer
 		: IDCounter(0)
 		, ArrayReplicationKey(0)
 		, CachedNumItems(INDEX_NONE)
-		, CachedNumItemsToConsiderForWriting(INDEX_NONE) { }
+		, CachedNumItemsToConsiderForWriting(INDEX_NONE)
+		, bUseDeltaStructSerialization(false)
+		, bCachedUseDeltaStructSerialization(false)
+		, bHasBeenSerialized(false) { }
+
+	~FFastArraySerializer() {}
+
+	/** Maps Element ReplicationID to Array Index.*/
+	TMap<int32, int32> ItemMap;
+
+	/** Counter used to assign IDs to new elements. */
+	int32 IDCounter;
 
-	TMap<int32, int32>	ItemMap;
-	int32				IDCounter;
-	int32				ArrayReplicationKey;
+	/** Counter used to track array replication. */
+	int32 ArrayReplicationKey;
 
-	TMap< int32, FFastArraySerializerGuidReferences > GuidReferencesMap;	// List of items that need to be re-serialized when the referenced objects are mapped
+	/** List of items that need to be re-serialized when the referenced objects are mapped */
+	TMap<int32, FFastArraySerializerGuidReferences> GuidReferencesMap;
 
+	/** List of items that need to be re-serialized when the referenced objects are mapped.*/
+	TMap<int32, FGuidReferencesMap> GuidReferencesMap_StructDelta;
+	
 	/** This must be called if you add or change an item in the array */
 	void MarkItemDirty(FFastArraySerializerItem & Item)
 	{
@@ -390,7 +434,9 @@ struct FFastArraySerializer
 		{
 			Item.ReplicationID = ++IDCounter;
 			if (IDCounter == INDEX_NONE)
+			{
 				IDCounter++;
+			}
 		}
 
 		Item.ReplicationKey++;
@@ -412,11 +458,23 @@ struct FFastArraySerializer
 	{
 		ArrayReplicationKey++;
 		if (ArrayReplicationKey == INDEX_NONE)
+		{
 			ArrayReplicationKey++;
+		}
 	}
 
-	template< typename Type, typename SerializerType >
-	static bool FastArrayDeltaSerialize( TArray<Type> &Items, FNetDeltaSerializeInfo& Parms, SerializerType& ArraySerializer );
+	/**
+	 * Performs "standard" delta serialization on the items in the FastArraySerializer.
+	 * This method relies more on the INetSerializeCB interface and custom logic and sends all properties
+	 * that aren't marked as SkipRep, regardless of whether or not they've changed.
+	 * This will be less CPU intensive, but require more bandwidth.
+	 *
+	 * @param Items				Array of items owned by ArraySerializer.
+	 * @param Parms				Set of parms that dictate what serialization will do / return.
+	 * @param ArraySerializer	The typed subclass of FFastArraySerializer that we're serializing.
+	 */
+	template<typename Type, typename SerializerType>
+	static bool FastArrayDeltaSerialize(TArray<Type>& Items, FNetDeltaSerializeInfo& Parms, SerializerType& ArraySerializer);
 
 	/**
 	 * Called before removing elements and after the elements themselves are notified.  The indices are valid for this function call only!
@@ -443,73 +501,571 @@ struct FFastArraySerializer
 	* Helper function for FastArrayDeltaSerialize to consolidate the logic of whether to consider writing an item in a fast TArray during network serialization.
 	* For client replay recording, we don't want to write any items that have been added to the array predictively.
 	*/
-	template< typename Type, typename SerializerType >
-	bool ShouldWriteFastArrayItem(const Type& Item, const bool bIsWritingOnClient)
+	template<typename Type, typename SerializerType>
+	bool ShouldWriteFastArrayItem(const Type& Item, const bool bIsWritingOnClient) const
 	{
-		if (bIsWritingOnClient)
-		{
-			return Item.ReplicationID != INDEX_NONE;
-		}
+		return !bIsWritingOnClient || Item.ReplicationID != INDEX_NONE;
+	}
 
-		return true;
+	const bool IsUsingStructDeltaSerialization() const
+	{
+		return bHasBeenSerialized ? bCachedUseDeltaStructSerialization : bUseDeltaStructSerialization;
+	}
+
+	const bool HasBeenSerialized() const
+	{
+		return bHasBeenSerialized;
 	}
 
 private:
+
+	static constexpr int32 MAX_NUM_CHANGED = 2048;
+	static constexpr int32 MAX_NUM_DELETED = 2048;
+
+	/** Struct containing common header data that is written / read when serializing Fast Arrays. */
+	struct FFastArraySerializerHeader
+	{
+		/** The current ArrayReplicationKey. */
+		int32 ArrayReplicationKey;
+
+		/** The previous ArrayReplicationKey. */
+		int32 BaseReplicationKey;
+
+		/** The number of changed elements (adds or removes). */
+		int32 NumChanged;
+
+		/**
+		 * The list of deleted elements.
+		 * When writing, this will be treated as IDs that are translated to indices prior to serialization.
+		 * When reading, this will be actual indices.
+		 */
+		TArray<int32, TInlineAllocator<8>> DeletedIndices;
+	};
+
+	/**
+	 * Helper struct that contains common methods / logic for standard Fast Array serialization and
+	 * Delta Struct Fast Array serialization.
+	 */
+	template<typename Type, typename SerializerType>
+	struct TFastArraySerializeHelper
+	{
+		/** Array element type struct. */
+		UScriptStruct* Struct;
+
+		/** Set of array elements we're serializing. */
+		TArray<Type>& Items;
+
+		/** The actual FFastArraySerializer struct we're serializing. */
+		SerializerType& ArraySerializer;
+
+		/** Cached DeltaSerialize params. */
+		FNetDeltaSerializeInfo& Parms;
+
+		/**
+		 * Conditionally rebuilds the ID to Index map for items.
+		 * This is generally only necessary on first serialization, or if we receive deletes and
+		 * can no longer trust our ordering is correct.
+		 */
+		void ConditionalRebuildItemMap();
+
+		/** Calculates the number of Items that actually need to be written. */
+		int32 CalcNumItemsForConsideration() const;
+
+		/** Conditionally logs the important state of the serializer. For debug purposes only. */
+		void ConditionalLogSerializerState(const TMap<int32, int32>* OldIDToKeyMap) const;
+
+		/**
+		 * Checks to see if the ArrayReplicationKey has changed, and if so creates a new DeltaState that
+		 * is passed out to the caller.
+		 * Note, this state may just be a copy of a previous state, or a brand new state.
+		 *
+		 * @return	True if the keys were different and a state was created.
+		 *			False if the keys were the same, and we can skip serialization.
+		 */
+		bool ConditionalCreateNewDeltaState(const TMap<int32, int32>& OldIDToKeyMap, const int32 BaseReplicationKey);
+
+		/**
+		 * Iterates over the current set of properties, comparing their keys with our old state, to figure
+		 * out which have changed and need to be serialized. Also populates a list of elements that are
+		 * no longer in our list (by ID).
+		 */
+		void BuildChangedAndDeletedBuffers(
+			TMap<int32, int32>& NewIDToKeyMap,
+			const TMap<int32, int32>* OldIDToKeyMap,
+			TArray<FFastArraySerializer_FastArrayDeltaSerialize_FIdxIDPair, TInlineAllocator<8>>& ChangedElements,
+			TArray<int32, TInlineAllocator<8>>& DeletedElements);
+
+		/** Writes out a FFastArraySerializerHeader */
+		void WriteDeltaHeader(FFastArraySerializerHeader& Header) const;
+
+		/** Reads in a FFastArraySerializerHeader */
+		bool ReadDeltaHeader(FFastArraySerializerHeader& Header) const;
+
+		/**
+		 * Manages any cleanup work that needs to be done after receiving elements,
+		 * such as looking for items that were implicitly deleted, removing all deleted items,
+		 * firing off any PostReceive / PostDeleted events, etc.
+		 */
+		template<typename GuidMapType>
+		void PostReceiveCleanup(
+			FFastArraySerializerHeader& Header,
+			TArray<int32, TInlineAllocator<8>>& ChangedIndices,
+			TArray<int32, TInlineAllocator<8>>& AddedIndices,
+			GuidMapType& GuidMap);
+	};
+
+	/**
+	 * Performs "struct delta" serialization on the items in the FastArraySerializer.
+	 * This method relies more directly on FRepLayout for management, and will only send properties
+	 * that have changed since the last update.
+	 * This is potentially more CPU intensive since we'll be doing comparisons, but should require less bandwidth.
+	 *
+	 * For this method to work, the following **must** be true:
+	 *	- Your array of FFastArraySerializerItems must be a top level UPROPERTY within your FFastArraySerializer.
+	 *	- Your array of FFastArraySerializerItems must **not** be marked RepSkip.
+	 *	- Your array of FFastArraySerializerItems must be the **only** replicated array of FFastArraySerializerItems within the FastArraySerializer.
+	 *		Note, it's OK to have multiple arrays of items, as long as only one is replicated (all others **must** be marked RepSkip).
+	 *	- Your FFastArraySerializer must not be nested in a static array.
+	 *	- Your array of FFastArraySerializerItem must not be nested in a static array.
+	 *
+	 * @param Items				Array of items owned by ArraySerializer.
+	 * @param Parms				Set of parms that dictate what serialization will do / return.
+	 * @param ArraySerializer	The typed subclass of FFastArraySerializer that we're serializing.
+	 */
+	template<typename Type, typename SerializerType>
+	static bool FastArrayDeltaSerialize_DeltaSerializeStructs(TArray<Type>& Items, FNetDeltaSerializeInfo& Parms, SerializerType& ArraySerializer);
+
+private:
+
 	// Cached item counts, used for fast sanity checking when writing.
-	int32				CachedNumItems;
-	int32				CachedNumItemsToConsiderForWriting;
+	int32 CachedNumItems;
+	int32 CachedNumItemsToConsiderForWriting;
+
+protected:
+
+	uint32 bUseDeltaStructSerialization : 1;
+
+private:
+
+	uint32 bCachedUseDeltaStructSerialization : 1;
+
+	uint32 bHasBeenSerialized : 1;
 };
 
-// Struct used only in FFastArraySerializer::FastArrayDeltaSerialize, however, declaring it within the templated function
-// causes crashes on some clang compilers
-struct FFastArraySerializer_FastArrayDeltaSerialize_FIdxIDPair
+template<typename Type, typename SerializerType>
+void FFastArraySerializer::TFastArraySerializeHelper<Type, SerializerType>::ConditionalRebuildItemMap()
 {
-	FFastArraySerializer_FastArrayDeltaSerialize_FIdxIDPair(int32 _idx, int32 _id) : Idx(_idx), ID(_id) { }
-	int32	Idx;
-	int32	ID;
+	if ((Parms.bUpdateUnmappedObjects || Parms.Writer == NULL) && ArraySerializer.ItemMap.Num() != Items.Num())
+	{
+		SCOPE_CYCLE_COUNTER(STAT_NetSerializeFastArray_BuildMap);
+		UE_LOG(LogNetFastTArray, Log, TEXT("FastArrayDeltaSerialize: Recreating Items map. Struct: %s, Items.Num: %d Map.Num: %d"), *Struct->GetOwnerStruct()->GetName(), Items.Num(), ArraySerializer.ItemMap.Num());
+		ArraySerializer.ItemMap.Empty();
+		for (int32 i = 0; i < Items.Num(); ++i)
+		{
+			if (Items[i].ReplicationID == INDEX_NONE)
+			{
+				if (Parms.Writer)
+				{
+					UE_LOG(LogNetFastTArray, Warning, TEXT("FastArrayDeltaSerialize: Item with uninitialized ReplicationID. Struct: %s, ItemIndex: %i"), *Struct->GetOwnerStruct()->GetName(), i);
+				}
+				else
+				{
+					// This is benign for clients, they may add things to their local array without assigning a ReplicationID
+					continue;
+				}
+			}
+
+			ArraySerializer.ItemMap.Add(Items[i].ReplicationID, i);
+		}
+	}
+}
+
+template<typename Type, typename SerializerType>
+int32 FFastArraySerializer::TFastArraySerializeHelper<Type, SerializerType>::CalcNumItemsForConsideration() const
+{
+	int32 Count = 0;
+
+	// Count the number of items in the current array that may be written. On clients, items that were predicted will be skipped.
+	for (const Type& Item : Items)
+	{
+		if (ArraySerializer.template ShouldWriteFastArrayItem<Type, SerializerType>(Item, Parms.bIsWritingOnClient))
+		{
+			Count++;
+		}
+	}
+
+	return Count;
 };
 
-/** The function that implements Fast TArray Replication  */
-template< typename Type, typename SerializerType >
-bool FFastArraySerializer::FastArrayDeltaSerialize( TArray<Type> &Items, FNetDeltaSerializeInfo& Parms, SerializerType& ArraySerializer )
+template<typename Type, typename SerializerType>
+void FFastArraySerializer::TFastArraySerializeHelper<Type, SerializerType>::ConditionalLogSerializerState(const TMap<int32, int32>* OldIDToKeyMap) const
 {
-	SCOPE_CYCLE_COUNTER(STAT_NetSerializeFastArray);
-	class UScriptStruct* InnerStruct = Type::StaticStruct();
+	// Log out entire state of current/base state
+	if (UE_LOG_ACTIVE(LogNetFastTArray, Log))
+	{
+		FString CurrentState = FString::Printf(TEXT("Current: %d "), ArraySerializer.ArrayReplicationKey);
+		for (const Type& Item : Items)
+		{
+			CurrentState += FString::Printf(TEXT("[%d/%d], "), Item.ReplicationID, Item.ReplicationKey);
+		}
+		UE_LOG(LogNetFastTArray, Log, TEXT("%s"), *CurrentState);
 
-	UE_LOG(LogNetFastTArray, Log, TEXT("FastArrayDeltaSerialize for %s. %s. %s"), *InnerStruct->GetName(), *InnerStruct->GetOwnerStruct()->GetName(), Parms.Reader ? TEXT("Reading") : TEXT("Writing"));
 
-	if ( Parms.bUpdateUnmappedObjects || Parms.Writer == NULL )
+		FString ClientStateStr = FString::Printf(TEXT("Client: %d "), Parms.OldState ? ((FNetFastTArrayBaseState*)Parms.OldState)->ArrayReplicationKey : 0);
+		if (OldIDToKeyMap)
+		{
+			for (const TPair<int32, int32>& KeyValuePair : *OldIDToKeyMap)
+			{
+				ClientStateStr += FString::Printf(TEXT("[%d/%d], "), KeyValuePair.Key, KeyValuePair.Value);
+			}
+		}
+		UE_LOG(LogNetFastTArray, Log, TEXT("%s"), *ClientStateStr);
+	}
+}
+
+template<typename Type, typename SerializerType>
+bool FFastArraySerializer::TFastArraySerializeHelper<Type, SerializerType>::ConditionalCreateNewDeltaState(const TMap<int32, int32>& OldIDToKeyMap, const int32 BaseReplicationKey)
+{
+	if (ArraySerializer.ArrayReplicationKey == BaseReplicationKey)
 	{
-		//---------------
-		// Build ItemMap if necessary. This maps ReplicationID to our local index into the Items array.
-		//---------------
-		if (ArraySerializer.ItemMap.Num() != Items.Num())
+		// If the keys didn't change, only update the item count caches if necessary.
+		if (ArraySerializer.CachedNumItems == INDEX_NONE ||
+			ArraySerializer.CachedNumItems != Items.Num() ||
+			ArraySerializer.CachedNumItemsToConsiderForWriting == INDEX_NONE)
+		{
+			ArraySerializer.CachedNumItems = Items.Num();
+			ArraySerializer.CachedNumItemsToConsiderForWriting = CalcNumItemsForConsideration();
+		}
+
+		if (UNLIKELY(OldIDToKeyMap.Num() != ArraySerializer.CachedNumItemsToConsiderForWriting))
+		{
+			UE_LOG(LogNetFastTArray, Warning, TEXT("OldMap size (%d) does not match item count (%d)"), OldIDToKeyMap.Num(), ArraySerializer.CachedNumItemsToConsiderForWriting);
+		}
+
+		if (Parms.OldState)
+		{
+			// Nothing changed and we had a valid old state, so just use/share the existing state. No need to create a new one.
+			*Parms.NewState = Parms.OldState->AsShared();
+		}
+		else
+		{
+			// Nothing changed but we don't have an existing state of our own yet so we need to make one here.
+			FNetFastTArrayBaseState* NewState = new FNetFastTArrayBaseState();
+			*Parms.NewState = MakeShareable(NewState);
+			NewState->ArrayReplicationKey = ArraySerializer.ArrayReplicationKey;
+		}
+		
+		return false;
+	}
+
+	return true;
+}
+
+template<typename Type, typename SerializerType>
+void FFastArraySerializer::TFastArraySerializeHelper<Type, SerializerType>::BuildChangedAndDeletedBuffers(
+	TMap<int32, int32>& NewIDToKeyMap,
+	const TMap<int32, int32>* OldIDToKeyMap,
+	TArray<FFastArraySerializer_FastArrayDeltaSerialize_FIdxIDPair, TInlineAllocator<8>>& ChangedElements,
+	TArray<int32, TInlineAllocator<8>>& DeletedElements)
+{
+	ConditionalLogSerializerState(OldIDToKeyMap);
+
+	const int32 NumConsideredItems = CalcNumItemsForConsideration();
+
+	int32 DeleteCount = (OldIDToKeyMap ? OldIDToKeyMap->Num() : 0) - NumConsideredItems; // Note: this is incremented when we add new items below.
+	UE_LOG(LogNetFastTArray, Log, TEXT("NetSerializeItemDeltaFast: %s. DeleteCount: %d"), *Parms.DebugName, DeleteCount);
+
+	//--------------------------------------------
+	// Find out what is new or what has changed
+	//--------------------------------------------
+	for (int32 i = 0; i < Items.Num(); ++i)
+	{
+		Type& Item = Items[i];
+
+		UE_LOG(LogNetFastTArray, Log, TEXT("    Array[%d] - ID %d. CL %d."), i, Item.ReplicationID, Item.ReplicationKey);
+		if (!ArraySerializer.template ShouldWriteFastArrayItem<Type, SerializerType>(Item, Parms.bIsWritingOnClient))
+		{
+			// On clients, this will skip items that were added predictively.
+			continue;
+		}
+		if (Item.ReplicationID == INDEX_NONE)
 		{
-			SCOPE_CYCLE_COUNTER(STAT_NetSerializeFastArray_BuildMap);
-			UE_LOG(LogNetFastTArray, Log, TEXT("FastArrayDeltaSerialize: Recreating Items map. Struct: %s, Items.Num: %d Map.Num: %d"), *InnerStruct->GetOwnerStruct()->GetName(), Items.Num(), ArraySerializer.ItemMap.Num() );
+			ArraySerializer.MarkItemDirty(Item);
+		}
+		NewIDToKeyMap.Add(Item.ReplicationID, Item.ReplicationKey);
 
-			ArraySerializer.ItemMap.Empty();
-			for (int32 i=0; i < Items.Num(); ++i)
+		const int32* OldValuePtr = OldIDToKeyMap ? OldIDToKeyMap->Find(Item.ReplicationID) : NULL;
+		if (OldValuePtr)
+		{
+			if (*OldValuePtr == Item.ReplicationKey)
 			{
-				if ( Items[i].ReplicationID == INDEX_NONE  )
+				UE_LOG(LogNetFastTArray, Log, TEXT("       Stayed The Same - Skipping"));
+
+				// Stayed the same, it might have moved but we dont care
+				continue;
+			}
+			else
+			{
+				UE_LOG(LogNetFastTArray, Log, TEXT("       Changed! Was: %d. Element ID: %d. %s"), *OldValuePtr, Item.ReplicationID, *Item.GetDebugString());
+
+				// Changed
+				ChangedElements.Add(FFastArraySerializer_FastArrayDeltaSerialize_FIdxIDPair(i, Item.ReplicationID));
+			}
+		}
+		else
+		{
+			UE_LOG(LogNetFastTArray, Log, TEXT("       New! Element ID: %d. %s"), Item.ReplicationID, *Item.GetDebugString());
+
+			// The item really should have a valid ReplicationID but in the case of loading from a save game,
+			// items may not have been marked dirty individually. Its ok to just assign them one here.
+			// New
+			ChangedElements.Add(FFastArraySerializer_FastArrayDeltaSerialize_FIdxIDPair(i, Item.ReplicationID));
+			++DeleteCount; // We added something new, so our initial DeleteCount value must be incremented.
+		}
+	}
+
+	// Find out what was deleted
+	if (DeleteCount > 0 && OldIDToKeyMap)
+	{
+		for (auto It = OldIDToKeyMap->CreateConstIterator(); It; ++It)
+		{
+			if (!NewIDToKeyMap.Contains(It.Key()))
+			{
+				UE_LOG(LogNetFastTArray, Log, TEXT("   Deleting ID: %d"), It.Key());
+
+				DeletedElements.Add(It.Key());
+				if (--DeleteCount <= 0)
 				{
-					if (Parms.Writer)
-					{
-						UE_LOG( LogNetFastTArray, Warning, TEXT( "FastArrayDeltaSerialize: Item with uninitialized ReplicationID. Struct: %s, ItemIndex: %i" ), *InnerStruct->GetOwnerStruct()->GetName(), i );
-					}
-					else
-					{
-						// This is benign for clients, they may add things to their local array without assigning a ReplicationID
-						continue;
-					}
+					break;
 				}
-				ArraySerializer.ItemMap.Add( Items[i].ReplicationID, i );
 			}
 		}
 	}
+}
+
+template<typename Type, typename SerializerType>
+void FFastArraySerializer::TFastArraySerializeHelper<Type, SerializerType>::WriteDeltaHeader(FFastArraySerializerHeader& Header) const
+{
+	FBitWriter& Writer = *Parms.Writer;
+
+	Writer << Header.ArrayReplicationKey;
+	Writer << Header.BaseReplicationKey;
+
+	int32 NumDeletes = Header.DeletedIndices.Num();
+	Writer << NumDeletes;
+
+	Writer << Header.NumChanged;
+
+	UE_LOG(LogNetFastTArray, Log, TEXT("   Writing Bunch. NumChange: %d. NumDel: %d [%d/%d]"),
+		Header.NumChanged, Header.DeletedIndices.Num(), Header.ArrayReplicationKey, Header.BaseReplicationKey);
+
+	// Serialize deleted items, just by their ID
+	for (auto It = Header.DeletedIndices.CreateIterator(); It; ++It)
+	{
+		int32 ID = *It;
+		Writer << ID;
+		UE_LOG(LogNetFastTArray, Log, TEXT("   Deleted ElementID: %d"), ID);
+	}
+}
+
+template<typename Type, typename SerializerType>
+bool FFastArraySerializer::TFastArraySerializeHelper<Type, SerializerType>::ReadDeltaHeader(FFastArraySerializerHeader& Header) const
+{
+	FBitReader& Reader = *Parms.Reader;
+
+	//---------------
+	// Read header
+	//---------------
+
+	Reader << Header.ArrayReplicationKey;
+	Reader << Header.BaseReplicationKey;
+
+	int32 NumDeletes;
+	Reader << NumDeletes;
+
+	UE_LOG(LogNetFastTArray, Log, TEXT("Received [%d/%d]."), Header.ArrayReplicationKey, Header.BaseReplicationKey);
+
+	if (NumDeletes > MAX_NUM_DELETED)
+	{
+		UE_LOG(LogNetFastTArray, Warning, TEXT("NumDeletes > MAX_NUM_DELETED: %d."), NumDeletes);
+		Reader.SetError();
+		return false;
+	}
+
+	Reader << Header.NumChanged;
+
+	if (Header.NumChanged > MAX_NUM_CHANGED)
+	{
+		UE_LOG(LogNetFastTArray, Warning, TEXT("NumChanged > MAX_NUM_CHANGED: %d."), Header.NumChanged);
+		Reader.SetError();
+		return false;
+	}
+
+	UE_LOG(LogNetFastTArray, Log, TEXT("Read NumChanged: %d NumDeletes: %d."), Header.NumChanged, NumDeletes);
+
+	//---------------
+	// Read deleted elements
+	//---------------
+	if (NumDeletes > 0)
+	{
+		for (int32 i = 0; i < NumDeletes; ++i)
+		{
+			int32 ElementID;
+			Reader << ElementID;
+
+			int32* ElementIndexPtr = ArraySerializer.ItemMap.Find(ElementID);
+			if (ElementIndexPtr)
+			{
+				int32 DeleteIndex = *ElementIndexPtr;
+				Header.DeletedIndices.Add(DeleteIndex);
+				UE_LOG(LogNetFastTArray, Log, TEXT("   Adding ElementID: %d for deletion"), ElementID);
+			}
+			else
+			{
+				UE_LOG(LogNetFastTArray, Log, TEXT("   Couldn't find ElementID: %d for deletion!"), ElementID);
+			}
+		}
+	}
+
+	return true;
+}
+
+template<typename Type, typename SerializerType>
+template<typename GuidMapType>
+void FFastArraySerializer::TFastArraySerializeHelper<Type, SerializerType>::PostReceiveCleanup(
+	FFastArraySerializerHeader& Header,
+	TArray<int32, TInlineAllocator<8>>& ChangedIndices,
+	TArray<int32, TInlineAllocator<8>>& AddedIndices,
+	GuidMapType& GuidMap)
+{
+	// ---------------------------------------------------------
+	// Look for implicit deletes that would happen due to Naks
+	// ---------------------------------------------------------
+
+	for (int32 idx = 0; idx < Items.Num(); ++idx)
+	{
+		Type& Item = Items[idx];
+		if (Item.MostRecentArrayReplicationKey < Header.ArrayReplicationKey && Item.MostRecentArrayReplicationKey > Header.BaseReplicationKey)
+		{
+			// Make sure this wasn't an explicit delete in this bunch (otherwise we end up deleting an extra element!)
+			if (!Header.DeletedIndices.Contains(idx))
+			{
+				// This will happen in normal conditions in network replays.
+				UE_LOG(LogNetFastTArray, Log, TEXT("Adding implicit delete for ElementID: %d. MostRecentArrayReplicationKey: %d. Current Payload: [%d/%d]"),
+					Item.ReplicationID, Item.MostRecentArrayReplicationKey, Header.ArrayReplicationKey, Header.BaseReplicationKey);
+
+				Header.DeletedIndices.Add(idx);
+			}
+		}
+	}
+
+	// Increment keys so that a client can re-serialize the array if needed, such as for client replay recording.
+	// Must check the size of DeleteIndices instead of NumDeletes to handle implicit deletes.
+	if (Header.DeletedIndices.Num() > 0 || Header.NumChanged > 0)
+	{
+		ArraySerializer.IncrementArrayReplicationKey();
+	}
+
+	// ---------------------------------------------------------
+	// Invoke all callbacks: removed -> added -> changed
+	// ---------------------------------------------------------
+
+	int32 PreRemoveSize = Items.Num();
+	int32 FinalSize = PreRemoveSize - Header.DeletedIndices.Num();
+	for (int32 idx : Header.DeletedIndices)
+	{
+		if (Items.IsValidIndex(idx))
+		{
+			// Remove the deleted element's tracked GUID references
+			if (GuidMap.Remove(Items[idx].ReplicationID) > 0)
+			{
+				Parms.bGuidListsChanged = true;
+			}
+
+			// Call the delete callbacks now, actually remove them at the end
+			Items[idx].PreReplicatedRemove(ArraySerializer);
+		}
+	}
+	ArraySerializer.PreReplicatedRemove(Header.DeletedIndices, FinalSize);
+
+	if (PreRemoveSize != Items.Num())
+	{
+		UE_LOG(LogNetFastTArray, Error, TEXT("Item size changed after PreReplicatedRemove! PremoveSize: %d  Item.Num: %d"),
+			PreRemoveSize, Items.Num());
+	}
+
+	for (int32 idx : AddedIndices)
+	{
+		Items[idx].PostReplicatedAdd(ArraySerializer);
+	}
+	ArraySerializer.PostReplicatedAdd(AddedIndices, FinalSize);
+
+	for (int32 idx : ChangedIndices)
+	{
+		Items[idx].PostReplicatedChange(ArraySerializer);
+	}
+	ArraySerializer.PostReplicatedChange(ChangedIndices, FinalSize);
+
+	if (PreRemoveSize != Items.Num())
+	{
+		UE_LOG(LogNetFastTArray, Error, TEXT("Item size changed after PostReplicatedAdd/PostReplicatedChange! PremoveSize: %d  Item.Num: %d"),
+			PreRemoveSize, Items.Num());
+	}
+
+	if (Header.DeletedIndices.Num() > 0)
+	{
+		Header.DeletedIndices.Sort();
+		for (int32 i = Header.DeletedIndices.Num() - 1; i >= 0; --i)
+		{
+			int32 DeleteIndex = Header.DeletedIndices[i];
+			if (Items.IsValidIndex(DeleteIndex))
+			{
+				Items.RemoveAtSwap(DeleteIndex, 1, false);
+
+				UE_LOG(LogNetFastTArray, Log, TEXT("   Deleting: %d"), DeleteIndex);
+			}
+		}
+
+		// Clear the map now that the indices are all shifted around. This kind of sucks, we could use slightly better data structures here I think.
+		// This will force the ItemMap to be rebuilt for the current Items array
+		ArraySerializer.ItemMap.Empty();
+	}
+}
+
+/** The function that implements Fast TArray Replication  */
+template<typename Type, typename SerializerType >
+bool FFastArraySerializer::FastArrayDeltaSerialize(TArray<Type> &Items, FNetDeltaSerializeInfo& Parms, SerializerType& ArraySerializer)
+{
+	// It's possible that we end up calling this method on clients before we've actually received
+	// anything from the server (Net Conditions, Static Actors, etc.)
+	// That should be fine though, because none of the GUID Tracking work should actually do anything
+	// until after we've received.
+	if (ArraySerializer.bHasBeenSerialized && ArraySerializer.bCachedUseDeltaStructSerialization)
+	{
+		return FastArrayDeltaSerialize_DeltaSerializeStructs(Items, Parms, ArraySerializer);
+	}
+
+	SCOPE_CYCLE_COUNTER(STAT_NetSerializeFastArray);
+	class UScriptStruct* InnerStruct = Type::StaticStruct();
+
+	UE_LOG(LogNetFastTArray, Log, TEXT("FastArrayDeltaSerialize for %s. %s. %s"), *InnerStruct->GetName(), *InnerStruct->GetOwnerStruct()->GetName(), Parms.Reader ? TEXT("Reading") : TEXT("Writing"));
+
+	TFastArraySerializeHelper<Type, SerializerType> Helper{
+		InnerStruct,
+		Items,
+		ArraySerializer,
+		Parms
+	};
+
+	//---------------
+	// Build ItemMap if necessary. This maps ReplicationID to our local index into the Items array.
+	//---------------
+	Helper.ConditionalRebuildItemMap();
 
 	if ( Parms.GatherGuidReferences )
 	{
+		UE_LOG(LogNetFastTArray, Log, TEXT("FastArrayDeltaSerialize for %s. %s. GatherGuidReferences"), *InnerStruct->GetName(), *InnerStruct->GetOwnerStruct()->GetName());
+
 		// Loop over all tracked guids, and return what we have
 		for ( const auto& GuidReferencesPair : ArraySerializer.GuidReferencesMap )
 		{
@@ -529,6 +1085,8 @@ bool FFastArraySerializer::FastArrayDeltaSerialize( TArray<Type> &Items, FNetDel
 
 	if ( Parms.MoveGuidToUnmapped )
 	{
+		UE_LOG(LogNetFastTArray, Log, TEXT("FastArrayDeltaSerialize for %s. %s. MovedGuidToUnmapped"), *InnerStruct->GetName(), *InnerStruct->GetOwnerStruct()->GetName());
+
 		bool bFound = false;
 
 		const FNetworkGUID GUID = *Parms.MoveGuidToUnmapped;
@@ -551,6 +1109,8 @@ bool FFastArraySerializer::FastArrayDeltaSerialize( TArray<Type> &Items, FNetDel
 
 	if ( Parms.bUpdateUnmappedObjects )
 	{
+		UE_LOG(LogNetFastTArray, Log, TEXT("FastArrayDeltaSerialize for %s. %s. UpdateUnmappedObjects"), *InnerStruct->GetName(), *InnerStruct->GetOwnerStruct()->GetName());
+
 		// Loop over each item that has unmapped objects
 		for ( auto It = ArraySerializer.GuidReferencesMap.CreateIterator(); It; ++It )
 		{
@@ -614,11 +1174,13 @@ bool FFastArraySerializer::FastArrayDeltaSerialize( TArray<Type> &Items, FNetDel
 				FNetBitReader Reader( Parms.Map, GuidReferences.Buffer.GetData(), GuidReferences.NumBufferBits );
 
 				// Read the property (which should serialize any newly mapped objects as well)
-				bool bHasUnmapped = false;
-				Parms.NetSerializeCB->NetSerializeStruct( InnerStruct, Reader, Parms.Map, ThisElement, bHasUnmapped );
+				Parms.Struct = InnerStruct;
+				Parms.Data = ThisElement;
+				Parms.Reader = &Reader;
+				Parms.NetSerializeCB->NetSerializeStruct(Parms);
 
 				// Let the element know it changed
-				ThisElement->PostReplicatedChange( ArraySerializer );
+				ThisElement->PostReplicatedChange(ArraySerializer);
 			}
 
 			// If we have no more guids, we can remove this item for good
@@ -637,8 +1199,24 @@ bool FFastArraySerializer::FastArrayDeltaSerialize( TArray<Type> &Items, FNetDel
 		return true;
 	}
 
+	// If we've made it this far, it means that we're going to be serializing something.
+	// So, it should be safe for us to update our cached state.
+	// Also, make sure that we hit the right path if we need to.
+	if (!ArraySerializer.bHasBeenSerialized)
+	{
+		ArraySerializer.bHasBeenSerialized = true;
+		ArraySerializer.bCachedUseDeltaStructSerialization = ArraySerializer.bUseDeltaStructSerialization && Parms.bSupportsFastArrayDeltaStructSerialization;
+
+		if (ArraySerializer.bCachedUseDeltaStructSerialization)
+		{
+			return FastArrayDeltaSerialize_DeltaSerializeStructs(Items, Parms, ArraySerializer);
+		}
+	}
+
 	if (Parms.Writer)
 	{
+		UE_LOG(LogNetFastTArray, Log, TEXT("FastArrayDeltaSerialize for %s. %s. Writing"), *InnerStruct->GetName(), *InnerStruct->GetOwnerStruct()->GetName());
+
 		//-----------------------------
 		// Saving
 		//-----------------------------	
@@ -646,62 +1224,23 @@ bool FFastArraySerializer::FastArrayDeltaSerialize( TArray<Type> &Items, FNetDel
 		FBitWriter& Writer = *Parms.Writer;
 
 		// Get the old map if its there
-		TMap<int32, int32> * OldMap = Parms.OldState ? &((FNetFastTArrayBaseState*)Parms.OldState)->IDToCLMap : NULL;
-		int32 BaseReplicationKey = Parms.OldState ? ((FNetFastTArrayBaseState*)Parms.OldState)->ArrayReplicationKey : -1;
-
-		// Helper for counting num elements we should consider
-		auto CalcNumItemsForConsideration = [&ArraySerializer, &Items, &Parms]()
-		{
-			int32 Count = 0;
-
-			// Count the number of items in the current array that may be written. On clients, items that were predicted will be skipped.
-			for (const Type& Item : Items)
-			{
-				if (ArraySerializer.template ShouldWriteFastArrayItem<Type, SerializerType>(Item, Parms.bIsWritingOnClient))
-				{
-					Count++;
-				}
-			}
-			return Count;
-		};
+		TMap<int32, int32> * OldMap = nullptr;
+		int32 BaseReplicationKey = INDEX_NONE;
 
 		// See if the array changed at all. If the ArrayReplicationKey matches we can skip checking individual items
-		if (Parms.OldState && (ArraySerializer.ArrayReplicationKey == BaseReplicationKey))
+		if (Parms.OldState)
 		{
-			// Double check the old map is valid and that we will consider writing the same number of elements that are in the old map.
-			if (ensureMsgf(OldMap, TEXT("Invalid OldMap")))
-			{
-				// If the keys didn't change, only update the item count caches if necessary.
-				if (ArraySerializer.CachedNumItems == INDEX_NONE ||
-					ArraySerializer.CachedNumItems != Items.Num() ||
-					ArraySerializer.CachedNumItemsToConsiderForWriting == INDEX_NONE)
-				{
-					ArraySerializer.CachedNumItems = Items.Num();
-					ArraySerializer.CachedNumItemsToConsiderForWriting = CalcNumItemsForConsideration();
-				}
+			OldMap = &((FNetFastTArrayBaseState*)Parms.OldState)->IDToCLMap;
+			BaseReplicationKey = ((FNetFastTArrayBaseState*)Parms.OldState)->ArrayReplicationKey;
 
-				if (UNLIKELY(OldMap->Num() != ArraySerializer.CachedNumItemsToConsiderForWriting))
-				{
-					UE_LOG(LogNetFastTArray, Warning, TEXT("OldMap size (%d) does not match item count (%d)"), OldMap->Num(), ArraySerializer.CachedNumItemsToConsiderForWriting);
-				}
-			}
-
-			if (Parms.OldState)
-			{
-				// Nothing changed and we had a valid old state, so just use/share the existing state. No need to create a new one.
-				*Parms.NewState = Parms.OldState->AsShared();
-			}
-			else
+			// If we didn't create a new delta state, that implies nothing changed,
+			// so we're done.
+			if (!Helper.ConditionalCreateNewDeltaState(*OldMap, BaseReplicationKey))
 			{
-				// Nothing changed but we don't have an existing state of our own yet so we need to make one here.
-				FNetFastTArrayBaseState * NewState = new FNetFastTArrayBaseState();
-				*Parms.NewState = MakeShareable( NewState );
-				NewState->ArrayReplicationKey = ArraySerializer.ArrayReplicationKey;
+				return false;
 			}
-			return false;
 		}
 
-
 		// Create a new map from the current state of the array		
 		FNetFastTArrayBaseState * NewState = new FNetFastTArrayBaseState();
 
@@ -710,130 +1249,28 @@ bool FFastArraySerializer::FastArrayDeltaSerialize( TArray<Type> &Items, FNetDel
 		TMap<int32, int32> & NewMap = NewState->IDToCLMap;
 		NewState->ArrayReplicationKey = ArraySerializer.ArrayReplicationKey;
 
+		FFastArraySerializerHeader Header{
+			ArraySerializer.ArrayReplicationKey,
+			BaseReplicationKey,
+			0
+		};
 
-		int32 NumConsideredItems = CalcNumItemsForConsideration();
-
-		TArray<FFastArraySerializer_FastArrayDeltaSerialize_FIdxIDPair, TInlineAllocator<8> >	ChangedElements;
-		TArray<int32, TInlineAllocator<8> >		DeletedElements;
-
-		int32 DeleteCount = (OldMap ? OldMap->Num() : 0) - NumConsideredItems; // Note: this is incremented when we add new items below.
-		UE_LOG(LogNetFastTArray, Log, TEXT("NetSerializeItemDeltaFast: %s. DeleteCount: %d"), *Parms.DebugName, DeleteCount);
-
-		// Log out entire state of current/base state
-		if (UE_LOG_ACTIVE(LogNetFastTArray,Log))
-		{
-			FString CurrentState = FString::Printf(TEXT("Current: %d "), ArraySerializer.ArrayReplicationKey );
-			for (int32 i=0; i < Items.Num(); ++i)
-			{
-				CurrentState += FString::Printf(TEXT("[%d/%d], "), Items[i].ReplicationID, Items[i].ReplicationKey);
-			}
-			UE_LOG(LogNetFastTArray, Log, TEXT("%s"), *CurrentState);
+		TArray<FFastArraySerializer_FastArrayDeltaSerialize_FIdxIDPair, TInlineAllocator<8>> ChangedElements;
 
+		Helper.BuildChangedAndDeletedBuffers(NewMap, OldMap, ChangedElements, Header.DeletedIndices);
 		
-			FString ClientStateStr = FString::Printf(TEXT("Client: %d "),  Parms.OldState ? ((FNetFastTArrayBaseState*)Parms.OldState)->ArrayReplicationKey: 0);
-			if (OldMap)
-			{
-				for (auto It = OldMap->CreateIterator(); It; ++It)
-				{
-					ClientStateStr += FString::Printf(TEXT("[%d/%d], "), It.Key(), It.Value() );
-				}
-			}
-			UE_LOG(LogNetFastTArray, Log, TEXT("%s"), *ClientStateStr);
-		}
-
-
-		//--------------------------------------------
-		// Find out what is new or what has changed
-		//--------------------------------------------
-		for (int32 i=0; i < Items.Num(); ++i)
-		{
-			UE_LOG(LogNetFastTArray, Log, TEXT("    Array[%d] - ID %d. CL %d."), i, Items[i].ReplicationID, Items[i].ReplicationKey);
-			if (!ArraySerializer.template ShouldWriteFastArrayItem<Type, SerializerType>(Items[i], Parms.bIsWritingOnClient))
-			{
-				// On clients, this will skip items that were added predictively.
-				continue;
-			}
-			if (Items[i].ReplicationID == INDEX_NONE)
-			{
-				ArraySerializer.MarkItemDirty(Items[i]);
-			}
-			NewMap.Add( Items[i].ReplicationID, Items[i].ReplicationKey );
-
-			int32* OldValuePtr = OldMap ? OldMap->Find( Items[i].ReplicationID ) : NULL;
-			if (OldValuePtr)
-			{
-				if (*OldValuePtr == Items[i].ReplicationKey)
-				{
-					UE_LOG(LogNetFastTArray, Log, TEXT("       Stayed The Same - Skipping"));
-
-					// Stayed the same, it might have moved but we dont care
-					continue;
-				}
-				else
-				{
-					UE_LOG(LogNetFastTArray, Log, TEXT("       Changed! Was: %d. Element ID: %d. %s"), *OldValuePtr, Items[i].ReplicationID, *Items[i].GetDebugString());
-
-					// Changed
-					ChangedElements.Add(FFastArraySerializer_FastArrayDeltaSerialize_FIdxIDPair(i, Items[i].ReplicationID));
-				}
-			}
-			else
-			{
-				UE_LOG(LogNetFastTArray, Log, TEXT("       New! Element ID: %d. %s"), Items[i].ReplicationID, *Items[i].GetDebugString());
-				
-				// The item really should have a valid ReplicationID but in the case of loading from a save game,
-				// items may not have been marked dirty individually. Its ok to just assign them one here.
-				// New
-				ChangedElements.Add(FFastArraySerializer_FastArrayDeltaSerialize_FIdxIDPair(i, Items[i].ReplicationID));
-				DeleteCount++; // We added something new, so our initial DeleteCount value must be incremented.
-			}
-		}
-
-		// Find out what was deleted
-		if( DeleteCount > 0 && OldMap)
-		{
-			for (auto It = OldMap->CreateIterator(); It; ++It)
-			{
-				if (!NewMap.Contains(It.Key()))
-				{
-					UE_LOG(LogNetFastTArray, Log, TEXT("   Deleting ID: %d"), It.Key());
-
-					DeletedElements.Add( It.Key() );
-					if (--DeleteCount <= 0)
-						break;
-				}
-			}
-		}
-
 		// Note: we used to early return false here if nothing had changed, but we still need to send
 		// a bunch with the array key / base key, so that clients can look for implicit deletes.
 
-		// The array replication key may have changed while adding new elemnts (in the call to MarkItemDirty above)
+		// The array replication key may have changed while adding new elements (in the call to MarkItemDirty above)
 		NewState->ArrayReplicationKey = ArraySerializer.ArrayReplicationKey;
 
 		//----------------------
 		// Write it out.
 		//----------------------
 
-		int32 ArrayReplicationKey = ArraySerializer.ArrayReplicationKey;
-		Writer << ArrayReplicationKey;
-		Writer << BaseReplicationKey;
-
-		uint32 ElemNum = DeletedElements.Num();
-		Writer << ElemNum;
-
-		ElemNum = ChangedElements.Num();
-		Writer << ElemNum;
-
-		UE_LOG(LogNetFastTArray, Log, TEXT("   Writing Bunch. NumChange: %d. NumDel: %d [%d/%d]"), ChangedElements.Num(), DeletedElements.Num(), ArrayReplicationKey, BaseReplicationKey );
-
-		// Serialize deleted items, just by their ID
-		for (auto It = DeletedElements.CreateIterator(); It; ++It)
-		{
-			int32 ID = *It;
-			Writer << ID;
-			UE_LOG(LogNetFastTArray, Log, TEXT("   Deleted ElementID: %d"), ID);
-		}
+		Header.NumChanged = ChangedElements.Num();
+		Helper.WriteDeltaHeader(Header);
 
 		// Serialized new elements with their payload
 		for (auto It = ChangedElements.CreateIterator(); It; ++It)
@@ -846,87 +1283,34 @@ bool FFastArraySerializer::FastArrayDeltaSerialize( TArray<Type> &Items, FNetDel
 
 			UE_LOG(LogNetFastTArray, Log, TEXT("   Changed ElementID: %d"), ID);
 
-			bool bHasUnmapped = false;
-			Parms.NetSerializeCB->NetSerializeStruct(InnerStruct, Writer, Parms.Map, ThisElement, bHasUnmapped);
+			Parms.Struct = InnerStruct;
+			Parms.Data = ThisElement;
+			Parms.NetSerializeCB->NetSerializeStruct(Parms);
 		}
 	}
 	else
 	{
+		UE_LOG(LogNetFastTArray, Log, TEXT("FastArrayDeltaSerialize for %s. %s. Reading"), *InnerStruct->GetName(), *InnerStruct->GetOwnerStruct()->GetName());
+
 		//-----------------------------
 		// Loading
 		//-----------------------------	
 		check(Parms.Reader);
 		FBitReader& Reader = *Parms.Reader;
 
-		static const int32 MAX_NUM_CHANGED = 2048;
-		static const int32 MAX_NUM_DELETED = 2048;
-
-		//---------------
-		// Read header
-		//---------------
-		
-		int32 ArrayReplicationKey;
-		Reader << ArrayReplicationKey;
-		
-		int32 BaseReplicationKey;
-		Reader << BaseReplicationKey;
-
-		uint32 NumDeletes;
-		Reader << NumDeletes;
-
-		UE_LOG( LogNetFastTArray, Log, TEXT( "Received [%d/%d]." ), ArrayReplicationKey, BaseReplicationKey );
-	
-		if ( NumDeletes > MAX_NUM_DELETED )
-		{
-			UE_LOG( LogNetFastTArray, Warning, TEXT( "NumDeletes > MAX_NUM_DELETED: %d." ), NumDeletes );
-			Reader.SetError();
-			return false;;
-		}
-
-		uint32 NumChanged;
-		Reader << NumChanged;
-
-		if (NumChanged > MAX_NUM_CHANGED)
+		FFastArraySerializerHeader Header;
+		if (!Helper.ReadDeltaHeader(Header))
 		{
-			UE_LOG(LogNetFastTArray, Warning, TEXT("NumChanged > MAX_NUM_CHANGED: %d."), NumChanged);
-			Reader.SetError();
-			return false;;
+			return false;
 		}
 
-		UE_LOG( LogNetFastTArray, Log, TEXT( "Read NumChanged: %d NumDeletes: %d." ), NumChanged, NumDeletes );
-
-		TArray<int32, TInlineAllocator<8> > DeleteIndices;
-		TArray<int32, TInlineAllocator<8> > AddedIndices;
-		TArray<int32, TInlineAllocator<8> > ChangedIndices;
-
-		//---------------
-		// Read deleted elements
-		//---------------
-		if (NumDeletes > 0)
-		{
-			for (uint32 i = 0; i < NumDeletes; ++i)
-			{
-				int32 ElementID;
-				Reader << ElementID;
-
-				int32* ElementIndexPtr = ArraySerializer.ItemMap.Find(ElementID);
-				if (ElementIndexPtr)
-				{
-					int32 DeleteIndex = *ElementIndexPtr;
-					DeleteIndices.Add(DeleteIndex);
-					UE_LOG(LogNetFastTArray, Log, TEXT("   Adding ElementID: %d for deletion"), ElementID);
-				}
-				else
-				{
-					UE_LOG(LogNetFastTArray, Log, TEXT("   Couldn't find ElementID: %d for deletion!"), ElementID);
-				}
-			}
-		}
+		TArray<int32, TInlineAllocator<8>> ChangedIndices;
+		TArray<int32, TInlineAllocator<8>> AddedIndices;
 
 		//---------------
 		// Read Changed/New elements
 		//---------------
-		for(uint32 i=0; i < NumChanged; ++i)
+		for(int32 i = 0; i < Header.NumChanged; ++i)
 		{
 			int32 ElementID;
 			Reader << ElementID;
@@ -956,7 +1340,7 @@ bool FFastArraySerializer::FastArrayDeltaSerialize( TArray<Type> &Items, FNetDel
 			}
 
 			// Update this element's most recent array replication key
-			ThisElement->MostRecentArrayReplicationKey = ArrayReplicationKey;
+			ThisElement->MostRecentArrayReplicationKey = Header.ArrayReplicationKey;
 
 			// Update this element's replication key so that a client can re-serialize the array for client replay recording
 			ThisElement->ReplicationKey++;
@@ -967,8 +1351,9 @@ bool FFastArraySerializer::FastArrayDeltaSerialize( TArray<Type> &Items, FNetDel
 			// Remember where we started reading from, so that if we have unmapped properties, we can re-deserialize from this data later
 			FBitReaderMark Mark( Reader );
 
-			bool bHasUnmapped = false;
-			Parms.NetSerializeCB->NetSerializeStruct( InnerStruct, Reader, Parms.Map, ThisElement, bHasUnmapped );
+			Parms.Struct = InnerStruct;
+			Parms.Data = ThisElement;
+			Parms.NetSerializeCB->NetSerializeStruct(Parms);
 
 			if ( !Reader.IsError() )
 			{
@@ -1026,103 +1411,169 @@ bool FFastArraySerializer::FastArrayDeltaSerialize( TArray<Type> &Items, FNetDel
 			}
 		}
 
-		// ---------------------------------------------------------
-		// Look for implicit deletes that would happen due to Naks
-		// ---------------------------------------------------------
-		
-		for (int32 idx=0; idx < Items.Num(); ++idx)
+		Helper.template PostReceiveCleanup<>(Header, ChangedIndices, AddedIndices, ArraySerializer.GuidReferencesMap);
+	}
+
+	return true;
+}
+
+struct FFastArrayDeltaSerializeParams
+{
+	FNetDeltaSerializeInfo& DeltaSerializeInfo;
+	FFastArraySerializer& ArraySerializer;
+	TArray<FFastArraySerializer_FastArrayDeltaSerialize_FIdxIDPair, TInlineAllocator<8>>* WriteChangedElements = nullptr;
+	FNetFastTArrayBaseState* WriteBaseState = nullptr;
+	TArray<int32, TInlineAllocator<8>>* ReadChangedElements = nullptr;
+	TArray<int32, TInlineAllocator<8>>* ReadAddedElements = nullptr;
+	int32 ReadNumChanged = INDEX_NONE;
+};
+
+template<typename Type, typename SerializerType>
+bool FFastArraySerializer::FastArrayDeltaSerialize_DeltaSerializeStructs(TArray<Type>& Items, FNetDeltaSerializeInfo& Parms, SerializerType& ArraySerializer)
+{
+	SCOPE_CYCLE_COUNTER(STAT_NetSerializeFastArray_DeltaStruct);
+	class UScriptStruct* InnerStruct = Type::StaticStruct();
+
+	TFastArraySerializeHelper<Type, SerializerType> Helper{
+		InnerStruct,
+		Items,
+		ArraySerializer,
+		Parms
+	};
+
+	FFastArrayDeltaSerializeParams DeltaSerializeParams{
+		Parms,
+		ArraySerializer
+	};
+
+	//---------------
+	// Build ItemMap if necessary. This maps ReplicationID to our local index into the Items array.
+	//---------------
+	Helper.ConditionalRebuildItemMap();
+
+	if (Parms.GatherGuidReferences)
+	{
+		UE_LOG(LogNetFastTArray, Log, TEXT("FastArrayDeltaSerialize_DeltaSerializeStruct for %s. %s. GatherGuidReferences"), *InnerStruct->GetName(), *InnerStruct->GetOwnerStruct()->GetName());
+		Parms.NetSerializeCB->GatherGuidReferencesForFastArray(DeltaSerializeParams);
+		return true;
+	}
+	else if (Parms.MoveGuidToUnmapped)
+	{
+		UE_LOG(LogNetFastTArray, Log, TEXT("FastArrayDeltaSerialize_DeltaSerializeStruct for %s. %s. MoveGuidToUnmapped"), *InnerStruct->GetName(), *InnerStruct->GetOwnerStruct()->GetName());
+		return Parms.NetSerializeCB->MoveGuidToUnmappedForFastArray(DeltaSerializeParams);
+	}
+	else if (Parms.bUpdateUnmappedObjects)
+	{
+		UE_LOG(LogNetFastTArray, Log, TEXT("FastArrayDeltaSerialize_DeltaSerializeStruct for %s. %s. UpdateUnmappedObjects"), *InnerStruct->GetName(), *InnerStruct->GetOwnerStruct()->GetName());
+		Parms.NetSerializeCB->UpdateUnmappedGuidsForFastArray(DeltaSerializeParams);
+		return true;
+	}
+	else if (Parms.Writer)
+	{
+		UE_LOG(LogNetFastTArray, Log, TEXT("FastArrayDeltaSerialize_DeltaSerializeStruct for %s. %s. Writing"), *InnerStruct->GetName(), *InnerStruct->GetOwnerStruct()->GetName());
+
+		//-----------------------------
+		// Saving
+		//-----------------------------	
+		check(Parms.Struct);
+		FBitWriter& Writer = *Parms.Writer;
+
+		// Get the old map if its there
+		TMap<int32, int32>* OldItemMap = nullptr;
+		int32 BaseReplicationKey = INDEX_NONE;
+		int32 OldChangelistHistory = INDEX_NONE;
+
+		// See if the array changed at all. If the ArrayReplicationKey matches we can skip checking individual items
+		if (Parms.OldState)
 		{
-			Type& Item = Items[idx];
-			if (Item.MostRecentArrayReplicationKey < ArrayReplicationKey && Item.MostRecentArrayReplicationKey > BaseReplicationKey)
+			OldItemMap = &((FNetFastTArrayBaseState*)Parms.OldState)->IDToCLMap;
+			BaseReplicationKey = ((FNetFastTArrayBaseState*)Parms.OldState)->ArrayReplicationKey;
+			OldChangelistHistory = ((FNetFastTArrayBaseState*)Parms.OldState)->ChangelistHistory;
+
+			if (!Helper.ConditionalCreateNewDeltaState(*OldItemMap, BaseReplicationKey))
 			{
-				// Make sure this wasn't an explicit delete in this bunch (otherwise we end up deleting an extra element!)
-				if (DeleteIndices.Contains(idx) == false)
-				{
-					// This will happen in normal conditions in network replays.
-					UE_LOG( LogNetFastTArray, Log, TEXT( "Adding implicit delete for ElementID: %d. MostRecentArrayReplicationKey: %d. Current Payload: [%d/%d]"), Item.ReplicationID, Item.MostRecentArrayReplicationKey, ArrayReplicationKey, BaseReplicationKey);
-					
-					DeleteIndices.Add(idx);
-				}
+				return false;
 			}
 		}
 
-		// Increment keys so that a client can re-serialize the array if needed, such as for client replay recording.
-		// Must check the size of DeleteIndices instead of NumDeletes to handle implicit deletes.
-		if ( DeleteIndices.Num() > 0 || NumChanged > 0 )
-		{
-			ArraySerializer.IncrementArrayReplicationKey();
-		}
+		// Create a new map from the current state of the array		
+		FNetFastTArrayBaseState* NewState = new FNetFastTArrayBaseState();
 
-		// ---------------------------------------------------------
-		// Invoke all callbacks: removed -> added -> changed
-		// ---------------------------------------------------------
+		check(Parms.NewState);
 
-		int32 PreRemoveSize = Items.Num();
-		int32 FinalSize = PreRemoveSize - DeleteIndices.Num();
-		for (int32 idx : DeleteIndices)
-		{
-			if (Items.IsValidIndex(idx))
-			{
-				// Remove the deleted element's tracked GUID references
-				if (ArraySerializer.GuidReferencesMap.Remove(Items[idx].ReplicationID) > 0)
-				{
-					Parms.bGuidListsChanged = true;
-				}
+		*Parms.NewState = MakeShareable(NewState);
+		NewState->ArrayReplicationKey = ArraySerializer.ArrayReplicationKey;
 
-				// Call the delete callbacks now, actually remove them at the end
-				Items[idx].PreReplicatedRemove(ArraySerializer);
-			}
-		}
-		ArraySerializer.PreReplicatedRemove(DeleteIndices, FinalSize);
+		TMap<int32, int32>& NewItemMap = NewState->IDToCLMap;
 
-		if (PreRemoveSize != Items.Num())
-		{
-			UE_LOG( LogNetFastTArray, Error, TEXT( "Item size changed after PreReplicatedRemove! PremoveSize: %d  Item.Num: %d"), PreRemoveSize, Items.Num() );
-		}
+		FFastArraySerializerHeader Header{
+			ArraySerializer.ArrayReplicationKey,
+			BaseReplicationKey,
+			0
+		};
 
-		for (int32 idx : AddedIndices)
-		{
-			Items[idx].PostReplicatedAdd(ArraySerializer);
-		}
-		ArraySerializer.PostReplicatedAdd(AddedIndices, FinalSize);
+		TArray<FFastArraySerializer_FastArrayDeltaSerialize_FIdxIDPair, TInlineAllocator<8>> ChangedElements;
+		Helper.BuildChangedAndDeletedBuffers(NewItemMap, OldItemMap, ChangedElements, Header.DeletedIndices);
 
-		for (int32 idx : ChangedIndices)
-		{
-			Items[idx].PostReplicatedChange(ArraySerializer);
-		}	
-		ArraySerializer.PostReplicatedChange(ChangedIndices, FinalSize);
+		// Note: we used to early return false here if nothing had changed, but we still need to send
+		// a bunch with the array key / base key, so that clients can look for implicit deletes.
+
+		// The array replication key may have changed while adding new elemnts (in the call to MarkItemDirty above)
+		NewState->ArrayReplicationKey = ArraySerializer.ArrayReplicationKey;
+
+		//----------------------
+		// Write it out.
+		//----------------------
+
+		Header.NumChanged = ChangedElements.Num();
+		Helper.WriteDeltaHeader(Header);
+
+		NewState->ChangelistHistory = OldChangelistHistory;
+
+		DeltaSerializeParams.WriteChangedElements = &ChangedElements;
+		DeltaSerializeParams.WriteBaseState = NewState;
+
+		return Parms.NetSerializeCB->NetDeltaSerializeForFastArray(DeltaSerializeParams);
+	}
+	else
+	{
+		UE_LOG(LogNetFastTArray, Log, TEXT("FastArrayDeltaSerialize_DeltaSerializeStruct for %s. %s. Reading"), *InnerStruct->GetName(), *InnerStruct->GetOwnerStruct()->GetName());
 
-		if (PreRemoveSize != Items.Num())
+		//-----------------------------
+		// Loading
+		//-----------------------------	
+		check(Parms.Reader);
+		FBitReader& Reader = *Parms.Reader;
+
+		FFastArraySerializerHeader Header;
+		if (!Helper.ReadDeltaHeader(Header))
 		{
-			UE_LOG( LogNetFastTArray, Error, TEXT( "Item size changed after PostReplicatedAdd/PostReplicatedChange! PremoveSize: %d  Item.Num: %d"), PreRemoveSize, Items.Num() );
+			return false;
 		}
 
-		if (DeleteIndices.Num() > 0)
-		{
-			DeleteIndices.Sort();
-			for (int32 i = DeleteIndices.Num() - 1; i >= 0; --i)
-			{
-				int32 DeleteIndex = DeleteIndices[i];
-				if (Items.IsValidIndex(DeleteIndex))
-				{
-					Items.RemoveAtSwap(DeleteIndex, 1, false);
+		TArray<int32, TInlineAllocator<8>> ChangedIndices;
+		TArray<int32, TInlineAllocator<8>> AddedIndices;
 
-					UE_LOG(LogNetFastTArray, Log, TEXT("   Deleting: %d"), DeleteIndex);
-				}
-			}
+		DeltaSerializeParams.ReadAddedElements = &AddedIndices;
+		DeltaSerializeParams.ReadChangedElements = &ChangedIndices;
+		DeltaSerializeParams.ReadNumChanged = Header.NumChanged;
 
-			// Clear the map now that the indices are all shifted around. This kind of sucks, we could use slightly better data structures here I think.
-			// This will force the ItemMap to be rebuilt for the current Items array
-			ArraySerializer.ItemMap.Empty();
+		if (!Parms.NetSerializeCB->NetDeltaSerializeForFastArray(DeltaSerializeParams))
+		{
+			return false;
 		}
+
+		//---------------
+		// Read Changed/New elements
+		//---------------
+
+		Helper.template PostReceiveCleanup<>(Header, ChangedIndices, AddedIndices, ArraySerializer.GuidReferencesMap_StructDelta);
 	}
 
 	return true;
 }
 
 
-
-
 /**
  *	===================== Vector NetSerialization customization. =====================
  *	Provides custom NetSerilization for FVectors.
diff --git a/Engine/Source/Runtime/Engine/Private/DataReplication.cpp b/Engine/Source/Runtime/Engine/Private/DataReplication.cpp
--- a/Engine/Source/Runtime/Engine/Private/DataReplication.cpp
+++ b/Engine/Source/Runtime/Engine/Private/DataReplication.cpp
@@ -18,6 +18,7 @@
 #include "Engine/Engine.h"
 #include "Engine/NetConnection.h"
 #include "Net/NetworkGranularMemoryLogging.h"
+#include "Misc/ScopeExit.h"
 
 DECLARE_CYCLE_STAT(TEXT("Custom Delta Property Rep Time"), STAT_NetReplicateCustomDeltaPropTime, STATGROUP_Game);
 DECLARE_CYCLE_STAT(TEXT("ReceiveRPC"), STAT_NetReceiveRPC, STATGROUP_Game);
@@ -41,59 +42,235 @@ static FAutoConsoleVariableRef CVarNetRPCDebug(
 	TEXT(" 1: Print bunches as they are sent."),
 	ECVF_Default);
 
+int32 GSupportsFastArrayDelta = 0;
+static FAutoConsoleVariableRef CVarSupportsFastArrayDelta(
+	TEXT("net.SupportFastArrayDelta"),
+	GSupportsFastArrayDelta,
+	TEXT("Whether or not Fast Array Struct Delta Serialization is enabled.")
+);
+
+PRAGMA_DISABLE_DEPRECATION_WARNINGS
 class FNetSerializeCB : public INetSerializeCB
 {
+private:
+
+	// This is an acceleration so if we make back to back requests for the same type
+	// we don't have to do repeated lookups.
+	struct FCachedRequestState
+	{
+		UClass* ObjectClass = nullptr;
+		UScriptStruct* Struct = nullptr;
+		TSharedPtr<FRepLayout> RepLayout;
+		bool bWasRequestFromClass = false;
+	};
+
 public:
-	FNetSerializeCB() : Driver( NULL ) { check( 0 ); }
 
-	FNetSerializeCB( UNetDriver * InNetDriver ) : Driver( InNetDriver ) { }
+	FNetSerializeCB():
+		Driver(nullptr)
+	{
+		check(0);
+	}
+
+	FNetSerializeCB(UNetDriver* InNetDriver):
+		Driver(InNetDriver)
+	{
+	}
+
+	void SetChangelistMgr(TSharedPtr<FReplicationChangelistMgr> InChangelistMgr)
+	{
+		ChangelistMgr = InChangelistMgr;
+	}
+
+private:
+
+	void UpdateCachedRepLayout()
+	{
+		if (!CachedRequestState.RepLayout.IsValid())
+		{
+			if (CachedRequestState.bWasRequestFromClass)
+			{
+				CachedRequestState.RepLayout = Driver->GetObjectClassRepLayout(CachedRequestState.ObjectClass);
+			}
+			else
+			{
+				CachedRequestState.RepLayout = Driver->GetStructRepLayout(CachedRequestState.Struct);
+			}
+		}
+	}
+
+	void UpdateCachedState(UClass* ObjectClass, UStruct* Struct)
+	{
+		if (CachedRequestState.ObjectClass != ObjectClass)
+		{
+			CachedRequestState.ObjectClass = ObjectClass;
+			CachedRequestState.Struct = CastChecked<UScriptStruct>(Struct);
+			CachedRequestState.bWasRequestFromClass = true;
+			CachedRequestState.RepLayout.Reset();
+		}
+	}
+
+	void UpdateCachedState(UStruct* Struct)
+	{
+		if (CachedRequestState.Struct != Struct || CachedRequestState.ObjectClass != nullptr)
+		{
+			CachedRequestState.ObjectClass = nullptr;
+			CachedRequestState.Struct = CastChecked<UScriptStruct>(Struct);
+			CachedRequestState.bWasRequestFromClass = false;
+			CachedRequestState.RepLayout.Reset();
+		}
+	}
+
+public:
 
-	virtual void NetSerializeStruct( UScriptStruct* Struct, FBitArchive& Ar, UPackageMap* Map, void* Data, bool& bHasUnmapped )
+	virtual void NetSerializeStruct(FNetDeltaSerializeInfo& Params) override final
 	{
-		if (Struct->StructFlags & STRUCT_NetSerializeNative)
+		UpdateCachedState(Params.Struct);
+		FBitArchive& Ar = Params.Reader ? static_cast<FBitArchive&>(*Params.Reader) : static_cast<FBitArchive&>(*Params.Writer);
+		Params.bOutHasMoreUnmapped = false;
+
+		if (EnumHasAnyFlags(CachedRequestState.Struct->StructFlags, STRUCT_NetSerializeNative))
 		{
-			UScriptStruct::ICppStructOps* CppStructOps = Struct->GetCppStructOps();
-			check(CppStructOps); // else should not have STRUCT_NetSerializeNative
+			UScriptStruct::ICppStructOps* CppStructOps = CachedRequestState.Struct->GetCppStructOps();
+			check(CppStructOps);
 			bool bSuccess = true;
-			if (!CppStructOps->NetSerialize(Ar, Map, bSuccess, Data))
+
+			if (!CppStructOps->NetSerialize(Ar, Params.Map, bSuccess, Params.Data))
 			{
-				bHasUnmapped = true;
+				Params.bOutHasMoreUnmapped = true;
 			}
 
 			if (!bSuccess)
 			{
-				UE_LOG(LogRep, Warning, TEXT("NetSerializeStruct: Native NetSerialize %s failed."), *Struct->GetFullName());
+				UE_LOG(LogRep, Warning, TEXT("NetSerializeStruct: Native NetSerialize %s failed."), *Params.Struct->GetFullName());
 			}
 		}
 		else
 		{
-			TSharedPtr<FRepLayout> RepLayout = Driver->GetStructRepLayout(Struct);
+			UpdateCachedRepLayout();
+			UPackageMapClient* PackageMapClient = ((UPackageMapClient*)Params.Map);
 
-			UPackageMapClient* PackageMapClient = ( ( UPackageMapClient* )Map );
-
-			if ( PackageMapClient && PackageMapClient->GetConnection()->InternalAck )
+			if (PackageMapClient && PackageMapClient->GetConnection()->InternalAck)
 			{
-				if ( Ar.IsSaving() )
+				if (Ar.IsSaving())
 				{
 					TArray< uint16 > Changed;
-					RepLayout->SendProperties_BackwardsCompatible( nullptr, nullptr, (uint8*)Data, PackageMapClient->GetConnection(), static_cast< FNetBitWriter& >( Ar ), Changed );
+					CachedRequestState.RepLayout->SendProperties_BackwardsCompatible(nullptr, nullptr, (uint8*)Params.Data, PackageMapClient->GetConnection(), static_cast<FNetBitWriter&>(Ar), Changed);
 				}
 				else
 				{
 					bool bHasGuidsChanged = false;
-					RepLayout->ReceiveProperties_BackwardsCompatible( PackageMapClient->GetConnection(), nullptr, Data, static_cast< FNetBitReader& >( Ar ), bHasUnmapped, false, bHasGuidsChanged );
+					CachedRequestState.RepLayout->ReceiveProperties_BackwardsCompatible(PackageMapClient->GetConnection(), nullptr, Params.Data, static_cast<FNetBitReader&>(Ar), Params.bOutHasMoreUnmapped, false, bHasGuidsChanged);
 				}
 			}
 			else
 			{
-				RepLayout->SerializePropertiesForStruct( Struct, Ar, Map, Data, bHasUnmapped );
+				CachedRequestState.RepLayout->SerializePropertiesForStruct(Params.Struct, Ar, Params.Map, Params.Data, Params.bOutHasMoreUnmapped);
 			}
 		}
 	}
 
-	UNetDriver * Driver;
+	virtual bool NetDeltaSerializeForFastArray(FFastArrayDeltaSerializeParams& Params) override final
+	{
+		UpdateCachedState(Params.DeltaSerializeInfo.Object->GetClass(), Params.DeltaSerializeInfo.Struct);
+		UpdateCachedRepLayout();
+		return CachedRequestState.RepLayout->DeltaSerializeFastArrayProperty(Params, ChangelistMgr.Get());
+	}
+
+	virtual void GatherGuidReferencesForFastArray(FFastArrayDeltaSerializeParams& Params) override final
+	{
+		UpdateCachedState(Params.DeltaSerializeInfo.Object->GetClass(), Params.DeltaSerializeInfo.Struct);
+		UpdateCachedRepLayout();
+		CachedRequestState.RepLayout->GatherGuidReferencesForFastArray(Params);
+	}
+
+	virtual bool MoveGuidToUnmappedForFastArray(FFastArrayDeltaSerializeParams& Params) override final
+	{
+		UpdateCachedState(Params.DeltaSerializeInfo.Object->GetClass(), Params.DeltaSerializeInfo.Struct);
+		UpdateCachedRepLayout();
+		return CachedRequestState.RepLayout->MoveMappedObjectToUnmappedForFastArray(Params);
+	}
+
+	virtual void UpdateUnmappedGuidsForFastArray(FFastArrayDeltaSerializeParams& Params) override final
+	{
+		UpdateCachedState(Params.DeltaSerializeInfo.Object->GetClass(), Params.DeltaSerializeInfo.Struct);
+		UpdateCachedRepLayout();
+		CachedRequestState.RepLayout->UpdateUnmappedGuidsForFastArray(Params);
+	}
+
+public:
+
+	// These can go away once we do a full merge of Custom Delta and RepLayout.
+
+	static bool SendCustomDeltaProperty(
+		const FRepLayout& RepLayout,
+		FNetDeltaSerializeInfo& Params,
+		uint16 CustomDeltaProperty)
+	{
+		return RepLayout.SendCustomDeltaProperty(Params, CustomDeltaProperty);
+	}
+
+	static bool ReceiveCustomDeltaProperty(
+		const FRepLayout& RepLayout,
+		FNetDeltaSerializeInfo& Params,
+		UStructProperty* ReplicatedProp,
+		uint32& StaticArrayIndex,
+		int32& Offset)
+	{
+		return RepLayout.ReceiveCustomDeltaProperty(Params, ReplicatedProp, StaticArrayIndex, Offset);
+	}
+
+	static void GatherGuidReferencesForCustomDeltaProperties(const FRepLayout& RepLayout, FNetDeltaSerializeInfo& Params)
+	{
+		RepLayout.GatherGuidReferencesForCustomDeltaProperties(Params);
+	}
+
+	static bool MoveMappedObjectToUnmappedForCustomDeltaProperties(
+		const FRepLayout& RepLayout,
+		FNetDeltaSerializeInfo& Params,
+		TMap<int32, UStructProperty*>& UnmappedCustomProperties)
+	{
+		return RepLayout.MoveMappedObjectToUnmappedForCustomDeltaProperties(Params, UnmappedCustomProperties);
+	}
+
+	static void UpdateUnmappedObjectsForCustomDeltaProperties(
+		const FRepLayout& RepLayout,
+		FNetDeltaSerializeInfo& Params,
+		TArray<TPair<int32, UStructProperty*>>& CompletelyMappedProperties,
+		TArray<TPair<int32, UStructProperty*>>& UpdatedProperties)
+	{
+		RepLayout.UpdateUnmappedObjectsForCustomDeltaProperties(Params, CompletelyMappedProperties, UpdatedProperties);
+	}
+
+	static void PreSendCustomDeltaProperties(
+		const FRepLayout& RepLayout,
+		UObject* Object,
+		UNetConnection* Connection,
+		FReplicationChangelistMgr& ChangelistMgr,
+		TMap<int32, TSharedPtr<INetDeltaBaseState>>& CustomDeltaStates)
+	{
+		RepLayout.PreSendCustomDeltaProperties(Object, Connection, ChangelistMgr, CustomDeltaStates);
+	}
+
+	static void PostSendCustomDeltaProperties(
+		const FRepLayout& RepLayout,
+		UObject* Object,
+		UNetConnection* Connection,
+		FReplicationChangelistMgr& ChangelistMgr,
+		TMap<int32, TSharedPtr<INetDeltaBaseState>>& CustomDeltaStates)
+	{
+		RepLayout.PostSendCustomDeltaProperties(Object, Connection, ChangelistMgr, CustomDeltaStates);
+	}
+
+private:
+
+	UNetDriver* Driver;
+	FCachedRequestState CachedRequestState;
+	TSharedPtr<FReplicationChangelistMgr> ChangelistMgr;
 };
+PRAGMA_ENABLE_DEPRECATION_WARNINGS
 
+PRAGMA_DISABLE_DEPRECATION_WARNINGS
 FObjectReplicator::FObjectReplicator() :
 	ObjectClass(nullptr),
 	ObjectPtr(nullptr),
@@ -111,7 +288,9 @@ FObjectReplicator::~FObjectReplicator()
 {
 	CleanUp();
 }
+PRAGMA_ENABLE_DEPRECATION_WARNINGS
 
+PRAGMA_DISABLE_DEPRECATION_WARNINGS
 bool FObjectReplicator::SerializeCustomDeltaProperty( UNetConnection * Connection, void* Src, UProperty * Property, uint32 ArrayIndex, FNetBitWriter & OutBunch, TSharedPtr<INetDeltaBaseState> &NewFullState, TSharedPtr<INetDeltaBaseState> & OldState )
 {
 	check( NewFullState.IsValid() == false ); // NewState is passed in as NULL and instantiated within this function if necessary
@@ -128,17 +307,18 @@ bool FObjectReplicator::SerializeCustomDeltaProperty( UNetConnection * Connectio
 		return false;
 	}
 
-	FNetDeltaSerializeInfo Parms;
-
 	FNetSerializeCB NetSerializeCB( Connection->Driver );
 
-	Parms.Writer				= &OutBunch;
-	Parms.Map					= Connection->PackageMap;
-	Parms.OldState				= OldState.Get();
-	Parms.NewState				= &NewFullState;
-	Parms.NetSerializeCB		= &NetSerializeCB;
-	Parms.bIsWritingOnClient	= (Connection->Driver && Connection->Driver->GetWorld()) ? Connection->Driver->GetWorld()->IsRecordingClientReplay() : false;
-
+	FNetDeltaSerializeInfo Parms;
+	Parms.Data = Property->ContainerPtrToValuePtr<void>(Src, ArrayIndex);
+	Parms.Object = reinterpret_cast<UObject*>(Src);
+	Parms.Connection = Connection;
+	Parms.Writer = &OutBunch;
+	Parms.Map = Connection->PackageMap;
+	Parms.OldState = OldState.Get();
+	Parms.NewState = &NewFullState;
+	Parms.NetSerializeCB = &NetSerializeCB;
+	Parms.bIsWritingOnClient = (Connection->Driver && Connection->Driver->GetWorld()) ? Connection->Driver->GetWorld()->IsRecordingClientReplay() : false;
 
 	UScriptStruct::ICppStructOps * CppStructOps = StructProperty->Struct->GetCppStructOps();
 
@@ -151,7 +331,39 @@ bool FObjectReplicator::SerializeCustomDeltaProperty( UNetConnection * Connectio
 		OutBunch.SerializeIntPacked( ArrayIndex );
 	}
 
-	return CppStructOps->NetDeltaSerialize( Parms, Property->ContainerPtrToValuePtr<void>( Src, ArrayIndex ) );
+	return CppStructOps->NetDeltaSerialize(Parms, Parms.Data);
+}
+PRAGMA_ENABLE_DEPRECATION_WARNINGS
+
+bool FObjectReplicator::SendCustomDeltaProperty(UObject* InObject, UProperty* Property, uint32 ArrayIndex, FNetBitWriter& OutBunch, TSharedPtr<INetDeltaBaseState>& NewFullState, TSharedPtr<INetDeltaBaseState>& OldState)
+{
+	return SendCustomDeltaProperty(InObject, Property->RepIndex + ArrayIndex, OutBunch, NewFullState, OldState);
+}
+
+bool FObjectReplicator::SendCustomDeltaProperty(UObject* InObject, uint16 CustomDeltaProperty, FNetBitWriter& OutBunch, TSharedPtr<INetDeltaBaseState>& NewFullState, TSharedPtr<INetDeltaBaseState>& OldState)
+{
+	check(!NewFullState.IsValid()); // NewState is passed in as NULL and instantiated within this function if necessary
+	check(RepLayout);
+
+	SCOPE_CYCLE_COUNTER(STAT_NetSerializeItemDeltaTime);
+
+	UNetDriver* const ConnectionDriver = Connection->GetDriver();
+	FNetSerializeCB NetSerializeCB(ConnectionDriver);
+	NetSerializeCB.SetChangelistMgr(ChangelistMgr);
+
+	FNetDeltaSerializeInfo Parms;
+	Parms.Object = InObject;
+	Parms.Writer = &OutBunch;
+	Parms.Map = Connection->PackageMap;
+	Parms.OldState = OldState.Get();
+	Parms.NewState = &NewFullState;
+	Parms.NetSerializeCB = &NetSerializeCB;
+	Parms.bIsWritingOnClient = ConnectionDriver && ConnectionDriver->GetWorld() && ConnectionDriver->GetWorld()->IsRecordingClientReplay();
+	Parms.PropertyRepIndex = CustomDeltaProperty;
+	Parms.bSupportsFastArrayDeltaStructSerialization = bSupportsFastArrayDelta;
+	Parms.Connection = Connection;
+
+	return FNetSerializeCB::SendCustomDeltaProperty(*RepLayout, Parms, CustomDeltaProperty);
 }
 
 /** 
@@ -160,25 +372,30 @@ bool FObjectReplicator::SerializeCustomDeltaProperty( UNetConnection * Connectio
 **/
 void FObjectReplicator::InitRecentProperties(uint8* Source)
 {
-	check(GetObject() != NULL);
-	check(Connection != NULL);
+	// TODO: Could we just use the cached ObjectPtr here?
+	UObject* MyObject = GetObject();
+
+	check(MyObject);
+	check(Connection);
 	check(!RepState.IsValid());
 
-	const bool bIsServer = Connection->Driver->IsServer();
-	const bool bCreateSendingState = bIsServer || Connection->Driver->MaySendProperties();
+	UNetDriver* const ConnectionDriver = Connection->GetDriver();
+	const bool bIsServer = ConnectionDriver->IsServer();
+	const bool bCreateSendingState = bIsServer || ConnectionDriver->MaySendProperties();
+	const FRepLayout& LocalRepLayout = *RepLayout;
 
-	UClass* InObjectClass = GetObject()->GetClass();
+	UClass* InObjectClass = MyObject->GetClass();
 
 	// Initialize the RepState memory
 	// Clients don't need RepChangedPropertyTracker's, as they are mainly
 	// used temporarily disable property replication, or store data
 	// for replays (and the DemoNetDriver will be acts as a server during recording).
-	TSharedPtr<FRepChangedPropertyTracker> RepChangedPropertyTracker = bCreateSendingState ? Connection->Driver->FindOrCreateRepChangedPropertyTracker(GetObject()) : nullptr;
+	TSharedPtr<FRepChangedPropertyTracker> RepChangedPropertyTracker = bCreateSendingState ? ConnectionDriver->FindOrCreateRepChangedPropertyTracker(MyObject) : nullptr;
 
 	// If acting as a server and are InternalAck, that means we're recording.
 	// In that case, we don't need to create any receiving state, as no one will be sending data to us.
 	ECreateRepStateFlags Flags = (Connection->InternalAck && bIsServer) ? ECreateRepStateFlags::SkipCreateReceivingState : ECreateRepStateFlags::None;
-	RepState = RepLayout->CreateRepState(Source, RepChangedPropertyTracker, Flags);
+	RepState = LocalRepLayout.CreateRepState(Source, RepChangedPropertyTracker, Flags);
 
 	if (!bCreateSendingState)
 	{
@@ -186,30 +403,29 @@ void FObjectReplicator::InitRecentProperties(uint8* Source)
 		return;
 	}
 
+	bSupportsFastArrayDelta = !!GSupportsFastArrayDelta;
+
+	// TODO: CDOCustomDeltaState, CheckpointCustomDeltaState, RecentCustomDeltaState, and Retirement could all be moved into SendingRepState.
+	//			This would allow us to skip allocating these containers for receivers completely.
+	//			This logic would also be easily moved to FRepLayout::CreateRepState.
+
+	// We should just update this method to accept an object pointer.
+	UObject* UseObject = reinterpret_cast<UObject*>(Source);
+
 	// Init custom delta property state
-	for ( TFieldIterator<UProperty> It( InObjectClass ); It; ++It )
+	for (const uint16 CustomDeltaProperty : LocalRepLayout.GetLifetimeCustomDeltaProperties())
 	{
-		if ( It->PropertyFlags & CPF_Net )
-		{
-			if ( IsCustomDeltaProperty( *It ) )
-			{
-				// We have to handle dynamic properties of the array individually
-				for ( int32 ArrayIdx = 0; ArrayIdx < It->ArrayDim; ++ArrayIdx )
-				{
-					FOutBunch DeltaState( Connection->PackageMap );
-					TSharedPtr<INetDeltaBaseState> & NewState = RecentCustomDeltaState.FindOrAdd(It->RepIndex + ArrayIdx);
-					NewState.Reset();					
+		FOutBunch DeltaState(Connection->PackageMap);
+		TSharedPtr<INetDeltaBaseState>& NewState = RecentCustomDeltaState.FindOrAdd(CustomDeltaProperty);
+		NewState.Reset();					
 
-					TSharedPtr<INetDeltaBaseState> OldState;
+		TSharedPtr<INetDeltaBaseState> OldState;
 
-					SerializeCustomDeltaProperty( Connection, Source, *It, ArrayIdx, DeltaState, NewState, OldState );
+		SendCustomDeltaProperty(UseObject, CustomDeltaProperty, DeltaState, NewState, OldState);
 
-					// Store the initial delta state in case we need it for when we're asked to resend all data since channel was first opened (bResendAllDataSinceOpen)
-					CDOCustomDeltaState.Add( It->RepIndex + ArrayIdx, NewState ); 
-					CheckpointCustomDeltaState.Add( It->RepIndex + ArrayIdx, NewState );
-				}
-			}
-		}
+		// Store the initial delta state in case we need it for when we're asked to resend all data since channel was first opened (bResendAllDataSinceOpen)
+		CDOCustomDeltaState.Add(CustomDeltaProperty, NewState);
+		CheckpointCustomDeltaState.Add(CustomDeltaProperty, NewState);
 	}
 }
 
@@ -290,7 +506,9 @@ void FObjectReplicator::InitWithObject( UObject* InObject, UNetConnection * InCo
 
 	InitRecentProperties( Source );
 
+PRAGMA_DISABLE_DEPRECATION_WARNINGS
 	RepLayout->GetLifetimeCustomDeltaProperties( LifetimeCustomDeltaProperties, LifetimeCustomDeltaPropertyConditions );
+PRAGMA_ENABLE_DEPRECATION_WARNINGS
 
 	Connection->Driver->AllOwnedReplicators.Add(this);
 }
@@ -339,38 +557,41 @@ void FObjectReplicator::CleanUp()
 	bHasReplicatedProperties	= false;
 	bOpenAckCalled				= false;
 
+PRAGMA_DISABLE_DEPRECATION_WARNINGS
 	// Cleanup custom delta state
 	RecentCustomDeltaState.Empty();
 	CheckpointCustomDeltaState.Empty();
 
 	LifetimeCustomDeltaProperties.Empty();
 	LifetimeCustomDeltaPropertyConditions.Empty();
+PRAGMA_ENABLE_DEPRECATION_WARNINGS
 
 	RepState = nullptr;
 	CheckpointRepState = nullptr;
 }
 
-void FObjectReplicator::StartReplicating( class UActorChannel * InActorChannel )
+void FObjectReplicator::StartReplicating(class UActorChannel * InActorChannel)
 {
 	check(OwningChannel == nullptr);
 	check(InActorChannel != nullptr);
+	check(InActorChannel->Connection == Connection);
 
 	UObject* const Object = GetObject();
-	if ( Object == nullptr )
+	if (Object == nullptr)
 	{
 		UE_LOG(LogRep, Error, TEXT("StartReplicating: Object == nullptr"));
 		return;
 	}
 
-	if ( !ensureMsgf( ObjectClass != nullptr, TEXT( "StartReplicating: ObjectClass == nullptr. Object = %s. Channel actor = %s. %s" ), *GetFullNameSafe( Object ), *GetFullNameSafe( InActorChannel->GetActor() ), *InActorChannel->Connection->Describe() ) )
+	if (!ensureMsgf(ObjectClass != nullptr, TEXT( "StartReplicating: ObjectClass == nullptr. Object = %s. Channel actor = %s. %s" ), *GetFullNameSafe(Object), *GetFullNameSafe(InActorChannel->GetActor()), *InActorChannel->Connection->Describe()))
 	{
 		return;
 	}
 
-	if ( UClass* const ObjectPtrClass = Object->GetClass() )
+	if (UClass* const ObjectPtrClass = Object->GetClass())
 	{
 		// Something is overwriting a bit in the ObjectClass pointer so it's becoming invalid - fix up the pointer to prevent crashing later until the real cause can be identified.
-		if ( !ensureMsgf( ObjectClass == ObjectPtrClass, TEXT( "StartReplicating: ObjectClass and ObjectPtr's class are not equal and they should be. Object = %s. Channel actor = %s. %s" ), *GetFullNameSafe( Object ), *GetFullNameSafe( InActorChannel->GetActor() ), *InActorChannel->Connection->Describe() ) )
+		if (!ensureMsgf(ObjectClass == ObjectPtrClass, TEXT("StartReplicating: ObjectClass and ObjectPtr's class are not equal and they should be. Object = %s. Channel actor = %s. %s"), *GetFullNameSafe(Object), *GetFullNameSafe(InActorChannel->GetActor()), *InActorChannel->Connection->Describe()))
 		{
 			ObjectClass = ObjectPtrClass;
 		}
@@ -378,28 +599,33 @@ void FObjectReplicator::StartReplicating( class UActorChannel * InActorChannel )
 
 	OwningChannel = InActorChannel;
 
+	UNetDriver* const ConnectionNetDriver = Connection->GetDriver();
+
 	// Cache off netGUID so if this object gets deleted we can close it
-	ObjectNetGUID = OwningChannel->Connection->Driver->GuidCache->GetOrAssignNetGUID( Object );
-	check( !ObjectNetGUID.IsDefault() && ObjectNetGUID.IsValid() );
+	ObjectNetGUID = ConnectionNetDriver->GuidCache->GetOrAssignNetGUID( Object );
+	check(!ObjectNetGUID.IsDefault() && ObjectNetGUID.IsValid());
 
-	if (Connection->Driver->IsServer() || Connection->Driver->MaySendProperties())
+	if (ConnectionNetDriver->IsServer() || ConnectionNetDriver->MaySendProperties())
 	{
 		// Allocate retirement list.
 		// SetNum now constructs, so this is safe
 		Retirement.SetNum(ObjectClass->ClassReps.Num());
 
-		const UWorld* const World = Connection->Driver->GetWorld();
+		const UWorld* const World = ConnectionNetDriver->GetWorld();
+		UNetDriver* const WorldNetDriver = World ? World->GetNetDriver() : nullptr;
+
 		// Prefer the changelist manager on the main net driver (so we share across net drivers if possible)
-		if (World && World->NetDriver && World->NetDriver->IsServer())
+		if (WorldNetDriver && WorldNetDriver->IsServer())
 		{
-			ChangelistMgr = World->NetDriver->GetReplicationChangeListMgr(Object);
+			ChangelistMgr = WorldNetDriver->GetReplicationChangeListMgr(Object);
 		}
 		else
 		{
-			ChangelistMgr = Connection->Driver->GetReplicationChangeListMgr(Object);
+			ChangelistMgr = ConnectionNetDriver->GetReplicationChangeListMgr(Object);
 		}
 	}
 }
+
 static FORCEINLINE void ValidateRetirementHistory( const FPropertyRetirement & Retire, const UObject* Object )
 {
 #if !UE_BUILD_SHIPPING
@@ -532,50 +758,55 @@ void FObjectReplicator::ReceivedNak( int32 NakPacketId )
 	FieldCache->bIncompatible = true;	\
 	continue;							\
 
-bool FObjectReplicator::ReceivedBunch( FNetBitReader& Bunch, const FReplicationFlags& RepFlags, const bool bHasRepLayout, bool& bOutHasUnmapped )
+bool FObjectReplicator::ReceivedBunch(FNetBitReader& Bunch, const FReplicationFlags& RepFlags, const bool bHasRepLayout, bool& bOutHasUnmapped)
 {
+	check(RepLayout);
+
 	UObject* Object = GetObject();
 
-	if ( Object == NULL )
+	if (!Object)
 	{
 		UE_LOG(LogNet, Verbose, TEXT("ReceivedBunch: Object == NULL"));
 		return false;
 	}
 
-	UPackageMap * PackageMap = OwningChannel->Connection->PackageMap;
+	UNetDriver* const ConnectionNetDriver = Connection->GetDriver();
+	UPackageMap* const PackageMap = Connection->PackageMap;
 
-	const bool bIsServer = OwningChannel->Connection->Driver->IsServer();
+	const bool bIsServer = ConnectionNetDriver->IsServer();
 	const bool bCanDelayRPCs = (CVarDelayUnmappedRPCs.GetValueOnGameThread() > 0) && !bIsServer;
 
-	const FClassNetCache * ClassCache = OwningChannel->Connection->Driver->NetCache->GetClassNetCache( ObjectClass );
+	const FClassNetCache* const ClassCache = ConnectionNetDriver->NetCache->GetClassNetCache(ObjectClass);
 
-	if ( ClassCache == NULL )
+	if (!ClassCache)
 	{
 		UE_LOG(LogNet, Error, TEXT("ReceivedBunch: ClassCache == NULL: %s"), *Object->GetFullName());
 		return false;
 	}
 
+	const FRepLayout& LocalRepLayout = *RepLayout;
 	bool bGuidsChanged = false;
 
 	// Handle replayout properties
-	if ( bHasRepLayout )
+	if (bHasRepLayout)
 	{
 		// Server shouldn't receive properties.
-		if ( bIsServer )
+		if (bIsServer)
 		{
-			UE_LOG( LogNet, Error, TEXT( "Server received RepLayout properties: %s" ), *Object->GetFullName() );
+			UE_LOG(LogNet, Error, TEXT("Server received RepLayout properties: %s"), *Object->GetFullName());
 			return false;
 		}
 
-		if ( !bHasReplicatedProperties )
+		if (!bHasReplicatedProperties)
 		{
-			bHasReplicatedProperties = true;		// Persistent, not reset until PostNetReceive is called
+			// Persistent, not reset until PostNetReceive is called
+			bHasReplicatedProperties = true;
 			PreNetReceive();
 		}
 
 		EReceivePropertiesFlags ReceivePropFlags = EReceivePropertiesFlags::None;
 
-		if (Connection->Driver->ShouldReceiveRepNotifiesForObject(Object))
+		if (ConnectionNetDriver->ShouldReceiveRepNotifiesForObject(Object))
 		{
 			ReceivePropFlags |= EReceivePropertiesFlags::RepNotifies;
 		}
@@ -587,44 +818,75 @@ bool FObjectReplicator::ReceivedBunch( FNetBitReader& Bunch, const FReplicationF
 
 		bool bLocalHasUnmapped = false;
 
-		if ( !RepLayout->ReceiveProperties(OwningChannel, ObjectClass, RepState->GetReceivingRepState(), ( void* )Object, Bunch, bLocalHasUnmapped, bGuidsChanged, ReceivePropFlags ) )
+		if (!LocalRepLayout.ReceiveProperties(OwningChannel, ObjectClass, RepState->GetReceivingRepState(), Object, Bunch, bLocalHasUnmapped, bGuidsChanged, ReceivePropFlags))
 		{
-			UE_LOG( LogRep, Error, TEXT( "RepLayout->ReceiveProperties FAILED: %s" ), *Object->GetFullName() );
+			UE_LOG(LogRep, Error, TEXT( "RepLayout->ReceiveProperties FAILED: %s" ), *Object->GetFullName());
 			return false;
 		}
 
-		if ( bLocalHasUnmapped )
-		{
-			bOutHasUnmapped = true;
-		}
+		bOutHasUnmapped |= bLocalHasUnmapped;
 	}
 
-	FNetFieldExportGroup* NetFieldExportGroup = OwningChannel->GetNetFieldExportGroupForClassNetCache( ObjectClass );
+	FNetFieldExportGroup* NetFieldExportGroup = OwningChannel->GetNetFieldExportGroupForClassNetCache(ObjectClass);
 
 	FNetBitReader Reader( Bunch.PackageMap );
 
 	// Read fields from stream
 	const FFieldNetCache * FieldCache = nullptr;
 	
+	// TODO:	As of now, we replicate all of our Custom Delta Properties immediately after our normal properties.
+	//			An optimization could be made here in the future if we replicated / received Custom Delta Properties in RepLayout
+	//			immediately with normal properties.
+	//
+	//			For the Standard case, we expect the RepLayout to be identical on Client and Server.
+	//				If the RepLayout doesn't have any Custom Delta Properties, everything stays as it is now.
+	//				If the RepLayout does have Custom Delta Properties, then:
+	//					1. We replicate a single bit indicating whether or not any were actually sent.
+	//					2. We replicate a packed int specifying the number of custom delta properties (if any were sent).
+	//					3. We replicate the Header and Payloads as normal.
+	//				This may increase bandwidth slightly, but it's likely negligible.
+	//
+	//			For the Backwards Compatible path, we do the above, except we always send the bit flag, and the count when set.
+	//				In that way, if Custom Delta Properties are added or removed, we can always rely on the bit field to try and
+	//				read them, and then throw them away if they are incompatible.
+	//
+	//			In both described cases, we could remove the first cast to a struct property below, and flags checks on the properties
+	//			as we could instead use the RepLayout cached command flags which would hopefully reduce cache misses.
+	//			This also means that we could leverage the bIsServer and bHasReplicatedProperties that have already taken place.
+	//
+	//			If we want maintain compatibility with older builds (mostly for replays), we could leave the branch in here for now
+	//			but short circuit it with a net version check, still allowing us to skip the cast in new versions.
+	//
+	//			This also becomes more convenient when we merge RepNotify handling.
+
+	FNetSerializeCB NetSerializeCB(ConnectionNetDriver);
+
+	FNetDeltaSerializeInfo Parms;
+	Parms.Map = PackageMap;
+	Parms.Reader = &Reader;
+	Parms.NetSerializeCB = &NetSerializeCB;
+	Parms.Connection = Connection;
+	Parms.Object = Object;
+
 	// Read each property/function blob into Reader (so we've safely jumped over this data in the Bunch/stream at this point)
-	while ( OwningChannel->ReadFieldHeaderAndPayload( Object, ClassCache, NetFieldExportGroup, Bunch, &FieldCache, Reader ) )
+	while (OwningChannel->ReadFieldHeaderAndPayload(Object, ClassCache, NetFieldExportGroup, Bunch, &FieldCache, Reader))
 	{
-		if ( Bunch.IsError() )
+		if (UNLIKELY(Bunch.IsError()))
 		{
-			UE_LOG( LogNet, Error, TEXT( "ReceivedBunch: Error reading field: %s" ), *Object->GetFullName() );
+			UE_LOG(LogNet, Error, TEXT("ReceivedBunch: Error reading field: %s"), *Object->GetFullName());
 			return false;
 		}
 
-		if ( FieldCache == nullptr )
+		else if (UNLIKELY(FieldCache == nullptr))
 		{
-			UE_LOG( LogNet, Warning, TEXT( "ReceivedBunch: FieldCache == nullptr: %s" ), *Object->GetFullName() );
+			UE_LOG(LogNet, Warning, TEXT("ReceivedBunch: FieldCache == nullptr: %s"), *Object->GetFullName());
 			continue;
 		}
 
-		if ( FieldCache->bIncompatible )
+		else if (UNLIKELY(FieldCache->bIncompatible))
 		{
 			// We've already warned about this property once, so no need to continue to do so
-			UE_LOG( LogNet, Verbose, TEXT( "ReceivedBunch: FieldCache->bIncompatible == true. Object: %s, Field: %s" ), *Object->GetFullName(), *FieldCache->Field->GetFName().ToString() );
+			UE_LOG(LogNet, Verbose, TEXT( "ReceivedBunch: FieldCache->bIncompatible == true. Object: %s, Field: %s" ), *Object->GetFullName(), *FieldCache->Field->GetFName().ToString());
 			continue;
 		}
 
@@ -638,76 +900,37 @@ bool FObjectReplicator::ReceivedBunch( FNetBitReader& Bunch, const FReplicationF
 				return false;
 			}
 
-			UScriptStruct * InnerStruct = ReplicatedProp->Struct;
-			UScriptStruct::ICppStructOps * CppStructOps = InnerStruct->GetCppStructOps();
-
-			check(CppStructOps);
-
-			// We should only be receiving custom delta properties (since RepLayout handles the rest)
-			if (!IsCustomDeltaProperty(ReplicatedProp))
-			{
-				UE_LOG( LogNet, Error, TEXT( "Client received non custom delta property value %s in %s" ), *ReplicatedProp->GetName(), *Object->GetFullName() );
-				return false;
-			}
-
 			// Call PreNetReceive if we haven't yet
 			if (!bHasReplicatedProperties)
 			{
-				bHasReplicatedProperties = true;		// Persistent, not reset until PostNetReceive is called
+				// Persistent, not reset until PostNetReceive is called
+				bHasReplicatedProperties = true;
 				PreNetReceive();
 			}
 
 #if !(UE_BUILD_SHIPPING || UE_BUILD_TEST)
 			{
 				FString DebugPropertyStr = CVarNetReplicationDebugProperty.GetValueOnAnyThread();
-				if (!DebugPropertyStr.IsEmpty() && ReplicatedProp->GetName().Contains(DebugPropertyStr) )
+				if (!DebugPropertyStr.IsEmpty() && ReplicatedProp->GetName().Contains(DebugPropertyStr))
 				{
 					UE_LOG(LogRep, Log, TEXT("Replicating Property[%d] %s on %s"), ReplicatedProp->RepIndex, *ReplicatedProp->GetName(), *Object->GetName());
 				}
 			}
 #endif
-
-			// Receive array index (static sized array, i.e. MemberVariable[4])
-			uint32 Element = 0;
-			if (ReplicatedProp->ArrayDim != 1)
+			uint32 StaticArrayIndex = 0;
+			int32 Offset = 0;
+			if (!FNetSerializeCB::ReceiveCustomDeltaProperty(LocalRepLayout, Parms, ReplicatedProp, StaticArrayIndex, Offset))
 			{
-				check(ReplicatedProp->ArrayDim >= 2);
-
-				Reader.SerializeIntPacked(Element);
-
-				if (Element >= (uint32)ReplicatedProp->ArrayDim)
-				{
-					UE_LOG(LogRep, Error, TEXT("Element index too large %s in %s"), *ReplicatedProp->GetName(), *Object->GetFullName());
-					return false;
-				}
+				// RepLayout should have already logged the error.
+				HANDLE_INCOMPATIBLE_PROP;
 			}
-
-			// Pointer to destination.
-			uint8* Data = ReplicatedProp->ContainerPtrToValuePtr<uint8>((uint8*)Object, Element);
-			TArray<uint8>	MetaData;
-			const PTRINT DataOffset = Data - (uint8*)Object;
-
-			FNetDeltaSerializeInfo Parms;
-
-			FNetSerializeCB NetSerializeCB(OwningChannel->Connection->Driver);
-
-			Parms.DebugName				= ReplicatedProp->GetName();
-			Parms.Struct				= InnerStruct;
-			Parms.Map					= PackageMap;
-			Parms.Reader				= &Reader;
-			Parms.NetSerializeCB		= &NetSerializeCB;
-			Parms.bIsWritingOnClient	= false;
-
-			// Call the custom delta serialize function to handle it
-			CppStructOps->NetDeltaSerialize(Parms, Data);
-
-			if (Reader.IsError())
+			else if (UNLIKELY(Reader.IsError()))
 			{
 				UE_LOG(LogNet, Error, TEXT("ReceivedBunch: NetDeltaSerialize - Reader.IsError() == true. Property: %s, Object: %s"), *Parms.DebugName, *Object->GetFullName());
 				HANDLE_INCOMPATIBLE_PROP
 			}
 
-			if (Reader.GetBitsLeft() != 0)
+			else if (UNLIKELY(Reader.GetBitsLeft() != 0))
 			{
 				UE_LOG( LogNet, Error, TEXT( "ReceivedBunch: NetDeltaSerialize - Mismatch read. Property: %s, Object: %s" ), *Parms.DebugName, *Object->GetFullName() );
 				HANDLE_INCOMPATIBLE_PROP
@@ -715,7 +938,10 @@ bool FObjectReplicator::ReceivedBunch( FNetBitReader& Bunch, const FReplicationF
 
 			if (Parms.bOutHasMoreUnmapped)
 			{
-				UnmappedCustomProperties.Add( DataOffset, ReplicatedProp );
+PRAGMA_DISABLE_DEPRECATION_WARNINGS
+				UnmappedCustomProperties.Add(Offset, ReplicatedProp);
+PRAGMA_ENABLE_DEPRECATION_WARNINGS
+
 				bOutHasUnmapped = true;
 			}
 
@@ -728,7 +954,8 @@ bool FObjectReplicator::ReceivedBunch( FNetBitReader& Bunch, const FReplicationF
 			UE_LOG(LogRepTraffic, Log, TEXT(" %s - %s"), *Object->GetName(), *Parms.DebugName);
 
 			// Notify the Object if this var is RepNotify
-			QueuePropertyRepNotify(Object, ReplicatedProp, Element, MetaData);
+			TArray<uint8> MetaData;
+			QueuePropertyRepNotify(Object, ReplicatedProp, StaticArrayIndex, MetaData);
 		}
 		// Handle function call
 		else if ( Cast< UFunction >( FieldCache->Field ) )
@@ -927,47 +1154,37 @@ bool FObjectReplicator::ReceivedRPC(FNetBitReader& Reader, const FReplicationFla
 
 void FObjectReplicator::UpdateGuidToReplicatorMap()
 {
-	SCOPE_CYCLE_COUNTER( STAT_NetUpdateGuidToReplicatorMap );
+	SCOPE_CYCLE_COUNTER(STAT_NetUpdateGuidToReplicatorMap);
 
 	if (Connection->Driver->IsServer())
 	{
 		return;
 	}
 
-	TSet< FNetworkGUID > LocalReferencedGuids;
+	TSet<FNetworkGUID> LocalReferencedGuids;
 	int32 LocalTrackedGuidMemoryBytes = 0;
 
+	check(RepLayout);
+
+	const FRepLayout& LocalRepLayout = *RepLayout;
+
 	// Gather guids on rep layout
-	if ( RepLayout.IsValid() && RepState.IsValid() )
+	if (RepState.IsValid())
 	{
-		RepLayout->GatherGuidReferences( RepState->GetReceivingRepState(), LocalReferencedGuids, LocalTrackedGuidMemoryBytes );
+		LocalRepLayout.GatherGuidReferences(RepState->GetReceivingRepState(), LocalReferencedGuids, LocalTrackedGuidMemoryBytes);
 	}
 
-	UObject* Object = GetObject();
-
-	// Gather guids on fast tarray
-	for ( const int32 CustomIndex : LifetimeCustomDeltaProperties )
+	if (UObject* Object = GetObject())
 	{
-		FRepRecord* Rep	= &ObjectClass->ClassReps[CustomIndex];
-
-		UStructProperty* StructProperty = CastChecked< UStructProperty >( Rep->Property );
+		FNetSerializeCB NetSerializeCB(Connection->Driver);
 
 		FNetDeltaSerializeInfo Parms;
+		Parms.NetSerializeCB = &NetSerializeCB;
+		Parms.GatherGuidReferences = &LocalReferencedGuids;
+		Parms.TrackedGuidMemoryBytes = &LocalTrackedGuidMemoryBytes;
+		Parms.Object = Object;
 
-		FNetSerializeCB NetSerializeCB( Connection->Driver );
-
-		Parms.NetSerializeCB			= &NetSerializeCB;
-		Parms.GatherGuidReferences		= &LocalReferencedGuids;
-		Parms.TrackedGuidMemoryBytes	= &LocalTrackedGuidMemoryBytes;
-
-		UScriptStruct::ICppStructOps* CppStructOps = StructProperty->Struct->GetCppStructOps();
-
-		Parms.Struct = StructProperty->Struct;
-
-		if ( Object != nullptr )
-		{
-			CppStructOps->NetDeltaSerialize( Parms, StructProperty->ContainerPtrToValuePtr<void>( Object, Rep->Index ) );
-		}
+		FNetSerializeCB::GatherGuidReferencesForCustomDeltaProperties(LocalRepLayout, Parms);
 	}
 
 	// Gather RPC guids
@@ -983,26 +1200,26 @@ void FObjectReplicator::UpdateGuidToReplicatorMap()
 	}
 
 	// Go over all referenced guids, and make sure we're tracking them in the GuidToReplicatorMap
-	for ( const FNetworkGUID& GUID : LocalReferencedGuids )
+	for (const FNetworkGUID& GUID : LocalReferencedGuids)
 	{
-		if ( !ReferencedGuids.Contains( GUID ) )
+		if (!ReferencedGuids.Contains(GUID))
 		{
-			Connection->Driver->GuidToReplicatorMap.FindOrAdd( GUID ).Add( this );
+			Connection->Driver->GuidToReplicatorMap.FindOrAdd(GUID).Add(this);
 		}
 	}
 
 	// Remove any guids that we were previously tracking but no longer should
-	for ( const FNetworkGUID& GUID : ReferencedGuids )
+	for (const FNetworkGUID& GUID : ReferencedGuids)
 	{
-		if ( !LocalReferencedGuids.Contains( GUID ) )
+		if (!LocalReferencedGuids.Contains(GUID))
 		{
-			TSet< FObjectReplicator* >& Replicators = Connection->Driver->GuidToReplicatorMap.FindChecked( GUID );
+			TSet<FObjectReplicator*>& Replicators = Connection->Driver->GuidToReplicatorMap.FindChecked(GUID);
 
-			Replicators.Remove( this );
+			Replicators.Remove(this);
 
-			if ( Replicators.Num() == 0 )
+			if (Replicators.Num() == 0)
 			{
-				Connection->Driver->GuidToReplicatorMap.Remove( GUID );
+				Connection->Driver->GuidToReplicatorMap.Remove(GUID);
 			}
 		}
 	}
@@ -1011,50 +1228,31 @@ void FObjectReplicator::UpdateGuidToReplicatorMap()
 	TrackedGuidMemoryBytes = LocalTrackedGuidMemoryBytes;
 	Connection->Driver->TotalTrackedGuidMemoryBytes += TrackedGuidMemoryBytes;
 
-	ReferencedGuids = MoveTemp( LocalReferencedGuids );
+	ReferencedGuids = MoveTemp(LocalReferencedGuids);
 }
 
-bool FObjectReplicator::MoveMappedObjectToUnmapped( const FNetworkGUID& GUID )
+bool FObjectReplicator::MoveMappedObjectToUnmapped(const FNetworkGUID& GUID)
 {
-	bool bFound = false;
+	check(RepLayout);
+	const FRepLayout& LocalRepLayout = *RepLayout;
 
-	if ( RepLayout.IsValid() )
-	{
-		if ( RepLayout->MoveMappedObjectToUnmapped( RepState->GetReceivingRepState(), GUID ) )
-		{
-			bFound = true;
-		}
-	}
+	bool bFound = LocalRepLayout.MoveMappedObjectToUnmapped(RepState->GetReceivingRepState(), GUID);
 
-	UObject* Object = GetObject();
-
-	for ( const int32 CustomIndex : LifetimeCustomDeltaProperties )
+	if (UObject* Object = GetObject())
 	{
-		FRepRecord* Rep	= &ObjectClass->ClassReps[CustomIndex];
-
-		UStructProperty* StructProperty = CastChecked< UStructProperty >( Rep->Property );
+		FNetSerializeCB NetSerializeCB(Connection->Driver);
 
 		FNetDeltaSerializeInfo Parms;
+		Parms.Connection = Connection;
+		Parms.Map = Connection->PackageMap;
+		Parms.Object = Object;
+		Parms.NetSerializeCB = &NetSerializeCB;
+		Parms.MoveGuidToUnmapped = &GUID;
+		Parms.Object = Object;
 
-		FNetSerializeCB NetSerializeCB( Connection->Driver );
-
-		Parms.NetSerializeCB			= &NetSerializeCB;
-		Parms.MoveGuidToUnmapped		= &GUID;
-
-		UScriptStruct::ICppStructOps* CppStructOps = StructProperty->Struct->GetCppStructOps();
-
-		Parms.Struct = StructProperty->Struct;
-
-		if ( Object != nullptr )
-		{
-			void* Data = StructProperty->ContainerPtrToValuePtr<void>( Object, Rep->Index );
-
-			if ( CppStructOps->NetDeltaSerialize( Parms, Data ) )
-			{
-				UnmappedCustomProperties.Add( (uint8*)Data - (uint8*)Object, StructProperty );
-				bFound = true;
-			}
-		}
+PRAGMA_DISABLE_DEPRECATION_WARNINGS
+		bFound |= FNetSerializeCB::MoveMappedObjectToUnmappedForCustomDeltaProperties(LocalRepLayout, Parms, UnmappedCustomProperties);
+PRAGMA_ENABLE_DEPRECATION_WARNINGS
 	}
 
 	return bFound;
@@ -1112,18 +1310,36 @@ void FObjectReplicator::ReplicateCustomDeltaProperties( FNetBitWriter & Bunch, F
 {
 	SCOPE_CYCLE_COUNTER(STAT_NetReplicateCustomDeltaPropTime);
 
-	if ( LifetimeCustomDeltaProperties.Num() == 0 )
+	check(RepLayout);
+	const FRepLayout& LocalRepLayout = *RepLayout;
+	const TArrayView<const uint16> LocalLifetimeCustomDeltaProperties = LocalRepLayout.GetLifetimeCustomDeltaProperties();
+
+	if (!LocalLifetimeCustomDeltaProperties.Num())
 	{
 		// No custom properties
 		return;
 	}
 
+	// TODO: See comments in ReceivedBunch. This code should get merged into RepLayout, to help optimize
+	//			the receiving end, and make things more consistent.
+
 	UObject* Object = GetObject();
 
-	check( Object );
-	check( OwningChannel );
+	check(Object);
+	check(OwningChannel);
+	check(Connection == OwningChannel->Connection);
+
+	TMap<int32, TSharedPtr<INetDeltaBaseState>>& UsingCustomDeltaStates =
+		EResendAllDataState::None == Connection->ResendAllDataState ? RecentCustomDeltaState :
+		EResendAllDataState::SinceOpen == Connection->ResendAllDataState ? CDOCustomDeltaState :
+		CheckpointCustomDeltaState;
 
-	UNetConnection * OwningChannelConnection = OwningChannel->Connection;
+	FNetSerializeCB::PreSendCustomDeltaProperties(LocalRepLayout, Object, Connection, *ChangelistMgr, UsingCustomDeltaStates);
+
+	ON_SCOPE_EXIT
+	{
+		FNetSerializeCB::PostSendCustomDeltaProperties(LocalRepLayout, Object, Connection, *ChangelistMgr, UsingCustomDeltaStates);
+	};
 
 	// Initialize a map of which conditions are valid
 	const TStaticBitArray<COND_Max> ConditionMap = FSendingRepState::BuildConditionMapFromRepFlags(RepFlags);
@@ -1131,30 +1347,29 @@ void FObjectReplicator::ReplicateCustomDeltaProperties( FNetBitWriter & Bunch, F
 	// Make sure net field export group is registered
 	FNetFieldExportGroup* NetFieldExportGroup = OwningChannel->GetOrCreateNetFieldExportGroupForClassNetCache( Object );
 
-	FNetBitWriter TempBitWriter( OwningChannel->Connection->PackageMap, 1024 );
+	FNetBitWriter TempBitWriter( Connection->PackageMap, 1024 );
 
 	// Replicate those properties.
-	for ( int32 i = 0; i < LifetimeCustomDeltaProperties.Num(); i++ )
+	for (const uint16 CustomDeltaProperty : LocalLifetimeCustomDeltaProperties)
 	{
 		// Get info.
-		const int32				RetireIndex	= LifetimeCustomDeltaProperties[i];
-		FPropertyRetirement &	Retire		= Retirement[RetireIndex];
-		FRepRecord *			Rep			= &ObjectClass->ClassReps[RetireIndex];
-		UProperty *				It			= Rep->Property;
-		int32					Index		= Rep->Index;
+		FPropertyRetirement& Retire = Retirement[CustomDeltaProperty];
+		FRepRecord* Rep = &ObjectClass->ClassReps[CustomDeltaProperty];
 
-		if (LifetimeCustomDeltaPropertyConditions.IsValidIndex(i))
-		{
-			// Check the replication condition here
-			ELifetimeCondition RepCondition = LifetimeCustomDeltaPropertyConditions[i];
+PRAGMA_DISABLE_DEPRECATION_WARNINGS
+		UProperty* It = LocalRepLayout.GetPropertyForRepIndex(CustomDeltaProperty);
+PRAGMA_ENABLE_DEPRECATION_WARNINGS
 
-			check(RepCondition >= 0 && RepCondition < COND_Max);
+		int32 Index = Rep->Index;
 
-			if (!ConditionMap[RepCondition])
-			{
-				// We didn't pass the condition so don't replicate us
-				continue;
-			}
+		const ELifetimeCondition RepCondition = LocalRepLayout.GetPropertyLifetimeCondition(CustomDeltaProperty);
+
+		check(RepCondition >= 0 && RepCondition < COND_Max);
+
+		if (!ConditionMap[RepCondition])
+		{
+			// We didn't pass the condition so don't replicate us
+			continue;
 		}
 
 		// If this is a dynamic array, we do the delta here
@@ -1166,49 +1381,49 @@ void FObjectReplicator::ReplicateCustomDeltaProperties( FNetBitWriter & Bunch, F
 		{
 			if (Connection->ResendAllDataState == EResendAllDataState::SinceCheckpoint)
 			{
-				TSharedPtr<INetDeltaBaseState> OldState = CheckpointCustomDeltaState.FindChecked(RetireIndex);
+				TSharedPtr<INetDeltaBaseState>& OldState = CheckpointCustomDeltaState.FindChecked(CustomDeltaProperty);
 
-				if (!SerializeCustomDeltaProperty(OwningChannelConnection, (void*)Object, It, Index, TempBitWriter, NewState, OldState))
+				if (!SendCustomDeltaProperty(Object, CustomDeltaProperty, TempBitWriter, NewState, OldState))
 				{
 					continue;
 				}
 
 				// update checkpoint with new state
-				CheckpointCustomDeltaState.Add(RetireIndex, NewState);
+				OldState = NewState;
 			}
 			else
 			{
 				// If we are resending data since open, we don't want to affect the current state of channel/replication, so just do the minimum and send the data, and return
 				// In this case, we'll send all of the properties since the CDO, so use the initial CDO delta state
-				TSharedPtr<INetDeltaBaseState>& OldState = CDOCustomDeltaState.FindChecked(RetireIndex);
+				TSharedPtr<INetDeltaBaseState>& OldState = CDOCustomDeltaState.FindChecked(CustomDeltaProperty);
 
-				if (!SerializeCustomDeltaProperty(OwningChannelConnection, (void*)Object, It, Index, TempBitWriter, NewState, OldState))
+				if (!SendCustomDeltaProperty(Object, CustomDeltaProperty, TempBitWriter, NewState, OldState))
 				{
 					continue;
 				}
 			}
 
 			// Write property header and payload to the bunch
-			WritePropertyHeaderAndPayload( Object, It, NetFieldExportGroup, Bunch, TempBitWriter );
+			WritePropertyHeaderAndPayload(Object, It, NetFieldExportGroup, Bunch, TempBitWriter);
 
 			continue;
 		}
 
 		// Update Retirement records with this new state so we can handle packet drops.
 		// LastNext will be pointer to the last "Next" pointer in the list (so pointer to a pointer)
-		FPropertyRetirement** LastNext = UpdateAckedRetirements( Retire, OwningChannelConnection->OutAckPacketId, Object );
+		FPropertyRetirement** LastNext = UpdateAckedRetirements(Retire, Connection->OutAckPacketId, Object);
 
-		check( LastNext != NULL );
-		check( *LastNext == NULL );
+		check(LastNext != nullptr);
+		check(*LastNext == nullptr);
 
-		ValidateRetirementHistory( Retire, Object );
+		ValidateRetirementHistory(Retire, Object);
 
-		TSharedPtr<INetDeltaBaseState>& OldState = RecentCustomDeltaState.FindOrAdd( RetireIndex );
+		TSharedPtr<INetDeltaBaseState>& OldState = RecentCustomDeltaState.FindOrAdd(CustomDeltaProperty);
 
 		//-----------------------------------------
 		//	Do delta serialization on dynamic properties
 		//-----------------------------------------
-		const bool WroteSomething = SerializeCustomDeltaProperty( OwningChannelConnection, (void*)Object, It, Index, TempBitWriter, NewState, OldState );
+		const bool WroteSomething = SendCustomDeltaProperty(Object, CustomDeltaProperty, TempBitWriter, NewState, OldState);
 
 		if ( !WroteSomething )
 		{
@@ -1224,9 +1439,9 @@ void FObjectReplicator::ReplicateCustomDeltaProperties( FNetBitWriter & Bunch, F
 		OldState = NewState; 
 
 		// Write property header and payload to the bunch
-		WritePropertyHeaderAndPayload( Object, It, NetFieldExportGroup, Bunch, TempBitWriter );
+		WritePropertyHeaderAndPayload(Object, It, NetFieldExportGroup, Bunch, TempBitWriter);
 
-		NETWORK_PROFILER( GNetworkProfiler.TrackReplicateProperty( It, TempBitWriter.GetNumBits(), Connection ) );
+		NETWORK_PROFILER(GNetworkProfiler.TrackReplicateProperty(It, TempBitWriter.GetNumBits(), Connection));
 	}
 }
 
@@ -1348,20 +1563,24 @@ void FObjectReplicator::PostSendBunch( FPacketIdRange & PacketRange, uint8 bReli
 		return;
 	}
 
+	check(RepLayout);
+
 	// Don't update retirement records for reliable properties. This is ok to do only if we also pause replication on the channel until the acks have gone through.
 	bool SkipRetirementUpdate = OwningChannel->bPausedUntilReliableACK;
 
+	const FRepLayout& LocalRepLayout = *RepLayout;
+
 	if (!SkipRetirementUpdate)
 	{
 		// Don't call if reliable, since the bunch will be resent. We dont want this to end up in the changelist history
 		// But is that enough? How does it know to delta against this latest state?
 
-		RepLayout->PostReplicate(RepState->GetSendingRepState(), PacketRange, bReliable ? true : false);
+		LocalRepLayout.PostReplicate(RepState->GetSendingRepState(), PacketRange, bReliable ? true : false);
 	}
 
-	for ( int32 i = 0; i < LifetimeCustomDeltaProperties.Num(); i++ )
+	for (uint16 LifetimePropertyIndex : LocalRepLayout.GetLifetimeCustomDeltaProperties())
 	{
-		FPropertyRetirement & Retire = Retirement[LifetimeCustomDeltaProperties[i]];
+		FPropertyRetirement & Retire = Retirement[LifetimePropertyIndex];
 
 		FPropertyRetirement* Next = Retire.Next;
 		FPropertyRetirement* Prev = &Retire;
@@ -1440,11 +1659,13 @@ void FObjectReplicator::CountBytes(FArchive& Ar) const
 		}
 	);
 
+PRAGMA_DISABLE_DEPRECATION_WARNINGS
 	GRANULAR_NETWORK_MEMORY_TRACKING_TRACK("LifetimeCustomDeltaProperties", LifetimeCustomDeltaProperties.CountBytes(Ar));
 
 	GRANULAR_NETWORK_MEMORY_TRACKING_TRACK("LifetimeCustomDeltaPropertyConditions", LifetimeCustomDeltaPropertyConditions.CountBytes(Ar));
 
 	GRANULAR_NETWORK_MEMORY_TRACKING_TRACK("UnmappedCustomProperties", UnmappedCustomProperties.CountBytes(Ar));
+PRAGMA_ENABLE_DEPRECATION_WARNINGS
 
 	GRANULAR_NETWORK_MEMORY_TRACKING_TRACK("RepNotifies", RepNotifies.CountBytes(Ar));
 
@@ -1604,6 +1825,8 @@ void FObjectReplicator::StartBecomingDormant()
 
 void FObjectReplicator::CallRepNotifies(bool bSkipIfChannelHasQueuedBunches)
 {
+	// This logic is mostly a copy of FRepLayout::CallRepNotifies, and they should be merged.
+
 	CSV_SCOPED_TIMING_STAT_EXCLUSIVE(RepNotifies);
 	UObject* Object = GetObject();
 
@@ -1694,7 +1917,7 @@ void FObjectReplicator::CallRepNotifies(bool bSkipIfChannelHasQueuedBunches)
 	}
 }
 
-void FObjectReplicator::UpdateUnmappedObjects( bool & bOutHasMoreUnmapped )
+void FObjectReplicator::UpdateUnmappedObjects(bool & bOutHasMoreUnmapped)
 {
 	UObject* Object = GetObject();
 	
@@ -1718,56 +1941,48 @@ void FObjectReplicator::UpdateUnmappedObjects( bool & bOutHasMoreUnmapped )
 
 	bool bSomeObjectsWereMapped = false;
 
-	// Let the rep layout update any unmapped properties
-	RepLayout->UpdateUnmappedObjects(ReceivingRepState, Connection->PackageMap, Object, bSomeObjectsWereMapped, bOutHasMoreUnmapped);
-
-	// Update unmapped objects for custom properties (currently just fast tarray)
-	for (auto It = UnmappedCustomProperties.CreateIterator(); It; ++It)
-	{
-		const int32			Offset			= It.Key();
-		UStructProperty*	StructProperty	= It.Value();
-		UScriptStruct*		InnerStruct		= StructProperty->Struct;
-
-		check( InnerStruct->StructFlags & STRUCT_NetDeltaSerializeNative );
-
-		UScriptStruct::ICppStructOps* CppStructOps = InnerStruct->GetCppStructOps();
+	check(RepLayout);
 
-		check( CppStructOps );
+	const FRepLayout& LocalRepLayout = *RepLayout;
 
-		FNetDeltaSerializeInfo Parms;
+	// Let the rep layout update any unmapped properties
+	LocalRepLayout.UpdateUnmappedObjects(ReceivingRepState, Connection->PackageMap, Object, bSomeObjectsWereMapped, bOutHasMoreUnmapped);
 
-		FNetSerializeCB NetSerializeCB( Connection->Driver );
+	FNetSerializeCB NetSerializeCB(Connection->Driver);
 
-		Parms.DebugName			= StructProperty->GetName();
-		Parms.Struct			= InnerStruct;
-		Parms.Map				= Connection->PackageMap;
-		Parms.NetSerializeCB	= &NetSerializeCB;
+	FNetDeltaSerializeInfo Parms;
+	Parms.Object = Object;
+	Parms.Connection = Connection;
+	Parms.Map = Connection->PackageMap;
+	Parms.NetSerializeCB = &NetSerializeCB;
 
-		Parms.bUpdateUnmappedObjects	= true;
-		Parms.bCalledPreNetReceive		= bSomeObjectsWereMapped;	// RepLayout used this to flag whether PreNetReceive was called
-		Parms.bIsWritingOnClient		= false;
-		Parms.Object					= Object;
+	Parms.bUpdateUnmappedObjects = true;
+	Parms.bCalledPreNetReceive = bSomeObjectsWereMapped;	// RepLayout used this to flag whether PreNetReceive was called
+	
 
-		// Call the custom delta serialize function to handle it
-		CppStructOps->NetDeltaSerialize( Parms, (uint8*)Object + Offset );
+	TArray<TPair<int32, UStructProperty*>> CompletelyMappedProperties;
+	TArray<TPair<int32, UStructProperty*>> UpdatedProperties;
+	FNetSerializeCB::UpdateUnmappedObjectsForCustomDeltaProperties(LocalRepLayout, Parms, CompletelyMappedProperties, UpdatedProperties);
 
-		// Merge in results
-		bSomeObjectsWereMapped	|= Parms.bOutSomeObjectsWereMapped;
-		bOutHasMoreUnmapped		|= Parms.bOutHasMoreUnmapped;
+	bSomeObjectsWereMapped |= Parms.bOutSomeObjectsWereMapped;
+	bOutHasMoreUnmapped |= Parms.bOutHasMoreUnmapped;
 
-		if ( Parms.bOutSomeObjectsWereMapped )
-		{
-			// If we mapped a property, call the rep notify
-			TArray<uint8> MetaData;
-			QueuePropertyRepNotify( Object, StructProperty, 0, MetaData );
-		}
+	// This should go away when UnmappedCustomProperties goes away, and when RepNotifies
+	// are merged with RepState RepNotifies.
+PRAGMA_DISABLE_DEPRECATION_WARNINGS
+	for (const TPair<int32, UStructProperty*>& KVP : UpdatedProperties)
+	{
+		TArray<uint8> MetaData;
+		QueuePropertyRepNotify(Object, KVP.Value, 0, MetaData);
+	}
 
-		// If this property no longer has unmapped objects, we can stop checking it
-		if ( !Parms.bOutHasMoreUnmapped )
-		{
-			It.RemoveCurrent();
-		}
+	// This is just for the sake of keeping UnmappedCustomProperties up to date.
+	// Remove this when that property goes away.
+	for (const TPair<int32, UStructProperty*>& KVP : CompletelyMappedProperties)
+	{
+		UnmappedCustomProperties.Remove(KVP.Key);
 	}
+PRAGMA_ENABLE_DEPRECATION_WARNINGS
 
 	if (bSomeObjectsWereMapped)
 	{
diff --git a/Engine/Source/Runtime/Engine/Private/NetSerialization.cpp b/Engine/Source/Runtime/Engine/Private/NetSerialization.cpp
new file mode 100644
--- /dev/null
+++ b/Engine/Source/Runtime/Engine/Private/NetSerialization.cpp
@@ -0,0 +1,7 @@
+// Copyright 1998-2019 Epic Games, Inc. All Rights Reserved.
+
+#include "Engine/NetSerialization.h"
+
+DEFINE_STAT(STAT_NetSerializeFastArray);
+DEFINE_STAT(STAT_NetSerializeFastArray_BuildMap);
+DEFINE_STAT(STAT_NetSerializeFastArray_DeltaStruct);
\ No newline at end of file
diff --git a/Engine/Source/Runtime/Engine/Private/NetworkProfiler.cpp b/Engine/Source/Runtime/Engine/Private/NetworkProfiler.cpp
--- a/Engine/Source/Runtime/Engine/Private/NetworkProfiler.cpp
+++ b/Engine/Source/Runtime/Engine/Private/NetworkProfiler.cpp
@@ -100,6 +100,7 @@ FNetworkProfiler::FNetworkProfiler()
 ,	bHasNoticeableNetworkTrafficOccured(false)
 ,	bIsTrackingEnabled(false)
 ,	LastAddress( nullptr )
+,	IgnorePropertyCount(0)
 {
 }
 
@@ -501,7 +502,7 @@ void FNetworkProfiler::TrackReplicateActor( const AActor* Actor, FReplicationFla
  */
 void FNetworkProfiler::TrackReplicateProperty( const UProperty* Property, uint16 NumBits, UNetConnection* Connection )
 {
-	if( bIsTrackingEnabled )
+	if(bIsTrackingEnabled && !!!IgnorePropertyCount)
 	{
 		SCOPE_LOCK_REF(CriticalSection);
 
@@ -708,7 +709,7 @@ void FNetworkProfiler::TrackEndContentBlock( UObject* Object, uint16 NumBits, UN
 
 void FNetworkProfiler::TrackWritePropertyHandle( uint16 NumBits, UNetConnection* Connection )
 {
-	if ( bIsTrackingEnabled )
+	if (bIsTrackingEnabled && !!!IgnorePropertyCount)
 	{
 		SCOPE_LOCK_REF(CriticalSection);
 
diff --git a/Engine/Source/Runtime/Engine/Private/RepLayout.cpp b/Engine/Source/Runtime/Engine/Private/RepLayout.cpp
--- a/Engine/Source/Runtime/Engine/Private/RepLayout.cpp
+++ b/Engine/Source/Runtime/Engine/Private/RepLayout.cpp
@@ -20,10 +20,12 @@
 #include "Algo/Sort.h"
 #include "Net/NetworkGranularMemoryLogging.h"
 #include "Serialization/ArchiveCountMem.h"
+#include "Containers/SortedMap.h"
 
 DECLARE_CYCLE_STAT(TEXT("RepLayout AddPropertyCmd"), STAT_RepLayout_AddPropertyCmd, STATGROUP_Game);
 DECLARE_CYCLE_STAT(TEXT("RepLayout InitFromObjectClass"), STAT_RepLayout_InitFromObjectClass, STATGROUP_Game);
 DECLARE_CYCLE_STAT(TEXT("RepLayout BuildShadowOffsets"), STAT_RepLayout_BuildShadowOffsets, STATGROUP_Game);
+DECLARE_CYCLE_STAT(TEXT("RepLayout DeltaSerializeFastArray"), STAT_RepLayout_DeltaSerializeFastArray, STATGROUP_Game);
 
 // LogRepProperties is very spammy, and the logs are in a very hot code path,
 // so prevent anything less than a warning from even being compiled in on
@@ -119,6 +121,78 @@ FConsoleVariableSinkHandle CreateMaxArrayMemoryCVarAndRegisterSink()
 FConsoleVariableSinkHandle MaxRepArraySizeHandle = CreateMaxArraySizeCVarAndRegisterSink();
 FConsoleVariableSinkHandle MaxRepArrayMemorySink = CreateMaxArrayMemoryCVarAndRegisterSink();
 
+struct FLifetimeCustomDeltaProperty
+{
+	/** If this is a Fast Array Serializer property, this will be the command index for the Fast Array Item array. */
+	int32 FastArrayItemsCommand = INDEX_NONE;
+
+	/**
+	 * If this is a Fast Array Serializer property, this will be the instance number in the class.
+	 * This is used to lookup Changelists.
+	 */
+	int32 FastArrayNumber = INDEX_NONE;
+};
+
+/**
+ * Acceleration tracking which properties are custom delta.
+ * This will ultimately replace the need for FObjectReplicator::LifetimeCustomDeltaProperties.
+ */
+struct FLifetimeCustomDeltaState
+{
+	TSortedMap<uint16, FLifetimeCustomDeltaProperty> LifetimeCustomDeltaProperties;
+
+	int32 NumFastArrayItems = 0;
+
+	/**
+	 * This is less than ideal.
+	 * For now, we want a fast way to let users query the list of lifetime custom delta property indices.
+	 * This is mostly to support FObjectReplicator since we're going to remove its (per object per connection!!) cache.
+	 * However, we wouldn't even need this if Custom Delta Properties were more completely handled by RepLayout.
+	 */
+	TArray<uint16> LifetimeCustomDeltaPropertiesParentCache;
+
+	void CountBytes(FArchive& Ar) const
+	{
+		LifetimeCustomDeltaProperties.CountBytes(Ar);
+		LifetimeCustomDeltaPropertiesParentCache.CountBytes(Ar);
+	}
+};
+
+//~ Some of this complexity could go away if we introduced a new Compare step to Custom Delta Serializers
+//~ instead of just relying on the standard serialization stuff. That would be a bigger backwards compatibility
+//~ risk, however.
+
+struct FCustomDeltaSinglePropertyHistory
+{
+	TMap<int32, TArray<uint16>> ChangelistByID;
+};
+
+struct FCustomDeltaChangedHistory
+{
+	// We could just use a Sparse array here, but in this case it's more convenient
+	// to work with a bitfield and array directly.
+
+	TArray<FCustomDeltaSinglePropertyHistory> FastArrayChangelists;
+	TBitArray<> bWasUpdated;
+};
+
+struct FCustomDeltaChangelistState
+{
+	/** The maximum number of individual changelists allowed.*/
+	static const int32 MAX_CHANGE_HISTORY = FRepChangelistState::MAX_CHANGE_HISTORY;
+
+	/** Circular buffer of changelists. */
+	FCustomDeltaChangedHistory ChangeHistory[MAX_CHANGE_HISTORY];
+
+	/** Index in the buffer where changelist history starts (i.e., the Oldest changelist). */
+	int32 HistoryStart = 0;
+
+	/** Index in the buffer where changelist history ends (i.e., the Newest changelist). */
+	int32 HistoryEnd = 0;
+
+	uint32 CompareIndex = 0;
+};
+
 #define ENABLE_PROPERTY_CHECKSUMS
 
 //#define SANITY_CHECK_MERGES
@@ -459,7 +533,12 @@ void FRepStateStaticBuffer::CountBytes(FArchive& Ar) const
 	);
 }
 
-FRepChangelistState::FRepChangelistState(const TSharedRef<const FRepLayout>& InRepLayout, const uint8* Source) :
+FRepChangelistState::FRepChangelistState(
+	const TSharedRef<const FRepLayout>& InRepLayout,
+	const uint8* Source,
+	FCustomDeltaChangelistState* InDeltaChangelistState):
+
+	CustomDeltaChangelistState(InDeltaChangelistState),
 	HistoryStart(0),
 	HistoryEnd(0),
 	CompareIndex(0),
@@ -472,9 +551,13 @@ void FRepChangelistState::CountBytes(FArchive& Ar) const
 	SharedSerialization.CountBytes(Ar);
 }
 
-FReplicationChangelistMgr::FReplicationChangelistMgr(const TSharedRef<const FRepLayout>& InRepLayout, const uint8* Source):
+FReplicationChangelistMgr::FReplicationChangelistMgr(
+	const TSharedRef<const FRepLayout>& InRepLayout,
+	const uint8* Source,
+	FCustomDeltaChangelistState* DeltaChangelistState):
+
 	LastReplicationFrame(0),
-	RepChangelistState(InRepLayout, Source)
+	RepChangelistState(InRepLayout, Source, DeltaChangelistState)
 {
 }
 
@@ -483,6 +566,19 @@ FReceivingRepState::FReceivingRepState(FRepStateStaticBuffer&& InStaticBuffer) :
 {
 }
 
+FRepLayout::FRepLayout() :
+	LayoutState(ERepLayoutState::Uninitialized),
+	RoleIndex(INDEX_NONE),
+	RemoteRoleIndex(-1),
+	Owner(nullptr)
+{}
+
+FRepLayout::~FRepLayout()
+{
+	// This should never happen.
+	check(ERepLayoutState::Uninitialized != LayoutState);
+}
+
 void FRepLayout::UpdateChangelistMgr(
 	FSendingRepState* RESTRICT RepState,
 	FReplicationChangelistMgr& InChangelistMgr,
@@ -513,79 +609,120 @@ void FReplicationChangelistMgr::CountBytes(FArchive& Ar) const
 	RepChangelistState.CountBytes(Ar);
 }
 
-uint16 FRepLayout::CompareProperties_r(
-	FSendingRepState* RESTRICT RepState,
-	const int32 CmdStart,
-	const int32 CmdEnd,
+struct FComparePropertiesSharedParams
+{
+	const bool bIsInitial;
+	const bool bForceFail;
+	const int16 RoleIndex;
+	const int16 RemoteRoleIndex;
+	const TArray<FRepParentCmd>& Parents;
+	const TArray<FRepLayoutCmd>& Cmds;
+};
+
+static uint16 CompareProperties_r(
+	const FComparePropertiesSharedParams& SharedParams,
+	const uint16 CmdStart,
+	const uint16 CmdEnd,
 	FRepShadowDataBuffer ShadowData,
 	const FConstRepObjectDataBuffer Data,
 	TArray<uint16>& Changed,
-	uint16 Handle,
-	const bool bIsInitial,
-	const bool bForceFail) const
+	uint16 Handle);
+
+static void CompareProperties_Array_r(
+	const FComparePropertiesSharedParams& SharedParams,
+	FRepShadowDataBuffer ShadowData,
+	const FConstRepObjectDataBuffer Data,
+	TArray<uint16>& Changed,
+	const uint16 CmdIndex,
+	const uint16 Handle);
+
+static void CompareParentProperties(
+	const FComparePropertiesSharedParams& SharedParams,
+	FSendingRepState* RESTRICT RepState,
+	FRepChangelistState* RESTRICT RepChangelistState,
+	const FConstRepObjectDataBuffer Data,
+	TArray<uint16>& Changed)
 {
+	FRepShadowDataBuffer ShadowData = RepChangelistState->StaticBuffer.GetData();
 	check(ShadowData);
 
 	FRepChangedPropertyTracker* RepChangedPropertyTracker = (RepState ? RepState->RepChangedPropertyTracker.Get() : nullptr);
-	for (int32 CmdIndex = CmdStart; CmdIndex < CmdEnd; CmdIndex++)
-	{
-		const FRepLayoutCmd& Cmd = Cmds[CmdIndex];
-		const FRepParentCmd& ParentCmd = Parents[Cmd.ParentIndex];
-
-		check(Cmd.Type != ERepLayoutCmdType::Return);
 
-		Handle++;
+	for (uint16 ParentIndex = 0; ParentIndex < SharedParams.Parents.Num(); ++ParentIndex)
+	{
+		const FRepParentCmd& Parent = SharedParams.Parents[ParentIndex];
 
-		const bool bIsLifetime = ((ParentCmd.Flags & ERepParentFlags::IsLifetime) != ERepParentFlags::None);
+		const bool bIsLifetime = EnumHasAnyFlags(Parent.Flags, ERepParentFlags::IsLifetime);
 
 		// Active state of a property applies to *all* connections.
 		// If the property is inactive, we can skip comparing it because we know it won't be sent.
 		// Further, this will keep the last active state of the property in the shadow buffer,
 		// meaning the next time the property becomes active it will be sent to all connections.
-		const bool bActive = (RepChangedPropertyTracker == nullptr || RepChangedPropertyTracker->Parents[Cmd.ParentIndex].Active);
-
-		const bool bShouldSkip = !bIsLifetime || !bActive || (ParentCmd.Condition == COND_InitialOnly && !bIsInitial);
-
-		if (Cmd.Type == ERepLayoutCmdType::DynamicArray)
-		{
-			if (bShouldSkip)
-			{
-				CmdIndex = Cmd.EndCmd - 1;		// The -1 to handle the ++ in the for loop
-				continue;
-			}
+		const bool bIsActive = !RepChangedPropertyTracker || RepChangedPropertyTracker->Parents[ParentIndex].Active;
 
-			// Once we hit an array, start using a stack based approach
-			CompareProperties_Array_r(RepState, ShadowData + Cmd, Data + Cmd, Changed, CmdIndex, Handle, bIsInitial, bForceFail);
-			CmdIndex = Cmd.EndCmd - 1;		// The -1 to handle the ++ in the for loop
-			continue;
-		}
+		const bool bShouldSkip = !bIsLifetime || !bIsActive || (Parent.Condition == COND_InitialOnly && !SharedParams.bIsInitial);
 
 		if (bShouldSkip)
 		{
 			continue;
 		}
 
+		const FRepLayoutCmd& Cmd = SharedParams.Cmds[Parent.CmdStart];
+		const uint16 Handle = Cmd.RelativeHandle;
+
 		// RepState may be null in the case where a deprecated version of this function is called.
 		// In that case, just allow this to fail and perform the old logic.
-		if (RepState && Cmd.ParentIndex == RoleIndex)
+		if (UNLIKELY(RepState && ParentIndex == SharedParams.RoleIndex))
 		{
-			const ENetRole ObjectRole = *(const ENetRole*)(Data + Cmd).Data;
-			if (bForceFail || RepState->SavedRole != ObjectRole)
+			const ENetRole ObjectRole = *(const ENetRole*)(Data + Parent).Data;
+			if (SharedParams.bForceFail || RepState->SavedRole != ObjectRole)
 			{
 				RepState->SavedRole = ObjectRole;
 				Changed.Add(Handle);
 			}
 		}
-		else if (RepState && Cmd.ParentIndex == RemoteRoleIndex)
+		else if (UNLIKELY(RepState && ParentIndex == SharedParams.RemoteRoleIndex))
 		{
-			const ENetRole ObjectRemoteRole = *(const ENetRole*)(Data + Cmd).Data;
-			if (bForceFail || RepState->SavedRemoteRole != ObjectRemoteRole)
+			const ENetRole ObjectRemoteRole = *(const ENetRole*)(Data + Parent).Data;
+			if (SharedParams.bForceFail || RepState->SavedRemoteRole != ObjectRemoteRole)
 			{
 				RepState->SavedRemoteRole = ObjectRemoteRole;
 				Changed.Add(Handle);
 			}
 		}
-		else if (bForceFail || !PropertiesAreIdentical(Cmd, (ShadowData + Cmd).Data, (Data + Cmd).Data))
+		else
+		{
+			// Note, Handle - 1 to account for CompareProperties_r incrementing handles.
+			CompareProperties_r(SharedParams, Parent.CmdStart, Parent.CmdEnd, ShadowData, Data, Changed, Handle - 1);
+		}
+	}
+}
+
+static uint16 CompareProperties_r(
+	const FComparePropertiesSharedParams& SharedParams,
+	const uint16 CmdStart,
+	const uint16 CmdEnd,
+	FRepShadowDataBuffer ShadowData,
+	const FConstRepObjectDataBuffer Data,
+	TArray<uint16>& Changed,
+	uint16 Handle)
+{
+	for (int32 CmdIndex = CmdStart; CmdIndex < CmdEnd; ++CmdIndex)
+	{
+		const FRepLayoutCmd& Cmd = SharedParams.Cmds[CmdIndex];
+
+		check(Cmd.Type != ERepLayoutCmdType::Return);
+
+		++Handle;
+
+		if (Cmd.Type == ERepLayoutCmdType::DynamicArray)
+		{
+			// Once we hit an array, start using a stack based approach
+			CompareProperties_Array_r(SharedParams, ShadowData + Cmd, Data + Cmd, Changed, CmdIndex, Handle);
+			CmdIndex = Cmd.EndCmd - 1;		// The -1 to handle the ++ in the for loop
+			continue;
+		}
+		else if (SharedParams.bForceFail || !PropertiesAreIdentical(Cmd, (ShadowData + Cmd).Data, (Data + Cmd).Data))
 		{
 			StoreProperty(Cmd, (ShadowData + Cmd).Data, (Data + Cmd.Offset).Data);
 			Changed.Add(Handle);
@@ -595,17 +732,15 @@ uint16 FRepLayout::CompareProperties_r(
 	return Handle;
 }
 
-void FRepLayout::CompareProperties_Array_r(
-	FSendingRepState* RESTRICT RepState,
+static void CompareProperties_Array_r(
+	const FComparePropertiesSharedParams& SharedParams,
 	FRepShadowDataBuffer ShadowData,
 	const FConstRepObjectDataBuffer Data,
 	TArray<uint16>& Changed,
 	const uint16 CmdIndex,
-	const uint16 Handle,
-	const bool bIsInitial,
-	const bool bForceFail) const
+	const uint16 Handle)
 {
-	const FRepLayoutCmd& Cmd = Cmds[CmdIndex];
+	const FRepLayoutCmd& Cmd = SharedParams.Cmds[CmdIndex];
 
 	FScriptArray* ShadowArray = (FScriptArray*)ShadowData.Data;
 	FScriptArray* Array = (FScriptArray*)Data.Data;
@@ -624,12 +759,17 @@ void FRepLayout::CompareProperties_Array_r(
 	const FConstRepObjectDataBuffer ArrayData(Array->GetData());
 	FRepShadowDataBuffer ShadowArrayData(ShadowArray->GetData());
 
-	for (int32 i = 0; i < ArrayNum; i++)
 	{
-		const int32 ArrayElementOffset = i * Cmd.ElementSize;
-		const bool bNewForceFail = bForceFail || i >= ShadowArrayNum;
+		const bool bOldForceFail = SharedParams.bForceFail;
+		bool& bForceFail = const_cast<bool&>(SharedParams.bForceFail);
+		TGuardValue<bool> ForceFailGuard(bForceFail, bForceFail);
 
-		LocalHandle = CompareProperties_r(RepState, CmdIndex + 1, Cmd.EndCmd - 1, ShadowArrayData + ArrayElementOffset, ArrayData + ArrayElementOffset, ChangedLocal, LocalHandle, bIsInitial, bNewForceFail);
+		for (int32 i = 0; i < ArrayNum; i++)
+		{
+			const int32 ArrayElementOffset = i * Cmd.ElementSize;
+			bForceFail = bOldForceFail || i >= ShadowArrayNum;
+			LocalHandle = CompareProperties_r(SharedParams, CmdIndex + 1, Cmd.EndCmd - 1, ShadowArrayData + ArrayElementOffset, ArrayData + ArrayElementOffset, ChangedLocal, LocalHandle);
+		}
 	}
 
 	if (ChangedLocal.Num())
@@ -642,7 +782,7 @@ void FRepLayout::CompareProperties_Array_r(
 	else if (ArrayNum != ShadowArrayNum)
 	{
 		// If nothing below us changed, we either shrunk, or we grew and our inner was an array that didn't have any elements
-		check(ArrayNum < ShadowArrayNum || Cmds[CmdIndex + 1].Type == ERepLayoutCmdType::DynamicArray);
+		check(ArrayNum < ShadowArrayNum || SharedParams.Cmds[CmdIndex + 1].Type == ERepLayoutCmdType::DynamicArray);
 
 		// Array got smaller, send the array handle to force array size change
 		Changed.Add(Handle);
@@ -674,7 +814,16 @@ bool FRepLayout::CompareProperties(
 	TArray<uint16>& Changed = NewHistoryItem.Changed;
 	Changed.Empty();
 
-	CompareProperties_r(RepState, 0, Cmds.Num() - 1, RepChangelistState->StaticBuffer.GetData(), Data, Changed, 0, RepFlags.bNetInitial, false);
+	FComparePropertiesSharedParams SharedParams{
+		/*bIsInitial=*/ !!RepFlags.bNetInitial,
+		/*bForceFail=*/ false,
+		RoleIndex,
+		RemoteRoleIndex,
+		Parents,
+		Cmds
+	};
+
+	CompareParentProperties(SharedParams, RepState, RepChangelistState, Data, Changed);
 
 	if (Changed.Num() == 0)
 	{
@@ -1556,14 +1705,15 @@ const FRepSerializedPropertyInfo* FRepSerializationSharedInfo::WriteSharedProper
 
 void FRepLayout::SendProperties_r(
 	FSendingRepState* RESTRICT RepState,
-	FRepChangedPropertyTracker* ChangedTracker,
 	FNetBitWriter& Writer,
 	const bool bDoChecksum,
 	FRepHandleIterator& HandleIterator,
 	const FConstRepObjectDataBuffer SourceData,
 	const int32 ArrayDepth,
-	const FRepSerializationSharedInfo& SharedInfo) const
+	const FRepSerializationSharedInfo* const RESTRICT SharedInfo) const
 {
+	const bool bDoSharedSerialization = SharedInfo && !!GNetSharedSerializedData;
+
 	while (HandleIterator.NextHandle())
 	{
 		const FRepLayoutCmd& Cmd = Cmds[HandleIterator.CmdIndex];
@@ -1600,7 +1750,7 @@ void FRepLayout::SendProperties_r(
 			check(ArrayHandleIterator.ArrayElementSize> 0);
 			check(ArrayHandleIterator.NumHandlesPerElement> 0);
 
-			SendProperties_r(RepState, ChangedTracker, Writer, bDoChecksum, ArrayHandleIterator, ArrayData, ArrayDepth + 1, SharedInfo);
+			SendProperties_r(RepState, Writer, bDoChecksum, ArrayHandleIterator, ArrayData, ArrayDepth + 1, SharedInfo);
 
 			check(HandleIterator.ChangelistIterator.ChangedIndex - OldChangedIndex == ArrayChangedCount);				// Make sure we read correct amount
 			check(HandleIterator.ChangelistIterator.Changed[HandleIterator.ChangelistIterator.ChangedIndex] == 0);	// Make sure we are at the end
@@ -1610,7 +1760,7 @@ void FRepLayout::SendProperties_r(
 			WritePropertyHandle(Writer, 0, bDoChecksum);		// Signify end of dynamic array
 			continue;
 		}
-		else
+		else if (RepState)
 		{
 			if (Cmd.ParentIndex == RoleIndex)
 			{
@@ -1631,11 +1781,11 @@ void FRepLayout::SendProperties_r(
 
 		const FRepSerializedPropertyInfo* SharedPropInfo = nullptr;
 
-		if ((GNetSharedSerializedData != 0) && ((Cmd.Flags & ERepLayoutFlags::IsSharedSerialization) != ERepLayoutFlags::None))
+		if (bDoSharedSerialization && EnumHasAnyFlags(Cmd.Flags, ERepLayoutFlags::IsSharedSerialization))
 		{
 			FGuid PropertyGuid(HandleIterator.CmdIndex, HandleIterator.ArrayIndex, ArrayDepth, (int32)((PTRINT)Data.Data & 0xFFFFFFFF));
 
-			SharedPropInfo = SharedInfo.SharedPropertyInfo.FindByPredicate([&](const FRepSerializedPropertyInfo& Info) 
+			SharedPropInfo = SharedInfo->SharedPropertyInfo.FindByPredicate([&](const FRepSerializedPropertyInfo& Info) 
 			{ 
 				return (Info.Guid == PropertyGuid); 
 			});
@@ -1666,7 +1816,7 @@ void FRepLayout::SendProperties_r(
 				BitWriterMark.Copy(Writer, StandardBuffer);
 				BitWriterMark.Pop(Writer);
 
-				Writer.SerializeBitsWithOffset(SharedInfo.SerializedProperties->GetData(), SharedPropInfo->BitOffset, SharedPropInfo->BitLength);
+				Writer.SerializeBitsWithOffset(SharedInfo->SerializedProperties->GetData(), SharedPropInfo->BitOffset, SharedPropInfo->BitLength);
 				
 				TArray<uint8> SharedBuffer;
 				BitWriterMark.Copy(Writer, SharedBuffer);
@@ -1679,7 +1829,7 @@ void FRepLayout::SendProperties_r(
 			else
 #endif
 			{
-				Writer.SerializeBitsWithOffset(SharedInfo.SerializedProperties->GetData(), SharedPropInfo->BitOffset, SharedPropInfo->BitLength);
+				Writer.SerializeBitsWithOffset(SharedInfo->SerializedProperties->GetData(), SharedPropInfo->BitOffset, SharedPropInfo->BitLength);
 			}
 
 			NETWORK_PROFILER(GNetworkProfiler.TrackReplicateProperty(ParentCmd.Property, SharedPropInfo->PropBitLength, nullptr));
@@ -1751,7 +1901,7 @@ void FRepLayout::SendProperties(
 	FChangelistIterator ChangelistIterator(Changed, 0);
 	FRepHandleIterator HandleIterator(ChangelistIterator, Cmds, BaseHandleToCmdIndex, 0, 1, 0, Cmds.Num() - 1);
 
-	SendProperties_r(RepState, ChangedTracker, Writer, bDoChecksum, HandleIterator, Data, 0, SharedInfo);
+	SendProperties_r(RepState, Writer, bDoChecksum, HandleIterator, Data, 0, &SharedInfo);
 
 	if (NumBits != Writer.GetNumBits())
 	{
@@ -2852,7 +3002,7 @@ FGuidReferences::~FGuidReferences()
 }
 
 void FRepLayout::GatherGuidReferences_r(
-	FGuidReferencesMap* GuidReferencesMap,
+	const FGuidReferencesMap* GuidReferencesMap,
 	TSet<FNetworkGUID>& OutReferencedGuids,
 	int32& OutTrackedGuidMemoryBytes) const
 {
@@ -2886,6 +3036,28 @@ void FRepLayout::GatherGuidReferences(
 	}
 }
 
+void FRepLayout::GatherGuidReferencesForCustomDeltaProperties(FNetDeltaSerializeInfo& Params) const
+{
+	if (LifetimeCustomPropertyState)
+	{
+		FRepObjectDataBuffer ObjectData(Params.Object);
+		for (const auto& KVP : LifetimeCustomPropertyState->LifetimeCustomDeltaProperties)
+		{
+			const FRepParentCmd& Parent = Parents[KVP.Key];
+
+			// Static cast is safe here, because this property wouldn't have been marked CustomDelta otherwise.
+			UStructProperty* StructProperty = static_cast<UStructProperty*>(Parent.Property);
+			UScriptStruct::ICppStructOps* CppStructOps = StructProperty->Struct->GetCppStructOps();
+
+			Params.Struct = StructProperty->Struct;
+			Params.PropertyRepIndex = KVP.Key;
+			Params.Data = ObjectData + Parent;
+
+			CppStructOps->NetDeltaSerialize(Params, Params.Data);
+		}
+	}
+}
+
 bool FRepLayout::MoveMappedObjectToUnmapped_r(FGuidReferencesMap* GuidReferencesMap, const FNetworkGUID& GUID) const
 {
 	bool bFoundGUID = false;
@@ -2918,8 +3090,46 @@ bool FRepLayout::MoveMappedObjectToUnmapped_r(FGuidReferencesMap* GuidReferences
 
 bool FRepLayout::MoveMappedObjectToUnmapped(FReceivingRepState* RESTRICT RepState, const FNetworkGUID& GUID) const
 {
-	return MoveMappedObjectToUnmapped_r(&RepState->GuidReferencesMap, GUID);
+	return ERepLayoutState::Normal == LayoutState && MoveMappedObjectToUnmapped_r(&RepState->GuidReferencesMap, GUID);
+}
+
+PRAGMA_DISABLE_DEPRECATION_WARNINGS
+bool FRepLayout::MoveMappedObjectToUnmappedForCustomDeltaProperties(FNetDeltaSerializeInfo& Params) const
+{
+	TMap<int32, UStructProperty*> Dummy;
+	return MoveMappedObjectToUnmappedForCustomDeltaProperties(Params, Dummy);
+}
+
+bool FRepLayout::MoveMappedObjectToUnmappedForCustomDeltaProperties(FNetDeltaSerializeInfo& Params, TMap<int32, UStructProperty*>& UnmappedCustomProperties) const
+{
+	bool bFound = false;
+
+	if (LifetimeCustomPropertyState)
+	{
+		FRepObjectDataBuffer ObjectData(Params.Object);
+		for (const auto& KVP : LifetimeCustomPropertyState->LifetimeCustomDeltaProperties)
+		{
+			const FRepParentCmd& Parent = Parents[KVP.Key];
+
+			// Static cast is safe here, because this property wouldn't have been marked CustomDelta otherwise.
+			UStructProperty* StructProperty = static_cast<UStructProperty*>(Parent.Property);
+			UScriptStruct::ICppStructOps* CppStructOps = StructProperty->Struct->GetCppStructOps();
+
+			Params.Struct = StructProperty->Struct;
+			Params.Data = ObjectData + Parent;
+			Params.PropertyRepIndex = KVP.Key;
+
+			if (CppStructOps->NetDeltaSerialize(Params, Params.Data))
+			{
+				UnmappedCustomProperties.Add(Parent.Offset, StructProperty);
+				bFound = true;
+			}
+		}
+	}
+
+	return bFound;
 }
+PRAGMA_ENABLE_DEPRECATION_WARNINGS
 
 void FRepLayout::UpdateUnmappedObjects_r(
 	FReceivingRepState* RESTRICT RepState, 
@@ -2952,15 +3162,26 @@ void FRepLayout::UpdateUnmappedObjects_r(
 		{
 			check(Cmd.Type == ERepLayoutCmdType::DynamicArray);
 			
-			FScriptArray* ShadowArray = (FScriptArray*)(ShadowData + Cmd).Data;
-			FScriptArray* Array = (FScriptArray*)(Data + AbsOffset).Data;
-			
-			FRepShadowDataBuffer ShadowArrayData(ShadowArray->GetData());
-			FRepObjectDataBuffer ArrayData(Array->GetData());
-			
-			const int32 NewMaxOffset = FMath::Min(ShadowArray->Num() * Cmd.ElementSize, Array->Num() * Cmd.ElementSize);
+			if (ShadowData)
+			{
+				FScriptArray* ShadowArray = (FScriptArray*)(ShadowData + Cmd).Data;
+				FScriptArray* Array = (FScriptArray*)(Data + AbsOffset).Data;
+
+				FRepShadowDataBuffer ShadowArrayData(ShadowArray->GetData());
+				FRepObjectDataBuffer ArrayData(Array->GetData());
+
+				const int32 NewMaxOffset = FMath::Min(ShadowArray->Num() * Cmd.ElementSize, Array->Num() * Cmd.ElementSize);
 
-			UpdateUnmappedObjects_r(RepState, GuidReferences.Array, OriginalObject, PackageMap, ShadowArrayData, ArrayData, NewMaxOffset, bOutSomeObjectsWereMapped, bOutHasMoreUnmapped);
+				UpdateUnmappedObjects_r(RepState, GuidReferences.Array, OriginalObject, PackageMap, ShadowArrayData, ArrayData, NewMaxOffset, bOutSomeObjectsWereMapped, bOutHasMoreUnmapped);
+			}
+			else
+			{
+				FScriptArray* Array = (FScriptArray*)(Data + AbsOffset).Data;
+				FRepObjectDataBuffer ArrayData(Array->GetData());
+				const int32 NewMaxOffset = Array->Num() * Cmd.ElementSize;
+
+				UpdateUnmappedObjects_r(RepState, GuidReferences.Array, OriginalObject, PackageMap, nullptr, ArrayData, NewMaxOffset, bOutSomeObjectsWereMapped, bOutHasMoreUnmapped);
+			}
 			continue;
 		}
 
@@ -3005,8 +3226,10 @@ void FRepLayout::UpdateUnmappedObjects_r(
 				bOutSomeObjectsWereMapped = true;
 			}
 
+			const bool bUpdateShadowState = (ShadowData && INDEX_NONE != Parent.RepNotifyNumParams);
+
 			// Copy current value over so we can check to see if it changed
-			if (INDEX_NONE != Parent.RepNotifyNumParams)
+			if (bUpdateShadowState)
 			{
 				StoreProperty(Cmd, ShadowData + Cmd, Data + AbsOffset);
 			}
@@ -3018,8 +3241,14 @@ void FRepLayout::UpdateUnmappedObjects_r(
 			Cmd.Property->NetSerializeItem(Reader, PackageMap, Data + AbsOffset);
 
 			// Check to see if this property changed
-			if (INDEX_NONE != Parent.RepNotifyNumParams)
+			if (bUpdateShadowState)
 			{
+				// I have a sneaking suspicion that this is broken.
+				// AbsOffset could be Cmd.Offset, but we also may be recursing into an Array, and that
+				// would mean it could be Cmd.Offset + (ArrayIndex * ElementOffset)
+				// That could cause us to trigger RepNotifies more often for Dynamic Array properties.
+				// That goes for the above too.
+
 				if (Parent.RepNotifyCondition == REPNOTIFY_Always || !PropertiesAreIdentical(Cmd, ShadowData + Cmd, Data + AbsOffset))
 				{
 					// If this properties needs an OnRep, queue that up to be handled later
@@ -3060,6 +3289,166 @@ void FRepLayout::UpdateUnmappedObjects(
 	}
 }
 
+PRAGMA_DISABLE_DEPRECATION_WARNINGS
+void FRepLayout::UpdateUnmappedObjectsForCustomDeltaProperties(FNetDeltaSerializeInfo& Params) const
+{
+	TArray<TPair<int32, UStructProperty*>> Dummy;
+	UpdateUnmappedObjectsForCustomDeltaProperties(Params, Dummy, Dummy);
+}
+PRAGMA_ENABLE_DEPRECATION_WARNINGS
+
+void FRepLayout::UpdateUnmappedObjectsForCustomDeltaProperties(FNetDeltaSerializeInfo& Params, TArray<TPair<int32, UStructProperty*>>& CompletelyUpdated, TArray<TPair<int32, UStructProperty*>>& Updated) const
+{
+	bool bOutSomeObjectsWereMapped = false;
+	bool bOutHasMoreUnmapped = false;
+
+	if (LifetimeCustomPropertyState)
+	{
+		FRepObjectDataBuffer ObjectData(Params.Object);
+		for (const auto& KVP : LifetimeCustomPropertyState->LifetimeCustomDeltaProperties)
+		{
+			const FRepParentCmd& Parent = Parents[KVP.Key];
+
+			// Static cast is safe here, because this property wouldn't have been marked CustomDelta otherwise.
+			UStructProperty* StructProperty = static_cast<UStructProperty*>(Parent.Property);
+			UScriptStruct::ICppStructOps* CppStructOps = StructProperty->Struct->GetCppStructOps();
+
+			Params.DebugName = Parent.CachedPropertyName.ToString();
+			Params.Struct = StructProperty->Struct;
+			Params.bOutSomeObjectsWereMapped = false;
+			Params.bOutHasMoreUnmapped = false;
+			Params.PropertyRepIndex = KVP.Key;
+			Params.Data = ObjectData + Parent;
+
+			// Call the custom delta serialize function to handle it
+			CppStructOps->NetDeltaSerialize(Params, Params.Data);
+
+			if (Params.bOutSomeObjectsWereMapped)
+			{
+				Updated.Emplace(Parent.Offset, StructProperty);
+			}
+
+			if (!Params.bOutHasMoreUnmapped)
+			{
+				CompletelyUpdated.Emplace(Parent.Offset, StructProperty);
+			}
+
+			bOutSomeObjectsWereMapped |= Params.bOutSomeObjectsWereMapped;
+			bOutHasMoreUnmapped |= Params.bOutHasMoreUnmapped;
+		}
+	}
+
+	Params.bOutSomeObjectsWereMapped = bOutSomeObjectsWereMapped;
+	Params.bOutHasMoreUnmapped = bOutHasMoreUnmapped;
+}
+
+bool FRepLayout::SendCustomDeltaProperty(
+	FNetDeltaSerializeInfo& Params,
+	UProperty* Property,
+	int32 StaticArrayIndex) const
+{
+	check(Property->ArrayDim > StaticArrayIndex);
+	return SendCustomDeltaProperty(Params, Property->RepIndex + StaticArrayIndex);
+}
+
+bool FRepLayout::SendCustomDeltaProperty(FNetDeltaSerializeInfo& Params, uint16 CustomDeltaRepIndex) const
+{
+	const FRepParentCmd& Parent = Parents[CustomDeltaRepIndex];
+	if (!ensure(EnumHasAnyFlags(Parent.Flags, ERepParentFlags::IsCustomDelta)))
+	{
+		return false;
+	}
+
+	UStructProperty* StructProperty = static_cast<UStructProperty*>(Parent.Property);
+	UScriptStruct::ICppStructOps * CppStructOps = StructProperty->Struct->GetCppStructOps();
+
+	check(CppStructOps); // else should not have STRUCT_NetSerializeNative
+
+	Params.DebugName = Parent.CachedPropertyName.ToString();
+	Params.Struct = StructProperty->Struct;
+	Params.PropertyRepIndex = CustomDeltaRepIndex;
+	Params.Data = FRepObjectDataBuffer(Params.Object) + Parent;
+
+	bool bSupportsFastArrayDelta = Params.bSupportsFastArrayDeltaStructSerialization;
+
+	if (Params.bSupportsFastArrayDeltaStructSerialization &&
+		EnumHasAnyFlags(Parent.Flags, ERepParentFlags::IsFastArray) &&
+		!!LifetimeCustomPropertyState->NumFastArrayItems)
+	{
+		const FLifetimeCustomDeltaProperty& CustomDeltaProperty = LifetimeCustomPropertyState->LifetimeCustomDeltaProperties.FindChecked(CustomDeltaRepIndex);
+		bSupportsFastArrayDelta = CustomDeltaProperty.FastArrayNumber != INDEX_NONE;
+	}
+
+	Params.Writer->WriteBit(bSupportsFastArrayDelta);
+
+	if (Parent.Property->ArrayDim != 1)
+	{
+		uint32 StaticArrayIndex = Parent.ArrayIndex;
+		Params.Writer->SerializeIntPacked(StaticArrayIndex);
+	}
+
+	TGuardValue<bool> SupportFastArrayDeltaGuard(Params.bSupportsFastArrayDeltaStructSerialization, bSupportsFastArrayDelta);
+	return CppStructOps->NetDeltaSerialize(Params, Params.Data);
+}
+
+PRAGMA_DISABLE_DEPRECATION_WARNINGS
+bool FRepLayout::ReceiveCustomDeltaProperty(FNetDeltaSerializeInfo& Params, UStructProperty* Property) const
+{
+	uint32 StaticArrayIndex = 0;
+	int32 Offset = 0;
+	return ReceiveCustomDeltaProperty(Params, Property, StaticArrayIndex, Offset);
+}
+PRAGMA_ENABLE_DEPRECATION_WARNINGS
+
+bool FRepLayout::ReceiveCustomDeltaProperty(FNetDeltaSerializeInfo& Params, UStructProperty* Property, uint32& StaticArrayIndex, int32& Offset) const
+{
+	if (Params.Connection->EngineNetworkProtocolVersion >= EEngineNetworkVersionHistory::HISTORY_FAST_ARRAY_DELTA_STRUCT)
+	{
+		Params.bSupportsFastArrayDeltaStructSerialization = !!Params.Reader->ReadBit();
+	}
+	else
+	{
+		Params.bSupportsFastArrayDeltaStructSerialization = false;
+	}
+
+	// Receive array index (static sized array, i.e. MemberVariable[4])
+	if (Property->ArrayDim != 1)
+	{
+		check(Property->ArrayDim >= 2);
+
+		Params.Reader->SerializeIntPacked(StaticArrayIndex);
+
+		if (StaticArrayIndex >= (uint32)Property->ArrayDim)
+		{
+			UE_LOG(LogRep, Error, TEXT("Element index too large %s in %s"), *Property->GetName(), *Params.Object->GetFullName());
+			return false;
+		}
+	}
+
+	const FRepParentCmd& Parent = Parents[Property->RepIndex + StaticArrayIndex];
+
+	// We should only be receiving custom delta properties (since RepLayout handles the rest)
+	if (!EnumHasAnyFlags(Parent.Flags, ERepParentFlags::IsCustomDelta))
+	{
+		UE_LOG(LogNet, Error, TEXT("Client received non custom delta property value %s in %s"), *Parent.CachedPropertyName.ToString(), *Params.Object->GetFullName());
+		return false;
+	}
+
+	Offset = Parent.Offset;
+
+	UScriptStruct* InnerStruct = Property->Struct;
+	UScriptStruct::ICppStructOps* CppStructOps = InnerStruct->GetCppStructOps();
+
+	check(CppStructOps);
+
+	Params.DebugName = Parent.CachedPropertyName.ToString();
+	Params.Struct = InnerStruct;
+	Params.PropertyRepIndex = Property->RepIndex + StaticArrayIndex;
+	Params.Data = FRepObjectDataBuffer(Params.Object) + Parent;
+
+	return CppStructOps->NetDeltaSerialize(Params, Params.Data);
+}
+
 void FRepLayout::CallRepNotifies(FReceivingRepState* RepState, UObject* Object) const
 {
 	if (RepState->RepNotifies.Num() == 0)
@@ -3385,6 +3774,12 @@ static bool DiffProperties_r(FDiffPropertiesSharedParams& Params, TDiffPropertie
 		const FRepLayoutCmd& Cmd = Params.Cmds[CmdIndex];
 		const FRepParentCmd& Parent = Params.Parents[Cmd.ParentIndex];
 
+		if (EnumHasAnyFlags(Parent.Flags, ERepParentFlags::IsCustomDelta))
+		{
+			CmdIndex = Parent.CmdEnd - 1;
+			continue;
+		}
+
 		check(ERepLayoutCmdType::Return != Cmd.Type);
 
 		if (ERepLayoutCmdType::DynamicArray == Cmd.Type)
@@ -3520,7 +3915,7 @@ struct TDiffStablePropertiesStackParams
 };
 
 template<ERepDataBufferType DestinationType, ERepDataBufferType SourceType>
-bool DiffStableProperties_r(FDiffStablePropertiesSharedParams& Params, TDiffStablePropertiesStackParams<DestinationType, SourceType>& StackParams)
+static bool DiffStableProperties_r(FDiffStablePropertiesSharedParams& Params, TDiffStablePropertiesStackParams<DestinationType, SourceType>& StackParams)
 {
 	// Note, it's never possible for the Source or Destination data to be null.
 	// At the top level, both will always be valid (pointing to the memory of each element in each respective buffer).
@@ -3535,6 +3930,12 @@ bool DiffStableProperties_r(FDiffStablePropertiesSharedParams& Params, TDiffStab
 		const FRepLayoutCmd& Cmd = Params.Cmds[CmdIndex];
 		const FRepParentCmd& Parent = Params.Parents[Cmd.ParentIndex];
 
+		if (EnumHasAnyFlags(Parent.Flags, ERepParentFlags::IsCustomDelta))
+		{
+			CmdIndex = Parent.CmdEnd - 1;
+			continue;
+		}
+
 		check(ERepLayoutCmdType::Return != Cmd.Type);
 
 		if (ERepLayoutCmdType::DynamicArray == Cmd.Type)
@@ -3867,7 +4268,7 @@ static int32 InitFromProperty_r(
 		const uint32 ArrayChecksum = AddArrayCmd(Cmds, ArrayProp, Offset + GetOffsetForProperty<BuildType>(*ArrayProp), RelativeHandle, ParentIndex, ParentChecksum, StaticArrayIndex, ServerConnection);
 
 		InitFromProperty_r<BuildType>(Cmds, ArrayProp->Inner, 0, 0, ParentIndex, ArrayChecksum, 0, ServerConnection);
-		
+
 		AddReturnCmd(Cmds);
 
 		Cmds[CmdStart].EndCmd = Cmds.Num();		// Patch in the offset to jump over our array inner elements
@@ -3881,14 +4282,10 @@ static int32 InitFromProperty_r(
 	{
 		UScriptStruct* Struct = StructProp->Struct;
 
-		if (Struct->StructFlags & STRUCT_NetDeltaSerializeNative)
+		if (EnumHasAnyFlags(Struct->StructFlags, STRUCT_NetSerializeNative))
 		{
-			// Custom delta serializers handles outside of FRepLayout
-			return RelativeHandle;
-		}
+			UE_CLOG(EnumHasAnyFlags(Struct->StructFlags, STRUCT_NetDeltaSerializeNative), LogRep, Warning, TEXT("RepLayout InitFromProperty_r: Struct marked both NetSerialize and NetDeltaSerialize: %s"), *StructProp->GetName());
 
-		if (Struct->StructFlags & STRUCT_NetSerializeNative)
-		{
 			RelativeHandle++;
 			AddPropertyCmd(Cmds, Property, Offset + GetOffsetForProperty<BuildType>(*Property), RelativeHandle, ParentIndex, ParentChecksum, StaticArrayIndex, ServerConnection);
 			return RelativeHandle;
@@ -3966,6 +4363,11 @@ static FORCEINLINE void SetupRepStructFlags(FRepParentCmd& Parent, const bool bS
 		if (!bSkipCustomDeltaCheck && EnumHasAnyFlags(Struct->StructFlags, STRUCT_NetDeltaSerializeNative))
 		{
 			Parent.Flags |= ERepParentFlags::IsCustomDelta;
+
+			if (Struct->IsChildOf(FFastArraySerializer::StaticStruct()))
+			{
+				Parent.Flags |= ERepParentFlags::IsFastArray;
+			}
 		}
 
 		if (EnumHasAnyFlags(Struct->StructFlags, STRUCT_NetSerializeNative))
@@ -4194,12 +4596,11 @@ void FRepLayout::InitFromClass(UClass* InObjectClass, const UNetConnection* Serv
 	SCOPE_CYCLE_UOBJECT(ObjectClass, InObjectClass);
 
 	const bool bIsObjectActor = InObjectClass->IsChildOf(AActor::StaticClass());
-	RoleIndex = -1;
-	RemoteRoleIndex = -1;
-	FirstNonCustomParent = -1;
+	RoleIndex = INDEX_NONE;
+	RemoteRoleIndex = INDEX_NONE;
 
 	int32 RelativeHandle = 0;
-	int32 LastOffset = -1;
+	int32 LastOffset = INDEX_NONE;
 
 	InObjectClass->SetUpRuntimeReplicationData();
 	Parents.Empty(InObjectClass->ClassReps.Num());
@@ -4216,11 +4617,13 @@ void FRepLayout::InitFromClass(UClass* InObjectClass, const UNetConnection* Serv
 		check(ParentHandle == i);
 		check(Parents[i].Property->RepIndex + Parents[i].ArrayIndex == i);
 
+		const int32 ParentOffset = Property->ElementSize * ArrayIdx;
+
 		Parents[ParentHandle].CmdStart = Cmds.Num();
-		RelativeHandle = InitFromProperty_r<ERepBuildType::Class>(Cmds, Property, Property->ElementSize * ArrayIdx, RelativeHandle, ParentHandle, 0, ArrayIdx, ServerConnection);
+		RelativeHandle = InitFromProperty_r<ERepBuildType::Class>(Cmds, Property, ParentOffset, RelativeHandle, ParentHandle, 0, ArrayIdx, ServerConnection);
 		Parents[ParentHandle].CmdEnd = Cmds.Num();
 		Parents[ParentHandle].Flags |= ERepParentFlags::IsConditional;
-		Parents[ParentHandle].Offset = GetOffsetForProperty<ERepBuildType::Class>(*Property);
+		Parents[ParentHandle].Offset = GetOffsetForProperty<ERepBuildType::Class>(*Property) + ParentOffset;
 
 		if (Parents[i].CmdEnd > Parents[i].CmdStart)
 		{
@@ -4236,12 +4639,6 @@ void FRepLayout::InitFromClass(UClass* InObjectClass, const UNetConnection* Serv
 			Parents[ParentHandle].Flags |= ERepParentFlags::IsConfig;
 		}
 
-		// Hijack the first non custom property for identifying this as a rep layout block
-		if (FirstNonCustomParent == -1 && Property->ArrayDim == 1 && ((Parents[ParentHandle].Flags & ERepParentFlags::IsCustomDelta) == ERepParentFlags::None))
-		{
-			FirstNonCustomParent = ParentHandle;
-		}
-
 		if (bIsObjectActor)
 		{
 		    // Find Role/RemoteRole property indexes so we can swap them on the client
@@ -4301,22 +4698,81 @@ void FRepLayout::InitFromClass(UClass* InObjectClass, const UNetConnection* Serv
 			Parents[ParentIndex].RepNotifyNumParams = RepNotifyFunc->NumParms;
 		}
 
-		if ((Parents[ParentIndex].Flags & ERepParentFlags::IsCustomDelta) != ERepParentFlags::None)
+		if (EnumHasAnyFlags(Parents[ParentIndex].Flags, ERepParentFlags::IsCustomDelta))
 		{
-			// We don't handle custom properties in the FRepLayout class
-			continue;
-		}
+			// We'll track all Custom Lifetime Properties here, and we'll handle Fast Array Serialization
+			// specially.
 
-		Parents[ParentIndex].Flags |= ERepParentFlags::IsLifetime;
+			// Note, there are in engine cases where we aren't using FFastArraySerialize, but are using
+			// Custom Delta (FGameplayDebuggerNetPack).
 
-		if (ParentIndex == RemoteRoleIndex)
-		{
-			// We handle remote role specially, since it can change between connections when downgraded
-			// So we force it on the conditional list
-			check(LifetimeProps[i].Condition == COND_None || LifetimeProps[i].Condition == COND_Never);
-			LifetimeProps[i].Condition = COND_Custom;
-			continue;
-		}
+			// Also note, we still don't mark these properties as Lifetime. This should help maintain behavior
+			// in Diff Properties / Compare Properties.
+
+			if (!LifetimeCustomPropertyState)
+			{
+				LifetimeCustomPropertyState.Reset(new FLifetimeCustomDeltaState());
+			}
+
+			const int32 CustomIndex = LifetimeCustomPropertyState->LifetimeCustomDeltaPropertiesParentCache.Add(ParentIndex);
+			FLifetimeCustomDeltaProperty& CustomDeltaProperty = LifetimeCustomPropertyState->LifetimeCustomDeltaProperties.Add(ParentIndex);
+
+			// If we're a FastArraySerializer, we'll look for our replicated item type.
+			// We do this by looking for an array property whose inner type is an FFastArraySerializerItem.
+			// Note, this isn't perfect. With the way the interface is set up now, there's no technically
+			// enforced requirements that the Array of items lives within the Fast Array Serializer, that the
+			// Array of items is marked up as a UProperty, that the Array of items is not marked RepSkip,
+			// or that there's not multiple arrays of FastArraySerializerItems.
+			//
+			// However, comments imply these, and typically they are true (certainly, any engine cases follow this).
+			// Further, these layouts are only needed for the new Delta Struct Serialization feature, so this won't break backwards compat.
+
+			if (EnumHasAnyFlags(Parents[ParentIndex].Flags, ERepParentFlags::IsFastArray))
+			{
+				int32 FastArrayItemArrayCmd = INDEX_NONE;
+
+				for (int32 CmdIndex = Parents[ParentIndex].CmdStart; CmdIndex < Parents[ParentIndex].CmdEnd; ++CmdIndex)
+				{
+					const FRepLayoutCmd& Cmd = Cmds[CmdIndex];
+					if (ERepLayoutCmdType::DynamicArray == Cmd.Type)
+					{
+						if (UStructProperty* MaybeFastArray = Cast<UStructProperty>(static_cast<UArrayProperty*>(Cmd.Property)->Inner))
+						{
+							if (MaybeFastArray->Struct->IsChildOf(FFastArraySerializerItem::StaticStruct()))
+							{
+								FastArrayItemArrayCmd = CmdIndex;
+								break;
+							}
+						}
+
+						CmdIndex = Cmd.EndCmd - 1;
+					}
+				}
+
+				if (FastArrayItemArrayCmd != INDEX_NONE)
+				{
+					++LifetimeCustomPropertyState->NumFastArrayItems;
+					CustomDeltaProperty.FastArrayItemsCommand = FastArrayItemArrayCmd;
+				}
+				else
+				{
+					UE_LOG(LogRep, Warning, TEXT("FRepLayout::InitFromClass: Unable to find Fast Array Item array in Fast Array Serializer: %s"), *Parents[ParentIndex].CachedPropertyName.ToString());
+				}
+			}
+
+			continue;
+		}
+
+		Parents[ParentIndex].Flags |= ERepParentFlags::IsLifetime;
+
+		if (ParentIndex == RemoteRoleIndex)
+		{
+			// We handle remote role specially, since it can change between connections when downgraded
+			// So we force it on the conditional list
+			check(LifetimeProps[i].Condition == COND_None || LifetimeProps[i].Condition == COND_Never);
+			LifetimeProps[i].Condition = COND_Custom;
+			continue;
+		}
 
 		if (LifetimeProps[i].Condition == COND_None)
 		{
@@ -4329,6 +4785,19 @@ void FRepLayout::InitFromClass(UClass* InObjectClass, const UNetConnection* Serv
 		BuildHandleToCmdIndexTable_r(0, Cmds.Num() - 1, BaseHandleToCmdIndex);
 	}
 
+	if (LifetimeCustomPropertyState && !!LifetimeCustomPropertyState->NumFastArrayItems)
+	{
+		int32 Counter = 0;
+		for (auto& KVP : LifetimeCustomPropertyState->LifetimeCustomDeltaProperties)
+		{
+			if (EnumHasAnyFlags(Parents[KVP.Key].Flags, ERepParentFlags::IsFastArray) &&
+				INDEX_NONE != KVP.Value.FastArrayItemsCommand)
+			{
+				KVP.Value.FastArrayNumber = Counter++;
+			}
+		}
+	}
+
 	BuildShadowOffsets<ERepBuildType::Class>(InObjectClass, Parents, Cmds, ShadowDataBufferSize, LayoutState);
 
 	Owner = InObjectClass;
@@ -4573,6 +5042,7 @@ void FRepLayout::BuildChangeList_r(
 	const int32 CmdEnd,
 	const FConstRepObjectDataBuffer Data,
 	const int32 HandleOffset,
+	const bool bForceArraySends,
 	TArray<uint16>& Changed) const
 {
 	for (int32 CmdIndex = CmdStart; CmdIndex < CmdEnd; CmdIndex++)
@@ -4599,7 +5069,7 @@ void FRepLayout::BuildChangeList_r(
 			for (int32 i = 0; i < Array->Num(); i++)
 			{
 				const int32 ArrayElementOffset = Cmd.ElementSize * i;
-				BuildChangeList_r(ArrayHandleToCmdIndex, ArrayCmdStart, ArrayCmdEnd, ArrayData + ArrayElementOffset, i * NumHandlesPerElement, ChangedLocal);
+				BuildChangeList_r(ArrayHandleToCmdIndex, ArrayCmdStart, ArrayCmdEnd, ArrayData + ArrayElementOffset, i * NumHandlesPerElement, bForceArraySends, ChangedLocal);
 			}
 
 			if (ChangedLocal.Num())
@@ -4609,6 +5079,15 @@ void FRepLayout::BuildChangeList_r(
 				Changed.Append(ChangedLocal);					// Append the change list under the array
 				Changed.Add(0);									// Null terminator
 			}
+			else if (bForceArraySends)
+			{
+				// Note the handle, that there were 0 changed elements, and the terminator.
+				// This will force anyone sending this changelist later to send the array size
+				// (which is most likely 0 in this case).
+				Changed.Add(Cmd.RelativeHandle + HandleOffset);
+				Changed.Add(0);
+				Changed.Add(0);
+			}
 
 			CmdIndex = Cmd.EndCmd - 1;		// The -1 to handle the ++ in the for loop
 			continue;
@@ -4657,6 +5136,12 @@ void FRepLayout::BuildSharedSerialization_r(
 
 		const FConstRepObjectDataBuffer Data = SourceData + ArrayOffset + Cmd;
 
+		// Custom Deltas are not supported for shared serialization at this time.
+		if (EnumHasAnyFlags(ParentCmd.Flags, ERepParentFlags::IsCustomDelta))
+		{
+			continue;
+		}
+
 		if (Cmd.Type == ERepLayoutCmdType::DynamicArray)
 		{
 			const FScriptArray* Array = (FScriptArray *)Data.Data;
@@ -4799,7 +5284,7 @@ void FRepLayout::SendPropertiesForRPC(
 			{
 				if (!Parents[i].Property->Identical_InContainer(Data.Data, NULL, Parents[i].ArrayIndex))
 				{
-					BuildChangeList_r(BaseHandleToCmdIndex, Parents[i].CmdStart, Parents[i].CmdEnd, Data, 0, Changed);
+					BuildChangeList_r(BaseHandleToCmdIndex, Parents[i].CmdStart, Parents[i].CmdEnd, Data, 0, false, Changed);
 				}
 			}
 
@@ -5049,7 +5534,16 @@ FRepStateStaticBuffer FRepLayout::CreateShadowBuffer(const FConstRepObjectDataBu
 
 TSharedPtr<FReplicationChangelistMgr> FRepLayout::CreateReplicationChangelistMgr(const UObject* InObject) const
 {
-	return MakeShareable(new FReplicationChangelistMgr(AsShared(), (const uint8*)InObject->GetArchetype()));
+	// ChangelistManager / ChangelistState will hold onto a unique pointer for this
+	// so no need to worry about deleting it here.
+
+	FCustomDeltaChangelistState* DeltaChangelistState = nullptr;
+	if (LifetimeCustomPropertyState && !!LifetimeCustomPropertyState->NumFastArrayItems)
+	{
+		DeltaChangelistState = new FCustomDeltaChangelistState();
+	}
+
+	return MakeShareable(new FReplicationChangelistMgr(AsShared(), (const uint8*)InObject->GetArchetype(), DeltaChangelistState));
 }
 
 TUniquePtr<FRepState> FRepLayout::CreateRepState(
@@ -5152,23 +5646,37 @@ void FRepLayout::DestructProperties(FRepStateStaticBuffer& InShadowData) const
 	InShadowData.Buffer.Empty();
 }
 
-void FRepLayout::GetLifetimeCustomDeltaProperties(TArray<int32>& OutCustom, TArray<ELifetimeCondition>& OutConditions)
+PRAGMA_DISABLE_DEPRECATION_WARNINGS
+
+void FRepLayout::GetLifetimeCustomDeltaProperties(TArray<int32>& OutCustom, TArray<ELifetimeCondition>& OutConditions) const
 {
 	OutCustom.Empty();
 	OutConditions.Empty();
 
-	for (int32 i = 0; i < Parents.Num(); i++)
+	if (LifetimeCustomPropertyState)
 	{
-		if ((Parents[i].Flags & ERepParentFlags::IsCustomDelta) != ERepParentFlags::None)
+		for (const auto& KVP : LifetimeCustomPropertyState->LifetimeCustomDeltaProperties)
 		{
-			check(Parents[i].Property->RepIndex + Parents[i].ArrayIndex == i);
-
-			OutCustom.Add(i);
-			OutConditions.Add(Parents[i].Condition);
+			OutCustom.Add(KVP.Key);
+			OutConditions.Add(Parents[KVP.Key].Condition);
 		}
 	}
 }
 
+PRAGMA_ENABLE_DEPRECATION_WARNINGS
+
+const TArrayView<const uint16> FRepLayout::GetLifetimeCustomDeltaProperties() const
+{
+	if (LifetimeCustomPropertyState)
+	{
+		return LifetimeCustomPropertyState->LifetimeCustomDeltaPropertiesParentCache;
+	}
+	else
+	{
+		return TArrayView<const uint16>();
+	}
+}
+
 void FRepLayout::AddReferencedObjects(FReferenceCollector& Collector)
 {
 	UProperty* Current = nullptr;
@@ -5191,6 +5699,653 @@ void FRepLayout::AddReferencedObjects(FReferenceCollector& Collector)
 	}
 }
 
+// TODO: There's a better way to do this, but it requires more changes.
+//			Ideally, we bring Retirements management, etc., into RepLayout.
+//			What we could do instead of using standard changelists is use individual circular buffers 
+//			for any given array element. Each time we see a new Fast Array Rep ID, we'd add a new element
+//			to the buffer. Once all connections have acked that history, we would remove it from the buffer.
+//			We could still enforce a hard limit to the size of the buffers.
+//
+//			We might technically be able to do that without bringing retirement stuff in,
+//			but it would require us exposing the handles to NetSerialization.h (where fast TArrays live)
+//			and that gets included all over the place, and it'd be better to not.
+
+
+// It's important to note that unlike normal RepLayout properties which require changelists to know
+// what changed and am accumulated "lifetime changelist" to catch up late joiners or people with bad
+// connections, Fast Arrays don't need that.
+//
+// Instead, Fast Arrays track Rep IDs for each individual element, and for the array itself.
+// Initially, all keys start as 0 and are incremented as the items are marked dirty.
+// If we ever receive a NAK, we revert our IDs back to the last ACKed ID state.
+//
+// This means that at any given time, each connection knows what its current state is, and
+// can just compare that against the "real state" stored on the Fast Array.
+//
+// In FastArrayDeltaSerialize (and the DeltaSerializeStructs variant), we use those IDs to determine
+// what items actually need to be sent to a client.
+//
+// We also store the Changelist History ID alongside the Array / Element Rep IDs, and that is reset in the
+// case of NAKs and starts at a designated invalid state initially.
+//
+// Changelists used here are purely accelerations. In ideal scenarios, we do property
+// comparisons once per frame and store those changelists, just like normal rep layouts.
+// Also like normal RepLayouts, that will be shared across all connections.
+//
+// When we go to replicate a Fast Array, we just accumulate the changelists for all of its
+// dirty items, and then send those.
+//
+// On initial send, or if a given connection receives too many NAKs and falls outside of our History Range,
+// we will devolve into sending the full state of the dirty items. At that point we're guaranteed that
+// the connection is up to date (if it receives a NAK on that packet, its state will be reset back to
+// being outside history range, and we try again).
+
+void FRepLayout::PreSendCustomDeltaProperties(
+	UObject* Object,
+	UNetConnection* Connection,
+	FReplicationChangelistMgr& ChangelistMgr,
+	TMap<int32, TSharedPtr<INetDeltaBaseState>>& CustomDeltaStates) const
+{
+	if (!Connection->InternalAck)
+	{
+		const FLifetimeCustomDeltaState& LocalLifetimeCustomPropertyState = *LifetimeCustomPropertyState;
+
+		if (LocalLifetimeCustomPropertyState.NumFastArrayItems)
+		{
+			FRepChangelistState& ChangelistState = *ChangelistMgr.GetRepChangelistState();
+			FCustomDeltaChangelistState& CustomDeltaChangelistState = *ChangelistState.CustomDeltaChangelistState;
+
+			if (CustomDeltaChangelistState.CompareIndex != static_cast<uint32>(GFrameCounter))
+			{
+				const bool bIsInitial = (CustomDeltaChangelistState.CompareIndex == 0);
+				CustomDeltaChangelistState.CompareIndex = GFrameCounter;
+
+				// Setup new buffers.
+				++CustomDeltaChangelistState.HistoryEnd;
+
+				if ((CustomDeltaChangelistState.HistoryEnd - CustomDeltaChangelistState.HistoryStart) == FCustomDeltaChangelistState::MAX_CHANGE_HISTORY)
+				{
+					++CustomDeltaChangelistState.HistoryStart;
+				}
+
+				const int32 NewHistoryIndex = CustomDeltaChangelistState.HistoryEnd % FCustomDeltaChangelistState::MAX_CHANGE_HISTORY;
+				const int32 StartHistoryIndex = CustomDeltaChangelistState.HistoryStart % FCustomDeltaChangelistState::MAX_CHANGE_HISTORY;
+
+				FCustomDeltaChangedHistory& NewHistoryItem = CustomDeltaChangelistState.ChangeHistory[NewHistoryIndex];
+				FCustomDeltaChangedHistory& StartHistoryItem = CustomDeltaChangelistState.ChangeHistory[StartHistoryIndex];
+
+				if (!bIsInitial)
+				{
+					bool bUpdatedAnything = false;
+					const FConstRepObjectDataBuffer ObjectData(Object);
+
+					for (const auto& KVP : LocalLifetimeCustomPropertyState.LifetimeCustomDeltaProperties)
+					{
+						const uint16 RepIndex = KVP.Key;
+						const FLifetimeCustomDeltaProperty& CustomDeltaProperty = KVP.Value;
+
+						if (INDEX_NONE != CustomDeltaProperty.FastArrayItemsCommand)
+						{
+							const FRepParentCmd& FastArrayCmd = Parents[RepIndex];
+							const FFastArraySerializer& FastArray = *reinterpret_cast<const FFastArraySerializer*>((ObjectData + FastArrayCmd).Data);
+
+							if (FastArray.HasBeenSerialized() && FastArray.IsUsingStructDeltaSerialization())
+							{
+								const TSharedPtr<INetDeltaBaseState>* OldState = CustomDeltaStates.Find(RepIndex);
+								const FNetFastTArrayBaseState* FastArrayBaseState = static_cast<const FNetFastTArrayBaseState*>(OldState ? OldState->Get() : nullptr);
+
+								if (!FastArrayBaseState || FastArray.ArrayReplicationKey != FastArrayBaseState->ArrayReplicationKey)
+								{
+									bUpdatedAnything = true;
+								}
+							}
+						}
+					}
+
+					// If none of our structs were dirty, then revert our history, unless we haven't created one yet.
+					// This is because we won't know what items are going to use delta until at least the first replication.
+					if (!bUpdatedAnything)
+					{
+						--CustomDeltaChangelistState.HistoryEnd;
+					}
+				}
+
+				NewHistoryItem.FastArrayChangelists.SetNum(LocalLifetimeCustomPropertyState.NumFastArrayItems);
+				NewHistoryItem.bWasUpdated.Init(false, LocalLifetimeCustomPropertyState.NumFastArrayItems);
+			}
+		}
+	}
+}
+
+void FRepLayout::PostSendCustomDeltaProperties(
+	UObject* Object,
+	UNetConnection* Connection,
+	FReplicationChangelistMgr& ChangelistMgr,
+	TMap<int32, TSharedPtr<INetDeltaBaseState>>& CustomDeltaStates) const
+{
+}
+
+bool FRepLayout::DeltaSerializeFastArrayProperty(FFastArrayDeltaSerializeParams& Params, FReplicationChangelistMgr* ChangelistMgr) const
+{
+	SCOPE_CYCLE_COUNTER(STAT_RepLayout_DeltaSerializeFastArray);
+
+	// A portion of this work could be shared across all Fast Array Properties for a given object,
+	// but that would be easier to do if the Custom Delta Serialization was completely encapsulated in FRepLayout.
+
+	check(LifetimeCustomPropertyState);
+
+	FNetDeltaSerializeInfo& DeltaSerializeInfo = Params.DeltaSerializeInfo;
+	const int32 ParentIndex = DeltaSerializeInfo.PropertyRepIndex;
+	const FLifetimeCustomDeltaProperty& CustomDeltaProperty = LifetimeCustomPropertyState->LifetimeCustomDeltaProperties.FindChecked(ParentIndex);
+	const FRepParentCmd& Parent = Parents[ParentIndex];
+	const int32 CmdIndex = CustomDeltaProperty.FastArrayItemsCommand;
+
+	if (INDEX_NONE == CmdIndex)
+	{
+		// This should have already been caught by InitFromClass.
+		// So, log with a lower verbosity.
+		UE_LOG(LogRep, Log, TEXT("FRepLayout::DeltaSerializeFastArrayProperty: Invalid fast array items command index! %s"), *Parent.CachedPropertyName.ToString())
+		return false;
+	}
+
+	const FRepLayoutCmd& FastArrayItemCmd = Cmds[CmdIndex];
+	const int32 ElementSize = FastArrayItemCmd.ElementSize;
+
+	const int32 ItemLayoutStart = CmdIndex + 1;
+	const int32 ItemLayoutEnd = FastArrayItemCmd.EndCmd - 1;
+
+	UObject* Object = DeltaSerializeInfo.Object;
+	UPackageMapClient* PackageMap = (UPackageMapClient*)DeltaSerializeInfo.Map;
+	UNetConnection* Connection = DeltaSerializeInfo.Connection;
+	const bool bIsWriting = !!DeltaSerializeInfo.Writer;
+	const bool bInternalAck = !!Connection->InternalAck;
+
+	FRepObjectDataBuffer ObjectData(Object);
+	FScriptArray* ObjectArray = (FScriptArray*)(ObjectData + FastArrayItemCmd).Data;
+	FRepObjectDataBuffer ObjectArrayData(ObjectArray->GetData());
+
+	const int32 ObjectArrayNum = ObjectArray->Num();
+	FNetFieldExportGroup* NetFieldExportGroup = nullptr;
+
+	if (bInternalAck)
+	{
+		// Note, PackageMap should hold onto the strong reference for us, so we use raw pointers where
+		// we can.
+		const FString OwnerPathName = Owner->GetPathName();
+		TSharedPtr<FNetFieldExportGroup> LocalNetFieldExportGroup = PackageMap->GetNetFieldExportGroup(OwnerPathName);
+
+		if (!LocalNetFieldExportGroup.IsValid())
+		{
+			UE_LOG(LogRepProperties, VeryVerbose, TEXT("DeltaSerializeFastArrayProperty: Create Netfield Export Group."))
+			LocalNetFieldExportGroup = CreateNetfieldExportGroup();
+			PackageMap->AddNetFieldExportGroup(OwnerPathName, LocalNetFieldExportGroup);
+		}
+
+		NetFieldExportGroup = LocalNetFieldExportGroup.Get();
+	}
+
+	if (bIsWriting)
+	{
+		FNetBitWriter& Writer = static_cast<FNetBitWriter&>(*DeltaSerializeInfo.Writer);
+		auto& ChangedElements = *Params.WriteChangedElements;
+		
+		// This is a list of changelists to send, corresponding to items in ChangedElements.
+		TArray<TArray<uint16>> Changelists;
+		
+
+		const FHandleToCmdIndex& ArrayHandleToCmd = BaseHandleToCmdIndex[FastArrayItemCmd.RelativeHandle - 1];
+		const TArray<FHandleToCmdIndex>& ArrayHandleToCmdIndex = *ArrayHandleToCmd.HandleToCmdIndex;
+		const int32 HandlePerElement = ArrayHandleToCmdIndex.Num();
+
+		// We only need to create changelists if we're not in a demo.
+		// Note, we may change this in the future to also delta structs for replays.
+		if (!bInternalAck)
+		{
+			// First, we'll create our changelists.
+			if (ChangelistMgr)
+			{
+				FRepChangelistState& RepChangelistState = *ChangelistMgr->GetRepChangelistState();
+				FCustomDeltaChangelistState& DeltaChangelistState = *RepChangelistState.CustomDeltaChangelistState;
+
+				int32& NewChangelistHistory = Params.WriteBaseState->ChangelistHistory;
+				const int32 LastSentHistory = NewChangelistHistory;
+				NewChangelistHistory = DeltaChangelistState.HistoryEnd;
+
+				FRepShadowDataBuffer ShadowData(RepChangelistState.StaticBuffer.GetData());
+				FScriptArray* ShadowArray = (FScriptArray*)(ShadowData + FastArrayItemCmd).Data;
+				FRepShadowDataBuffer ShadowArrayData(ShadowArray->GetData());
+
+				const int32 FastArrayNumber = CustomDeltaProperty.FastArrayNumber;
+
+				// Check to see whether or not we need to update the global changelist shared between connections.
+
+				// Note, we explicitly pass in a 0 handles everywhere below.
+				// This is because each item will be received individually, the receiving side indices won't
+				// necessarily match our indices, and we already track the changelists separately.
+				{
+					const int32 RelativeNewHistory = NewChangelistHistory % FCustomDeltaChangelistState::MAX_CHANGE_HISTORY;
+					FCustomDeltaChangedHistory& HistoryItem = DeltaChangelistState.ChangeHistory[RelativeNewHistory];
+
+					FBitReference HistoryItemWasUpdated = HistoryItem.bWasUpdated[FastArrayNumber];
+					if (!HistoryItemWasUpdated)
+					{
+						HistoryItemWasUpdated = true;
+
+						// Update our shadow array, and reset our pointer in case we reallocated.
+						const int32 ShadowArrayNum = ShadowArray->Num();
+						FScriptArrayHelper Helper((UArrayProperty*)FastArrayItemCmd.Property, ShadowArray);
+						Helper.Resize(ObjectArrayNum);
+						ShadowArrayData = ShadowArray->GetData();
+
+						TMap<int32, TArray<uint16>>& HistoryChangelists = HistoryItem.FastArrayChangelists[CustomDeltaProperty.FastArrayNumber].ChangelistByID;
+
+						// Right now, we rely on the Changed Elements passed in to create our new Changelist.
+						// Unfortunately, it's possible that this particular sender was a couple of History Items behind us.
+						// This isn't particularly ideal, and means that we could end up with multiple histories with redundant
+						// changelists.
+						//
+						// That means more CPU spent merging them, and more bandwidth sending those redundant items.
+						// To get around that, we could keep a FNetFastTArrayBaseState in our changelist history.
+						// That would mean copying the logic in FFastArraySerializer::TFastArraySerializeHelper::BuildChangedAndDeletedBuffers
+						// and making it generic to FFastArraySerializer and FFastArraySerializerItem as opposed to templated.
+						//
+						// Note, we may need to make that change anyway, because there might be a pathological case with this.
+						//
+						// What happens if we reorder items in the array, or remove an entry in the middle ?
+						// In that case, right now, we have no way to tell this, and so we may end comparing shadow states
+						// from mismatched items, generating full changelists.
+						//
+						// That would basically devolve us back to the old behavior of sending the full reordered items
+						// at least once. After that, we'd be back to proper diffing.
+
+						const bool bIsInitial = DeltaChangelistState.HistoryStart == (DeltaChangelistState.HistoryEnd - 1);
+						for (auto& IDIndexPair : ChangedElements)
+						{
+							// If this is a new item, just send everything.
+							const bool bForceFail = ShadowArrayNum <= IDIndexPair.Idx;
+
+							// Go ahead and do a property compare here, regardless of what we'll actually use below.
+							// This is to prevent issues where someone with an initial / outdated connection doesn't properly
+							// update the changelists in our history, but does update the shadow state inadvertently.
+							FComparePropertiesSharedParams SharedParams{
+								/*bIsInitial=*/ bIsInitial,
+								/*bForceFail=*/ bForceFail,
+								RoleIndex,
+								RemoteRoleIndex,
+								Parents,
+								Cmds
+							};
+
+							const int32 ArrayElementOffset = ElementSize * IDIndexPair.Idx;
+
+							FConstRepObjectDataBuffer ElementData(ObjectArrayData + ArrayElementOffset);
+							FRepShadowDataBuffer ElementShadowData(ShadowArrayData + ArrayElementOffset);
+
+							TArray<uint16>& HistoryChangelist = HistoryChangelists.Emplace(IDIndexPair.ID);
+							CompareProperties_r(SharedParams, ItemLayoutStart, ItemLayoutEnd, ElementShadowData, ElementData, HistoryChangelist, 0);
+
+							if (HistoryChangelist.Num())
+							{
+								HistoryChangelist.Add(0);
+							}
+						}
+					}
+				}
+
+				// Now, merge all of the changelists we need together.
+				// If we're sufficiently far back, or if this is our first transmission, then we'll just force fail
+				// and send all changes. Note, this won't be all changes since the beginning, but just all changes
+				// for the currently dirty items.
+
+				if (INDEX_NONE != LastSentHistory && ((NewChangelistHistory - LastSentHistory) < FCustomDeltaChangelistState::MAX_CHANGE_HISTORY))
+				{
+					Changelists.SetNum(ChangedElements.Num());
+					for (int32 ChangelistHistory = LastSentHistory + 1; ChangelistHistory <= NewChangelistHistory; ++ChangelistHistory)
+					{
+						const int32 RelativeHistory = ChangelistHistory % FCustomDeltaChangelistState::MAX_CHANGE_HISTORY;
+						FCustomDeltaChangedHistory& HistoryItem = DeltaChangelistState.ChangeHistory[RelativeHistory];
+
+						const FConstRepObjectDataBuffer ConstObjectData(ObjectData);
+
+						if (HistoryItem.bWasUpdated[FastArrayNumber])
+						{
+							for (int32 i = 0; i < ChangedElements.Num(); ++i)
+							{
+								const TMap<int32, TArray<uint16>>& HistoryChangelists = HistoryItem.FastArrayChangelists[FastArrayNumber].ChangelistByID;
+								const auto& IDIndexPair = ChangedElements[i];
+
+								if (const TArray<uint16>* FoundChangelist = HistoryChangelists.Find(IDIndexPair.ID))
+								{
+									if (FoundChangelist->Num() > 1)
+									{
+										// This is basically the DynamicArray case from MergeChangelists, but specialized.
+										// We could probably just make that more generic.
+										// Might also be worth creating a helper Lambda / Struct to create ChangelistIterators / HandleIterators.
+
+										TArray<uint16>& ElementChangelist = Changelists[i];
+										TArray<uint16> Temp = MoveTemp(ElementChangelist);
+
+										const FConstRepObjectDataBuffer ElementData(ObjectArrayData + (IDIndexPair.Idx * ElementSize));
+
+										FChangelistIterator FoundChangelistIterator(*FoundChangelist, 0);
+										FRepHandleIterator FoundHandleIterator(FoundChangelistIterator, Cmds, ArrayHandleToCmdIndex, ElementSize, 1, ItemLayoutStart, ItemLayoutEnd);
+
+										if (Temp.Num() == 0)
+										{
+											PruneChangeList_r(FoundHandleIterator, ElementData, ElementChangelist);
+										}
+										else
+										{
+											FChangelistIterator ElementChangelistIterator(ElementChangelist, 0);
+											FRepHandleIterator ElementHandleIterator(ElementChangelistIterator, Cmds, ArrayHandleToCmdIndex, ElementSize, 1, ItemLayoutStart, ItemLayoutEnd);
+
+											MergeChangeList_r(FoundHandleIterator, ElementHandleIterator, ElementData, ElementChangelist);
+										}
+
+										ElementChangelist.Add(0);
+									}
+								}
+							}
+						}
+					}
+				}
+			}
+
+			if (Changelists.Num() == 0)
+			{
+				// If we didn't end up building changelists earlier for whatever reason, go ahead and just
+				// build a full changelist for all changed elements.
+
+				// This could have happened if we were sending initially, we were outside of history range, or
+				// we didn't have a changelist manager.
+
+				Changelists.SetNum(ChangedElements.Num());
+				for (int32 i = 0; i < ChangedElements.Num(); ++i)
+				{
+					const auto& IDIndexPair = ChangedElements[i];
+					const int32 ArrayElementOffset = ElementSize * IDIndexPair.Idx;
+
+					FConstRepObjectDataBuffer ElementData(ObjectArrayData + ArrayElementOffset);
+
+					TArray<uint16>& Changelist = Changelists[i];
+					BuildChangeList_r(ArrayHandleToCmdIndex, ItemLayoutStart, ItemLayoutEnd, ElementData, 0, true, Changelist);
+
+					if (Changelist.Num())
+					{
+						Changelist.Add(0);
+					}
+				}
+			}
+		}
+
+		// Ignore tracking properties in Network Profiler below.
+		// We will rely on the normal custom delta property tracking which happens elsewhere.
+		NETWORK_PROFILER_IGNORE_PROPERTY_SCOPE
+
+		// Now that we have our changelists setup, we can send the data.
+		for (int32 i = 0; i < ChangedElements.Num(); ++i)
+		{
+			const auto& IDIndexPair = ChangedElements[i];
+			uint32 ID = ChangedElements[i].ID;
+			Writer << ID;
+
+			const int32 ArrayElementOffset = ElementSize * IDIndexPair.Idx;
+			FConstRepObjectDataBuffer ElementData(ObjectArrayData + ArrayElementOffset);
+
+			if (bInternalAck)
+			{
+				SendAllProperties_BackwardsCompatible_r(
+					nullptr,
+					Writer,
+					false,
+					PackageMap,
+					NetFieldExportGroup,
+					ItemLayoutStart,
+					ItemLayoutEnd,
+					ElementData);
+			}
+			else
+			{
+				TArray<uint16>& Changelist = Changelists[i];
+				const bool bAnythingToSend = Changelist.Num() > 1;
+				Writer.WriteBit(!!bAnythingToSend);
+
+				if (bAnythingToSend)
+				{
+					FChangelistIterator ChangelistIterator(Changelist, 0);
+					FRepHandleIterator HandleIterator(
+						ChangelistIterator,
+						Cmds,
+						ArrayHandleToCmdIndex,
+						ElementSize,
+						1,
+						ItemLayoutStart,
+						ItemLayoutEnd);
+
+					SendProperties_r(
+						/*RepState=*/ nullptr,
+						Writer,
+						/*bDoChecksum=*/ false,
+						HandleIterator,
+						ElementData,
+						/*ArrayDepth=*/ 1,
+						/*SharedInfo=*/ nullptr);
+
+					WritePropertyHandle(Writer, 0, false);
+				}
+			}
+			
+		}
+
+		return !Writer.IsError();
+	}
+	else
+	{
+		FNetBitReader& Reader = static_cast<FNetBitReader&>(*DeltaSerializeInfo.Reader);
+		TArray<int32, TInlineAllocator<8>>& ChangedElements = *Params.ReadChangedElements;
+		TArray<int32, TInlineAllocator<8>>& AddedElements = *Params.ReadAddedElements;
+
+		FFastArraySerializer& ArraySerializer = Params.ArraySerializer;
+		FScriptArrayHelper FastArrayHelper((UArrayProperty*)FastArrayItemCmd.Property, ObjectArray);
+
+		bool bOutGuidsChanged = false;
+		bool bOutHasUnmapped = false;
+
+		// WARNING! Don't attempt to use ObjectArrayData below, always rely on FastArrayHelper.
+		// The helper may reallocate the array, and invalidate that pointer.
+
+		for (int32 i = 0; i < Params.ReadNumChanged; ++i)
+		{
+			uint32 ID = 0;
+			Reader << ID;
+
+			int32* ElementIndexPtr = ArraySerializer.ItemMap.Find(ID);
+			int32 ElementIndex = 0;
+			FFastArraySerializerItem* ThisElement = nullptr;
+
+			if (!ElementIndexPtr)
+			{
+				UE_LOG(LogNetFastTArray, Log, TEXT("   New. ID: %d. New Element!"), ID);
+
+				ElementIndex = FastArrayHelper.AddValue();
+				ArraySerializer.ItemMap.Add(ID, ElementIndex);
+				AddedElements.Add(ElementIndex);
+
+				ThisElement = (FFastArraySerializerItem*)(FastArrayHelper.GetRawPtr(ElementIndex));
+				ThisElement->ReplicationID = ID;
+			}
+			else
+			{
+				ElementIndex = *ElementIndexPtr;
+				ThisElement = (FFastArraySerializerItem*)(FastArrayHelper.GetRawPtr(ElementIndex));
+				ChangedElements.Add(ElementIndex);
+
+				UE_LOG(LogNetFastTArray, Log, TEXT("   Changed. ID: %d -> Idx: %d"), ID, ElementIndex);
+			}
+
+			// Update this element's most recent array replication key
+			ThisElement->MostRecentArrayReplicationKey = ArraySerializer.ArrayReplicationKey;
+
+			// Update this element's replication key so that a client can re-serialize the array for client replay recording
+			ThisElement->ReplicationKey++;
+
+			FGuidReferencesMap& GuidReferences = ArraySerializer.GuidReferencesMap_StructDelta.FindOrAdd(ID);
+
+			if (bInternalAck)
+			{
+				const bool bSuccess = ReceiveProperties_BackwardsCompatible_r(
+					nullptr,
+					NetFieldExportGroup,
+					Reader,
+					ItemLayoutStart,
+					ItemLayoutEnd,
+					nullptr,
+					FastArrayHelper.GetRawPtr(0),
+					ThisElement,
+					&GuidReferences,
+					bOutHasUnmapped,
+					bOutGuidsChanged);
+
+				if (!bSuccess)
+				{
+					UE_LOG(LogNetFastTArray, Warning, TEXT("FRepLayout::DeltaSerializeFastArrayProperty: Failed to receive backwards compat properties!"));
+					return false;
+				}
+			}
+			else
+			{
+				const bool bAnythingSent = !!Reader.ReadBit();
+				if (!bAnythingSent)
+				{
+					continue;
+				}
+
+				FReceivePropertiesSharedParams SharedParams{
+					/*bDoChecksum=*/ false,
+					/*bSkipRoleSwap=*/ false,
+					/*RepNotifies=*/ nullptr,
+					Reader,
+					bOutHasUnmapped,
+					bOutGuidsChanged,
+					Parents,
+					Cmds
+				};
+
+				FReceivePropertiesStackParams StackParams{
+					ThisElement,
+					nullptr,
+					&GuidReferences,
+					ItemLayoutStart,
+					ItemLayoutEnd
+				};
+
+				// Read the first handle, and then start receiving properties.
+				ReadPropertyHandle(SharedParams);
+				if (ReceiveProperties_r(SharedParams, StackParams))
+				{
+					if (0 != SharedParams.ReadHandle)
+					{
+						UE_LOG(LogRep, Error, TEXT("ReceiveFastArrayItem: Invalid property terminator handle - Handle=%d"), SharedParams.ReadHandle);
+						return false;
+					}
+				}
+				else
+				{
+					UE_LOG(LogNetFastTArray, Warning, TEXT("FRepLayout::DeltaSerializeFastArrayProperty: Failed to received properties"));
+					return false;
+				}
+			}
+
+			if (Reader.IsError())
+			{
+				UE_LOG(LogNetFastTArray, Warning, TEXT("FRepLayout::DeltaSerializeFastArrayProperty: Reader.IsError() == true"));
+				return false;
+			}
+
+			DeltaSerializeInfo.bGuidListsChanged |= bOutGuidsChanged;
+			DeltaSerializeInfo.bOutHasMoreUnmapped |= bOutHasUnmapped;
+		}
+
+		return true;
+	}
+}
+
+void FRepLayout::GatherGuidReferencesForFastArray(FFastArrayDeltaSerializeParams& Params) const
+{
+	const FFastArraySerializer& ArraySerializer = Params.ArraySerializer;
+	TSet<FNetworkGUID>& GatherGuids = *Params.DeltaSerializeInfo.GatherGuidReferences;
+	int32 TrackedGuidMemory = 0;
+
+	for (auto& GuidReferencesPair : ArraySerializer.GuidReferencesMap_StructDelta)
+	{
+		GatherGuidReferences_r(&GuidReferencesPair.Value, GatherGuids, TrackedGuidMemory);
+	}
+
+	if (Params.DeltaSerializeInfo.TrackedGuidMemoryBytes)
+	{
+		*Params.DeltaSerializeInfo.TrackedGuidMemoryBytes += TrackedGuidMemory;
+	}
+}
+
+bool FRepLayout::MoveMappedObjectToUnmappedForFastArray(FFastArrayDeltaSerializeParams& Params) const
+{
+	FFastArraySerializer& ArraySerializer = Params.ArraySerializer;
+	const FNetworkGUID& MoveToUnmapped = *Params.DeltaSerializeInfo.MoveGuidToUnmapped;
+
+	bool bFound = false;
+	for (auto& GuidReferencesPair : ArraySerializer.GuidReferencesMap_StructDelta)
+	{
+		bFound |= MoveMappedObjectToUnmapped_r(&GuidReferencesPair.Value, MoveToUnmapped);
+	}
+	return bFound;
+}
+
+void FRepLayout::UpdateUnmappedGuidsForFastArray(FFastArrayDeltaSerializeParams& Params) const
+ {
+	check(LifetimeCustomPropertyState);
+
+	FNetDeltaSerializeInfo& DeltaSerializeInfo = Params.DeltaSerializeInfo;
+	const int32 ParentIndex = DeltaSerializeInfo.PropertyRepIndex;
+	const FLifetimeCustomDeltaProperty& CustomDeltaProperty = LifetimeCustomPropertyState->LifetimeCustomDeltaProperties.FindChecked(ParentIndex);
+	const FRepParentCmd& Parent = Parents[ParentIndex];
+	const int32 CmdIndex = CustomDeltaProperty.FastArrayItemsCommand;
+	
+	if (INDEX_NONE == CmdIndex)
+	{
+		// This should have already been caught by InitFromClass.
+		// So, log with a lower verbosity.
+		UE_LOG(LogRep, Log, TEXT("FRepLayout::UpdateUnmappedGuidsForFastArray: Invalid fast array items command index! %s"), *Parent.CachedPropertyName.ToString())
+		return;
+	}
+
+	const FRepLayoutCmd& FastArrayItemCmd = Cmds[CmdIndex];
+	const int32 ElementSize = FastArrayItemCmd.ElementSize;
+
+	UObject* Object = DeltaSerializeInfo.Object;
+	UPackageMap* PackageMap = DeltaSerializeInfo.Map;
+
+	FRepObjectDataBuffer ObjectData(Object);
+	FScriptArray* ScriptArray = (FScriptArray*)(ObjectData + FastArrayItemCmd).Data;
+	FRepObjectDataBuffer ArrayData(ScriptArray->GetData());
+
+	FFastArraySerializer& ArraySerializer = Params.ArraySerializer;
+
+	bool bOutSomeObjectsWereMapped = false;
+	bool bOutHasMoreUnmapped = false;
+
+	for (auto& GuidReferencesPair : ArraySerializer.GuidReferencesMap_StructDelta)
+	{
+		const int32 ItemIndex = ArraySerializer.ItemMap[GuidReferencesPair.Key];
+		const int32 ArrayElementOffset = ItemIndex * ElementSize;
+		FRepObjectDataBuffer ElementData(ArrayData + ArrayElementOffset);
+
+		UpdateUnmappedObjects_r(nullptr, &GuidReferencesPair.Value, Object, PackageMap, nullptr, ElementData, ElementSize, bOutSomeObjectsWereMapped, bOutHasMoreUnmapped);
+
+		DeltaSerializeInfo.bOutHasMoreUnmapped |= bOutHasMoreUnmapped;
+		DeltaSerializeInfo.bOutSomeObjectsWereMapped |= bOutSomeObjectsWereMapped;
+	}
+}
+
 void FRepLayout::CountBytes(FArchive& Ar) const
 {
 	GRANULAR_NETWORK_MEMORY_TRACKING_INIT(Ar, "FRepLayout::CountBytes");
@@ -5199,6 +6354,13 @@ void FRepLayout::CountBytes(FArchive& Ar) const
 	GRANULAR_NETWORK_MEMORY_TRACKING_TRACK("BaseHandleToCmdIndex", BaseHandleToCmdIndex.CountBytes(Ar));
 	GRANULAR_NETWORK_MEMORY_TRACKING_TRACK("SharedInfoRPC", SharedInfoRPC.CountBytes(Ar));
 	GRANULAR_NETWORK_MEMORY_TRACKING_TRACK("SharedInfoRPCParentsChanged", SharedInfoRPCParentsChanged.CountBytes(Ar));
+	GRANULAR_NETWORK_MEMORY_TRACKING_TRACK("LifetimeCustomPropertyState",
+		if (LifetimeCustomPropertyState)
+		{
+			Ar.CountBytes(sizeof(FLifetimeCustomDeltaState), sizeof(FLifetimeCustomDeltaState));
+			LifetimeCustomPropertyState->CountBytes(Ar);
+		}
+	);
 }
 
 void FReceivingRepState::CountBytes(FArchive& Ar) const
diff --git a/Engine/Source/Runtime/Engine/Public/Net/DataReplication.h b/Engine/Source/Runtime/Engine/Public/Net/DataReplication.h
--- a/Engine/Source/Runtime/Engine/Public/Net/DataReplication.h
+++ b/Engine/Source/Runtime/Engine/Public/Net/DataReplication.h
@@ -39,6 +39,7 @@ struct FReplicatedActorProperty
 {
 	/** offset into the Actor where this reference is located - includes offsets from any outer structs */
 	int32 Offset;
+
 	/** Reference to property object */
 	const class UObjectPropertyBase* Property;
 
@@ -49,7 +50,6 @@ struct FReplicatedActorProperty
 
 /** FObjectReplicator
  *   Generic class that replicates properties for an object.
- *   All delta/diffing work is done in this class. 
  *	 Its primary job is to produce and consume chunks of properties/RPCs:
  *
  *		|----------------|
@@ -75,23 +75,59 @@ public:
 	FNetworkGUID									ObjectNetGUID;
 	UObject*										ObjectPtr;
 
+	//~ Retirement, RecentCustomDeltaState, CDOCustomDeltaState, and CheckpointCustomDeltaState
+	//~ could all be moved to FSendingRepState, as they are not needed when solely receiving.
+
 	TArray<FPropertyRetirement>						Retirement;					// Property retransmission.
 	TMap<int32, TSharedPtr<INetDeltaBaseState> >	RecentCustomDeltaState;		// This is the delta state we need to compare with when determining what to send to a client for custom delta properties
 	TMap<int32, TSharedPtr<INetDeltaBaseState> >	CDOCustomDeltaState;		// Same as RecentCustomDeltaState, but this will always remain as the initial CDO version. We use this to send all properties since channel was first opened (for bResendAllDataSinceOpen)
 	TMap<int32, TSharedPtr<INetDeltaBaseState> >	CheckpointCustomDeltaState;	// Same as RecentCustomDeltaState, but will represent the state at the last checkpoint
 
+	UE_DEPRECATED(4.23, "This property will be removed in future versions. Use FRepLayout::GetLifetimeCustomDeltaProperties instead.")
 	TArray< int32 >									LifetimeCustomDeltaProperties;
+
+	UE_DEPRECATED(4.23, "This property will be removed in future versions. Use FRepLayout::GetRepPropertyCondition instead.")
 	TArray< ELifetimeCondition >					LifetimeCustomDeltaPropertyConditions;
 
-	uint32											bLastUpdateEmpty	: 1;	// True if last update (ReplicateActor) produced no replicated properties
-	uint32											bOpenAckCalled		: 1;
-	uint32											bForceUpdateUnmapped: 1;	// True if we need to do an unmapped check next frame
+	/** True if last update (ReplicateActor) produced no replicated properties */
+	uint32 bLastUpdateEmpty : 1;
+
+	/** Whether or not the Actor Channel on which we're replicating has been Opened / Acked by the receiver. */
+	uint32 bOpenAckCalled : 1;
+
+	/** True if we need to do an unmapped check next frame. */
+	uint32 bForceUpdateUnmapped: 1;
+
+	/** Whether or not we've already replicated properties this frame. */
+	uint32 bHasReplicatedProperties : 1;
+	
+private:
 
-	UNetConnection *								Connection;					// Connection this replicator was created on
-	class UActorChannel	*							OwningChannel;
+	uint32 bSupportsFastArrayDelta : 1;
 
+public:	
+
+	/** Connection this replicator was created on. */
+	UNetConnection* Connection;
+
+	/** The Actor Channel that we're replicating on. This expected to be owned by Connection. */
+	class UActorChannel* OwningChannel;
+
+	//~ This property is only used as an acceleration to avoid GUID tracking / updates when we know there are
+	//~ none unmapped. This will just be removed in future versions, but if the performance hit is unacceptable
+	//~ it could be replaced by a TBitArray in FReceivingRepState, where each bit flags whether or not
+	//~ a given index in FRepLayout::LifetimeCustomDeltaProperties has unmapped guids.
+	//~ That should be both more CPU and Memory efficient than the current approach.
+	//~ TODO: Remove FRepLayout::MoveMappedToUnmappedForCustomDeltaProperties overload when this goes away.
+	//~ TODO: Remove FRepLayout::UpdateUnmappedObjectsForCustomDeltaProperties overload when this goes away.
+	UE_DEPRECATED(4.23, "This property will be removed in future versions.")
 	TMap< int32, UStructProperty* >					UnmappedCustomProperties;
 
+	//~ TODO: RepNotifies here could likely just be removed in favor of using the existing array on FReceivingRepState.
+	//~			In most cases, RepNotifyMetaData isn't really used, but if it is useful, then that should also
+	//~			be moved to FReceivingRepState and we should combine the RepNotify logic here with that in
+	//~			FRepLayout.
+
 	TArray< UProperty*,TInlineAllocator< 32 > >		RepNotifies;
 	TMap< UProperty*, TArray<uint8> >				RepNotifyMetaData;
 
@@ -152,8 +188,15 @@ public:
 	/** Takes Data, and compares against shadow state to log differences */
 	bool ValidateAgainstState( const UObject* ObjectState );
 
+	UE_DEPRECATED(4.23, "This method will be removed in future versions. Please use SendCustomDeltaProperty instead.")
 	static bool SerializeCustomDeltaProperty( UNetConnection * Connection, void* Src, UProperty * Property, uint32 ArrayIndex, FNetBitWriter & OutBunch, TSharedPtr<INetDeltaBaseState> & NewFullState, TSharedPtr<INetDeltaBaseState> & OldState );
 
+	//~ Both of these should be private, IMO, but we'll leave them public for now for back compat
+	//~ in case anyone was using SerializeCustomDeltaProperty already.
+
+	bool SendCustomDeltaProperty(UObject* InObject, UProperty* Property, uint32 ArrayIndex, FNetBitWriter& OutBunch, TSharedPtr<INetDeltaBaseState> & NewFullState, TSharedPtr<INetDeltaBaseState> & OldState);
+	bool SendCustomDeltaProperty(UObject* InObject, uint16 CustomDeltaProperty, FNetBitWriter& OutBunch, TSharedPtr<INetDeltaBaseState>& NewFullState, TSharedPtr<INetDeltaBaseState>& OldState);
+
 	/** Packet was dropped */
 	void	ReceivedNak( int32 NakPacketId );
 
@@ -178,8 +221,6 @@ public:
 
 	void	ForceRefreshUnreliableProperties();
 
-	bool bHasReplicatedProperties;
-
 	void QueueRemoteFunctionBunch( UFunction* Func, FOutBunch &Bunch );
 
 	bool ReadyForDormancy(bool debug=false);
diff --git a/Engine/Source/Runtime/Engine/Public/Net/GuidReferences.h b/Engine/Source/Runtime/Engine/Public/Net/GuidReferences.h
new file mode 100644
--- /dev/null
+++ b/Engine/Source/Runtime/Engine/Public/Net/GuidReferences.h
@@ -0,0 +1,102 @@
+// Copyright 1998-2019 Epic Games, Inc. All Rights Reserved.
+
+#pragma once
+
+#include "CoreMinimal.h"
+#include "Serialization/BitReader.h"
+#include "Misc/NetworkGuid.h"
+
+class FGuidReferences;
+
+using FGuidReferencesMap = TMap<int32, FGuidReferences>;
+
+/**
+ * References to Objects (including Actors, Components, etc.) are replicated as NetGUIDs, since
+ * the literal memory pointers will be different across game instances. In these cases, actual
+ * replicated data for the Object will be handled elsewhere (either on its own Actor channel,
+ * or on its Owning Actor's channel, as a replicated subobject).
+ *
+ * This class helps manage those references for specific replicated properties.
+ * A FGuidReferences instance will be created for each Replicated Property that is a reference to an object.
+ *
+ * Guid References may also be nested in properties (like dynamic arrays), so we'll recursively track
+ * those as well.
+ */
+class ENGINE_API FGuidReferences
+{
+public:
+	FGuidReferences():
+		NumBufferBits(0),
+		Array(nullptr)
+	{}
+
+	FGuidReferences(
+		FBitReader&					InReader,
+		FBitReaderMark&				InMark,
+		const TSet<FNetworkGUID>&	InUnmappedGUIDs,
+		const TSet<FNetworkGUID>&	InMappedDynamicGUIDs,
+		const int32					InParentIndex,
+		const int32					InCmdIndex
+	):
+		ParentIndex(InParentIndex),
+		CmdIndex(InCmdIndex),
+		UnmappedGUIDs(InUnmappedGUIDs),
+		MappedDynamicGUIDs(InMappedDynamicGUIDs),
+		Array(nullptr)
+	{
+		NumBufferBits = InReader.GetPosBits() - InMark.GetPos();
+		InMark.Copy(InReader, Buffer);
+	}
+
+	FGuidReferences(
+		FGuidReferencesMap*	InArray,
+		const int32			InParentIndex,
+		const int32			InCmdIndex
+	):
+		ParentIndex(InParentIndex),
+		CmdIndex(InCmdIndex),
+		NumBufferBits(0),
+		Array(InArray)
+	{}
+
+	~FGuidReferences();
+
+	void CountBytes(FArchive& Ar) const
+	{
+		UnmappedGUIDs.CountBytes(Ar);
+		MappedDynamicGUIDs.CountBytes(Ar);
+		Buffer.CountBytes(Ar);
+
+		if (Array)
+		{
+			Array->CountBytes(Ar);
+			for (const auto& GuidReferencePair : *Array)
+			{
+				GuidReferencePair.Value.CountBytes(Ar);
+			}
+		}
+	}
+
+	/** The Property Command index of the top level property that references the GUID. */
+	int32 ParentIndex;
+
+	/** The Property Command index of the actual property that references the GUID. */
+	int32 CmdIndex;
+
+	int32 NumBufferBits;
+
+	/** GUIDs for objects that haven't been loaded / created yet. */
+	TSet<FNetworkGUID> UnmappedGUIDs;
+
+	/** GUIDs for dynamically spawned objects that have already been created. */
+	TSet<FNetworkGUID> MappedDynamicGUIDs;
+
+	/** A copy of the last network data read related to this GUID Reference. */
+	TArray<uint8> Buffer;
+
+	/**
+	 * If this FGuidReferences instance is owned by an Array Property that contains nested GUID References,
+	 * then this will be a valid FGuidReferencesMap containing the nested FGuidReferences.
+	 */
+	FGuidReferencesMap* Array;
+};
\ No newline at end of file
diff --git a/Engine/Source/Runtime/Engine/Public/Net/NetworkProfiler.h b/Engine/Source/Runtime/Engine/Public/Net/NetworkProfiler.h
--- a/Engine/Source/Runtime/Engine/Public/Net/NetworkProfiler.h
+++ b/Engine/Source/Runtime/Engine/Public/Net/NetworkProfiler.h
@@ -70,6 +70,9 @@ public:
 class FNetworkProfiler
 {
 private:
+
+	friend struct FNetworkProfilerScopedIgnoreReplicateProperties;
+
 	/** File writer used to serialize data.															*/
 	FArchive*								FileWriter;
 
@@ -162,6 +165,9 @@ private:
 	*/
 	int32 GetAddressTableIndex( const FString& Address );
 
+	// Used with FScopedIgnoreReplicateProperties.
+	uint32 IgnorePropertyCount;
+
 public:
 	/**
 	 * Constructor, initializing members.
@@ -406,8 +412,33 @@ public:
 /** Global network profiler instance. */
 extern ENGINE_API FNetworkProfiler GNetworkProfiler;
 
+#define NETWORK_PROFILER_IGNORE_PROPERTY_SCOPE const FNetworkProfilerScopedIgnoreReplicateProperties _NetProfilePrivate_IgnoreScope;
+
+/**
+ * Can be used to enforce a scope where we don't want to track properties.
+ * This is useful to prevent cases where we might inadvertently over count properties.
+ */
+struct FNetworkProfilerScopedIgnoreReplicateProperties
+{
+	FNetworkProfilerScopedIgnoreReplicateProperties()
+	{
+		++GNetworkProfiler.IgnorePropertyCount;
+	}
+
+	~FNetworkProfilerScopedIgnoreReplicateProperties()
+	{
+		--GNetworkProfiler.IgnorePropertyCount;
+	}
+
+private:
+
+	FNetworkProfilerScopedIgnoreReplicateProperties(const FNetworkProfilerScopedIgnoreReplicateProperties&) = delete;
+	FNetworkProfilerScopedIgnoreReplicateProperties(FNetworkProfilerScopedIgnoreReplicateProperties&&) = delete;
+};
+
 #else	// USE_NETWORK_PROFILER
 
 #define NETWORK_PROFILER(x)
+#define NETWORK_PROFILER_IGNORE_PROPERTY_SCOPE 
 
 #endif
diff --git a/Engine/Source/Runtime/Engine/Public/Net/RepLayout.h b/Engine/Source/Runtime/Engine/Public/Net/RepLayout.h
--- a/Engine/Source/Runtime/Engine/Public/Net/RepLayout.h
+++ b/Engine/Source/Runtime/Engine/Public/Net/RepLayout.h
@@ -20,6 +20,7 @@
 #include "Engine/EngineTypes.h"
 #include "UObject/GCObject.h"
 #include "Containers/StaticBitArray.h"
+#include "Net/GuidReferences.h"
 
 class FGuidReferences;
 class FNetFieldExportGroup;
@@ -435,7 +436,10 @@ private:
 
 	friend class FReplicationChangelistMgr;
 
-	FRepChangelistState(const TSharedRef<const FRepLayout>& InRepLayout, const uint8* Source);
+	FRepChangelistState(
+		const TSharedRef<const FRepLayout>& InRepLayout,
+		const uint8* Source,
+		struct FCustomDeltaChangelistState* CustomDeltaChangelistState);
 
 public:
 
@@ -445,6 +449,9 @@ public:
 	/** Circular buffer of changelists. */
 	FRepChangedHistory ChangeHistory[MAX_CHANGE_HISTORY];
 
+	/** Changelist state specific to Custom Delta properties. Only allocated if this RepLayout has Custom Delta Properties. */
+	TUniquePtr<struct FCustomDeltaChangelistState> CustomDeltaChangelistState;
+
 	/** Index in the buffer where changelist history starts (i.e., the Oldest changelist). */
 	int32 HistoryStart;
 
@@ -476,7 +483,10 @@ private:
 
 	friend class FRepLayout;
 
-	FReplicationChangelistMgr(const TSharedRef<const FRepLayout>& InRepLayout, const uint8* Source);
+	FReplicationChangelistMgr(
+		const TSharedRef<const FRepLayout>& InRepLayout,
+		const uint8* Source,
+		struct FCustomDeltaChangelistState* CustomDeltaChangelistState);
 
 public:
 
@@ -493,87 +503,6 @@ private:
 	FRepChangelistState RepChangelistState;
 };
 
-class FGuidReferences;
-
-typedef TMap<int32, FGuidReferences> FGuidReferencesMap;
-
-/**
- * References to Objects (including Actors, Components, etc.) are replicated as NetGUIDs, since
- * the literal memory pointers will be different across game instances. In these cases, actual
- * replicated data for the Object will be handled elsewhere (either on its own Actor channel,
- * or on its Owning Actor's channel, as a replicated subobject).
- *
- * This class helps manage those references for specific replicated properties.
- * A FGuidReferences instance will be created for each Replicated Property that is a reference to an object.
- * Note, that in the same way Rep Commands may be nested, FGuidReferences can be nested in the same way.
- */
-class FGuidReferences
-{
-public:
-	FGuidReferences() : NumBufferBits(0), Array(NULL) {}
-
-	FGuidReferences(
-		FBitReader&					InReader,
-		FBitReaderMark&				InMark,
-		const TSet<FNetworkGUID>&	InUnmappedGUIDs,
-		const TSet<FNetworkGUID>&	InMappedDynamicGUIDs,
-		const int32					InParentIndex,
-		const int32					InCmdIndex
-	):
-			UnmappedGUIDs(InUnmappedGUIDs),
-			MappedDynamicGUIDs(InMappedDynamicGUIDs),
-			Array(NULL),
-			ParentIndex(InParentIndex),
-			CmdIndex(InCmdIndex)
-	{
-		NumBufferBits = InReader.GetPosBits() - InMark.GetPos();
-		InMark.Copy(InReader, Buffer);
-	}
-
-	FGuidReferences(
-		FGuidReferencesMap*	InArray,
-		const int32			InParentIndex,
-		const int32			InCmdIndex
-	):
-
-			NumBufferBits(0),
-			Array(InArray),
-			ParentIndex(InParentIndex),
-			CmdIndex(InCmdIndex)
-	{}
-		
-	~FGuidReferences();
-
-	void CountBytes(FArchive& Ar) const
-	{
-		UnmappedGUIDs.CountBytes(Ar);
-		MappedDynamicGUIDs.CountBytes(Ar);
-		Buffer.CountBytes(Ar);
-	}
-
-	/** GUIDs for objects that haven't been loaded / created yet. */
-	TSet<FNetworkGUID> UnmappedGUIDs;
-
-	/** GUIDs for dynamically spawned objects that have already been created. */
-	TSet<FNetworkGUID> MappedDynamicGUIDs;
-
-	/** A copy of the last network data read related to this GUID Reference. */
-	TArray<uint8> Buffer;
-	int32 NumBufferBits;
-
-	/**
-	 * If this FGuidReferences instance is owned by an Array Property that contains nested GUID References,
-	 * then this will be a valid FGuidReferencesMap containing the nested FGuidReferences.
-	 */
-	FGuidReferencesMap* Array;
-
-	/** The Property Command index of the top level property that references the GUID. */
-	int32 ParentIndex;
-
-	/** The Property Command index of the actual property that references the GUID. */
-	int32 CmdIndex;
-};
-
 /** Replication State needed to track received properties. */
 class FReceivingRepState : public FNoncopyable
 {
@@ -779,7 +708,9 @@ enum class ERepParentFlags : uint32
 	IsCustomDelta		= (1 << 3),	//! This property uses custom delta compression. Mutually exclusive with IsNetSerialize.
 	IsNetSerialize		= (1 << 4), //! This property uses a custom net serializer. Mutually exclusive with IsCustomDelta.
 	IsStructProperty	= (1 << 5),	//! This property is a UStructProperty.
-	IsZeroConstructible	= (1 << 6)	//! This property is ZeroConstructible.
+	IsZeroConstructible	= (1 << 6),	//! This property is ZeroConstructible.
+	IsFastArray			= (1 << 7), //! This property is a FastArraySerializer. This can't be a ERepLayoutCmdType, because
+									//! these Custom Delta structs will have their inner properties tracked.
 };
 
 ENUM_CLASS_FLAGS(ERepParentFlags)
@@ -1142,22 +1073,13 @@ private:
 
 	friend struct FRepStateStaticBuffer;
 	friend class UPackageMapClient;
+	friend class FNetSerializeCB;
 
-	FRepLayout():
-		LayoutState(ERepLayoutState::Uninitialized),
-		FirstNonCustomParent(0),
-		RoleIndex(INDEX_NONE),
-		RemoteRoleIndex(-1),
-		Owner(nullptr)
-	{}
+	FRepLayout();
 
 public:
 
-	~FRepLayout()
-	{
-		// This should never happen.
-		check(ERepLayoutState::Uninitialized != LayoutState);
-	}
+	virtual ~FRepLayout();
 
 	/** Creates a new FRepLayout for the given class. */
 	ENGINE_API static TSharedPtr<FRepLayout> CreateFromClass(UClass* InObjectClass, const UNetConnection* ServerConnection = nullptr, const ECreateRepLayoutFlags Flags = ECreateRepLayoutFlags::None);
@@ -1389,7 +1311,7 @@ public:
 	{
 		UpdateUnmappedObjects(RepState->GetReceivingRepState(), PackageMap, Object, bOutSomeObjectsWereMapped, bOutHasMoreUnmapped);
 	}
-	
+
 	/**
 	 * Fire any RepNotifies that have been queued for an object while receiving properties.
 	 *
@@ -1498,7 +1420,17 @@ public:
 		TRepDataBuffer<DstType> Destination,
 		TConstRepDataBuffer<SrcType> Source) const;
 
-	void GetLifetimeCustomDeltaProperties(TArray<int32>& OutCustom, TArray<ELifetimeCondition>& OutConditions);
+	UE_DEPRECATED(4.23, "This method is deprecated. Please use GetLifetimeCustomDeltaProperties that returns an ArrayView and GetPropertyLifetimeCondition instead.")
+	void GetLifetimeCustomDeltaProperties(TArray<int32>& OutCustom, TArray<ELifetimeCondition>& OutConditions) const;
+
+	/** Returns the RepIndices of any Lifetime Custom Delta properties. */
+	const TArrayView<const uint16> GetLifetimeCustomDeltaProperties() const;
+
+	/** Returns the Lifetime Condition for the Replicated property associated with the given RepIndex. */
+	const ELifetimeCondition GetPropertyLifetimeCondition(uint16 RepIndex) const
+	{
+		return Parents[RepIndex].Condition;
+	}
 
 	/** @see SendProperties. */
 	void ENGINE_API SendPropertiesForRPC(
@@ -1606,6 +1538,12 @@ public:
 
 	void CountBytes(FArchive& Ar) const;
 
+	UE_DEPRECATED(4.23, "This method is only temporary to support existing behavior, and will be removed in future versions.")
+	UProperty* GetPropertyForRepIndex(uint32 RepIndex) const
+	{
+		return Parents[RepIndex].Property;
+	}
+
 private:
 
 	void InitFromClass(UClass* InObjectClass, const UNetConnection* ServerConnection, const ECreateRepLayoutFlags Flags);
@@ -1705,34 +1643,12 @@ private:
 
 	void SendProperties_r(
 		FSendingRepState* RESTRICT RepState,
-		FRepChangedPropertyTracker* ChangedTracker,
 		FNetBitWriter& Writer,
 		const bool bDoChecksum,
 		FRepHandleIterator& HandleIterator,
 		const FConstRepObjectDataBuffer SourceData,
 		const int32	 ArrayDepth,
-		const FRepSerializationSharedInfo& SharedInfo) const;
-
-	uint16 CompareProperties_r(
-		FSendingRepState* RESTRICT RepState,
-		const int32 CmdStart,
-		const int32 CmdEnd,
-		FRepShadowDataBuffer CompareData,
-		const FConstRepObjectDataBuffer Data,
-		TArray<uint16>& Changed,
-		uint16 Handle,
-		const bool bIsInitial,
-		const bool bForceFail) const;
-
-	void CompareProperties_Array_r(
-		FSendingRepState* RESTRICT RepState,
-		FRepShadowDataBuffer ShadowData,
-		const FConstRepObjectDataBuffer Data,
-		TArray<uint16>& Changed,
-		const uint16 CmdIndex,
-		const uint16 Handle,
-		const bool bIsInitial,
-		const bool bForceFail) const;
+		const FRepSerializationSharedInfo* const RESTRICT SharedInfo) const;
 
 	void BuildSharedSerialization(
 		const FConstRepObjectDataBuffer Data,
@@ -1783,7 +1699,7 @@ private:
 		bool& bOutGuidsChanged) const;
 
 	void GatherGuidReferences_r(
-		FGuidReferencesMap* GuidReferencesMap,
+		const FGuidReferencesMap* GuidReferencesMap,
 		TSet<FNetworkGUID>& OutReferencedGuids,
 		int32& OutTrackedGuidMemoryBytes) const;
 
@@ -1816,24 +1732,6 @@ private:
 
 	void SanityCheckChangeList(const FConstRepObjectDataBuffer Data, TArray<uint16>& Changed) const;
 
-	uint32 AddPropertyCmd(
-		UProperty* Property,
-		int32 Offset,
-		int32 RelativeHandle,
-		int32 ParentIndex,
-		uint32 ParentChecksum,
-		int32 StaticArrayIndex,
-		const UNetConnection* ServerConnection);
-
-	uint32 AddArrayCmd(
-		UArrayProperty* Property,
-		int32 Offset,
-		int32 RelativeHandle,
-		int32 ParentIndex,
-		uint32 ParentChecksum,
-		int32 StaticArrayIndex,
-		const UNetConnection* ServerConnection);
-
 	void SerializeProperties_DynamicArray_r(
 		FBitArchive& Ar, 
 		UPackageMap* Map,
@@ -1891,6 +1789,7 @@ private:
 		const int32 CmdEnd,
 		const FConstRepObjectDataBuffer Data,
 		const int32 HandleOffset,
+		const bool bForceArraySends,
 		TArray<uint16>& Changed) const;
 
 	void BuildHandleToCmdIndexTable_r(
@@ -1903,9 +1802,89 @@ private:
 	void CopyProperties(FRepStateStaticBuffer& ShadowData, const FConstRepObjectDataBuffer Source) const;
 	void DestructProperties(FRepStateStaticBuffer& RepStateStaticBuffer) const;
 
-	ERepLayoutState LayoutState;
+	/**
+	 * Similar to GatherGuidReferences, except works for CustomDeltaProperties.
+	 *
+	 * @param Params	Params that we'll use to gather guids.
+	 *					The GatherGuidReferences member must not be null.
+	 *					The Object member must not be null.
+	 */
+	void GatherGuidReferencesForCustomDeltaProperties(FNetDeltaSerializeInfo& Params) const;
+
+	/**
+	 * Similar to MoveMappedObjectToUnmapped, except works for CustomDeltaProperties.
+	 *
+	 * @param Params	Params that we'll use to gather guids.
+	 *					The MoveGuidToUnmapped member must not be null.
+	 *					The Object member must not be null.
+	 */
+	bool MoveMappedObjectToUnmappedForCustomDeltaProperties(FNetDeltaSerializeInfo& Params) const;
+
+	//~ Remove this when FObjectReplicator::UnmappedCustomProperties goes away.
+	UE_DEPRECATED(4.23, "This method is only temporary to support existing behavior, and will be removed in future versions.")
+	bool MoveMappedObjectToUnmappedForCustomDeltaProperties(FNetDeltaSerializeInfo& Params, TMap<int32, UStructProperty*>& UnmappedProperties) const;
+
+	/**
+	 * Similar to UpdateUnmappedObjects, except works for CustomDeltaProperties.
+	 *
+	 * @param Params					Params that we'll use to gather guids.
+	 *									The bUpdateUnmappedObjects member must be true.
+	 *									The PackageMap member must be non null.
+	 *									The Object member must be non null.
+	 */
+	void UpdateUnmappedObjectsForCustomDeltaProperties(FNetDeltaSerializeInfo& Params) const;
+
+	//~ Remove this when FObjectReplicator::UnmappedCustomProperties goes away.
+	UE_DEPRECATED(4.23, "This method is only temporary to support existing behavior, and will be removed in future versions.")
+	void UpdateUnmappedObjectsForCustomDeltaProperties(
+		FNetDeltaSerializeInfo& Params,
+		TArray<TPair<int32, UStructProperty*>>& CompletelyMapped,
+		TArray<TPair<int32, UStructProperty*>>& Updated) const;
+
+	bool SendCustomDeltaProperty(FNetDeltaSerializeInfo& Params, uint16 CustomDeltaRepIndex) const;
+
+	UE_DEPRECATED(4.23, "This method is only temporary to support existing behavior, and will be removed in future versions.")
+	bool SendCustomDeltaProperty(
+		FNetDeltaSerializeInfo& Params,
+		UProperty* Property,
+		int32 StaticArrayIndex) const;
+
+	/**
+	 * Attempts to receive the custom delta property.
+	 *
+	 * @param Params			Params that we'll use to receive.
+	 *							The PackageMap member must be non null.
+	 *							The Object member must be non null.
+	 *							The Reader member must be non null.
+	 * @param Property			The Property that we're trying to received.
+	 *							This is expected to be a valid Replicated property that is a CustomDelta type.
+	 */
+	bool ReceiveCustomDeltaProperty(FNetDeltaSerializeInfo& Params, UStructProperty* Property) const;
+
+	UE_DEPRECATED(4.23, "This method is only temporary to support existing behavior, and will be removed in future versions.")
+	bool ReceiveCustomDeltaProperty(FNetDeltaSerializeInfo& Params, UStructProperty* Property, uint32& StaticArrayIndex, int32& Offset) const;
+
+	bool DeltaSerializeFastArrayProperty(struct FFastArrayDeltaSerializeParams& Params, FReplicationChangelistMgr* ChangelistMgr) const;
+
+	void GatherGuidReferencesForFastArray(struct FFastArrayDeltaSerializeParams& Params) const;
 
-	int16 FirstNonCustomParent;
+	bool MoveMappedObjectToUnmappedForFastArray(struct FFastArrayDeltaSerializeParams& Params) const;
+
+	void UpdateUnmappedGuidsForFastArray(struct FFastArrayDeltaSerializeParams& Params) const;
+
+	void PreSendCustomDeltaProperties(
+		UObject* Object,
+		UNetConnection* Connection,
+		FReplicationChangelistMgr& ChangelistMgr,
+		TMap<int32, TSharedPtr<INetDeltaBaseState>>& CustomDeltaStates) const;
+
+	void PostSendCustomDeltaProperties(
+		UObject* Object,
+		UNetConnection* Connection,
+		FReplicationChangelistMgr& ChangelistMgr,
+		TMap<int32, TSharedPtr<INetDeltaBaseState>>& CustomDeltaStates) const;
+
+	ERepLayoutState LayoutState;
 
 	/** Index of the Role property in the Parents list. May be INDEX_NONE if Owner doesn't have the property. */
 	int16 RoleIndex;
@@ -1925,6 +1904,12 @@ private:
 	/** Converts a relative handle to the appropriate index into the Cmds array */
 	TArray<FHandleToCmdIndex> BaseHandleToCmdIndex;
 
+	/**
+	 * Special state tracking for Lifetime Custom Delta Properties.
+	 * Will only ever be valid if the Layout has Lifetime Custom Delta Properties.
+	 */
+	TUniquePtr<struct FLifetimeCustomDeltaState> LifetimeCustomPropertyState;
+
 	/** UClass, UStruct, or UFunction that this FRepLayout represents.*/
 	UStruct* Owner;
 
