Make sure NullRHI is correctly handled when using compact bulk data

diff --git a/Engine/Source/Runtime/Engine/Private/ContentStreaming.cpp b/Engine/Source/Runtime/Engine/Private/ContentStreaming.cpp
--- a/Engine/Source/Runtime/Engine/Private/ContentStreaming.cpp
+++ b/Engine/Source/Runtime/Engine/Private/ContentStreaming.cpp
@@ -744,7 +744,7 @@ FStreamingManagerCollection::FStreamingManagerCollection()
 {
 #if PLATFORM_SUPPORTS_TEXTURE_STREAMING
 	// Disable texture streaming if that was requested (needs to happen before the call to ProcessNewlyLoadedUObjects, as that can load textures)
-	if( !TEXTURE2DMIPMAP_USE_COMPACT_BULKDATA && FParse::Param( FCommandLine::Get(), TEXT( "NoTextureStreaming" ) ) )
+	if( FParse::Param( FCommandLine::Get(), TEXT( "NoTextureStreaming" ) ) )
 	{
 		CVarSetTextureStreaming.AsVariable()->Set(0, ECVF_SetByCommandline);
 	}
@@ -1126,31 +1126,26 @@ void FStreamingManagerCollection::AddOrRemoveTextureStreamingManagerIfNeeded(boo
 #if PLATFORM_SUPPORTS_TEXTURE_STREAMING
 	{
 		bUseTextureStreaming = CVarSetTextureStreaming.GetValueOnGameThread() != 0;
-		if (!bUseTextureStreaming && TEXTURE2DMIPMAP_USE_COMPACT_BULKDATA)
-		{
-			static bool bWarned;
-			if (!bWarned)
-			{
-				UE_LOG(LogContentStreaming, Warning, TEXT("Texture streaming cannot be disabled when FTexture2DMipMap::FCompactBulkData is used"));
-				bWarned = true;
-			}
-			CVarSetTextureStreaming.AsVariable()->Set(1, ECVF_SetByCode);
-			bUseTextureStreaming = true;
-		}
 	}
 
 	if( !GRHISupportsTextureStreaming || IsRunningDedicatedServer() )
 	{
-		if (TEXTURE2DMIPMAP_USE_COMPACT_BULKDATA)
-		{
-			UE_LOG(LogContentStreaming, Fatal, TEXT("RHI or target doesn't support texture streaming but cannot disable because TEXTURE2DMIPMAP_USE_COMPACT_BULKDATA is 1"));
-		}
-
 		bUseTextureStreaming = false;
 
 		// some code relies on r.TextureStreaming so we're going to disable it here to reflect the hardware capabilities and system needs
 		CVarSetTextureStreaming.AsVariable()->Set(0, ECVF_SetByCode);
 	}
+	else if (!bUseTextureStreaming && TEXTURE2DMIPMAP_USE_COMPACT_BULKDATA)
+	{
+		static bool bWarned;
+		if (!bWarned)
+		{
+			UE_LOG(LogContentStreaming, Warning, TEXT("Texture streaming cannot be disabled when FTexture2DMipMap::FCompactBulkData is used"));
+			bWarned = true;
+		}
+		CVarSetTextureStreaming.AsVariable()->Set(1, ECVF_SetByCode);
+		bUseTextureStreaming = true;
+	}
 #endif
 
 	if ( bUseTextureStreaming )
