Fix linker errors with FMeshMaterialShader::ValidateAfterBind() and inlining disabled. Since the compiler may ignore the FORCEINLINE macro, we have to ensure that the function can be called with static linkage.

#rnx

diff --git a/Engine/Source/Runtime/Engine/Private/Materials/MeshMaterialShader.cpp b/Engine/Source/Runtime/Engine/Private/Materials/MeshMaterialShader.cpp
--- a/Engine/Source/Runtime/Engine/Private/Materials/MeshMaterialShader.cpp
+++ b/Engine/Source/Runtime/Engine/Private/Materials/MeshMaterialShader.cpp
@@ -28,13 +28,6 @@ static inline bool ShouldCacheMeshShader(const FMeshMaterialShaderType* ShaderTy
 		InVertexFactoryType->ShouldCache(Platform, Material, ShaderType);
 }
 
-#if PLATFORM_WINDOWS && defined(__clang__)
-void FMeshMaterialShader::ValidateAfterBind()
-{
-	checkfSlow(PassUniformBuffer.IsInitialized(), TEXT("FMeshMaterialShader must bind a pass uniform buffer, even if it is just FSceneTexturesUniformParameters: %s"), GetType()->GetName());
-}
-#endif
-
 /**
  * Enqueues a compilation for a new shader of this type.
  * @param Platform - The platform to compile for.
@@ -173,7 +166,7 @@ FShader* FMeshMaterialShaderType::FinishCompileShader(
 	if (!Shader)
 	{
 		Shader = (*ConstructCompiledRef)(CompiledShaderInitializerType(this, CurrentJob.Output, Resource, UniformExpressionSet, MaterialShaderMapHash, InDebugDescription, ShaderPipelineType, CurrentJob.VFType));
-		((FMeshMaterialShader*)Shader)->ValidateAfterBind();
+		ValidateAfterBind((FMeshMaterialShader*)Shader);
 		CurrentJob.Output.ParameterMap.VerifyBindingsAreComplete(GetName(), CurrentJob.Output.Target, CurrentJob.VFType);
 	}
 
diff --git a/Engine/Source/Runtime/Renderer/Public/MeshMaterialShader.h b/Engine/Source/Runtime/Renderer/Public/MeshMaterialShader.h
--- a/Engine/Source/Runtime/Renderer/Public/MeshMaterialShader.h
+++ b/Engine/Source/Runtime/Renderer/Public/MeshMaterialShader.h
@@ -46,15 +46,11 @@ public:
 		return true;
 	}
 
-	// Clang treats FORCEINLINE as adivsory, and will not inline it on debug builds. Since Engine does not depend on the Renderer module, it fails to link against it.
-#if PLATFORM_WINDOWS && defined(__clang__)
-	void ValidateAfterBind();
-#else
-	FORCEINLINE void ValidateAfterBind()
+	// Declared as a friend, so that it can be called from other modules via static linkage, even if the compiler doesn't inline it.
+	FORCEINLINE friend void ValidateAfterBind(FMeshMaterialShader* Shader)
 	{
-		checkfSlow(PassUniformBuffer.IsInitialized(), TEXT("FMeshMaterialShader must bind a pass uniform buffer, even if it is just FSceneTexturesUniformParameters: %s"), GetType()->GetName());
+		checkfSlow(Shader->PassUniformBuffer.IsInitialized(), TEXT("FMeshMaterialShader must bind a pass uniform buffer, even if it is just FSceneTexturesUniformParameters: %s"), Shader->GetType()->GetName());
 	}
-#endif
 
 	void GetShaderBindings(
 		const FScene* Scene,
