Stop using sparse delegate storage temporarily while an issue with binding is resolved
#rnx

diff --git a/Engine/Source/Runtime/Engine/Classes/Components/ActorComponent.h b/Engine/Source/Runtime/Engine/Classes/Components/ActorComponent.h
--- a/Engine/Source/Runtime/Engine/Classes/Components/ActorComponent.h
+++ b/Engine/Source/Runtime/Engine/Classes/Components/ActorComponent.h
@@ -69,8 +69,8 @@ FORCEINLINE EUpdateTransformFlags SkipPhysicsToEnum(bool bSkipPhysics){ return b
 
 class UActorComponent;
 
-DECLARE_DYNAMIC_MULTICAST_SPARSE_DELEGATE_TwoParams(FActorComponentActivatedSignature, UActorComponent, OnComponentActivated, UActorComponent*, Component, bool, bReset);
-DECLARE_DYNAMIC_MULTICAST_SPARSE_DELEGATE_OneParam(FActorComponentDeactivateSignature, UActorComponent, OnComponentDeactivated, UActorComponent*, Component);
+DECLARE_DYNAMIC_MULTICAST_DELEGATE_TwoParams(FActorComponentActivatedSignature, UActorComponent*, Component, bool, bReset);
+DECLARE_DYNAMIC_MULTICAST_DELEGATE_OneParam(FActorComponentDeactivateSignature, UActorComponent*, Component);
 
 DECLARE_MULTICAST_DELEGATE_OneParam(FActorComponentGlobalCreatePhysicsSignature, UActorComponent*);
 DECLARE_MULTICAST_DELEGATE_OneParam(FActorComponentGlobalDestroyPhysicsSignature, UActorComponent*);
diff --git a/Engine/Source/Runtime/Engine/Classes/Components/SceneComponent.h b/Engine/Source/Runtime/Engine/Classes/Components/SceneComponent.h
--- a/Engine/Source/Runtime/Engine/Classes/Components/SceneComponent.h
+++ b/Engine/Source/Runtime/Engine/Classes/Components/SceneComponent.h
@@ -94,7 +94,7 @@ FORCEINLINE EMoveComponentFlags operator&(EMoveComponentFlags Arg1,EMoveComponen
 FORCEINLINE void operator&=(EMoveComponentFlags& Dest,EMoveComponentFlags Arg)					{ Dest = EMoveComponentFlags(Dest & Arg); }
 FORCEINLINE void operator|=(EMoveComponentFlags& Dest,EMoveComponentFlags Arg)					{ Dest = EMoveComponentFlags(Dest | Arg); }
 
-DECLARE_DYNAMIC_MULTICAST_SPARSE_DELEGATE_OneParam(FPhysicsVolumeChanged, USceneComponent, PhysicsVolumeChangedDelegate, class APhysicsVolume*, NewVolume);
+DECLARE_DYNAMIC_MULTICAST_DELEGATE_OneParam(FPhysicsVolumeChanged, class APhysicsVolume*, NewVolume);
 DECLARE_EVENT_ThreeParams(USceneComponent, FTransformUpdated, USceneComponent* /*UpdatedComponent*/, EUpdateTransformFlags /*UpdateTransformFlags*/, ETeleportType /*Teleport*/);
 
 /**
diff --git a/Engine/Source/Runtime/Engine/Classes/GameFramework/Actor.h b/Engine/Source/Runtime/Engine/Classes/GameFramework/Actor.h
--- a/Engine/Source/Runtime/Engine/Classes/GameFramework/Actor.h
+++ b/Engine/Source/Runtime/Engine/Classes/GameFramework/Actor.h
@@ -36,24 +36,24 @@ struct FNetworkObjectInfo;
 ENGINE_API DECLARE_LOG_CATEGORY_EXTERN(LogActor, Log, Warning);
 
 // Delegate signatures
-DECLARE_DYNAMIC_MULTICAST_SPARSE_DELEGATE_FiveParams( FTakeAnyDamageSignature, AActor, OnTakeAnyDamage, AActor*, DamagedActor, float, Damage, const class UDamageType*, DamageType, class AController*, InstigatedBy, AActor*, DamageCauser );
-DECLARE_DYNAMIC_MULTICAST_SPARSE_DELEGATE_NineParams( FTakePointDamageSignature, AActor, OnTakePointDamage, AActor*, DamagedActor, float, Damage, class AController*, InstigatedBy, FVector, HitLocation, class UPrimitiveComponent*, FHitComponent, FName, BoneName, FVector, ShotFromDirection, const class UDamageType*, DamageType, AActor*, DamageCauser );
-DECLARE_DYNAMIC_MULTICAST_SPARSE_DELEGATE_SevenParams( FTakeRadialDamageSignature, AActor, OnTakeRadialDamage, AActor*, DamagedActor, float, Damage, const class UDamageType*, DamageType, FVector, Origin, FHitResult, HitInfo, class AController*, InstigatedBy, AActor*, DamageCauser );
-DECLARE_DYNAMIC_MULTICAST_SPARSE_DELEGATE_TwoParams( FActorBeginOverlapSignature, AActor, OnActorBeginOverlap, AActor*, OverlappedActor, AActor*, OtherActor );
-DECLARE_DYNAMIC_MULTICAST_SPARSE_DELEGATE_TwoParams( FActorEndOverlapSignature, AActor, OnActorEndOverlap, AActor*, OverlappedActor, AActor*, OtherActor );
-DECLARE_DYNAMIC_MULTICAST_SPARSE_DELEGATE_FourParams( FActorHitSignature, AActor, OnActorHit, AActor*, SelfActor, AActor*, OtherActor, FVector, NormalImpulse, const FHitResult&, Hit );
-
-DECLARE_DYNAMIC_MULTICAST_SPARSE_DELEGATE_OneParam( FActorBeginCursorOverSignature, AActor, OnBeginCursorOver, AActor*, TouchedActor );
-DECLARE_DYNAMIC_MULTICAST_SPARSE_DELEGATE_OneParam( FActorEndCursorOverSignature, AActor, OnEndCursorOver, AActor*, TouchedActor );
-DECLARE_DYNAMIC_MULTICAST_SPARSE_DELEGATE_TwoParams( FActorOnClickedSignature, AActor, OnClicked, AActor*, TouchedActor , FKey, ButtonPressed );
-DECLARE_DYNAMIC_MULTICAST_SPARSE_DELEGATE_TwoParams( FActorOnReleasedSignature, AActor, OnReleased, AActor*, TouchedActor , FKey, ButtonReleased );
-DECLARE_DYNAMIC_MULTICAST_SPARSE_DELEGATE_TwoParams( FActorOnInputTouchBeginSignature, AActor, OnInputTouchBegin, ETouchIndex::Type, FingerIndex, AActor*, TouchedActor );
-DECLARE_DYNAMIC_MULTICAST_SPARSE_DELEGATE_TwoParams( FActorOnInputTouchEndSignature, AActor, OnInputTouchEnd, ETouchIndex::Type, FingerIndex, AActor*, TouchedActor );
-DECLARE_DYNAMIC_MULTICAST_SPARSE_DELEGATE_TwoParams( FActorBeginTouchOverSignature, AActor, OnInputTouchEnter, ETouchIndex::Type, FingerIndex, AActor*, TouchedActor );
-DECLARE_DYNAMIC_MULTICAST_SPARSE_DELEGATE_TwoParams( FActorEndTouchOverSignature, AActor, OnInputTouchLeave, ETouchIndex::Type, FingerIndex, AActor*, TouchedActor );
-
-DECLARE_DYNAMIC_MULTICAST_SPARSE_DELEGATE_OneParam(FActorDestroyedSignature, AActor, OnDestroyed, AActor*, DestroyedActor );
-DECLARE_DYNAMIC_MULTICAST_SPARSE_DELEGATE_TwoParams(FActorEndPlaySignature, AActor, OnEndPlay, AActor*, Actor , EEndPlayReason::Type, EndPlayReason);
+DECLARE_DYNAMIC_MULTICAST_DELEGATE_FiveParams(FTakeAnyDamageSignature, AActor*, DamagedActor, float, Damage, const class UDamageType*, DamageType, class AController*, InstigatedBy, AActor*, DamageCauser);
+DECLARE_DYNAMIC_MULTICAST_DELEGATE_NineParams(FTakePointDamageSignature, AActor*, DamagedActor, float, Damage, class AController*, InstigatedBy, FVector, HitLocation, class UPrimitiveComponent*, FHitComponent, FName, BoneName, FVector, ShotFromDirection, const class UDamageType*, DamageType, AActor*, DamageCauser);
+DECLARE_DYNAMIC_MULTICAST_DELEGATE_SevenParams(FTakeRadialDamageSignature, AActor*, DamagedActor, float, Damage, const class UDamageType*, DamageType, FVector, Origin, FHitResult, HitInfo, class AController*, InstigatedBy, AActor*, DamageCauser);
+DECLARE_DYNAMIC_MULTICAST_DELEGATE_TwoParams(FActorBeginOverlapSignature, AActor*, OverlappedActor, AActor*, OtherActor);
+DECLARE_DYNAMIC_MULTICAST_DELEGATE_TwoParams(FActorEndOverlapSignature, AActor*, OverlappedActor, AActor*, OtherActor);
+DECLARE_DYNAMIC_MULTICAST_DELEGATE_FourParams(FActorHitSignature, AActor*, SelfActor, AActor*, OtherActor, FVector, NormalImpulse, const FHitResult&, Hit);
+
+DECLARE_DYNAMIC_MULTICAST_DELEGATE_OneParam(FActorBeginCursorOverSignature, AActor*, TouchedActor);
+DECLARE_DYNAMIC_MULTICAST_DELEGATE_OneParam(FActorEndCursorOverSignature, AActor*, TouchedActor);
+DECLARE_DYNAMIC_MULTICAST_DELEGATE_TwoParams(FActorOnClickedSignature, AActor*, TouchedActor, FKey, ButtonPressed);
+DECLARE_DYNAMIC_MULTICAST_DELEGATE_TwoParams(FActorOnReleasedSignature, AActor*, TouchedActor, FKey, ButtonReleased);
+DECLARE_DYNAMIC_MULTICAST_DELEGATE_TwoParams(FActorOnInputTouchBeginSignature, ETouchIndex::Type, FingerIndex, AActor*, TouchedActor);
+DECLARE_DYNAMIC_MULTICAST_DELEGATE_TwoParams(FActorOnInputTouchEndSignature, ETouchIndex::Type, FingerIndex, AActor*, TouchedActor);
+DECLARE_DYNAMIC_MULTICAST_DELEGATE_TwoParams(FActorBeginTouchOverSignature, ETouchIndex::Type, FingerIndex, AActor*, TouchedActor);
+DECLARE_DYNAMIC_MULTICAST_DELEGATE_TwoParams(FActorEndTouchOverSignature, ETouchIndex::Type, FingerIndex, AActor*, TouchedActor);
+
+DECLARE_DYNAMIC_MULTICAST_DELEGATE_OneParam(FActorDestroyedSignature, AActor*, DestroyedActor);
+DECLARE_DYNAMIC_MULTICAST_DELEGATE_TwoParams(FActorEndPlaySignature, AActor*, Actor, EEndPlayReason::Type, EndPlayReason);
 
 DECLARE_DELEGATE_SixParams(FMakeNoiseDelegate, AActor*, float /*Loudness*/, class APawn*, const FVector&, float /*MaxRange*/, FName /*Tag*/);
 
