Adding a lot more logging when tearing down a world, and when adding and removing UI from the viewport so that we can better trace of what the engine is doing when preloading and loading worlds.

diff --git a/Engine/Source/Editor/UnrealEd/Private/PlayLevel.cpp b/Engine/Source/Editor/UnrealEd/Private/PlayLevel.cpp
--- a/Engine/Source/Editor/UnrealEd/Private/PlayLevel.cpp
+++ b/Engine/Source/Editor/UnrealEd/Private/PlayLevel.cpp
@@ -578,7 +578,7 @@ void UEditorEngine::TeardownPlaySession(FWorldContext& PieWorldContext)
 {
 	check(PieWorldContext.WorldType == EWorldType::PIE);
 	PlayWorld = PieWorldContext.World();
-	PlayWorld->bIsTearingDown = true;
+	PlayWorld->BeginTearingDown();
 
 	if (!PieWorldContext.RunAsDedicated)
 	{
diff --git a/Engine/Source/Runtime/Engine/Classes/Engine/GameViewportClient.h b/Engine/Source/Runtime/Engine/Classes/Engine/GameViewportClient.h
--- a/Engine/Source/Runtime/Engine/Classes/Engine/GameViewportClient.h
+++ b/Engine/Source/Runtime/Engine/Classes/Engine/GameViewportClient.h
@@ -866,7 +866,7 @@ private:
 	TMap<EMouseCursor::Type, TSharedPtr<FHardwareCursor>> HardwareCursors;
 
 	/** Map of Software Cursor Widgets*/
-	TMap<EMouseCursor::Type, TSharedRef<SWidget>> CursorWidgets;
+	TMap<EMouseCursor::Type, TSharedPtr<SWidget>> CursorWidgets;
 
 	/** Controls if the Map of Software Cursor Widgets is used */
 	bool bUseSoftwareCursorWidgets;
diff --git a/Engine/Source/Runtime/Engine/Classes/Engine/World.h b/Engine/Source/Runtime/Engine/Classes/Engine/World.h
--- a/Engine/Source/Runtime/Engine/Classes/Engine/World.h
+++ b/Engine/Source/Runtime/Engine/Classes/Engine/World.h
@@ -2763,6 +2763,9 @@ public:
 	/** Handle Exec/Console Commands related to the World */
 	bool Exec( UWorld* InWorld, const TCHAR* Cmd, FOutputDevice& Ar=*GLog );
 
+	/** Mark the world as being torn down */
+	void BeginTearingDown();
+
 private:
 	/** Utility function to handle Exec/Console Commands related to the Trace Tags */
 	bool HandleTraceTagCommand( const TCHAR* Cmd, FOutputDevice& Ar );
diff --git a/Engine/Source/Runtime/Engine/Private/GameEngine.cpp b/Engine/Source/Runtime/Engine/Private/GameEngine.cpp
--- a/Engine/Source/Runtime/Engine/Private/GameEngine.cpp
+++ b/Engine/Source/Runtime/Engine/Private/GameEngine.cpp
@@ -1102,7 +1102,7 @@ void UGameEngine::PreExit()
 		UWorld* const World = WorldList[WorldIndex].World();
 		if ( World != NULL )
 		{
-			World->bIsTearingDown = true;
+			World->BeginTearingDown();
 
 			// Cancel any pending connection to a server
 			CancelPending(World);
diff --git a/Engine/Source/Runtime/Engine/Private/GameViewportClient.cpp b/Engine/Source/Runtime/Engine/Private/GameViewportClient.cpp
--- a/Engine/Source/Runtime/Engine/Private/GameViewportClient.cpp
+++ b/Engine/Source/Runtime/Engine/Private/GameViewportClient.cpp
@@ -834,17 +834,13 @@ EMouseCursor::Type UGameViewportClient::GetCursor(FViewport* InViewport, int32 X
 
 void UGameViewportClient::SetVirtualCursorWidget(EMouseCursor::Type Cursor, UUserWidget* UserWidget)
 {
-	if (ensure(UserWidget))
+	TSharedPtr<SWidget>& ExistingWidget = CursorWidgets.FindOrAdd(Cursor);
+	TSharedPtr<SWidget> NewWidget = UserWidget ? UserWidget->TakeWidget() : TSharedPtr<SWidget>();
+	if (NewWidget != ExistingWidget)
 	{
-		if (CursorWidgets.Contains(Cursor))
-		{
-			TSharedRef<SWidget>* Widget = CursorWidgets.Find(Cursor);
-			(*Widget) = UserWidget->TakeWidget();
-		}
-		else
-		{
-			CursorWidgets.Add(Cursor, UserWidget->TakeWidget());
-		}
+		// Pure safety
+		ExistingWidget.Reset();
+		ExistingWidget = NewWidget;
 	}
 }
 
@@ -852,8 +848,7 @@ void UGameViewportClient::AddSoftwareCursor(EMouseCursor::Type Cursor, const FSo
 {
 	if (CursorClass.IsValid())
 	{
-		UClass* Class = CursorClass.TryLoadClass<UUserWidget>();
-		if (Class)
+		if (UClass* Class = CursorClass.TryLoadClass<UUserWidget>())
 		{
 			UUserWidget* UserWidget = CreateWidget(GetGameInstance(), Class);
 			AddCursorWidget(Cursor, UserWidget);
@@ -888,18 +883,19 @@ TOptional<TSharedRef<SWidget>> UGameViewportClient::MapCursor(FViewport* InViewp
 	{
 		if (CursorReply.GetCursorType() != EMouseCursor::None)
 		{
-			const TSharedRef<SWidget>* CursorWidgetPtr = CursorWidgets.Find(CursorReply.GetCursorType());
+			const TSharedPtr<SWidget> CursorWidgetPtr = CursorWidgets.FindRef(CursorReply.GetCursorType());
 
-			if (CursorWidgetPtr != nullptr)
+			if (CursorWidgetPtr.IsValid())
 			{
-				return *CursorWidgetPtr;
+				return CursorWidgetPtr.ToSharedRef();
 			}
 			else
 			{
-				UE_LOG(LogPlayerManagement, Warning, TEXT("UGameViewportClient::MapCursor: Could not find cursor to map to %d."),int32(CursorReply.GetCursorType()));
+				UE_LOG(LogPlayerManagement, Warning, TEXT("UGameViewportClient::MapCursor: Could not find cursor to map to %d."), int32(CursorReply.GetCursorType()));
 			}
 		}
 	}
+
 	return TOptional<TSharedRef<SWidget>>();
 }
 
diff --git a/Engine/Source/Runtime/Engine/Private/UnrealEngine.cpp b/Engine/Source/Runtime/Engine/Private/UnrealEngine.cpp
--- a/Engine/Source/Runtime/Engine/Private/UnrealEngine.cpp
+++ b/Engine/Source/Runtime/Engine/Private/UnrealEngine.cpp
@@ -12018,13 +12018,14 @@ bool UEngine::LoadMap( FWorldContext& WorldContext, FURL URL, class UPendingNetG
 
 	// send a callback message
 	FCoreUObjectDelegates::PreLoadMap.Broadcast(URL.Map);
+
 	// make sure there is a matching PostLoadMap() no matter how we exit
 	struct FPostLoadMapCaller
 	{
-		bool bCalled;
 		FPostLoadMapCaller()
 			: bCalled(false)
 		{}
+
 		~FPostLoadMapCaller()
 		{
 			if (!bCalled)
@@ -12032,6 +12033,19 @@ bool UEngine::LoadMap( FWorldContext& WorldContext, FURL URL, class UPendingNetG
 				FCoreUObjectDelegates::PostLoadMapWithWorld.Broadcast(nullptr);
 			}
 		}
+
+		void Broadcast(UWorld* World)
+		{
+			if (ensure(!bCalled))
+			{
+				bCalled = true;
+				FCoreUObjectDelegates::PostLoadMapWithWorld.Broadcast(World);
+			}
+		}
+
+	private:
+		bool bCalled;
+
 	} PostLoadMapCaller;
 
 	// Cancel any pending texture streaming requests.  This avoids a significant delay on consoles 
@@ -12074,7 +12088,7 @@ bool UEngine::LoadMap( FWorldContext& WorldContext, FURL URL, class UPendingNetG
 	// Unload the current world
 	if( WorldContext.World() )
 	{
-		WorldContext.World()->bIsTearingDown = true;
+		WorldContext.World()->BeginTearingDown();
 
 		if(!URL.HasOption(TEXT("quiet")) )
 		{
@@ -12506,8 +12520,7 @@ bool UEngine::LoadMap( FWorldContext& WorldContext, FURL URL, class UPendingNetG
 	WorldContext.World()->BeginPlay();
 
 	// send a callback message
-	PostLoadMapCaller.bCalled = true;
-	FCoreUObjectDelegates::PostLoadMapWithWorld.Broadcast(WorldContext.World());
+	PostLoadMapCaller.Broadcast(WorldContext.World());
 
 	WorldContext.World()->bWorldWasLoadedThisTick = true;
 
diff --git a/Engine/Source/Runtime/Engine/Private/World.cpp b/Engine/Source/Runtime/Engine/Private/World.cpp
--- a/Engine/Source/Runtime/Engine/Private/World.cpp
+++ b/Engine/Source/Runtime/Engine/Private/World.cpp
@@ -2408,6 +2408,12 @@ void UWorld::AddToWorld( ULevel* Level, const FTransform& LevelTransform, bool b
 #endif // PERF_TRACK_DETAILED_ASYNC_STATS
 }
 
+void UWorld::BeginTearingDown()
+{
+	bIsTearingDown = true;
+	UE_LOG(LogWorld, Log, TEXT("BeginTearingDown for %s"), *GetOutermost()->GetName());
+}
+
 void UWorld::RemoveFromWorld( ULevel* Level, bool bAllowIncrementalRemoval )
 {
 	SCOPE_CYCLE_COUNTER(STAT_RemoveFromWorldTime);
@@ -2422,6 +2428,8 @@ void UWorld::RemoveFromWorld( ULevel* Level, bool bAllowIncrementalRemoval )
 	// If the level may be removed incrementally then there must also be no level pending visibility
 	if ( ((CurrentLevelPendingVisibility == nullptr) || (!bAllowIncrementalRemoval && (CurrentLevelPendingVisibility != Level))) && Level->bIsVisible )
 	{
+		UE_LOG(LogWorld, Log, TEXT("UWorld::RemoveFromWorld for %s"), *Level->GetOutermost()->GetName());
+
 		// Keep track of timing.
 		double StartTime = FPlatformTime::Seconds();	
 
@@ -5835,8 +5843,7 @@ UWorld* FSeamlessTravelHandler::Tick()
 				CurrentWorld->GetGameState()->SeamlessTravelTransitionCheckpoint(!bSwitchedToDefaultMap);
 			}
 			
-
-			CurrentWorld->bIsTearingDown = true;
+			CurrentWorld->BeginTearingDown();
 
 			// If it's not still playing, destroy the demo net driver before we start renaming actors.
 			if ( CurrentWorld->DemoNetDriver && !CurrentWorld->DemoNetDriver->IsPlaying() && !CurrentWorld->DemoNetDriver->bRecordMapChanges)
diff --git a/Engine/Source/Runtime/UMG/Private/Components/WidgetComponent.cpp b/Engine/Source/Runtime/UMG/Private/Components/WidgetComponent.cpp
--- a/Engine/Source/Runtime/UMG/Private/Components/WidgetComponent.cpp
+++ b/Engine/Source/Runtime/UMG/Private/Components/WidgetComponent.cpp
@@ -1175,18 +1175,21 @@ void UWidgetComponent::RemoveWidgetFromScreen()
 	{
 		bAddedToScreen = false;
 
-		if ( UGameViewportClient* ViewportClient = GetWorld()->GetGameViewport() )
+		if (UWorld* World = GetWorld())
 		{
-			TSharedPtr<IGameLayerManager> LayerManager = ViewportClient->GetGameLayerManager();
-			if ( LayerManager.IsValid() )
+			if (UGameViewportClient* ViewportClient = World->GetGameViewport())
 			{
 				if (ULocalPlayer* TargetPlayer = GetOwnerPlayer())
 				{
-					TSharedPtr<IGameLayer> Layer = LayerManager->FindLayerForPlayer(TargetPlayer, SharedLayerName);
-					if (Layer.IsValid())
+					TSharedPtr<IGameLayerManager> LayerManager = ViewportClient->GetGameLayerManager();
+					if (LayerManager.IsValid())
 					{
-						TSharedPtr<FWorldWidgetScreenLayer> ScreenLayer = StaticCastSharedPtr<FWorldWidgetScreenLayer>(Layer);
-						ScreenLayer->RemoveComponent(this);
+						TSharedPtr<IGameLayer> Layer = LayerManager->FindLayerForPlayer(TargetPlayer, SharedLayerName);
+						if (Layer.IsValid())
+						{
+							TSharedPtr<FWorldWidgetScreenLayer> ScreenLayer = StaticCastSharedPtr<FWorldWidgetScreenLayer>(Layer);
+							ScreenLayer->RemoveComponent(this);
+						}
 					}
 				}
 			}
diff --git a/Engine/Source/Runtime/UMG/Private/UserWidget.cpp b/Engine/Source/Runtime/UMG/Private/UserWidget.cpp
--- a/Engine/Source/Runtime/UMG/Private/UserWidget.cpp
+++ b/Engine/Source/Runtime/UMG/Private/UserWidget.cpp
@@ -2170,6 +2170,15 @@ UUserWidget* UUserWidget::CreateInstanceInternal(UObject* Outer, TSubclassOf<UUs
 	}
 #endif
 
+#if !UE_BUILD_SHIPPING
+	// Check if the world is being torn down before we create a widget for it.
+	if (World)
+	{
+		// Look for indications that widgets are being created for a dead and dying world.
+		ensureMsgf(!World->bIsTearingDown, TEXT("Widget Class %s - Attempting to be created while tearing down the world."), *UserWidgetClass->GetName());
+	}
+#endif
+
 	if (!Outer)
 	{
 		FMessageLog("PIE").Error(FText::Format(LOCTEXT("OuterNull", "Unable to create the widget {0}, no outer provided."), FText::FromName(UserWidgetClass->GetFName())));
