Reduce log level of message indicating that an actor was destroyed during UActorChannel::ProcessBunch, it's spammy and not indicative of a bug.

diff --git a/Engine/Source/Runtime/Engine/Private/DataChannel.cpp b/Engine/Source/Runtime/Engine/Private/DataChannel.cpp
--- a/Engine/Source/Runtime/Engine/Private/DataChannel.cpp
+++ b/Engine/Source/Runtime/Engine/Private/DataChannel.cpp
@@ -2585,7 +2585,7 @@ void UActorChannel::ProcessBunch( FInBunch & Bunch )
 		// note that this is a legitimate occurrence, particularly on client to server RPCs
 		if ( !Actor || Actor->IsPendingKill() )
 		{
-			UE_LOG( LogNet, Verbose, TEXT( "UActorChannel::ProcessBunch: Actor was destroyed during Replicator.ReceivedBunch processing" ) );
+			UE_LOG( LogNet, VeryVerbose, TEXT( "UActorChannel::ProcessBunch: Actor was destroyed during Replicator.ReceivedBunch processing" ) );
 			// If we lose the actor on this channel, we can no longer process bunches, so consider this channel broken
 			Broken = 1;		
 			break;
