Enable Net.SupportFastArrayDelta by default.

diff --git a/Engine/Source/Runtime/Engine/Private/DataReplication.cpp b/Engine/Source/Runtime/Engine/Private/DataReplication.cpp
--- a/Engine/Source/Runtime/Engine/Private/DataReplication.cpp
+++ b/Engine/Source/Runtime/Engine/Private/DataReplication.cpp
@@ -42,7 +42,7 @@ static FAutoConsoleVariableRef CVarNetRPCDebug(
 	TEXT(" 1: Print bunches as they are sent."),
 	ECVF_Default);
 
-int32 GSupportsFastArrayDelta = 0;
+int32 GSupportsFastArrayDelta = 1;
 static FAutoConsoleVariableRef CVarSupportsFastArrayDelta(
 	TEXT("net.SupportFastArrayDelta"),
 	GSupportsFastArrayDelta,
