Adding some logging for timers in the timer manager to help track down server spikes in the event that we spend a long time firing timer delegates, driven by the cvar TimerManager.DumpTimerLogsThreshold which defaults to 1000 ms and is defined as 15ms in the dedicated server engine ini file.

diff --git a/Engine/Source/Runtime/Engine/Private/TimerManager.cpp b/Engine/Source/Runtime/Engine/Private/TimerManager.cpp
--- a/Engine/Source/Runtime/Engine/Private/TimerManager.cpp
+++ b/Engine/Source/Runtime/Engine/Private/TimerManager.cpp
@@ -23,6 +23,12 @@ CSV_DECLARE_CATEGORY_MODULE_EXTERN(CORE_API, Basic);
 /** Track the last assigned handle globally */
 uint64 FTimerManager::LastAssignedSerialNumber = 0;
 
+static float DumpTimerLogsThreshold = 1000.f;
+static FAutoConsoleVariableRef CVarDumpTimerLogsThreshold(
+	TEXT("TimerManager.DumpTimerLogsThreshold"), DumpTimerLogsThreshold,
+	TEXT("Threshold (in milliseconds) after which we log timer info to try and help track down spikes in the timer code."),
+	ECVF_Default);
+
 namespace
 {
 	void DescribeFTimerDataSafely(FOutputDevice& Ar, const FTimerData& Data)
@@ -154,7 +160,7 @@ FString FTimerUnifiedDelegate::ToString() const
 		FunctionName = NotBoundName;
 	}
 
-	return FString::Printf(TEXT("%s,%s,%s"), bDynDelegate ? TEXT("DELEGATE") : TEXT("DYN DELEGATE"), Object == nullptr ? TEXT("NO OBJ") : *Object->GetPathName(), *FunctionName.ToString());
+	return FString::Printf(TEXT("%s,%s,%s"), bDynDelegate ? TEXT("DYN DELEGATE") : TEXT("DELEGATE"), Object == nullptr ? TEXT("NO OBJ") : *Object->GetPathName(), *FunctionName.ToString());
 }
 
 // ---------------------------------
@@ -547,6 +553,8 @@ void FTimerManager::Tick(float DeltaTime)
 		return;
 	}
 
+	const double StartTime = FPlatformTime::Seconds();
+
 	InternalTime += DeltaTime;
 
 	UWorld* const OwningWorld = OwningGameInstance ? OwningGameInstance->GetWorld() : nullptr;
@@ -636,6 +644,14 @@ void FTimerManager::Tick(float DeltaTime)
 		}
 	}
 
+	// help us hunt down outliers that cause our timer manager times to spike.  Recommended that users set meaningful DumpTimerLogsThresholds in appropriate ini files if they are seeing spikes in the timer manager.
+	const double DeltaT = (FPlatformTime::Seconds() - StartTime) * 1000.f;
+	if (DeltaT >= DumpTimerLogsThreshold)
+	{
+		UE_LOG(LogEngine, Log, TEXT("TimerManager's time threshold of %.2fms exceeded with a deltaT of %.4f, dumping list of timers to help identify why it took so long to process the list of timers."), DumpTimerLogsThreshold, DeltaT);
+		ListTimers();
+	}
+
 	// Timer has been ticked.
 	LastTickedFrame = GFrameCounter;
 
