#ue4 - The server no longer forces a client movement correction if a server hitch caused it to think client movement was being delayed. Added a configurable hitch threshold (GameNetworkManager->ServerForcedUpdateHitchThreshold) and hitch cooldown (ServerForcedUpdateHitchCooldown) which allows time for the server to process buffered RPCs.

diff --git a/Engine/Config/BaseGame.ini b/Engine/Config/BaseGame.ini
--- a/Engine/Config/BaseGame.ini
+++ b/Engine/Config/BaseGame.ini
@@ -14,6 +14,8 @@ MAXNEARZEROVELOCITYSQUARED=9.0f
 CLIENTADJUSTUPDATECOST=180.0f
 MAXCLIENTUPDATEINTERVAL=0.25f
 MaxClientForcedUpdateDuration=1.0f
+ServerForcedUpdateHitchThreshold=0.150f
+ServerForcedUpdateHitchCooldown=0.100f
 MaxMoveDeltaTime=0.125f
 MaxClientSmoothingDeltaTime=0.50f
 ClientNetSendMoveDeltaTime=0.0166
diff --git a/Engine/Source/Runtime/Engine/Classes/GameFramework/GameNetworkManager.h b/Engine/Source/Runtime/Engine/Classes/GameFramework/GameNetworkManager.h
--- a/Engine/Source/Runtime/Engine/Classes/GameFramework/GameNetworkManager.h
+++ b/Engine/Source/Runtime/Engine/Classes/GameFramework/GameNetworkManager.h
@@ -132,6 +132,14 @@ class ENGINE_API AGameNetworkManager : public AInfo
 	/** MaxClientForcedUpdateDuration is the maximum time duration over which the server will force updates, after MAXCLIENTUPDATEINTERVAL is initially exceeded. */
 	UPROPERTY(GlobalConfig)
 	float MaxClientForcedUpdateDuration;
+	
+	/** Ignore forced client movement updates when server hitches for longer than this duration. */
+	UPROPERTY(GlobalConfig)
+	float ServerForcedUpdateHitchThreshold;
+
+	/** Ignore forced client movement updates when server hitch was detected within this amount of time in the past. */
+	UPROPERTY(GlobalConfig)
+	float ServerForcedUpdateHitchCooldown;
 
 	/** MaxMoveDeltaTime is the default maximum time delta of CharacterMovement ServerMoves. Should be less than or equal to MAXCLIENTUPDATEINTERVAL, otherwise server will interfere by forcing position updates. */
 	UPROPERTY(GlobalConfig)
diff --git a/Engine/Source/Runtime/Engine/Classes/GameFramework/PlayerController.h b/Engine/Source/Runtime/Engine/Classes/GameFramework/PlayerController.h
--- a/Engine/Source/Runtime/Engine/Classes/GameFramework/PlayerController.h
+++ b/Engine/Source/Runtime/Engine/Classes/GameFramework/PlayerController.h
@@ -1887,6 +1887,12 @@ protected:
 	/** Set the SpawnLocation for use when changing states or when there is no pawn or spectator. */
 	virtual void SetSpawnLocation(const FVector& NewLocation);
 
+	/** Last real time (undilated) recorded in TickActor() when checking for forced client movement updates. */
+	float LastMovementUpdateTime;
+
+	/** Last real time (undilated) a hitch was detected in TickActor() when checking for forced client movement updates. */
+	float LastMovementHitch;
+
 public:
 	/** Get the location used when initially created, or when changing states when there is no pawn or spectator. */
 	FVector GetSpawnLocation() const { return SpawnLocation; }
diff --git a/Engine/Source/Runtime/Engine/Private/GameNetworkManager.cpp b/Engine/Source/Runtime/Engine/Private/GameNetworkManager.cpp
--- a/Engine/Source/Runtime/Engine/Private/GameNetworkManager.cpp
+++ b/Engine/Source/Runtime/Engine/Private/GameNetworkManager.cpp
@@ -28,6 +28,8 @@ AGameNetworkManager::AGameNetworkManager(const FObjectInitializer& ObjectInitial
 	CLIENTADJUSTUPDATECOST = 180.0f;
 	MAXCLIENTUPDATEINTERVAL = 0.25f;
 	MaxClientForcedUpdateDuration=1.0f;
+	ServerForcedUpdateHitchThreshold = 0.150f;
+	ServerForcedUpdateHitchCooldown = 0.100f;
 	MaxMoveDeltaTime = 0.125f;
 	MaxClientSmoothingDeltaTime = 0.50f;
 	ClientNetSendMoveDeltaTime = 0.0166f;
diff --git a/Engine/Source/Runtime/Engine/Private/PlayerController.cpp b/Engine/Source/Runtime/Engine/Private/PlayerController.cpp
--- a/Engine/Source/Runtime/Engine/Private/PlayerController.cpp
+++ b/Engine/Source/Runtime/Engine/Private/PlayerController.cpp
@@ -103,6 +103,9 @@ APlayerController::APlayerController(const FObjectInitializer& ObjectInitializer
 	DefaultClickTraceChannel = ECollisionChannel::ECC_Visibility;
 	HitResultTraceDistance = 100000.f;
 
+	LastMovementUpdateTime = 0.f;
+	LastMovementHitch = 0.f;
+
 	bCinemaDisableInputMove = false;
 	bCinemaDisableInputLook = false;
 
@@ -4535,9 +4538,15 @@ void APlayerController::TickActor( float DeltaSeconds, ELevelTick TickType, FAct
 							}
 						}
 						
+						const float CurrentRealTime = World->GetRealTimeSeconds();
+						const bool bHitch = (CurrentRealTime - LastMovementUpdateTime) > GameNetworkManager->ServerForcedUpdateHitchThreshold && (LastMovementUpdateTime != 0);
+						LastMovementHitch = bHitch ? CurrentRealTime : LastMovementHitch;
+						const bool bRecentHitch = bHitch || (CurrentRealTime - LastMovementHitch < GameNetworkManager->ServerForcedUpdateHitchCooldown);
+						LastMovementUpdateTime = CurrentRealTime;
+
 						// Trigger forced update if allowed
-						if (ForcedUpdateInterval > 0.f && PawnTimeSinceUpdate > FMath::Max<float>(DeltaSeconds+0.06f, ForcedUpdateInterval * GetPawn()->GetActorTimeDilation()))
-						{						
+						if (!bRecentHitch && ForcedUpdateInterval > 0.f && PawnTimeSinceUpdate > FMath::Max<float>(DeltaSeconds+0.06f, ForcedUpdateInterval * GetPawn()->GetActorTimeDilation()))
+						{
 							//UE_LOG(LogPlayerController, Warning, TEXT("ForcedMovementTick. PawnTimeSinceUpdate: %f, DeltaSeconds: %f, DeltaSeconds+: %f"), PawnTimeSinceUpdate, DeltaSeconds, DeltaSeconds+0.06f);
 							const USkeletalMeshComponent* PawnMesh = GetPawn()->FindComponentByClass<USkeletalMeshComponent>();
 							if (!ServerData->bForcedUpdateDurationExceeded && (!PawnMesh || !PawnMesh->IsSimulatingPhysics()))
