Improve the apply skinning algorithm use to retarget the skinning between different meshes.

diff --git a/Engine/Source/Runtime/Engine/Public/Rendering/SkeletalMeshLODImporterData.cpp b/Engine/Source/Runtime/Engine/Public/Rendering/SkeletalMeshLODImporterData.cpp
--- a/Engine/Source/Runtime/Engine/Public/Rendering/SkeletalMeshLODImporterData.cpp
+++ b/Engine/Source/Runtime/Engine/Public/Rendering/SkeletalMeshLODImporterData.cpp
@@ -224,35 +224,55 @@ bool FSkeletalMeshImportData::ApplyRigToGeo(FSkeletalMeshImportData& Other)
 			TArray<FWedgeInfo> NearestWedges;
 			FVector SearchPosition = Points[NewVertexIndex];
 			OctreeQueryHelper.FindNearestWedgeIndexes(SearchPosition, NearestWedges);
+			//The best old wedge match is base on those weight ratio
+			const int32 UVWeightRatioIndex = 0;
+			const int32 NormalWeightRatioIndex = 1;
+			const float MatchWeightRatio[3] = { 0.99f, 0.01f };
 			if (NearestWedges.Num() > 0)
 			{
-				float MinDistance = MAX_FLT;
-				float MinNormalDiff = MAX_FLT;
 				int32 BestOldVertexIndex = INDEX_NONE;
+				float MaxUVDistance = 0.0f;
+				float MaxNormalDelta = 0.0f;
+				TArray<float> UvDistances;
+				UvDistances.Reserve(NearestWedges.Num());
+				TArray<float> NormalDeltas;
+				NormalDeltas.Reserve(NearestWedges.Num());
 				for (const FWedgeInfo& WedgeInfo : NearestWedges)
 				{
 					int32 OldWedgeIndex = WedgeInfo.WedgeIndex;
 					int32 OldVertexIndex = Other.Wedges[OldWedgeIndex].VertexIndex;
 					int32 OldFaceIndex = (OldWedgeIndex / 3);
 					int32 OldFaceCorner = (OldWedgeIndex % 3);
-					FVector OldNormal = Other.Faces[OldFaceIndex].TangentZ[OldFaceCorner];
-
-					const FVector& OtherPosition = Other.Points[OldVertexIndex];
-					float VectorDelta = FVector::DistSquared(OtherPosition, SearchPosition);
-					if (VectorDelta <= (MinDistance + KINDA_SMALL_NUMBER))
+					const FVector2D& OldUV = Other.Wedges[OldWedgeIndex].UVs[0];
+					const FVector& OldNormal = Other.Faces[OldFaceIndex].TangentZ[OldFaceCorner];
+					float UVDelta = FVector2D::DistSquared(CurWedgeUV, OldUV);
+					float NormalDelta = FMath::Abs(FMath::Acos(FVector::DotProduct(NewNormal, OldNormal)));
+					if (UVDelta > MaxUVDistance)
+					{
+						MaxUVDistance = UVDelta;
+					}
+					UvDistances.Add(UVDelta);
+					if (NormalDelta > MaxNormalDelta)
 					{
-						if (VectorDelta < MinDistance - KINDA_SMALL_NUMBER)
-						{
-							MinDistance = VectorDelta;
-							MinNormalDiff = MAX_FLT;
-						}
-						float AngleDiff = FMath::Abs(FMath::Acos(FVector::DotProduct(NewNormal, OldNormal)));
-						if (AngleDiff < MinNormalDiff)
-						{
-							MinNormalDiff = AngleDiff;
-							BestOldVertexIndex = OldVertexIndex;
-						}
+						MaxNormalDelta = NormalDelta;
 					}
+					NormalDeltas.Add(NormalDelta);
+				}
+				float BestContribution = 0.0f;
+				for (int32 NearestWedgeIndex = 0; NearestWedgeIndex < UvDistances.Num(); ++NearestWedgeIndex)
+				{
+					float Contribution = ((MaxUVDistance - UvDistances[NearestWedgeIndex])/MaxUVDistance)*MatchWeightRatio[UVWeightRatioIndex];
+					Contribution += ((MaxNormalDelta - NormalDeltas[NearestWedgeIndex]) / MaxNormalDelta)*MatchWeightRatio[NormalWeightRatioIndex];
+					if (Contribution > BestContribution)
+					{
+						BestContribution = Contribution;
+						BestOldVertexIndex = Other.Wedges[NearestWedges[NearestWedgeIndex].WedgeIndex].VertexIndex;
+					}
+				}
+				if (BestOldVertexIndex == INDEX_NONE)
+				{
+					//Use the first NearestWedges entry, we end up here because all NearestWedges entries all equals, so the ratio will be zero in such a case
+					BestOldVertexIndex = Other.Wedges[NearestWedges[0].WedgeIndex].VertexIndex;
 				}
 				OldToNewRemap[BestOldVertexIndex].AddUnique(NewVertexIndex);
 			}
@@ -611,7 +631,7 @@ void FOctreeQueryHelper::FindNearestWedgeIndexes(const FVector& SearchPosition,
 	float MinSquaredDistance = MAX_FLT;
 	OutNearestWedges.Empty();
 	
-	FVector Extend(1.0f);
+	FVector Extend(2.0f);
 	for (int i = 0; i < 2; ++i)
 	{
 		TWedgeInfoPosOctree::TConstIterator<> OctreeIter((*WedgePosOctree));
@@ -636,11 +656,8 @@ void FOctreeQueryHelper::FindNearestWedgeIndexes(const FVector& SearchPosition,
 			for (const FWedgeInfo& WedgeInfo : CurNode.GetElements())
 			{
 				float VectorDelta = FVector::DistSquared(SearchPosition, WedgeInfo.Position);
-				if (VectorDelta <= (MinSquaredDistance + SMALL_NUMBER))
-				{
-					MinSquaredDistance = FMath::Min(VectorDelta, MinSquaredDistance);
-					OutNearestWedges.Add(WedgeInfo);
-				}
+				MinSquaredDistance = FMath::Min(VectorDelta, MinSquaredDistance);
+				OutNearestWedges.Add(WedgeInfo);
 			}
 			OctreeIter.Advance();
 		}
