Adding new material expression: Pre-skinned local bounds. This provides the bounds of the base pose, prior to any transforms.

diff --git a/Engine/Shaders/Private/SceneData.ush b/Engine/Shaders/Private/SceneData.ush
--- a/Engine/Shaders/Private/SceneData.ush
+++ b/Engine/Shaders/Private/SceneData.ush
@@ -30,13 +30,14 @@ struct FPrimitiveSceneData
 	float4 NonUniformScale;
 	float3 LocalObjectBoundsMin;
 	float3 LocalObjectBoundsMax;
+	float3 PreSkinnedLocalBounds;
 	uint LightingChannelMask;
 	uint LightmapDataIndex;
 	int SingleCaptureIndex;
 };
 
 // Stride of a single primitive's data in float4's, must match C++
-#define PRIMITIVE_SCENE_DATA_STRIDE 26
+#define PRIMITIVE_SCENE_DATA_STRIDE 27
 
 // Fetch from scene primitive buffer
 FPrimitiveSceneData GetPrimitiveData(uint PrimitiveId) 
@@ -89,6 +90,8 @@ FPrimitiveSceneData GetPrimitiveData(uint PrimitiveId)
 
 	PrimitiveData.SingleCaptureIndex = asint(View.PrimitiveSceneData[PrimitiveBaseOffset + 25].x);
 
+	PrimitiveData.PreSkinnedLocalBounds = View.PrimitiveSceneData[PrimitiveBaseOffset + 26].xyz;
+
 	return PrimitiveData;
 }
 
diff --git a/Engine/Source/Developer/MaterialBaking/Private/ExportMaterialProxy.h b/Engine/Source/Developer/MaterialBaking/Private/ExportMaterialProxy.h
--- a/Engine/Source/Developer/MaterialBaking/Private/ExportMaterialProxy.h
+++ b/Engine/Source/Developer/MaterialBaking/Private/ExportMaterialProxy.h
@@ -106,6 +106,11 @@ struct FExportMaterialCompiler : public FProxyMaterialCompiler
 		return Compiler->ObjectBounds();
 	}
 
+	virtual int32 PreSkinnedLocalBounds() override
+	{
+		return Compiler->PreSkinnedLocalBounds();
+	}
+
 	virtual int32 CameraVector() override
 	{
 		return Compiler->Constant3(0.0f, 0.0f, 1.0f);
diff --git a/Engine/Source/Developer/MaterialUtilities/Private/MaterialUtilities.cpp b/Engine/Source/Developer/MaterialUtilities/Private/MaterialUtilities.cpp
--- a/Engine/Source/Developer/MaterialUtilities/Private/MaterialUtilities.cpp
+++ b/Engine/Source/Developer/MaterialUtilities/Private/MaterialUtilities.cpp
@@ -282,6 +282,11 @@ struct FExportMaterialCompiler : public FProxyMaterialCompiler
 		return Compiler->ObjectBounds();
 	}
 
+	virtual int32 PreSkinnedLocalBounds() override
+	{
+		return Compiler->PreSkinnedLocalBounds();
+	}
+
 	virtual int32 CameraVector() override
 	{
 		return Compiler->Constant3(0.0f, 0.0f, 1.0f);
diff --git a/Engine/Source/Editor/UnrealEd/Private/Lightmass/LightmassRender.cpp b/Engine/Source/Editor/UnrealEd/Private/Lightmass/LightmassRender.cpp
--- a/Engine/Source/Editor/UnrealEd/Private/Lightmass/LightmassRender.cpp
+++ b/Engine/Source/Editor/UnrealEd/Private/Lightmass/LightmassRender.cpp
@@ -116,6 +116,11 @@ struct FLightmassMaterialCompiler : public FProxyMaterialCompiler
 		return Compiler->Constant3(0,0,0);
 	}
 
+	virtual int32 PreSkinnedLocalBounds() override
+	{
+		return Compiler->Constant3(0, 0, 0);
+	}
+
 	virtual int32 DistanceCullFade() override
 	{
 		return Compiler->Constant(1.0f);
diff --git a/Engine/Source/Runtime/Engine/Classes/Components/SkinnedMeshComponent.h b/Engine/Source/Runtime/Engine/Classes/Components/SkinnedMeshComponent.h
--- a/Engine/Source/Runtime/Engine/Classes/Components/SkinnedMeshComponent.h
+++ b/Engine/Source/Runtime/Engine/Classes/Components/SkinnedMeshComponent.h
@@ -794,6 +794,9 @@ public:
 	virtual int32 GetNumMaterials() const override;
 	//~ End UPrimitiveComponent Interface
 
+	/** Get the pre-skinning local space bounds for this component. */
+	FBoxSphereBounds GetPreSkinnedLocalBounds() const;
+
 	/**
 	 *	Sets the value of the bForceWireframe flag and reattaches the component as necessary.
 	 *
diff --git a/Engine/Source/Runtime/Engine/Classes/Engine/SkeletalMesh.h b/Engine/Source/Runtime/Engine/Classes/Engine/SkeletalMesh.h
--- a/Engine/Source/Runtime/Engine/Classes/Engine/SkeletalMesh.h
+++ b/Engine/Source/Runtime/Engine/Classes/Engine/SkeletalMesh.h
@@ -489,11 +489,11 @@ public:
 
 	/** Get the extended bounds of this mesh (imported bounds plus bounds extension) */
 	UFUNCTION(BlueprintCallable, Category = Mesh)
-	FBoxSphereBounds GetBounds();
+	FBoxSphereBounds GetBounds() const;
 
 	/** Get the original imported bounds of the skel mesh */
 	UFUNCTION(BlueprintCallable, Category = Mesh)
-	FBoxSphereBounds GetImportedBounds();
+	FBoxSphereBounds GetImportedBounds() const;
 
 	/** Set the original imported bounds of the skel mesh, will recalculate extended bounds */
 	void SetImportedBounds(const FBoxSphereBounds& InBounds);
diff --git a/Engine/Source/Runtime/Engine/Classes/Materials/MaterialExpressionPreSkinnedLocalBounds.h b/Engine/Source/Runtime/Engine/Classes/Materials/MaterialExpressionPreSkinnedLocalBounds.h
new file mode 100644
--- /dev/null
+++ b/Engine/Source/Runtime/Engine/Classes/Materials/MaterialExpressionPreSkinnedLocalBounds.h
@@ -0,0 +1,21 @@
+// Copyright 1998-2019 Epic Games, Inc. All Rights Reserved.
+
+#pragma once
+
+#include "CoreMinimal.h"
+#include "UObject/ObjectMacros.h"
+#include "Materials/MaterialExpression.h"
+#include "MaterialExpressionPreSkinnedLocalBounds.generated.h"
+
+UCLASS(collapsecategories, hidecategories = Object)
+class UMaterialExpressionPreSkinnedLocalBounds : public UMaterialExpression
+{
+	GENERATED_UCLASS_BODY()
+
+	//~ Begin UMaterialExpression Interface
+#if WITH_EDITOR
+	virtual int32 Compile(class FMaterialCompiler* Compiler, int32 OutputIndex) override;
+	virtual void GetCaption(TArray<FString>& OutCaptions) const override;
+#endif
+	//~ End UMaterialExpression Interface
+};
\ No newline at end of file
diff --git a/Engine/Source/Runtime/Engine/Private/Components/SkinnedMeshComponent.cpp b/Engine/Source/Runtime/Engine/Private/Components/SkinnedMeshComponent.cpp
--- a/Engine/Source/Runtime/Engine/Private/Components/SkinnedMeshComponent.cpp
+++ b/Engine/Source/Runtime/Engine/Private/Components/SkinnedMeshComponent.cpp
@@ -1056,6 +1056,26 @@ FBoxSphereBounds USkinnedMeshComponent::CalcMeshBound(const FVector& RootOffset,
 	return NewBounds;
 }
 
+FBoxSphereBounds USkinnedMeshComponent::GetPreSkinnedLocalBounds() const 
+{
+	const USkinnedMeshComponent* const MasterPoseComponentInst = MasterPoseComponent.Get();
+
+	if (SkeletalMesh)
+	{
+		// Get the Pre-skinned bounds from the skeletal mesh. Note that these bounds are the "ExtendedBounds", so they can be tweaked on the SkeletalMesh   
+		return SkeletalMesh->GetBounds();
+	}
+	else if(MasterPoseComponentInst && MasterPoseComponentInst->SkeletalMesh)
+	{
+		// Get the bounds from the master pose if ther is no skeletal mesh
+		return MasterPoseComponentInst->SkeletalMesh->GetBounds();
+	}
+	else
+	{
+		// Fall back
+		return FBoxSphereBounds(ForceInitToZero);
+	}
+}
 
 FMatrix USkinnedMeshComponent::GetBoneMatrix(int32 BoneIdx) const
 {
diff --git a/Engine/Source/Runtime/Engine/Private/Materials/HLSLMaterialTranslator.h b/Engine/Source/Runtime/Engine/Private/Materials/HLSLMaterialTranslator.h
--- a/Engine/Source/Runtime/Engine/Private/Materials/HLSLMaterialTranslator.h
+++ b/Engine/Source/Runtime/Engine/Private/Materials/HLSLMaterialTranslator.h
@@ -3179,6 +3179,11 @@ protected:
 		return GetPrimitiveProperty(MCT_Float3, TEXT("ObjectBounds"), TEXT("ObjectBounds.xyz"));
 	}
 
+	virtual int32 PreSkinnedLocalBounds() override
+	{
+		return GetPrimitiveProperty(MCT_Float3, TEXT("PreSkinnedLocalBounds"), TEXT("PreSkinnedLocalBounds.xyz"));
+	}
+
 	virtual int32 DistanceCullFade() override
 	{
 		bUsesDistanceCullFade = true;
diff --git a/Engine/Source/Runtime/Engine/Private/Materials/MaterialExpressions.cpp b/Engine/Source/Runtime/Engine/Private/Materials/MaterialExpressions.cpp
--- a/Engine/Source/Runtime/Engine/Private/Materials/MaterialExpressions.cpp
+++ b/Engine/Source/Runtime/Engine/Private/Materials/MaterialExpressions.cpp
@@ -123,6 +123,7 @@
 #include "Materials/MaterialExpressionOneMinus.h"
 #include "Materials/MaterialExpressionPanner.h"
 #include "Materials/MaterialExpressionParameter.h"
+#include "Materials/MaterialExpressionPreSkinnedLocalBounds.h"
 #include "Materials/MaterialExpressionPreviousFrameSwitch.h"
 #include "Materials/MaterialExpressionReroute.h"
 #include "Materials/MaterialExpressionScalarParameter.h"
@@ -9444,6 +9445,47 @@ void UMaterialExpressionObjectBounds::GetCaption(TArray<FString>& OutCaptions) c
 }
 #endif // WITH_EDITOR
 
+///////////////////////////////////////////////////////////////////////////////
+// UMaterialExpressionPreSkinnedLocalBounds
+///////////////////////////////////////////////////////////////////////////////
+UMaterialExpressionPreSkinnedLocalBounds::UMaterialExpressionPreSkinnedLocalBounds(const FObjectInitializer& ObjectInitializer)
+	: Super(ObjectInitializer)
+{
+#if WITH_EDITORONLY_DATA
+	// Structure to hold one-time initialization
+	struct FConstructorStatics
+	{
+		FText NAME_Vectors;
+		FConstructorStatics()
+			: NAME_Vectors(LOCTEXT("Vectors", "Vectors"))
+		{
+		}
+	};
+	static FConstructorStatics ConstructorStatics;
+
+	MenuCategories.Add(ConstructorStatics.NAME_Vectors);
+
+	bShaderInputData = true;
+#endif
+}
+
+#if WITH_EDITOR
+int32 UMaterialExpressionPreSkinnedLocalBounds::Compile(class FMaterialCompiler* Compiler, int32 OutputIndex)
+{
+	if (Material && Material->MaterialDomain == MD_DeferredDecal)
+	{
+		return CompilerError(Compiler, TEXT("Expression not available in the deferred decal material domain."));
+	}
+
+	return Compiler->PreSkinnedLocalBounds();
+}
+
+void UMaterialExpressionPreSkinnedLocalBounds::GetCaption(TArray<FString>& OutCaptions) const
+{
+	OutCaptions.Add(TEXT("Pre-Skinned Local Bounds"));
+}
+#endif // WITH_EDITOR
+
 ///////////////////////////////////////////////////////////////////////////////
 // UMaterialExpressionDistanceCullFade
 ///////////////////////////////////////////////////////////////////////////////
diff --git a/Engine/Source/Runtime/Engine/Private/PrimitiveSceneProxy.cpp b/Engine/Source/Runtime/Engine/Private/PrimitiveSceneProxy.cpp
--- a/Engine/Source/Runtime/Engine/Private/PrimitiveSceneProxy.cpp
+++ b/Engine/Source/Runtime/Engine/Private/PrimitiveSceneProxy.cpp
@@ -277,6 +277,7 @@ void FPrimitiveSceneProxy::UpdateUniformBuffer()
 				ActorPosition, 
 				Bounds, 
 				LocalBounds, 
+				GetPreSkinnedLocalBounds(),
 				bReceivesDecals, 
 				HasDistanceFieldRepresentation(), 
 				HasDynamicIndirectShadowCasterRepresentation(), 
diff --git a/Engine/Source/Runtime/Engine/Private/PrimitiveUniformShaderParameters.cpp b/Engine/Source/Runtime/Engine/Private/PrimitiveUniformShaderParameters.cpp
--- a/Engine/Source/Runtime/Engine/Private/PrimitiveUniformShaderParameters.cpp
+++ b/Engine/Source/Runtime/Engine/Private/PrimitiveUniformShaderParameters.cpp
@@ -48,7 +48,8 @@ FPrimitiveSceneShaderData::FPrimitiveSceneShaderData(const FPrimitiveSceneProxy*
 		PreviousLocalToWorld,
 		Proxy->GetActorPosition(), 
 		Proxy->GetBounds(), 
-		Proxy->GetLocalBounds(), 
+		Proxy->GetLocalBounds(),
+		Proxy->GetPreSkinnedLocalBounds(),
 		Proxy->ReceivesDecals(), 
 		Proxy->HasDistanceFieldRepresentation(), 
 		Proxy->HasDynamicIndirectShadowCasterRepresentation(), 
@@ -107,4 +108,6 @@ void FPrimitiveSceneShaderData::Setup(const FPrimitiveUniformShaderParameters& P
 
 	Data[25] = FVector4(0.0f, 0.0f, 0.0f, 0.0f);
 	Data[25].X = *(const float*)&PrimitiveUniformShaderParameters.SingleCaptureIndex;
+
+	Data[26] = FVector4(PrimitiveUniformShaderParameters.PreSkinnedLocalBounds, 0.0f); // .w unused
 }
\ No newline at end of file
diff --git a/Engine/Source/Runtime/Engine/Private/SceneManagement.cpp b/Engine/Source/Runtime/Engine/Private/SceneManagement.cpp
--- a/Engine/Source/Runtime/Engine/Private/SceneManagement.cpp
+++ b/Engine/Source/Runtime/Engine/Private/SceneManagement.cpp
@@ -305,16 +305,29 @@ void FDynamicPrimitiveUniformBuffer::Set(
 	const FMatrix& PreviousLocalToWorld,
 	const FBoxSphereBounds& WorldBounds,
 	const FBoxSphereBounds& LocalBounds,
+	const FBoxSphereBounds& PreSkinnedLocalBounds,
 	bool bReceivesDecals,
 	bool bHasPrecomputedVolumetricLightmap,
 	bool bUseEditorDepthTest)
 {
 	check(IsInRenderingThread());
 	UniformBuffer.SetContents(
-		GetPrimitiveUniformShaderParameters(LocalToWorld, PreviousLocalToWorld, WorldBounds.Origin, WorldBounds, LocalBounds, bReceivesDecals, false, false, false, bHasPrecomputedVolumetricLightmap, bUseEditorDepthTest, GetDefaultLightingChannelMask(), 1.0f, INDEX_NONE, INDEX_NONE));
+		GetPrimitiveUniformShaderParameters(LocalToWorld, PreviousLocalToWorld, WorldBounds.Origin, WorldBounds, LocalBounds, PreSkinnedLocalBounds, bReceivesDecals, false, false, false, bHasPrecomputedVolumetricLightmap, bUseEditorDepthTest, GetDefaultLightingChannelMask(), 1.0f, INDEX_NONE, INDEX_NONE));
 	UniformBuffer.InitResource();
 }
 
+void FDynamicPrimitiveUniformBuffer::Set(
+	const FMatrix& LocalToWorld,
+	const FMatrix& PreviousLocalToWorld,
+	const FBoxSphereBounds& WorldBounds,
+	const FBoxSphereBounds& LocalBounds,
+	bool bReceivesDecals,
+	bool bHasPrecomputedVolumetricLightmap,
+	bool bUseEditorDepthTest)
+{
+	Set(LocalToWorld, PreviousLocalToWorld, WorldBounds, LocalBounds, LocalBounds, bReceivesDecals, bHasPrecomputedVolumetricLightmap, bUseEditorDepthTest);
+}
+
 FLightMapInteraction FLightMapInteraction::Texture(
 	const class ULightMapTexture2D* const* InTextures,
 	const ULightMapTexture2D* InSkyOcclusionTexture,
diff --git a/Engine/Source/Runtime/Engine/Private/SkeletalMesh.cpp b/Engine/Source/Runtime/Engine/Private/SkeletalMesh.cpp
--- a/Engine/Source/Runtime/Engine/Private/SkeletalMesh.cpp
+++ b/Engine/Source/Runtime/Engine/Private/SkeletalMesh.cpp
@@ -321,12 +321,12 @@ void USkeletalMesh::PostInitProperties()
 	Super::PostInitProperties();
 }
 
-FBoxSphereBounds USkeletalMesh::GetBounds()
+FBoxSphereBounds USkeletalMesh::GetBounds() const
 {
 	return ExtendedBounds;
 }
 
-FBoxSphereBounds USkeletalMesh::GetImportedBounds()
+FBoxSphereBounds USkeletalMesh::GetImportedBounds() const
 {
 	return ImportedBounds;
 }
@@ -863,6 +863,7 @@ void USkeletalMesh::PostEditChangeProperty(FPropertyChangedEvent& PropertyChange
 		BuildPhysicsData();
 	}
 
+	bool bHasToReregisterComponent = false;
 	if(UProperty* MemberProperty = PropertyChangedEvent.MemberProperty)
 	{
 		if(MemberProperty->GetFName() == GET_MEMBER_NAME_CHECKED(USkeletalMesh, PositiveBoundsExtension) ||
@@ -871,10 +872,16 @@ void USkeletalMesh::PostEditChangeProperty(FPropertyChangedEvent& PropertyChange
 			// If the bounds extensions change, recalculate extended bounds.
 			ValidateBoundsExtension();
 			CalculateExtendedBounds();
+			bHasToReregisterComponent = true;
 		}
 	}
+		
+	if (PropertyThatChanged && PropertyThatChanged->GetFName() == GET_MEMBER_NAME_CHECKED(USkeletalMesh, PostProcessAnimBlueprint))
+	{
+		bHasToReregisterComponent = true;
+	}
 
-	if(PropertyThatChanged && PropertyThatChanged->GetFName() == GET_MEMBER_NAME_CHECKED(USkeletalMesh, PostProcessAnimBlueprint))
+	if (bHasToReregisterComponent)
 	{
 		TArray<UActorComponent*> ComponentsToReregister;
 		for(TObjectIterator<USkeletalMeshComponent> It; It; ++It)
@@ -888,6 +895,7 @@ void USkeletalMesh::PostEditChangeProperty(FPropertyChangedEvent& PropertyChange
 		FMultiComponentReregisterContext ReregisterContext(ComponentsToReregister);
 	}
 
+
 	if (PropertyThatChanged && PropertyChangedEvent.MemberProperty)
 	{
 		if (PropertyChangedEvent.MemberProperty->GetFName() == FName(TEXT("SamplingInfo")))
@@ -3361,6 +3369,9 @@ FSkeletalMeshSceneProxy::FSkeletalMeshSceneProxy(const USkinnedMeshComponent* Co
 	// Force inset shadows if capsule shadows are requested, as they can't be supported with full scene shadows
 	bCastInsetShadow = bCastInsetShadow || bCastCapsuleDirectShadow;
 
+	// Get the pre-skinned local bounds
+	PreSkinnedLocalBounds = Component->GetPreSkinnedLocalBounds();
+
 	const USkinnedMeshComponent* SkinnedMeshComponent = Cast<const USkinnedMeshComponent>(Component);
 	if(SkinnedMeshComponent && SkinnedMeshComponent->bPerBoneMotionBlur)
 	{
diff --git a/Engine/Source/Runtime/Engine/Public/MaterialCompiler.h b/Engine/Source/Runtime/Engine/Public/MaterialCompiler.h
--- a/Engine/Source/Runtime/Engine/Public/MaterialCompiler.h
+++ b/Engine/Source/Runtime/Engine/Public/MaterialCompiler.h
@@ -139,7 +139,8 @@ public:
 	virtual int32 WorldPosition(EWorldPositionIncludedOffsets WorldPositionIncludedOffsets) = 0;
 	virtual int32 ObjectWorldPosition() = 0;
 	virtual int32 ObjectRadius() = 0;
-	virtual int32 ObjectBounds() = 0;	
+	virtual int32 ObjectBounds() = 0;
+	virtual int32 PreSkinnedLocalBounds() = 0;
 	virtual int32 DistanceCullFade() = 0;
 	virtual int32 ActorWorldPosition() = 0;
 	virtual int32 ParticleMacroUV() = 0;
@@ -391,6 +392,7 @@ public:
 	virtual int32 ObjectWorldPosition() override { return Compiler->ObjectWorldPosition(); }
 	virtual int32 ObjectRadius() override { return Compiler->ObjectRadius(); }
 	virtual int32 ObjectBounds() override { return Compiler->ObjectBounds(); }
+	virtual int32 PreSkinnedLocalBounds() override { return Compiler->PreSkinnedLocalBounds(); }
 	virtual int32 DistanceCullFade() override { return Compiler->DistanceCullFade(); }
 	virtual int32 ActorWorldPosition() override { return Compiler->ActorWorldPosition(); }
 	virtual int32 ParticleMacroUV() override { return Compiler->ParticleMacroUV(); }
diff --git a/Engine/Source/Runtime/Engine/Public/PrimitiveSceneProxy.h b/Engine/Source/Runtime/Engine/Public/PrimitiveSceneProxy.h
--- a/Engine/Source/Runtime/Engine/Public/PrimitiveSceneProxy.h
+++ b/Engine/Source/Runtime/Engine/Public/PrimitiveSceneProxy.h
@@ -425,6 +425,7 @@ public:
 	inline bool IsLocalToWorldDeterminantNegative() const { return bIsLocalToWorldDeterminantNegative; }
 	inline const FBoxSphereBounds& GetBounds() const { return Bounds; }
 	inline const FBoxSphereBounds& GetLocalBounds() const { return LocalBounds; }
+	virtual FBoxSphereBounds GetPreSkinnedLocalBounds() const { return LocalBounds; }
 	inline FName GetOwnerName() const { return OwnerName; }
 	inline FName GetResourceName() const { return ResourceName; }
 	inline FName GetLevelName() const { return LevelName; }
diff --git a/Engine/Source/Runtime/Engine/Public/PrimitiveUniformShaderParameters.h b/Engine/Source/Runtime/Engine/Public/PrimitiveUniformShaderParameters.h
--- a/Engine/Source/Runtime/Engine/Public/PrimitiveUniformShaderParameters.h
+++ b/Engine/Source/Runtime/Engine/Public/PrimitiveUniformShaderParameters.h
@@ -33,6 +33,7 @@ BEGIN_GLOBAL_SHADER_PARAMETER_STRUCT(FPrimitiveUniformShaderParameters,ENGINE_AP
 	SHADER_PARAMETER_EX(FVector4,NonUniformScale,EShaderPrecisionModifier::Half)
 	SHADER_PARAMETER(FVector, LocalObjectBoundsMin)		// This is used in a custom material function (ObjectLocalBounds.uasset)
 	SHADER_PARAMETER(FVector, LocalObjectBoundsMax)		// This is used in a custom material function (ObjectLocalBounds.uasset)
+	SHADER_PARAMETER(FVector, PreSkinnedLocalBounds)	// Local space bounds, pre-skinning
 	SHADER_PARAMETER(uint32,LightingChannelMask)
 	SHADER_PARAMETER(uint32,LightmapDataIndex)
 	SHADER_PARAMETER(int32, SingleCaptureIndex)
@@ -45,6 +46,7 @@ inline FPrimitiveUniformShaderParameters GetPrimitiveUniformShaderParameters(
 	FVector ActorPosition,
 	const FBoxSphereBounds& WorldBounds,
 	const FBoxSphereBounds& LocalBounds,
+	const FBoxSphereBounds& PreSkinnedLocalBounds,
 	bool bReceivesDecals,
 	bool bHasDistanceFieldRepresentation,
 	bool bHasCapsuleRepresentation,
@@ -66,6 +68,7 @@ inline FPrimitiveUniformShaderParameters GetPrimitiveUniformShaderParameters(
 	Result.ObjectBounds = WorldBounds.BoxExtent;
 	Result.LocalObjectBoundsMin = LocalBounds.GetBoxExtrema(0); // 0 == minimum
 	Result.LocalObjectBoundsMax = LocalBounds.GetBoxExtrema(1); // 1 == maximum
+	Result.PreSkinnedLocalBounds = PreSkinnedLocalBounds.BoxExtent;
 	Result.ObjectOrientation = LocalToWorld.GetUnitAxis( EAxis::Z );
 	Result.ActorWorldPosition = ActorPosition;
 	Result.LightingChannelMask = LightingChannelMask;
@@ -97,10 +100,35 @@ inline FPrimitiveUniformShaderParameters GetPrimitiveUniformShaderParameters(
 	return Result;
 }
 
+/** Initializes the primitive uniform shader parameters. Pre-skinned local bounds default to LocalBounds */
+inline FPrimitiveUniformShaderParameters GetPrimitiveUniformShaderParameters(
+	const FMatrix& LocalToWorld,
+	const FMatrix& PreviousLocalToWorld,
+	FVector ActorPosition,
+	const FBoxSphereBounds& WorldBounds,
+	const FBoxSphereBounds& LocalBounds,
+	bool bReceivesDecals,
+	bool bHasDistanceFieldRepresentation,
+	bool bHasCapsuleRepresentation,
+	bool bUseSingleSampleShadowFromStationaryLights,
+	bool bUseVolumetricLightmap,
+	bool bUseEditorDepthTest,
+	uint32 LightingChannelMask,
+	float LpvBiasMultiplier,
+	uint32 LightmapDataIndex,
+	int32 SingleCaptureIndex
+)
+{
+	// Pass through call
+	return GetPrimitiveUniformShaderParameters(LocalToWorld, PreviousLocalToWorld, ActorPosition, WorldBounds, LocalBounds, LocalBounds, bReceivesDecals, bHasDistanceFieldRepresentation, bHasCapsuleRepresentation, 
+		bUseSingleSampleShadowFromStationaryLights, bUseVolumetricLightmap, bUseEditorDepthTest, LightingChannelMask, LpvBiasMultiplier, LightmapDataIndex, SingleCaptureIndex);
+}
+
 inline TUniformBufferRef<FPrimitiveUniformShaderParameters> CreatePrimitiveUniformBufferImmediate(
 	const FMatrix& LocalToWorld,
 	const FBoxSphereBounds& WorldBounds,
 	const FBoxSphereBounds& LocalBounds,
+	const FBoxSphereBounds& PreSkinnedLocalBounds,
 	bool bReceivesDecals,
 	bool bUseEditorDepthTest,
 	float LpvBiasMultiplier = 1.0f
@@ -108,7 +136,7 @@ inline TUniformBufferRef<FPrimitiveUniformShaderParameters> CreatePrimitiveUnifo
 {
 	check(IsInRenderingThread());
 	return TUniformBufferRef<FPrimitiveUniformShaderParameters>::CreateUniformBufferImmediate(
-		GetPrimitiveUniformShaderParameters(LocalToWorld, LocalToWorld, WorldBounds.Origin, WorldBounds, LocalBounds, bReceivesDecals, false, false, false, false, bUseEditorDepthTest, GetDefaultLightingChannelMask(), LpvBiasMultiplier, INDEX_NONE, INDEX_NONE),
+		GetPrimitiveUniformShaderParameters(LocalToWorld, LocalToWorld, WorldBounds.Origin, WorldBounds, LocalBounds, PreSkinnedLocalBounds, bReceivesDecals, false, false, false, false, bUseEditorDepthTest, GetDefaultLightingChannelMask(), LpvBiasMultiplier, INDEX_NONE, INDEX_NONE),
 		UniformBuffer_MultiFrame
 		);
 }
@@ -156,7 +184,7 @@ extern ENGINE_API TGlobalResource<FIdentityPrimitiveUniformBuffer> GIdentityPrim
 struct FPrimitiveSceneShaderData
 {
 	// Must match usf
-	enum { PrimitiveDataStrideInFloat4s = 26 };
+	enum { PrimitiveDataStrideInFloat4s = 27 };
 
 	FVector4 Data[PrimitiveDataStrideInFloat4s];
 
diff --git a/Engine/Source/Runtime/Engine/Public/SceneManagement.h b/Engine/Source/Runtime/Engine/Public/SceneManagement.h
--- a/Engine/Source/Runtime/Engine/Public/SceneManagement.h
+++ b/Engine/Source/Runtime/Engine/Public/SceneManagement.h
@@ -2200,6 +2200,17 @@ public:
 
 	TUniformBuffer<FPrimitiveUniformShaderParameters> UniformBuffer;
 
+	ENGINE_API void Set(
+		const FMatrix& LocalToWorld,
+		const FMatrix& PreviousLocalToWorld,
+		const FBoxSphereBounds& WorldBounds,
+		const FBoxSphereBounds& LocalBounds,
+		const FBoxSphereBounds& PreSkinnedLocalBounds,
+		bool bReceivesDecals,
+		bool bHasPrecomputedVolumetricLightmap,
+		bool bUseEditorDepthTest);
+
+	/** Pass-through implementation which calls the overloaded Set function with LocalBounds for PreSkinnedLocalBounds. */
 	ENGINE_API void Set(
 		const FMatrix& LocalToWorld,
 		const FMatrix& PreviousLocalToWorld,
diff --git a/Engine/Source/Runtime/Engine/Public/SkeletalMeshTypes.h b/Engine/Source/Runtime/Engine/Public/SkeletalMeshTypes.h
--- a/Engine/Source/Runtime/Engine/Public/SkeletalMeshTypes.h
+++ b/Engine/Source/Runtime/Engine/Public/SkeletalMeshTypes.h
@@ -241,6 +241,9 @@ public:
 	virtual bool HasDynamicIndirectShadowCasterRepresentation() const override;
 	virtual void GetShadowShapes(TArray<FCapsuleShape>& CapsuleShapes) const override;
 
+	/** Return the bounds for the pre-skinned primitive in local space */
+	virtual FBoxSphereBounds GetPreSkinnedLocalBounds() const override { return PreSkinnedLocalBounds; }
+
 	/** Returns a pre-sorted list of shadow capsules's bone indicies */
 	const TArray<uint16>& GetSortedShadowBoneIndices() const
 	{
@@ -264,7 +267,7 @@ public:
 	 */
 	virtual void DebugDrawPhysicsAsset(int32 ViewIndex, FMeshElementCollector& Collector, const FEngineShowFlags& EngineShowFlags) const;
 
-	/** Render the bones of the skeleton for debug display */
+	/** Render the bones of the skeleton for debug display */ 
 	void DebugDrawSkeleton(int32 ViewIndex, FMeshElementCollector& Collector, const FEngineShowFlags& EngineShowFlags) const;
 
 	virtual uint32 GetMemoryFootprint( void ) const override { return( sizeof( *this ) + GetAllocatedSize() ); }
@@ -355,6 +358,9 @@ protected:
 	/** Set of materials used by this scene proxy, safe to access from the game thread. */
 	TSet<UMaterialInterface*> MaterialsInUse_GameThread;
 	
+	/** The primitive's pre-skinned local space bounds. */
+	FBoxSphereBounds PreSkinnedLocalBounds;
+
 #if WITH_EDITORONLY_DATA
 	/** The component streaming distance multiplier */
 	float StreamingDistanceMultiplier;
