Reducing spam of timer manager logging when the timer is exceeded ... instead of dumping all active timers, dump the current timer(s) when the time has been exceeded.

diff --git a/Engine/Source/Runtime/Engine/Private/TimerManager.cpp b/Engine/Source/Runtime/Engine/Private/TimerManager.cpp
--- a/Engine/Source/Runtime/Engine/Private/TimerManager.cpp
+++ b/Engine/Source/Runtime/Engine/Private/TimerManager.cpp
@@ -554,6 +554,7 @@ void FTimerManager::Tick(float DeltaTime)
 	}
 
 	const double StartTime = FPlatformTime::Seconds();
+	bool TimeExceededLogged = false;
 
 	InternalTime += DeltaTime;
 
@@ -618,6 +619,26 @@ void FTimerManager::Tick(float DeltaTime)
 				}
 			}
 
+			// help us hunt down outliers that cause our timer manager times to spike.  Recommended that users set meaningful DumpTimerLogsThresholds in appropriate ini files if they are seeing spikes in the timer manager.
+			const double DeltaT = (FPlatformTime::Seconds() - StartTime) * 1000.f;
+			if (DeltaT >= DumpTimerLogsThreshold)
+			{
+				if (!TimeExceededLogged)
+				{
+					UE_LOG(LogEngine, Log, TEXT("TimerManager's time threshold of %.2fms exceeded with a deltaT of %.4f, dumping current timer data."), DumpTimerLogsThreshold, DeltaT);
+					TimeExceededLogged = true;
+				}
+				
+				if (Top)
+				{
+					DescribeFTimerDataSafely(Top);
+				}
+				else
+				{
+					UE_LOG(LogEngine, Log, TEXT("There was no timer data!"));
+				}
+			}
+
 			// test to ensure it didn't get cleared during execution
 			if (Top)
 			{
@@ -644,14 +665,6 @@ void FTimerManager::Tick(float DeltaTime)
 		}
 	}
 
-	// help us hunt down outliers that cause our timer manager times to spike.  Recommended that users set meaningful DumpTimerLogsThresholds in appropriate ini files if they are seeing spikes in the timer manager.
-	const double DeltaT = (FPlatformTime::Seconds() - StartTime) * 1000.f;
-	if (DeltaT >= DumpTimerLogsThreshold)
-	{
-		UE_LOG(LogEngine, Log, TEXT("TimerManager's time threshold of %.2fms exceeded with a deltaT of %.4f, dumping list of timers to help identify why it took so long to process the list of timers."), DumpTimerLogsThreshold, DeltaT);
-		ListTimers();
-	}
-
 	// Timer has been ticked.
 	LastTickedFrame = GFrameCounter;
 
