#UE4 Compile fix

diff --git a/Engine/Source/Runtime/Engine/Private/TimerManager.cpp b/Engine/Source/Runtime/Engine/Private/TimerManager.cpp
--- a/Engine/Source/Runtime/Engine/Private/TimerManager.cpp
+++ b/Engine/Source/Runtime/Engine/Private/TimerManager.cpp
@@ -631,7 +631,7 @@ void FTimerManager::Tick(float DeltaTime)
 				
 				if (Top)
 				{
-					DescribeFTimerDataSafely(Top);
+					DescribeFTimerDataSafely(*Top);
 				}
 				else
 				{
