Set all Fast Arrays to use Struct Delta Serialization by default.
Individual Fast Arrays can still be opted out.

diff --git a/Engine/Source/Runtime/Engine/Classes/Engine/NetSerialization.h b/Engine/Source/Runtime/Engine/Classes/Engine/NetSerialization.h
--- a/Engine/Source/Runtime/Engine/Classes/Engine/NetSerialization.h
+++ b/Engine/Source/Runtime/Engine/Classes/Engine/NetSerialization.h
@@ -418,7 +418,10 @@ struct FFastArraySerializer
 		, ArrayReplicationKey(0)
 		, CachedNumItems(INDEX_NONE)
 		, CachedNumItemsToConsiderForWriting(INDEX_NONE)
-		, DeltaFlags(EFastArraySerializerDeltaFlags::None) { }
+		, DeltaFlags(EFastArraySerializerDeltaFlags::None)
+	{
+		SetDeltaSerializationEnabled(true);
+	}
 
 	~FFastArraySerializer() {}
 
