Implement FSkeletalMeshVertexClothBuffer::ReleaseRHI() so it's RHI resources don't get released on the game thread.

diff --git a/Engine/Source/Runtime/Engine/Private/Rendering/SkeletalMeshVertexClothBuffer.cpp b/Engine/Source/Runtime/Engine/Private/Rendering/SkeletalMeshVertexClothBuffer.cpp
--- a/Engine/Source/Runtime/Engine/Private/Rendering/SkeletalMeshVertexClothBuffer.cpp
+++ b/Engine/Source/Runtime/Engine/Private/Rendering/SkeletalMeshVertexClothBuffer.cpp
@@ -80,6 +80,16 @@ void FSkeletalMeshVertexClothBuffer::InitRHI()
 	}
 }
 
+/**
+* Release the RHI resource for this vertex buffer
+*/
+void FSkeletalMeshVertexClothBuffer::ReleaseRHI()
+{
+	VertexBufferSRV.SafeRelease();
+	// call the base class's ReleaseRHI() since we overrode it
+	FVertexBuffer::ReleaseRHI();
+}
+
 /**
 * Serializer for this class
 * @param Ar - archive to serialize to
diff --git a/Engine/Source/Runtime/Engine/Public/Rendering/SkeletalMeshVertexClothBuffer.h b/Engine/Source/Runtime/Engine/Public/Rendering/SkeletalMeshVertexClothBuffer.h
--- a/Engine/Source/Runtime/Engine/Public/Rendering/SkeletalMeshVertexClothBuffer.h
+++ b/Engine/Source/Runtime/Engine/Public/Rendering/SkeletalMeshVertexClothBuffer.h
@@ -61,6 +61,11 @@ public:
 	*/
 	virtual void InitRHI() override;
 
+	/**
+	* Release the RHI resource for this vertex buffer
+	*/
+	virtual void ReleaseRHI() override;
+
 	/**
 	* @return text description for the resource type
 	*/
