Landscape transactions
- Don't create a material update context when a texture change if it doesn't affect any materials.
- Landscape undo/redo operations went from >30s to <3s.
- Tested by robert.osborn.

#rnx

diff --git a/Engine/Source/Runtime/Engine/Private/Texture.cpp b/Engine/Source/Runtime/Engine/Private/Texture.cpp
--- a/Engine/Source/Runtime/Engine/Private/Texture.cpp
+++ b/Engine/Source/Runtime/Engine/Private/Texture.cpp
@@ -199,8 +199,8 @@ void UTexture::PostEditChangeProperty(FPropertyChangedEvent& PropertyChangedEven
 	}
 	else if (!GDisableAutomaticTextureMaterialUpdateDependencies)
 	{
-		FMaterialUpdateContext UpdateContext;
 		// Update any material that uses this texture and must force a recompile of cache ressource
+		TArray<UMaterial*> MaterialsToUpdate;
 		TSet<UMaterial*> BaseMaterialsThatUseThisTexture;
 		for (TObjectIterator<UMaterialInterface> It; It; ++It)
 		{
@@ -214,12 +214,22 @@ void UTexture::PostEditChangeProperty(FPropertyChangedEvent& PropertyChangedEven
 				{
 					if (Material->IsTextureForceRecompileCacheRessource(this))
 					{
-						UpdateContext.AddMaterial(Material);
+						MaterialsToUpdate.Add(Material);
 						Material->UpdateMaterialShaderCacheAndTextureReferences();
 					}
 				}
 			}
 		}
+
+		if (MaterialsToUpdate.Num())
+		{
+			FMaterialUpdateContext UpdateContext;
+
+			for (UMaterial* MaterialToUpdate: MaterialsToUpdate)
+			{
+				UpdateContext.AddMaterial(MaterialToUpdate);
+			}
+		}
 	}
 
 	NumCinematicMipLevels = FMath::Max<int32>( NumCinematicMipLevels, 0 );
