In Editor, Landscape Heightmap and Weightmap Texture2D's MipLevels are forced resident only if in Procedural Landscape experimental mode.

diff --git a/Engine/Source/Runtime/Engine/Private/Texture2D.cpp b/Engine/Source/Runtime/Engine/Private/Texture2D.cpp
--- a/Engine/Source/Runtime/Engine/Private/Texture2D.cpp
+++ b/Engine/Source/Runtime/Engine/Private/Texture2D.cpp
@@ -33,6 +33,9 @@
 #include "Streaming/Texture2DStreamIn_IO_AsyncReallocate.h"
 #include "Streaming/Texture2DStreamIn_IO_Virtual.h"
 #include "Async/AsyncFileHandle.h"
+#if WITH_EDITOR
+#include "Settings/EditorExperimentalSettings.h"
+#endif
 
 UTexture2D::UTexture2D(const FObjectInitializer& ObjectInitializer)
 	: Super(ObjectInitializer)
@@ -1162,10 +1165,12 @@ bool UTexture2D::ShouldMipLevelsBeForcedResident() const
 		return true;
 	}
 
-	if (GIsEditor && (LODGroup == TEXTUREGROUP_Terrain_Heightmap || LODGroup == TEXTUREGROUP_Terrain_Weightmap))
+#if WITH_EDITOR
+	if (GIsEditor && GetMutableDefault<UEditorExperimentalSettings>()->bProceduralLandscape && (LODGroup == TEXTUREGROUP_Terrain_Heightmap || LODGroup == TEXTUREGROUP_Terrain_Weightmap))
 	{
 		return true;
 	}
+#endif
 
 	return false;
 }
