Undo uneeded Required Bones change

diff --git a/Engine/Source/Runtime/Engine/Classes/Components/SkeletalMeshComponent.h b/Engine/Source/Runtime/Engine/Classes/Components/SkeletalMeshComponent.h
--- a/Engine/Source/Runtime/Engine/Classes/Components/SkeletalMeshComponent.h
+++ b/Engine/Source/Runtime/Engine/Classes/Components/SkeletalMeshComponent.h
@@ -1288,7 +1288,7 @@ public:
 	 *
 	 * @param	LODIndex	Index of LOD [0-(MaxLOD-1)]
 	 */
-	bool RecalcRequiredBones(int32 LODIndex);
+	void RecalcRequiredBones(int32 LODIndex);
 
 	/** Computes the required bones in this SkeletalMeshComponent based on current SkeletalMesh, LOD and PhysicsAsset
 	  * @param	LODIndex	Index of LOD [0-(MaxLOD-1)]
@@ -1301,21 +1301,6 @@ public:
 	*/
 	void RecalcRequiredCurves();
 
-private:
-	/**
-	 * Recalculates the RequiredBones array in this SkeletalMeshComponent based on current SkeletalMesh, LOD and PhysicsAsset.
-	 * Only modifies members of the SMC (safe to call while parallel task is in flight
-	 *
-	 * @param	LODIndex	Index of LOD [0-(MaxLOD-1)]
-	 */
-	bool RecalcRequiredBonesInternalSMC(int32 LODIndex);
-
-	/**
-	 * Causes the anim instances this component owns to recalc their required bones.
-	 * RecalcRequiredBonesInternalSMC must be called first
-	 */
-	void RecalcRequiredBonesInternalAnimInstances();
-
 public:
 	//~ Begin UObject Interface.
 	virtual void Serialize(FArchive& Ar) override;
diff --git a/Engine/Source/Runtime/Engine/Private/Components/SkeletalMeshComponent.cpp b/Engine/Source/Runtime/Engine/Private/Components/SkeletalMeshComponent.cpp
--- a/Engine/Source/Runtime/Engine/Private/Components/SkeletalMeshComponent.cpp
+++ b/Engine/Source/Runtime/Engine/Private/Components/SkeletalMeshComponent.cpp
@@ -1649,33 +1649,18 @@ void USkeletalMeshComponent::ComputeRequiredBones(TArray<FBoneIndexType>& OutReq
 	FAnimationRuntime::EnsureParentsPresent(OutFillComponentSpaceTransformsRequiredBones, SkeletalMesh->RefSkeleton);
 }
 
-bool USkeletalMeshComponent::RecalcRequiredBones(int32 LODIndex)
-{
-	bool bSuccess = RecalcRequiredBonesInternalSMC(LODIndex);
-	RecalcRequiredBonesInternalAnimInstances();
-	return bSuccess;
-}
-
-bool USkeletalMeshComponent::RecalcRequiredBonesInternalSMC(int32 LODIndex)
+void USkeletalMeshComponent::RecalcRequiredBones(int32 LODIndex)
 {
 	if (!SkeletalMesh)
 	{
-		return false;
+		return;
 	}
 
 	ComputeRequiredBones(RequiredBones, FillComponentSpaceTransformsRequiredBones, LODIndex, /*bIgnorePhysicsAsset=*/ false);
 
 	BoneSpaceTransforms = SkeletalMesh->RefSkeleton.GetRefBonePose();
 
-	// Invalidate cached bones.
-	CachedBoneSpaceTransforms.Empty();
-	CachedComponentSpaceTransforms.Empty();
-	CachedCurve.Empty();
-	return true;
-}
-
-void USkeletalMeshComponent::RecalcRequiredBonesInternalAnimInstances()
-{
+	// make sure animation requiredBone to mark as dirty
 	if (AnimScriptInstance)
 	{
 		AnimScriptInstance->RecalcRequiredBones();
@@ -1695,6 +1680,11 @@ void USkeletalMeshComponent::RecalcRequiredBonesInternalAnimInstances()
 	// this should always happen
 	MarkRequiredCurveUpToDate();
 	bRequiredBonesUpToDate = true;
+
+	// Invalidate cached bones.
+	CachedBoneSpaceTransforms.Empty();
+	CachedComponentSpaceTransforms.Empty();
+	CachedCurve.Empty();
 }
 
 void USkeletalMeshComponent::MarkRequiredCurveUpToDate()
@@ -1964,7 +1954,13 @@ void USkeletalMeshComponent::RefreshBoneTransforms(FActorComponentTickFunction*
 	if (!bRequiredBonesUpToDate)
 	{
 		QUICK_SCOPE_CYCLE_COUNTER(STAT_USkeletalMeshComponent_RefreshBoneTransforms_RecalcRequiredBones);
-		RecalcRequiredBonesInternalSMC(PredictedLODLevel);
+		RecalcRequiredBones(PredictedLODLevel);
+	}
+	// if curves have to be refreshed
+	else if (!AreRequiredCurvesUpToDate())
+	{
+		QUICK_SCOPE_CYCLE_COUNTER(STAT_USkeletalMeshComponent_RefreshBoneTransforms_RecalcRequiredCurves);
+		RecalcRequiredCurves();
 	}
 
 	const bool bCachedShouldUseUpdateRateOptimizations = ShouldUseUpdateRateOptimizations() && AnimUpdateRateParams != nullptr;
@@ -2010,22 +2006,6 @@ void USkeletalMeshComponent::RefreshBoneTransforms(FActorComponentTickFunction*
 		return;
 	}
 
-	// Recalculate the RequiredBones array, if necessary
-	if (!bRequiredBonesUpToDate)
-	{
-		QUICK_SCOPE_CYCLE_COUNTER(STAT_USkeletalMeshComponent_RefreshBoneTransforms_RecalcRequiredBones);
-		RecalcRequiredBonesInternalAnimInstances();
-	}
-	// if curves have to be refreshed
-	else if (!AreRequiredCurvesUpToDate())
-	{
-		QUICK_SCOPE_CYCLE_COUNTER(STAT_USkeletalMeshComponent_RefreshBoneTransforms_RecalcRequiredCurves);
-		RecalcRequiredCurves();
-	}
-	
-	//Refresh UID list now curves have been updated
-	CurrentAnimCurveUIDFinder = (AnimScriptInstance) ? &AnimScriptInstance->GetRequiredBones().GetUIDToArrayLookupTable() : nullptr;
-
 	AnimEvaluationContext.SkeletalMesh = SkeletalMesh;
 	AnimEvaluationContext.AnimInstance = AnimScriptInstance;
 	AnimEvaluationContext.PostProcessAnimInstance = PostProcessAnimInstance;
