Material uniform buffer layouts are needed when caching uniform expressions so we cannot free unused material resources until the caching process is done

diff --git a/Engine/Source/Runtime/Engine/Classes/Materials/Material.h b/Engine/Source/Runtime/Engine/Classes/Materials/Material.h
--- a/Engine/Source/Runtime/Engine/Classes/Materials/Material.h
+++ b/Engine/Source/Runtime/Engine/Classes/Materials/Material.h
@@ -1109,7 +1109,7 @@ private:
 	void SetUsageByFlag(const EMaterialUsage Usage, const bool NewValue);
 
 	/** Sets up transient properties in MaterialResources. */
-	void UpdateResourceAllocations();
+	void UpdateResourceAllocations(FMaterialResourceDeferredDeletionArray* ResourcesToFree = nullptr);
 
 	/** to share code for PostLoad() and PostEditChangeProperty(), and UMaterialInstance::InitResources(), needs to be refactored */
 	void PropagateDataToMaterialProxy();
diff --git a/Engine/Source/Runtime/Engine/Classes/Materials/MaterialInstance.h b/Engine/Source/Runtime/Engine/Classes/Materials/MaterialInstance.h
--- a/Engine/Source/Runtime/Engine/Classes/Materials/MaterialInstance.h
+++ b/Engine/Source/Runtime/Engine/Classes/Materials/MaterialInstance.h
@@ -635,7 +635,7 @@ protected:
 	 * This is a helper used when recompiling MI's with static parameters.  
 	 * Assumes that the rendering thread command queue has been flushed by the caller.
 	 */
-	void UpdatePermutationAllocations();
+	void UpdatePermutationAllocations(FMaterialResourceDeferredDeletionArray* ResourcesToFree = nullptr);
 
 #if WITH_EDITOR
 	/**
diff --git a/Engine/Source/Runtime/Engine/Classes/Materials/MaterialInterface.h b/Engine/Source/Runtime/Engine/Classes/Materials/MaterialInterface.h
--- a/Engine/Source/Runtime/Engine/Classes/Materials/MaterialInterface.h
+++ b/Engine/Source/Runtime/Engine/Classes/Materials/MaterialInterface.h
@@ -33,6 +33,8 @@ struct FPrimitiveViewRelevance;
 struct FMaterialParameterInfo;
 struct FMaterialResourceLocOnDisk;
 
+typedef TArray<FMaterialResource*> FMaterialResourceDeferredDeletionArray;
+
 UENUM(BlueprintType)
 enum EMaterialUsage
 {
diff --git a/Engine/Source/Runtime/Engine/Private/Materials/Material.cpp b/Engine/Source/Runtime/Engine/Private/Materials/Material.cpp
--- a/Engine/Source/Runtime/Engine/Private/Materials/Material.cpp
+++ b/Engine/Source/Runtime/Engine/Private/Materials/Material.cpp
@@ -3063,7 +3063,10 @@ void UMaterial::CacheResourceShadersForRendering(bool bRegenerateId)
 		FlushResourceShaderMaps();
 	}
 
-	UpdateResourceAllocations();
+	// Resources cannot be deleted before uniform expressions are recached because
+	// UB layouts will be accessed and they are owned by material resources
+	FMaterialResourceDeferredDeletionArray ResourcesToFree;
+	UpdateResourceAllocations(&ResourcesToFree);
 
 	if (FApp::CanEverRender())
 	{
@@ -3118,7 +3121,21 @@ void UMaterial::CacheResourceShadersForRendering(bool bRegenerateId)
 			}
 		}
 
-		RecacheUniformExpressions(true);
+		// Material reload can call this function again. Don't recreate since mesh
+		// draw commands cache a pointer to the uniform buffer
+		RecacheUniformExpressions(!FPlatformProperties::RequiresCookedData());
+	}
+
+	if (ResourcesToFree.Num())
+	{
+		ENQUEUE_RENDER_COMMAND(CmdFreeUnusedMaterialResources)(
+			[ResourcesToFreeRT = MoveTemp(ResourcesToFree)](FRHICommandList&)
+		{
+			for (int32 Idx = 0; Idx < ResourcesToFreeRT.Num(); ++Idx)
+			{
+				delete ResourcesToFreeRT[Idx];
+			}
+		});
 	}
 }
 
@@ -3697,7 +3714,7 @@ void UMaterial::GetQualityLevelNodeUsage(TArray<bool, TInlineAllocator<EMaterial
 	}
 }
 
-void UMaterial::UpdateResourceAllocations()
+void UMaterial::UpdateResourceAllocations(FMaterialResourceDeferredDeletionArray* ResourcesToFree)
 {
 	if (FApp::CanEverRender())
 	{
@@ -3715,7 +3732,14 @@ void UMaterial::UpdateResourceAllocations()
 				FMaterialResource*& Resource = MaterialResources[Quality][Feature];
 				if (Feature != ActiveFeatureLevel || Quality != ActiveQualityLevel)
 				{
-					delete Resource;
+					if (ResourcesToFree)
+					{
+						ResourcesToFree->Add(Resource);
+					}
+					else
+					{
+						delete Resource;
+					}
 					Resource = nullptr;
 				}
 				else
diff --git a/Engine/Source/Runtime/Engine/Private/Materials/MaterialInstance.cpp b/Engine/Source/Runtime/Engine/Private/Materials/MaterialInstance.cpp
--- a/Engine/Source/Runtime/Engine/Private/Materials/MaterialInstance.cpp
+++ b/Engine/Source/Runtime/Engine/Private/Materials/MaterialInstance.cpp
@@ -2474,7 +2474,7 @@ FMaterialResource* UMaterialInstance::AllocatePermutationResource()
 	return new FMaterialResource();
 }
 
-void UMaterialInstance::UpdatePermutationAllocations()
+void UMaterialInstance::UpdatePermutationAllocations(FMaterialResourceDeferredDeletionArray* ResourcesToFree)
 {
 	if (bHasStaticPermutationResource)
 	{
@@ -2494,7 +2494,14 @@ void UMaterialInstance::UpdatePermutationAllocations()
 				FMaterialResource*& StaticPermResource = StaticPermutationMaterialResources[Quality][Feature];
 				if (Feature != ActiveFeatureLevel || Quality != ActiveQualityLevel)
 				{
-					delete StaticPermResource;
+					if (ResourcesToFree)
+					{
+						ResourcesToFree->Add(StaticPermResource);
+					}
+					else
+					{
+						delete StaticPermResource;
+					}
 					StaticPermResource = nullptr;
 				}
 				else
@@ -2535,7 +2542,8 @@ void UMaterialInstance::CacheResourceShadersForRendering()
 {
 	check(IsInGameThread() || IsAsyncLoading());
 
-	UpdatePermutationAllocations();
+	FMaterialResourceDeferredDeletionArray ResourcesToFree;
+	UpdatePermutationAllocations(&ResourcesToFree);
 	UpdateOverridableBaseProperties();
 
 	if (bHasStaticPermutationResource && FApp::CanEverRender())
@@ -2578,6 +2586,18 @@ void UMaterialInstance::CacheResourceShadersForRendering()
 	}
 
 	InitResources();
+
+	if (ResourcesToFree.Num())
+	{
+		ENQUEUE_RENDER_COMMAND(CmdFreeMaterialResources)(
+			[ResourcesToFreeRT = MoveTemp(ResourcesToFree)](FRHICommandList&)
+		{
+			for (int32 Idx = 0; Idx < ResourcesToFreeRT.Num(); ++Idx)
+			{
+				delete ResourcesToFreeRT[Idx];
+			}
+		});
+	}
 }
 
 void UMaterialInstance::CacheResourceShadersForCooking(EShaderPlatform ShaderPlatform, TArray<FMaterialResource*>& OutCachedMaterialResources)
diff --git a/Engine/Source/Runtime/Engine/Private/Materials/MaterialShared.cpp b/Engine/Source/Runtime/Engine/Private/Materials/MaterialShared.cpp
--- a/Engine/Source/Runtime/Engine/Private/Materials/MaterialShared.cpp
+++ b/Engine/Source/Runtime/Engine/Private/Materials/MaterialShared.cpp
@@ -2068,6 +2068,7 @@ void FMaterialRenderProxy::EvaluateUniformExpressions(FUniformExpressionCache& O
 		{
 			check(OutUniformExpressionCache.UniformBuffer->GetLayout() == UniformBufferStruct.GetLayout());
 			RHIUpdateUniformBuffer(OutUniformExpressionCache.UniformBuffer, TempBuffer);
+			OutUniformExpressionCache.UniformBuffer->UpdateLayoutReference(&UniformBufferStruct.GetLayout());
 		}
 		else
 		{
diff --git a/Engine/Source/Runtime/RHI/Public/RHIResources.h b/Engine/Source/Runtime/RHI/Public/RHIResources.h
--- a/Engine/Source/Runtime/RHI/Public/RHIResources.h
+++ b/Engine/Source/Runtime/RHI/Public/RHIResources.h
@@ -413,6 +413,13 @@ public:
 	}
 	const FRHIUniformBufferLayout& GetLayout() const { return *Layout; }
 
+	/** Same layout but different address */
+	void UpdateLayoutReference(const FRHIUniformBufferLayout* NewRef)
+	{
+		check(*Layout == *NewRef);
+		Layout = NewRef;
+	}
+
 #if VALIDATE_UNIFORM_BUFFER_LIFETIME
 	mutable int32 NumMeshCommandReferencesForDebugging = 0;
 #endif
