#UE4 When printing out FTimerUnifiedDelegates, non-dynamic delegates will now print their function address as well. Additionally the address will be resolved to a symbol name if symbols are available. Useful when using TimerManager.DumpTimerLogsThreshold or the ListTimers exec command.

diff --git a/Engine/Source/Runtime/Core/Public/Delegates/DelegateBase.h b/Engine/Source/Runtime/Core/Public/Delegates/DelegateBase.h
--- a/Engine/Source/Runtime/Core/Public/Delegates/DelegateBase.h
+++ b/Engine/Source/Runtime/Core/Public/Delegates/DelegateBase.h
@@ -100,6 +100,24 @@ public:
 
 #endif
 
+	/**
+	 * Returns the address of the method pointer which can be used to learn the address of the function that will be executed.
+	 * Returns nullptr if this delegate type does not directly invoke a function pointer.
+	 *
+	 * Note: Only intended to be used to aid debugging of delegates.
+	 *
+	 * @return The address of the function pointer that would be executed by this delegate
+	 */
+	void* GetBoundFunctionAddress() const
+	{
+		if (IDelegateInstance* Ptr = GetDelegateInstanceProtected())
+		{
+			return Ptr->GetBoundFunctionAddress();
+		}
+
+		return nullptr;
+	}
+
 	/**
 	 * If this is a UFunction or UObject delegate, return the UObject.
 	 *
diff --git a/Engine/Source/Runtime/Core/Public/Delegates/DelegateInstancesImpl.h b/Engine/Source/Runtime/Core/Public/Delegates/DelegateInstancesImpl.h
--- a/Engine/Source/Runtime/Core/Public/Delegates/DelegateInstancesImpl.h
+++ b/Engine/Source/Runtime/Core/Public/Delegates/DelegateInstancesImpl.h
@@ -71,6 +71,11 @@ public:
 
 #endif
 
+	virtual void* GetBoundFunctionAddress() const override final
+	{
+		return nullptr;
+	}
+
 	virtual UObject* GetUObject( ) const override final
 	{
 		return (UObject*)UserObjectPtr.Get();
@@ -229,6 +234,11 @@ public:
 
 #endif
 
+	virtual void* GetBoundFunctionAddress() const override final
+	{
+		return (void*)&MethodPtr;
+	}
+
 	virtual UObject* GetUObject() const override final
 	{
 		return nullptr;
@@ -410,6 +420,11 @@ public:
 
 #endif
 
+	virtual void* GetBoundFunctionAddress() const override final
+	{
+		return (void*)&MethodPtr;
+	}
+
 	virtual UObject* GetUObject( ) const override final
 	{
 		return nullptr;
@@ -563,6 +578,11 @@ public:
 
 #endif
 
+	virtual void* GetBoundFunctionAddress() const override final
+	{
+		return (void*)&MethodPtr;
+	}
+
 	virtual UObject* GetUObject( ) const override final
 	{
 		return (UObject*)UserObject.Get();
@@ -720,6 +740,11 @@ public:
 
 #endif
 
+	virtual void* GetBoundFunctionAddress() const override final
+	{
+		return (void*)&StaticFuncPtr;
+	}
+
 	virtual UObject* GetUObject( ) const override final
 	{
 		return nullptr;
@@ -858,6 +883,11 @@ public:
 
 #endif
 
+	virtual void* GetBoundFunctionAddress() const override final
+	{
+		return (void*)&Functor;
+	}
+
 	virtual UObject* GetUObject() const override final
 	{
 		return nullptr;
@@ -1003,6 +1033,11 @@ public:
 
 #endif
 
+	virtual void* GetBoundFunctionAddress() const override final
+	{
+		return (void*)&Functor;
+	}
+
 	virtual UObject* GetUObject() const override final
 	{
 		return ContextObject.Get();
diff --git a/Engine/Source/Runtime/Core/Public/Delegates/IDelegateInstance.h b/Engine/Source/Runtime/Core/Public/Delegates/IDelegateInstance.h
--- a/Engine/Source/Runtime/Core/Public/Delegates/IDelegateInstance.h
+++ b/Engine/Source/Runtime/Core/Public/Delegates/IDelegateInstance.h
@@ -85,6 +85,16 @@ public:
 
 #endif
 
+	/**
+	 * Returns the address of the method pointer which can be used to learn the address of the function that will be executed.
+	 * Returns nullptr if this delegate type does not directly invoke a function pointer.
+	 *
+	 * Note: Only intended to be used to aid debugging of delegates.
+	 *
+	 * @return The address of the function pointer that would be executed by this delegate
+	 */
+	virtual void* GetBoundFunctionAddress() const = 0;
+
 	/**
 	 * Returns the UObject that this delegate instance is bound to.
 	 *
diff --git a/Engine/Source/Runtime/Engine/Private/TimerManager.cpp b/Engine/Source/Runtime/Engine/Private/TimerManager.cpp
--- a/Engine/Source/Runtime/Engine/Private/TimerManager.cpp
+++ b/Engine/Source/Runtime/Engine/Private/TimerManager.cpp
@@ -12,6 +12,7 @@
 #include "Misc/TimeGuard.h"
 #include "ProfilingDebugging/CsvProfiler.h"
 #include "Algo/Transform.h"
+#include "HAL/PlatformStackWalk.h"
 
 DECLARE_CYCLE_STAT(TEXT("SetTimer"), STAT_SetTimer, STATGROUP_Engine);
 DECLARE_CYCLE_STAT(TEXT("SetTimeForNextTick"), STAT_SetTimerForNextTick, STATGROUP_Engine);
@@ -29,6 +30,12 @@ static FAutoConsoleVariableRef CVarDumpTimerLogsThreshold(
 	TEXT("Threshold (in milliseconds) after which we log timer info to try and help track down spikes in the timer code."),
 	ECVF_Default);
 
+static int32 DumpTimerLogSymbolNames = 1;
+static FAutoConsoleVariableRef CVarDumpTimerLogSymbolNames(
+	TEXT("TimerManager.DumpTimerLogSymbolNames"), DumpTimerLogSymbolNames,
+	TEXT("When logging timer info, symbol names will be included if set to 1."),
+	ECVF_Default);
+
 namespace
 {
 	void DescribeFTimerDataSafely(FOutputDevice& Ar, const FTimerData& Data)
@@ -139,28 +146,59 @@ void FTimerManager::OnCrash()
 FString FTimerUnifiedDelegate::ToString() const
 {
 	const UObject* Object = nullptr;
-	FName FunctionName = NAME_None;
+	FString FunctionNameStr;
 	bool bDynDelegate = false;
 
 	if (FuncDelegate.IsBound())
 	{
+		FName FunctionName;
 #if USE_DELEGATE_TRYGETBOUNDFUNCTIONNAME
 		FunctionName = FuncDelegate.TryGetBoundFunctionName();
 #endif
+		if (FunctionName.IsNone())
+		{
+			void* FuncPtr = FuncDelegate.GetBoundFunctionAddress();
+			if (FuncPtr)
+			{
+#if PLATFORM_64BITS
+				uint64 ProgramCounter = *((uint64*)FuncPtr);
+#else
+				uint64 ProgramCounter = *((uint32*)FuncPtr);
+#endif
+				FunctionNameStr = FString::Printf(TEXT("%#llx"), ProgramCounter);
+
+				if (DumpTimerLogSymbolNames)
+				{
+					FProgramCounterSymbolInfo SymbolInfo;
+					SymbolInfo.FunctionName[0] = 0;
+					SymbolInfo.Filename[0] = 0;
+					SymbolInfo.LineNumber = 0;
+					FPlatformStackWalk::ProgramCounterToSymbolInfo(ProgramCounter, SymbolInfo);
+					FunctionNameStr += FString::Printf(TEXT(" %s[%s:%d]"), ANSI_TO_TCHAR(SymbolInfo.FunctionName), ANSI_TO_TCHAR(SymbolInfo.Filename), SymbolInfo.LineNumber);
+				}
+			}
+			else
+			{
+				FunctionNameStr = TEXT("0x0");
+			}
+		}
+		else
+		{
+			FunctionNameStr = FunctionName.ToString();
+		}
 	}
 	else if (FuncDynDelegate.IsBound())
 	{
 		Object = FuncDynDelegate.GetUObject();
-		FunctionName = FuncDynDelegate.GetFunctionName();
+		FunctionNameStr = FuncDynDelegate.GetFunctionName().ToString();
 		bDynDelegate = true;
 	}
 	else
 	{
-		static FName NotBoundName(TEXT("NotBound!"));
-		FunctionName = NotBoundName;
+		FunctionNameStr = TEXT("NotBound!");
 	}
 
-	return FString::Printf(TEXT("%s,%s,%s"), bDynDelegate ? TEXT("DYN DELEGATE") : TEXT("DELEGATE"), Object == nullptr ? TEXT("NO OBJ") : *Object->GetPathName(), *FunctionName.ToString());
+	return FString::Printf(TEXT("%s,%s,%s"), bDynDelegate ? TEXT("DYN DELEGATE") : TEXT("DELEGATE"), Object == nullptr ? TEXT("NO OBJ") : *Object->GetPathName(), *FunctionNameStr);
 }
 
 // ---------------------------------
