Try to share initial property replication across connections
Uses new cvar "net.ShareInitialCompareState", disabled by default

diff --git a/Engine/Source/Runtime/Engine/Classes/Engine/EngineTypes.h b/Engine/Source/Runtime/Engine/Classes/Engine/EngineTypes.h
--- a/Engine/Source/Runtime/Engine/Classes/Engine/EngineTypes.h
+++ b/Engine/Source/Runtime/Engine/Classes/Engine/EngineTypes.h
@@ -3336,6 +3336,8 @@ struct FReplicationFlags
 			uint32 bIgnoreRPCs:1;
 			/** True if we should not swap the role and remote role of this actor when receiving properties. */
 			uint32 bSkipRoleSwap:1;
+			/** True if we should only compare role properties in CompareProperties */
+			uint32 bRolesOnly:1;
 		};
 
 		uint32	Value;
diff --git a/Engine/Source/Runtime/Engine/Private/RepLayout.cpp b/Engine/Source/Runtime/Engine/Private/RepLayout.cpp
--- a/Engine/Source/Runtime/Engine/Private/RepLayout.cpp
+++ b/Engine/Source/Runtime/Engine/Private/RepLayout.cpp
@@ -58,6 +58,9 @@ static FAutoConsoleVariableRef CVarUsePackedShadowBuffers(TEXT("Net.UsePackedSha
 int32 GShareShadowState = 1;
 static FAutoConsoleVariableRef CVarShareShadowState(TEXT("net.ShareShadowState"), GShareShadowState, TEXT("If true, work done to compare properties will be shared across connections"));
 
+int32 GShareInitialCompareState = 0;
+static FAutoConsoleVariableRef CVarShareInitialCompareState(TEXT("net.ShareInitialCompareState"), GShareInitialCompareState, TEXT("If true and net.ShareShadowState is enabled, attempt to also share initial replication compares across connections."));
+
 int32 MaxRepArraySize = UNetworkSettings::DefaultMaxRepArraySize;
 int32 MaxRepArrayMemory = UNetworkSettings::DefaultMaxRepArrayMemory;
 
@@ -709,6 +712,7 @@ FReplicationChangelistMgr::FReplicationChangelistMgr(
 	FCustomDeltaChangelistState* DeltaChangelistState):
 
 	LastReplicationFrame(0),
+	LastInitialReplicationFrame(0),
 	RepChangelistState(InRepLayout, Source, DeltaChangelistState)
 {
 }
@@ -739,21 +743,49 @@ void FRepLayout::UpdateChangelistMgr(
 	const FReplicationFlags& RepFlags,
 	const bool bForceCompare) const
 {
-	// See if we can re-use the work already done on a previous connection
-	// Rules:
-	//	1. We always compare once per frame (i.e. check LastReplicationFrame == ReplicationFrame)
-	//	2. We check LastCompareIndex > 1 so we can do at least one pass per connection to compare all properties
-	//		This is necessary due to how RemoteRole is manipulated per connection, so we need to give all connections a chance to see if it changed
-	//	3. We ALWAYS compare on bNetInitial to make sure we have a fresh changelist of net initial properties in this case
-	if (!bForceCompare && GShareShadowState && !RepFlags.bNetInitial && RepState->LastCompareIndex > 1 && InChangelistMgr.LastReplicationFrame == ReplicationFrame)
+	if (GShareInitialCompareState)
 	{
-		INC_DWORD_STAT_BY(STAT_NetSkippedDynamicProps, 1);
-		return;
+		// See if we can re-use the work already done on a previous connection
+		// Rules:
+		// 1. We have replicated this actor at least once this frame
+		// 2. This is not initial replication or we have done an initial replication this frame as well
+		if (!bForceCompare && GShareShadowState && (InChangelistMgr.LastReplicationFrame == ReplicationFrame) && (!RepFlags.bNetInitial || (InChangelistMgr.LastInitialReplicationFrame == ReplicationFrame)))
+		{
+			// If this is initial replication, or we have never replicated on this connection, force a role compare
+			if (RepFlags.bNetInitial || (RepState->LastCompareIndex == 0))
+			{
+				FReplicationFlags TempFlags = RepFlags;
+				TempFlags.bRolesOnly = true;
+				CompareProperties(RepState, &InChangelistMgr.RepChangelistState, (const uint8*)InObject, TempFlags);
+			}
+
+			INC_DWORD_STAT_BY(STAT_NetSkippedDynamicProps, 1);
+			return;
+		}
+	}
+	else
+	{
+		// See if we can re-use the work already done on a previous connection
+		// Rules:
+		//	1. We always compare once per frame (i.e. check LastReplicationFrame == ReplicationFrame)
+		//	2. We check LastCompareIndex > 1 so we can do at least one pass per connection to compare all properties
+		//		This is necessary due to how RemoteRole is manipulated per connection, so we need to give all connections a chance to see if it changed
+		//	3. We ALWAYS compare on bNetInitial to make sure we have a fresh changelist of net initial properties in this case
+		if (!bForceCompare && GShareShadowState && !RepFlags.bNetInitial && RepState->LastCompareIndex > 1 && InChangelistMgr.LastReplicationFrame == ReplicationFrame)
+		{
+			INC_DWORD_STAT_BY(STAT_NetSkippedDynamicProps, 1);
+			return;
+		}
 	}
 
 	CompareProperties(RepState, &InChangelistMgr.RepChangelistState, (const uint8*)InObject, RepFlags);
 
 	InChangelistMgr.LastReplicationFrame = ReplicationFrame;
+
+	if (GShareInitialCompareState && RepFlags.bNetInitial)
+	{
+		InChangelistMgr.LastInitialReplicationFrame = ReplicationFrame;
+	}
 }
 
 void FReplicationChangelistMgr::CountBytes(FArchive& Ar) const
@@ -788,6 +820,41 @@ static void CompareProperties_Array_r(
 	const uint16 CmdIndex,
 	const uint16 Handle);
 
+static void CompareRoleProperties(
+	const FComparePropertiesSharedParams& SharedParams,
+	FSendingRepState* RESTRICT RepState,
+	const FConstRepObjectDataBuffer Data,
+	TArray<uint16>& Changed)
+{
+	if (RepState && SharedParams.RoleIndex != INDEX_NONE && SharedParams.RemoteRoleIndex != INDEX_NONE)
+	{
+		// Verify that the order hasn't changed, otherwise ReceiveProperties will fail
+		check(SharedParams.RemoteRoleIndex < SharedParams.RoleIndex);
+
+		const FRepParentCmd& RemoteRoleParent = SharedParams.Parents[SharedParams.RemoteRoleIndex];
+		const FRepLayoutCmd& RemoteRoleCmd = SharedParams.Cmds[RemoteRoleParent.CmdStart];
+		const uint16 RemoteRoleHandle = RemoteRoleCmd.RelativeHandle;
+
+		const ENetRole ObjectRemoteRole = *(const ENetRole*)(Data + RemoteRoleParent).Data;
+		if (SharedParams.bForceFail || RepState->SavedRemoteRole != ObjectRemoteRole)
+		{
+			RepState->SavedRemoteRole = ObjectRemoteRole;
+			Changed.Add(RemoteRoleHandle);
+		}
+
+		const FRepParentCmd& RoleParent = SharedParams.Parents[SharedParams.RoleIndex];
+		const FRepLayoutCmd& RoleCmd = SharedParams.Cmds[RoleParent.CmdStart];
+		const uint16 RoleHandle = RoleCmd.RelativeHandle;
+
+		const ENetRole ObjectRole = *(const ENetRole*)(Data + RoleParent).Data;
+		if (SharedParams.bForceFail || RepState->SavedRole != ObjectRole)
+		{
+			RepState->SavedRole = ObjectRole;
+			Changed.Add(RoleHandle);
+		}
+	}
+}
+
 static void CompareParentProperties(
 	const FComparePropertiesSharedParams& SharedParams,
 	FSendingRepState* RESTRICT RepState,
@@ -975,7 +1042,14 @@ bool FRepLayout::CompareProperties(
 		Cmds
 	};
 
-	CompareParentProperties(SharedParams, RepState, RepChangelistState, Data, Changed);
+	if (RepFlags.bRolesOnly)
+	{
+		CompareRoleProperties(SharedParams, RepState, Data, Changed);
+	}
+	else
+	{
+		CompareParentProperties(SharedParams, RepState, RepChangelistState, Data, Changed);
+	}
 
 	if (Changed.Num() == 0)
 	{
@@ -2651,7 +2725,7 @@ static bool ReceiveProperties_r(FReceivePropertiesSharedParams& Params, FReceive
 		}
 		else
 		{
-			UE_LOG(LogRepProperties, VeryVerbose, TEXT("ReceiveProperties_r: Parent=%d, Cmd=%d, ArrayIndex=%d"), Cmd.ParentIndex, CmdIndex);
+			UE_LOG(LogRepProperties, VeryVerbose, TEXT("ReceiveProperties_r: Parent=%d, Cmd=%d"), Cmd.ParentIndex, CmdIndex);
 				
 			if (ERepLayoutCmdType::DynamicArray == Cmd.Type)
 			{
diff --git a/Engine/Source/Runtime/Engine/Public/Net/RepLayout.h b/Engine/Source/Runtime/Engine/Public/Net/RepLayout.h
--- a/Engine/Source/Runtime/Engine/Public/Net/RepLayout.h
+++ b/Engine/Source/Runtime/Engine/Public/Net/RepLayout.h
@@ -496,6 +496,8 @@ public:
 private:
 
 	uint32 LastReplicationFrame;
+	uint32 LastInitialReplicationFrame;
+
 	FRepChangelistState RepChangelistState;
 };
 
