#UE4 Now printing out timers after exceeding the log timers threshold before they are executed so they are not skipped if cleared during execution. The first one is printed immediately if it is not cleared as part of its execution. Also fully disabling timer logging by default (must be enabled via CVar)

diff --git a/Engine/Source/Runtime/Engine/Private/TimerManager.cpp b/Engine/Source/Runtime/Engine/Private/TimerManager.cpp
--- a/Engine/Source/Runtime/Engine/Private/TimerManager.cpp
+++ b/Engine/Source/Runtime/Engine/Private/TimerManager.cpp
@@ -24,10 +24,10 @@ CSV_DECLARE_CATEGORY_MODULE_EXTERN(CORE_API, Basic);
 /** Track the last assigned handle globally */
 uint64 FTimerManager::LastAssignedSerialNumber = 0;
 
-static float DumpTimerLogsThreshold = 1000.f;
+static float DumpTimerLogsThreshold = 0.f;
 static FAutoConsoleVariableRef CVarDumpTimerLogsThreshold(
 	TEXT("TimerManager.DumpTimerLogsThreshold"), DumpTimerLogsThreshold,
-	TEXT("Threshold (in milliseconds) after which we log timer info to try and help track down spikes in the timer code."),
+	TEXT("Threshold (in milliseconds) after which we log timer info to try and help track down spikes in the timer code. Disabled when set to 0"),
 	ECVF_Default);
 
 static int32 DumpTimerLogSymbolNames = 1;
@@ -592,7 +592,7 @@ void FTimerManager::Tick(float DeltaTime)
 	}
 
 	const double StartTime = FPlatformTime::Seconds();
-	bool TimeExceededLogged = false;
+	bool bDumpTimerLogsThresholdExceeded = false;
 
 	InternalTime += DeltaTime;
 
@@ -618,6 +618,11 @@ void FTimerManager::Tick(float DeltaTime)
 		{
 			// Timer has expired! Fire the delegate, then handle potential looping.
 
+			if (bDumpTimerLogsThresholdExceeded)
+			{
+				DescribeFTimerDataSafely(*GLog, *Top);
+			}
+
 			// Set the relevant level context for this timer
 			const int32 LevelCollectionIndex = OwningWorld ? OwningWorld->FindCollectionIndexByType(Top->LevelCollection) : INDEX_NONE;
 			
@@ -657,23 +662,23 @@ void FTimerManager::Tick(float DeltaTime)
 				}
 			}
 
-			// help us hunt down outliers that cause our timer manager times to spike.  Recommended that users set meaningful DumpTimerLogsThresholds in appropriate ini files if they are seeing spikes in the timer manager.
-			const double DeltaT = (FPlatformTime::Seconds() - StartTime) * 1000.f;
-			if (DeltaT >= DumpTimerLogsThreshold)
+			if (DumpTimerLogsThreshold > 0.f && !bDumpTimerLogsThresholdExceeded)
 			{
-				if (!TimeExceededLogged)
+				// help us hunt down outliers that cause our timer manager times to spike.  Recommended that users set meaningful DumpTimerLogsThresholds in appropriate ini files if they are seeing spikes in the timer manager.
+				const double DeltaT = (FPlatformTime::Seconds() - StartTime) * 1000.f;
+				if (DeltaT >= DumpTimerLogsThreshold)
 				{
+					bDumpTimerLogsThresholdExceeded = true;
 					UE_LOG(LogEngine, Log, TEXT("TimerManager's time threshold of %.2fms exceeded with a deltaT of %.4f, dumping current timer data."), DumpTimerLogsThreshold, DeltaT);
-					TimeExceededLogged = true;
-				}
-				
-				if (Top)
-				{
-					DescribeFTimerDataSafely(*GLog, *Top);
-				}
-				else
-				{
-					UE_LOG(LogEngine, Log, TEXT("There was no timer data!"));
+
+					if (Top)
+					{
+						DescribeFTimerDataSafely(*GLog, *Top);
+					}
+					else
+					{
+						UE_LOG(LogEngine, Log, TEXT("There was no timer data for the first timer after exceeding the time threshold!"));
+					}
 				}
 			}
 
