#ue4 - Re-enable jump apex position trace with "p.ForceJumpPeakSubstep". Add limited number of attempts per tick (MaxJumpApexAttemptsPerSimulation) to prevent deep recursion when special cases cause collision or other conditions which reactivate the apex condition.

diff --git a/Engine/Source/Runtime/Engine/Classes/GameFramework/CharacterMovementComponent.h b/Engine/Source/Runtime/Engine/Classes/GameFramework/CharacterMovementComponent.h
--- a/Engine/Source/Runtime/Engine/Classes/GameFramework/CharacterMovementComponent.h
+++ b/Engine/Source/Runtime/Engine/Classes/GameFramework/CharacterMovementComponent.h
@@ -668,6 +668,12 @@ protected:
 	/** Used when throttling "stuck in geometry" logging, to output the number of events we skipped if throttling. */
 	uint32 StuckWarningCountSinceNotify;
 
+	/**
+	 * Used to limit number of jump apex attempts per tick.
+	 * @see MaxJumpApexAttemptsPerSimulation
+	 */
+	int32 NumJumpApexAttempts;
+
 public:
 
 	/** Returns the location at the end of the last tick. */
@@ -721,6 +727,13 @@ public:
 	UPROPERTY(Category="Character Movement (General Settings)", EditAnywhere, BlueprintReadWrite, AdvancedDisplay, meta=(ClampMin="1", ClampMax="25", UIMin="1", UIMax="25"))
 	int32 MaxSimulationIterations;
 
+	/**
+	 * Max number of attempts per simulation to attempt to exactly reach the jump apex when falling movement reaches the top of the arc.
+	 * Limiting this prevents deep recursion when special cases cause collision or other conditions which reactivate the apex condition.
+	 */
+	UPROPERTY(Category="Character Movement (General Settings)", EditAnywhere, BlueprintReadWrite, AdvancedDisplay, meta=(ClampMin="1", ClampMax="4", UIMin="1", UIMax="4"))
+	int32 MaxJumpApexAttemptsPerSimulation;
+
 	/**
 	* Max distance we allow simulated proxies to depenetrate when moving out of anything but Pawns.
 	* This is generally more tolerant than with Pawns, because other geometry is either not moving, or is moving predictably with a bit of delay compared to on the server.
diff --git a/Engine/Source/Runtime/Engine/Private/Components/CharacterMovementComponent.cpp b/Engine/Source/Runtime/Engine/Private/Components/CharacterMovementComponent.cpp
--- a/Engine/Source/Runtime/Engine/Private/Components/CharacterMovementComponent.cpp
+++ b/Engine/Source/Runtime/Engine/Private/Components/CharacterMovementComponent.cpp
@@ -196,7 +196,7 @@ namespace CharacterMovementCVars
 		TEXT( "If 1, remove invalid replay samples that can occur due to oversampling (sampling at higher rate than physics is being ticked)" ),
 		ECVF_Default);
 
-	static int32 ForceJumpPeakSubstep = 0;
+	static int32 ForceJumpPeakSubstep = 1;
 	FAutoConsoleVariableRef CVarForceJumpPeakSubstep(
 		TEXT("p.ForceJumpPeakSubstep"),
 		ForceJumpPeakSubstep,
@@ -404,6 +404,8 @@ UCharacterMovementComponent::UCharacterMovementComponent(const FObjectInitialize
 	
 	MaxSimulationTimeStep = 0.05f;
 	MaxSimulationIterations = 8;
+	MaxJumpApexAttemptsPerSimulation = 2;
+	NumJumpApexAttempts = 0;
 
 	MaxDepenetrationWithGeometry = 500.f;
 	MaxDepenetrationWithGeometryAsProxy = 100.f;
@@ -2825,6 +2827,10 @@ void UCharacterMovementComponent::StartNewPhysics(float deltaTime, int32 Iterati
 
 	const bool bSavedMovementInProgress = bMovementInProgress;
 	bMovementInProgress = true;
+	if (Iterations == 0)
+	{
+		NumJumpApexAttempts = 0;
+	}
 
 	switch ( MovementMode )
 	{
@@ -4116,7 +4122,6 @@ void UCharacterMovementComponent::PhysFalling(float deltaTime, int32 Iterations)
 	FVector FallAcceleration = GetFallingLateralAcceleration(deltaTime);
 	FallAcceleration.Z = 0.f;
 	const bool bHasAirControl = (FallAcceleration.SizeSquared2D() > 0.f);
-	int32 NumApexAttempts = 0;
 
 	float remainingTime = deltaTime;
 	while( (remainingTime >= MIN_TICK_TIME) && (Iterations < MaxSimulationIterations) )
@@ -4192,7 +4197,7 @@ void UCharacterMovementComponent::PhysFalling(float deltaTime, int32 Iterations)
 		const FVector AirControlAccel = (Velocity - VelocityNoAirControl) / timeTick;
 
 		// See if we need to sub-step to exactly reach the apex. This is important for avoiding "cutting off the top" of the trajectory as framerate varies.
-		if (CharacterMovementCVars::ForceJumpPeakSubstep && OldVelocity.Z > 0.f && Velocity.Z <= 0.f && NumApexAttempts < 2)
+		if (CharacterMovementCVars::ForceJumpPeakSubstep && OldVelocity.Z > 0.f && Velocity.Z <= 0.f && NumJumpApexAttempts < MaxJumpApexAttemptsPerSimulation)
 		{
 			const FVector DerivedAccel = (Velocity - OldVelocity) / timeTick;
 			if (!FMath::IsNearlyZero(DerivedAccel.Z))
@@ -4211,7 +4216,7 @@ void UCharacterMovementComponent::PhysFalling(float deltaTime, int32 Iterations)
 					remainingTime += (timeTick - TimeToApex);
 					timeTick = TimeToApex;
 					Iterations--;
-					NumApexAttempts++;
+					NumJumpApexAttempts++;
 				}
 			}
 		}
