#UE4 Minor formatting fix for more platforms

diff --git a/Engine/Source/Runtime/Engine/Private/TimerManager.cpp b/Engine/Source/Runtime/Engine/Private/TimerManager.cpp
--- a/Engine/Source/Runtime/Engine/Private/TimerManager.cpp
+++ b/Engine/Source/Runtime/Engine/Private/TimerManager.cpp
@@ -165,7 +165,7 @@ FString FTimerUnifiedDelegate::ToString() const
 #else
 				uint64 ProgramCounter = *((uint32*)FuncPtr);
 #endif
-				FunctionNameStr = FString::Printf(TEXT("%#llx"), ProgramCounter);
+				FunctionNameStr = FString::Printf(TEXT("0x%llx"), ProgramCounter);
 
 				if (DumpTimerLogSymbolNames)
 				{
