Fixed crash due to shader not ready for rect light importance sanpled data. Now delayed later.
This is a simple change to delay the creation of rect light importance sampling textures for 4.22.
A follow up is planned in https://jira.it.epicgames.net/browse/UE-7183.

diff --git a/Engine/Source/Runtime/Engine/Private/Components/RectLightComponent.cpp b/Engine/Source/Runtime/Engine/Private/Components/RectLightComponent.cpp
--- a/Engine/Source/Runtime/Engine/Private/Components/RectLightComponent.cpp
+++ b/Engine/Source/Runtime/Engine/Private/Components/RectLightComponent.cpp
@@ -23,26 +23,6 @@ float GetRectLightBarnDoorMaxAngle()
 	return 88.f;
 }
 
-#if RHI_RAYTRACING
-FRectLightRayTracingData BuildRectLightMipTree(FRHICommandListImmediate& RHICmdList, UTexture* SourceTexture);
-#endif
-
-void URectLightComponent::UpdateRayTracingData()
-{
-#if RHI_RAYTRACING
-	if (IsRayTracingEnabled())
-	{
-		UTexture* NewSourceTexture = SourceTexture;
-		FRectLightRayTracingData* UpdatedRayTracingData = RayTracingData;
-		ENQUEUE_RENDER_COMMAND(BuildRectLightMipTree)(
-			[NewSourceTexture, UpdatedRayTracingData](FRHICommandListImmediate& RHICmdList) mutable
-		{
-			*UpdatedRayTracingData = BuildRectLightMipTree(RHICmdList, NewSourceTexture);
-		});
-	}
-#endif
-}
-
 URectLightComponent::URectLightComponent(const FObjectInitializer& ObjectInitializer)
 	: Super(ObjectInitializer)
 	, RayTracingData(new FRectLightRayTracingData())
@@ -65,8 +45,7 @@ URectLightComponent::URectLightComponent(const FObjectInitializer& ObjectInitial
 	SourceTexture = nullptr;
 	BarnDoorAngle = GetRectLightBarnDoorMaxAngle();
 	BarnDoorLength = 20.0f;
-
-	UpdateRayTracingData();
+	// RayTracingData will be initialised on the render thread.
 }
 
 FLightSceneProxy* URectLightComponent::CreateSceneProxy() const
@@ -81,7 +60,7 @@ void URectLightComponent::SetSourceTexture(UTexture* NewValue)
 	{
 		SourceTexture = NewValue;
 
-		UpdateRayTracingData();
+		// This will trigger a recreation of the LightSceneProxy and update RayTracingData accordingly if the texture has changed.
 		MarkRenderStateDirty();
 	}
 }
@@ -286,135 +265,3 @@ bool FRectLightSceneProxy::GetWholeSceneProjectedShadowInitializer(const FSceneV
 
 	return false;
 }
-
-#if RHI_RAYTRACING
-
-class FBuildRectLightMipTreeCS : public FGlobalShader
-{
-	DECLARE_SHADER_TYPE(FBuildRectLightMipTreeCS, Global)
-
-public:
-	static bool ShouldCompilePermutation(const FGlobalShaderPermutationParameters& Parameters)
-	{
-		return ShouldCompileRayTracingShadersForProject(Parameters.Platform);
-	}
-
-	static uint32 GetGroupSize()
-	{
-		return 16;
-	}
-
-	static void ModifyCompilationEnvironment(const FGlobalShaderPermutationParameters& Parameters, FShaderCompilerEnvironment& OutEnvironment)
-	{
-		FGlobalShader::ModifyCompilationEnvironment(Parameters, OutEnvironment);
-		OutEnvironment.SetDefine(TEXT("THREADGROUP_SIZE"), GetGroupSize());
-	}
-
-	FBuildRectLightMipTreeCS(const ShaderMetaType::CompiledShaderInitializerType& Initializer)
-		: FGlobalShader(Initializer)
-	{
-		TextureParameter.Bind(Initializer.ParameterMap, TEXT("RectLightTexture"));
-		TextureSamplerParameter.Bind(Initializer.ParameterMap, TEXT("TextureSampler"));
-		DimensionsParameter.Bind(Initializer.ParameterMap, TEXT("Dimensions"));
-		MipLevelParameter.Bind(Initializer.ParameterMap, TEXT("MipLevel"));
-		MipTreeParameter.Bind(Initializer.ParameterMap, TEXT("MipTree"));
-	}
-
-	FBuildRectLightMipTreeCS() {}
-
-	void SetParameters(
-		FRHICommandList& RHICmdList,
-		FTextureRHIRef Texture,
-		const FIntVector& Dimensions,
-		uint32 MipLevel,
-		FRWBuffer& MipTree
-	)
-	{
-		FComputeShaderRHIParamRef ShaderRHI = GetComputeShader();
-
-		SetShaderValue(RHICmdList, ShaderRHI, DimensionsParameter, Dimensions);
-		SetShaderValue(RHICmdList, ShaderRHI, MipLevelParameter, MipLevel);
-		SetTextureParameter(RHICmdList, ShaderRHI, TextureParameter, TextureSamplerParameter, TStaticSamplerState<SF_Bilinear>::GetRHI(), Texture);
-
-		check(MipTreeParameter.IsBound());
-		MipTreeParameter.SetBuffer(RHICmdList, ShaderRHI, MipTree);
-	}
-
-	void UnsetParameters(
-		FRHICommandList& RHICmdList,
-		EResourceTransitionAccess TransitionAccess,
-		EResourceTransitionPipeline TransitionPipeline,
-		FRWBuffer& MipTree,
-		FComputeFenceRHIParamRef Fence)
-	{
-		FComputeShaderRHIParamRef ShaderRHI = GetComputeShader();
-
-		MipTreeParameter.UnsetUAV(RHICmdList, ShaderRHI);
-		RHICmdList.TransitionResource(TransitionAccess, TransitionPipeline, MipTree.UAV, Fence);
-	}
-
-	virtual bool Serialize(FArchive& Ar)
-	{
-		bool bShaderHasOutdatedParameters = FGlobalShader::Serialize(Ar);
-		Ar << TextureParameter;
-		Ar << TextureSamplerParameter;
-		Ar << DimensionsParameter;
-		Ar << MipLevelParameter;
-		Ar << MipTreeParameter;
-		return bShaderHasOutdatedParameters;
-	}
-
-private:
-	FShaderResourceParameter TextureParameter;
-	FShaderResourceParameter TextureSamplerParameter;
-
-	FShaderParameter DimensionsParameter;
-	FShaderParameter MipLevelParameter;
-	FRWShaderParameter MipTreeParameter;
-};
-
-IMPLEMENT_SHADER_TYPE(, FBuildRectLightMipTreeCS, TEXT("/Engine/Private/Raytracing/BuildMipTreeCS.usf"), TEXT("BuildRectLightMipTreeCS"), SF_Compute)
-
-DECLARE_GPU_STAT_NAMED(BuildRectLightMipTreeStat, TEXT("build RectLight MipTree"));
-
-FRectLightRayTracingData BuildRectLightMipTree(FRHICommandListImmediate& RHICmdList, UTexture* SourceTexture)
-{
-	SCOPED_GPU_STAT(RHICmdList, BuildRectLightMipTreeStat);
-
-	check(IsInRenderingThread());
-	FRectLightRayTracingData Data;
-	FTextureRHIRef RhiTexture = SourceTexture ? SourceTexture->Resource->TextureRHI : GWhiteTexture->TextureRHI;
-
-	const auto ShaderMap = GetGlobalShaderMap(ERHIFeatureLevel::SM5);
-	TShaderMapRef<FBuildRectLightMipTreeCS> BuildRectLightMipTreeComputeShader(ShaderMap);
-	RHICmdList.SetComputeShader(BuildRectLightMipTreeComputeShader->GetComputeShader());
-
-	// Allocate MIP tree
-	FIntVector TextureSize = RhiTexture->GetSizeXYZ();
-	uint32 MipLevelCount = FMath::Min(FMath::CeilLogTwo(TextureSize.X), FMath::CeilLogTwo(TextureSize.Y));
-	Data.RectLightMipTreeDimensions = FIntVector(1 << MipLevelCount, 1 << MipLevelCount, 1);
-	uint32 NumElements = Data.RectLightMipTreeDimensions.X * Data.RectLightMipTreeDimensions.Y;
-	for (uint32 MipLevel = 1; MipLevel <= MipLevelCount; ++MipLevel)
-	{
-		uint32 NumElementsInLevel = (Data.RectLightMipTreeDimensions.X >> MipLevel) * (Data.RectLightMipTreeDimensions.Y >> MipLevel);
-		NumElements += NumElementsInLevel;
-	}
-
-	Data.RectLightMipTree.Initialize(sizeof(float), NumElements, PF_R32_FLOAT, BUF_UnorderedAccess | BUF_ShaderResource);
-
-	// Execute hierarchical build
-	for (uint32 MipLevel = 0; MipLevel <= MipLevelCount; ++MipLevel)
-	{
-		FComputeFenceRHIRef MipLevelFence = RHICmdList.CreateComputeFence(TEXT("RectLightMipTree Build"));
-		BuildRectLightMipTreeComputeShader->SetParameters(RHICmdList, RhiTexture, Data.RectLightMipTreeDimensions, MipLevel, Data.RectLightMipTree);
-		FIntVector MipLevelDimensions = FIntVector(Data.RectLightMipTreeDimensions.X >> MipLevel, Data.RectLightMipTreeDimensions.Y >> MipLevel, 1);
-		FIntVector NumGroups = FIntVector::DivideAndRoundUp(MipLevelDimensions, FBuildRectLightMipTreeCS::GetGroupSize());
-		DispatchComputeShader(RHICmdList, *BuildRectLightMipTreeComputeShader, NumGroups.X, NumGroups.Y, 1);
-		BuildRectLightMipTreeComputeShader->UnsetParameters(RHICmdList, EResourceTransitionAccess::ERWBarrier, EResourceTransitionPipeline::EComputeToCompute, Data.RectLightMipTree, MipLevelFence);
-	}
-	FComputeFenceRHIRef TransitionFence = RHICmdList.CreateComputeFence(TEXT("RectLightMipTree Transition"));
-	BuildRectLightMipTreeComputeShader->UnsetParameters(RHICmdList, EResourceTransitionAccess::ERWBarrier, EResourceTransitionPipeline::EComputeToCompute, Data.RectLightMipTree, TransitionFence);
-
-	return Data;
-}
-#endif // RHI_RAYTRACING
diff --git a/Engine/Source/Runtime/Engine/Public/RectLightSceneProxy.h b/Engine/Source/Runtime/Engine/Public/RectLightSceneProxy.h
--- a/Engine/Source/Runtime/Engine/Public/RectLightSceneProxy.h
+++ b/Engine/Source/Runtime/Engine/Public/RectLightSceneProxy.h
@@ -16,6 +16,8 @@ struct FRectLightRayTracingData
 #if RHI_RAYTRACING
 	FRWBuffer   RectLightMipTree;
 	FIntVector  RectLightMipTreeDimensions;
+	FGuid		TextureLightingGuid;
+	bool		bInitialised = false;
 #endif
 };
 
@@ -26,7 +28,7 @@ public:
 	float		SourceHeight;
 	float		BarnDoorAngle;
 	float		BarnDoorLength;
-	FRectLightRayTracingData* RayTracingData;
+	FRectLightRayTracingData* RayTracingData; // Render thread data only
 	UTexture*	SourceTexture;
 
 	FRectLightSceneProxy(const URectLightComponent* Component);
diff --git a/Engine/Source/Runtime/Renderer/Private/RayTracing/RayTracingRectLight.cpp b/Engine/Source/Runtime/Renderer/Private/RayTracing/RayTracingRectLight.cpp
--- a/Engine/Source/Runtime/Renderer/Private/RayTracing/RayTracingRectLight.cpp
+++ b/Engine/Source/Runtime/Renderer/Private/RayTracing/RayTracingRectLight.cpp
@@ -74,6 +74,139 @@ DECLARE_GPU_STAT_NAMED(RayTracingRectLight, TEXT("Ray Tracing RectLight"));
 
 IMPLEMENT_GLOBAL_SHADER_PARAMETER_STRUCT(FRectLightData, "RectLight");
 
+
+
+class FBuildRectLightMipTreeCS : public FGlobalShader
+{
+	DECLARE_SHADER_TYPE(FBuildRectLightMipTreeCS, Global)
+
+public:
+	static bool ShouldCompilePermutation(const FGlobalShaderPermutationParameters& Parameters)
+	{
+		return ShouldCompileRayTracingShadersForProject(Parameters.Platform);
+	}
+
+	static uint32 GetGroupSize()
+	{
+		return 16;
+	}
+
+	static void ModifyCompilationEnvironment(const FGlobalShaderPermutationParameters& Parameters, FShaderCompilerEnvironment& OutEnvironment)
+	{
+		FGlobalShader::ModifyCompilationEnvironment(Parameters, OutEnvironment);
+		OutEnvironment.SetDefine(TEXT("THREADGROUP_SIZE"), GetGroupSize());
+	}
+
+	FBuildRectLightMipTreeCS(const ShaderMetaType::CompiledShaderInitializerType& Initializer)
+		: FGlobalShader(Initializer)
+	{
+		TextureParameter.Bind(Initializer.ParameterMap, TEXT("RectLightTexture"));
+		TextureSamplerParameter.Bind(Initializer.ParameterMap, TEXT("TextureSampler"));
+		DimensionsParameter.Bind(Initializer.ParameterMap, TEXT("Dimensions"));
+		MipLevelParameter.Bind(Initializer.ParameterMap, TEXT("MipLevel"));
+		MipTreeParameter.Bind(Initializer.ParameterMap, TEXT("MipTree"));
+	}
+
+	FBuildRectLightMipTreeCS() {}
+
+	void SetParameters(
+		FRHICommandList& RHICmdList,
+		FTextureRHIRef Texture,
+		const FIntVector& Dimensions,
+		uint32 MipLevel,
+		FRWBuffer& MipTree
+	)
+	{
+		FComputeShaderRHIParamRef ShaderRHI = GetComputeShader();
+
+		SetShaderValue(RHICmdList, ShaderRHI, DimensionsParameter, Dimensions);
+		SetShaderValue(RHICmdList, ShaderRHI, MipLevelParameter, MipLevel);
+		SetTextureParameter(RHICmdList, ShaderRHI, TextureParameter, TextureSamplerParameter, TStaticSamplerState<SF_Bilinear>::GetRHI(), Texture);
+
+		check(MipTreeParameter.IsBound());
+		MipTreeParameter.SetBuffer(RHICmdList, ShaderRHI, MipTree);
+	}
+
+	void UnsetParameters(
+		FRHICommandList& RHICmdList,
+		EResourceTransitionAccess TransitionAccess,
+		EResourceTransitionPipeline TransitionPipeline,
+		FRWBuffer& MipTree,
+		FComputeFenceRHIParamRef Fence)
+	{
+		FComputeShaderRHIParamRef ShaderRHI = GetComputeShader();
+
+		MipTreeParameter.UnsetUAV(RHICmdList, ShaderRHI);
+		RHICmdList.TransitionResource(TransitionAccess, TransitionPipeline, MipTree.UAV, Fence);
+	}
+
+	virtual bool Serialize(FArchive& Ar)
+	{
+		bool bShaderHasOutdatedParameters = FGlobalShader::Serialize(Ar);
+		Ar << TextureParameter;
+		Ar << TextureSamplerParameter;
+		Ar << DimensionsParameter;
+		Ar << MipLevelParameter;
+		Ar << MipTreeParameter;
+		return bShaderHasOutdatedParameters;
+	}
+
+private:
+	FShaderResourceParameter TextureParameter;
+	FShaderResourceParameter TextureSamplerParameter;
+
+	FShaderParameter DimensionsParameter;
+	FShaderParameter MipLevelParameter;
+	FRWShaderParameter MipTreeParameter;
+};
+
+IMPLEMENT_SHADER_TYPE(, FBuildRectLightMipTreeCS, TEXT("/Engine/Private/Raytracing/BuildMipTreeCS.usf"), TEXT("BuildRectLightMipTreeCS"), SF_Compute)
+
+DECLARE_GPU_STAT_NAMED(BuildRectLightMipTreeStat, TEXT("build RectLight MipTree"));
+
+FRectLightRayTracingData BuildRectLightMipTree(FRHICommandListImmediate& RHICmdList, UTexture* SourceTexture)
+{
+	SCOPED_GPU_STAT(RHICmdList, BuildRectLightMipTreeStat);
+
+	check(IsInRenderingThread());
+	FRectLightRayTracingData Data;
+	FTextureRHIRef RhiTexture = SourceTexture ? SourceTexture->Resource->TextureRHI : GWhiteTexture->TextureRHI;
+
+	const auto ShaderMap = GetGlobalShaderMap(ERHIFeatureLevel::SM5);
+	TShaderMapRef<FBuildRectLightMipTreeCS> BuildRectLightMipTreeComputeShader(ShaderMap);
+	RHICmdList.SetComputeShader(BuildRectLightMipTreeComputeShader->GetComputeShader());
+
+	// Allocate MIP tree
+	FIntVector TextureSize = RhiTexture->GetSizeXYZ();
+	uint32 MipLevelCount = FMath::Min(FMath::CeilLogTwo(TextureSize.X), FMath::CeilLogTwo(TextureSize.Y));
+	Data.RectLightMipTreeDimensions = FIntVector(1 << MipLevelCount, 1 << MipLevelCount, 1);
+	uint32 NumElements = Data.RectLightMipTreeDimensions.X * Data.RectLightMipTreeDimensions.Y;
+	for (uint32 MipLevel = 1; MipLevel <= MipLevelCount; ++MipLevel)
+	{
+		uint32 NumElementsInLevel = (Data.RectLightMipTreeDimensions.X >> MipLevel) * (Data.RectLightMipTreeDimensions.Y >> MipLevel);
+		NumElements += NumElementsInLevel;
+	}
+
+	Data.RectLightMipTree.Initialize(sizeof(float), NumElements, PF_R32_FLOAT, BUF_UnorderedAccess | BUF_ShaderResource);
+
+	// Execute hierarchical build
+	for (uint32 MipLevel = 0; MipLevel <= MipLevelCount; ++MipLevel)
+	{
+		FComputeFenceRHIRef MipLevelFence = RHICmdList.CreateComputeFence(TEXT("RectLightMipTree Build"));
+		BuildRectLightMipTreeComputeShader->SetParameters(RHICmdList, RhiTexture, Data.RectLightMipTreeDimensions, MipLevel, Data.RectLightMipTree);
+		FIntVector MipLevelDimensions = FIntVector(Data.RectLightMipTreeDimensions.X >> MipLevel, Data.RectLightMipTreeDimensions.Y >> MipLevel, 1);
+		FIntVector NumGroups = FIntVector::DivideAndRoundUp(MipLevelDimensions, FBuildRectLightMipTreeCS::GetGroupSize());
+		DispatchComputeShader(RHICmdList, *BuildRectLightMipTreeComputeShader, NumGroups.X, NumGroups.Y, 1);
+		BuildRectLightMipTreeComputeShader->UnsetParameters(RHICmdList, EResourceTransitionAccess::ERWBarrier, EResourceTransitionPipeline::EComputeToCompute, Data.RectLightMipTree, MipLevelFence);
+	}
+	FComputeFenceRHIRef TransitionFence = RHICmdList.CreateComputeFence(TEXT("RectLightMipTree Transition"));
+	BuildRectLightMipTreeComputeShader->UnsetParameters(RHICmdList, EResourceTransitionAccess::ERWBarrier, EResourceTransitionPipeline::EComputeToCompute, Data.RectLightMipTree, TransitionFence);
+
+	return Data;
+}
+
+
+
 template <int TextureImportanceSampling>
 class FRectLightRGS : public FGlobalShader
 {
@@ -336,6 +469,21 @@ void RenderRayTracingRectLightInternal(
 	check(RectLightSceneInfo.Proxy->IsRectLight());
 	FRectLightSceneProxy* RectLightSceneProxy = (FRectLightSceneProxy*)RectLightSceneInfo.Proxy;
 
+	check(RectLightSceneProxy->RayTracingData);
+	if( !RectLightSceneProxy->RayTracingData->bInitialised // Test needed in case GRayTracingStochasticRectLight is turned on in editor,
+		|| RectLightSceneProxy->SourceTexture && RectLightSceneProxy->SourceTexture->GetLightingGuid() != RectLightSceneProxy->RayTracingData->TextureLightingGuid)
+	{
+		// We ignore TextureImportanceSampling and RectLightSceneProxy->HasSourceTexture() because uniform buffer expect a resource.
+		// So we always update.  
+		// dxr-todo: cache texture RayTracingData render side based on GUID in a database (render thread safe and avoid duplicating the work for each light using the same texture).
+		*RectLightSceneProxy->RayTracingData = BuildRectLightMipTree(RHICmdList, RectLightSceneProxy->SourceTexture);
+		RectLightSceneProxy->RayTracingData->bInitialised = true;
+		if (RectLightSceneProxy->SourceTexture)
+		{
+			RectLightSceneProxy->RayTracingData->TextureLightingGuid = RectLightSceneProxy->SourceTexture->GetLightingGuid();
+		}
+	}
+
 	FLightShaderParameters LightShaderParameters;
 	RectLightSceneProxy->GetLightShaderParameters(LightShaderParameters);
 
