Improved level load time by avoiding unnecessary DDC check - Saves ~5% in my test
#rnx

diff --git a/Engine/Source/Runtime/Engine/Private/StaticMesh.cpp b/Engine/Source/Runtime/Engine/Private/StaticMesh.cpp
--- a/Engine/Source/Runtime/Engine/Private/StaticMesh.cpp
+++ b/Engine/Source/Runtime/Engine/Private/StaticMesh.cpp
@@ -3777,13 +3777,15 @@ void UStaticMesh::CacheMeshData()
 		FStaticMeshSourceModel& SourceModel = SourceModels[LodIndex];
 		if (!SourceModel.MeshDescriptionBulkData.IsValid())
 		{
-			FString MeshDataKey;
-			if (GetMeshDataKey(LodIndex, MeshDataKey))
+			// Legacy assets used to store their source data in the RawMeshBulkData
+			// Migrate it to the new description if present
+			if (!SourceModel.RawMeshBulkData->IsEmpty())
 			{
-				// If the DDC key doesn't exist, convert the data and save it to DDC
-				if (!GetDerivedDataCacheRef().CachedDataProbablyExists(*MeshDataKey))
+				FString MeshDataKey;
+				if (GetMeshDataKey(LodIndex, MeshDataKey))
 				{
-					if (!SourceModel.RawMeshBulkData->IsEmpty())
+					// If the DDC key doesn't exist, convert the data and save it to DDC
+					if (!GetDerivedDataCacheRef().CachedDataProbablyExists(*MeshDataKey))
 					{
 						// Get the RawMesh for this LOD
 						FRawMesh TempRawMesh;
