Fixed case where skeletal mesh with "Render Static" option would crash editor in mobile preview

diff --git a/Engine/Source/Runtime/Engine/Private/SkeletalMesh.cpp b/Engine/Source/Runtime/Engine/Private/SkeletalMesh.cpp
--- a/Engine/Source/Runtime/Engine/Private/SkeletalMesh.cpp
+++ b/Engine/Source/Runtime/Engine/Private/SkeletalMesh.cpp
@@ -3747,8 +3747,7 @@ void FSkeletalMeshSceneProxy::DrawStaticElements(FStaticPrimitiveDrawInterface*
 					BatchElement.MaxVertexIndex = LODData.GetNumVertices() - 1;
 					BatchElement.NumPrimitives = Section.NumTriangles;
 					BatchElement.IndexBuffer = LODData.MultiSizeIndexContainer.GetIndexBuffer();
-					BatchElement.PrimitiveUniformBuffer = GetUniformBuffer();
-								
+													
 					PDI->DrawMesh(MeshElement, ScreenSize);
 				}
 			}
