Suppress a PVS studio warning. The warning is correct but it is expected that the loop condition is always false

diff --git a/Engine/Source/Runtime/Engine/Private/TextureDerivedData.cpp b/Engine/Source/Runtime/Engine/Private/TextureDerivedData.cpp
--- a/Engine/Source/Runtime/Engine/Private/TextureDerivedData.cpp
+++ b/Engine/Source/Runtime/Engine/Private/TextureDerivedData.cpp
@@ -1047,7 +1047,7 @@ static void SerializePlatformData(
 #endif
 		}
 
-		for (int32 MipIndex = 0; MipIndex < NumMips && MipIndex < OptionalMips; ++MipIndex )
+		for (int32 MipIndex = 0; MipIndex < NumMips && MipIndex < OptionalMips; ++MipIndex ) //-V654
 		{
 			PlatformData->Mips[MipIndex + FirstMipToSerialize].BulkData.SetBulkDataFlags(BULKDATA_Force_NOT_InlinePayload| BULKDATA_OptionalPayload);
 		}
