Fix uninitialized value

#rnx

diff --git a/Engine/Source/Runtime/Engine/Public/PrimitiveSceneProxy.h b/Engine/Source/Runtime/Engine/Public/PrimitiveSceneProxy.h
--- a/Engine/Source/Runtime/Engine/Public/PrimitiveSceneProxy.h
+++ b/Engine/Source/Runtime/Engine/Public/PrimitiveSceneProxy.h
@@ -98,14 +98,14 @@ public:
 class FHeightfieldComponentDescription
 {
 public:
-	FVector4 HeightfieldScaleBias;
-	FVector4 MinMaxUV;
-	FMatrix LocalToWorld;
-	FVector2D LightingAtlasLocation;
-	FIntRect HeightfieldRect;
-
-	int32 NumSubsections;
-	FVector4 SubsectionScaleAndBias;
+	FVector4 HeightfieldScaleBias = FVector4(ForceInit);
+	FVector4 MinMaxUV = FVector4(ForceInit);
+	FMatrix LocalToWorld = FMatrix::Identity;
+	FVector2D LightingAtlasLocation = FVector2D(ForceInit);
+	FIntRect HeightfieldRect; // Default initialized
+
+	int32 NumSubsections = 0;
+	FVector4 SubsectionScaleAndBias = FVector4(ForceInit);
 
 	FHeightfieldComponentDescription(const FMatrix& InLocalToWorld) :
 		LocalToWorld(InLocalToWorld)
