Fix Demo Recording Timeslicing to properly use milliseconds.

diff --git a/Engine/Source/Runtime/Engine/Private/DemoNetDriver.cpp b/Engine/Source/Runtime/Engine/Private/DemoNetDriver.cpp
--- a/Engine/Source/Runtime/Engine/Private/DemoNetDriver.cpp
+++ b/Engine/Source/Runtime/Engine/Private/DemoNetDriver.cpp
@@ -2426,7 +2426,7 @@ void UDemoNetDriver::TickDemoRecordFrame( float DeltaSeconds )
 	}
 
 	const double RecordFrameStartTime = FPlatformTime::Seconds();
-	const double RecordTimeLimit = (MaxDesiredRecordTimeMS * 1000.f);
+	const double RecordTimeLimit = (MaxDesiredRecordTimeMS / 1000.f);
 
 	// Mark any new streaming levels, so that they are saved out this frame
 	if (!HasLevelStreamingFixes())
