If we don't have a valid PerInstanceSMData assume a scale of 1.0

diff --git a/Engine/Source/Runtime/Engine/Private/InstancedStaticMesh.cpp b/Engine/Source/Runtime/Engine/Private/InstancedStaticMesh.cpp
--- a/Engine/Source/Runtime/Engine/Private/InstancedStaticMesh.cpp
+++ b/Engine/Source/Runtime/Engine/Private/InstancedStaticMesh.cpp
@@ -2168,8 +2168,8 @@ void UInstancedStaticMeshComponent::GetInstancesMinMaxScale(FVector& MinScale, F
 	}
 	else
 	{
-		MinScale = FVector(0);
-		MaxScale = FVector(0);
+		MinScale = FVector(1.0f);
+		MaxScale = FVector(1.0f);
 	}
 }
 
