Handle a missing archetype when initializing the shadow state buffer and log an error

diff --git a/Engine/Source/Runtime/Engine/Private/RepLayout.cpp b/Engine/Source/Runtime/Engine/Private/RepLayout.cpp
--- a/Engine/Source/Runtime/Engine/Private/RepLayout.cpp
+++ b/Engine/Source/Runtime/Engine/Private/RepLayout.cpp
@@ -5798,7 +5798,14 @@ TSharedPtr<FReplicationChangelistMgr> FRepLayout::CreateReplicationChangelistMgr
 		DeltaChangelistState = new FCustomDeltaChangelistState(LifetimeCustomPropertyState->NumFastArrayItems);
 	}
 
-	return MakeShareable(new FReplicationChangelistMgr(AsShared(), (const uint8*)InObject->GetArchetype(), DeltaChangelistState));
+	const uint8* ShadowStateSource = (const uint8*)InObject->GetArchetype();
+	if (ShadowStateSource == nullptr)
+	{
+		UE_LOG(LogRep, Error, TEXT("FRepLayout::CreateReplicationChangelistMgr: Invalid object archetype, initializing shadow state to current object state: %s"), *GetFullNameSafe(InObject));
+		ShadowStateSource = (const uint8*)InObject;
+	}
+
+	return MakeShareable(new FReplicationChangelistMgr(AsShared(), ShadowStateSource, DeltaChangelistState));
 }
 
 TUniquePtr<FRepState> FRepLayout::CreateRepState(
