Fix for 2D and spawn sounds not observing concurrency limits

diff --git a/Engine/Source/Runtime/Engine/Private/GameplayStatics.cpp b/Engine/Source/Runtime/Engine/Private/GameplayStatics.cpp
--- a/Engine/Source/Runtime/Engine/Private/GameplayStatics.cpp
+++ b/Engine/Source/Runtime/Engine/Private/GameplayStatics.cpp
@@ -1167,7 +1167,12 @@ void UGameplayStatics::PlaySound2D(const UObject* WorldContextObject, class USou
 
 		NewActiveSound.bIsUISound = true;
 		NewActiveSound.bAllowSpatialization = false;
-		NewActiveSound.ConcurrencySet.Add(ConcurrencySettings);
+
+		if (ConcurrencySettings)
+		{
+			NewActiveSound.ConcurrencySet.Add(ConcurrencySettings);
+		}
+
 		NewActiveSound.Priority = Sound->Priority;
 		NewActiveSound.SubtitlePriority = Sound->GetSubtitlePriority();
 
@@ -1328,7 +1333,11 @@ UAudioComponent* UGameplayStatics::SpawnSoundAttached(USoundBase* Sound, USceneC
 	Params.SetLocation(TestLocation);
 	Params.bStopWhenOwnerDestroyed = bStopWhenAttachedToDestroyed;
 	Params.AttenuationSettings = AttenuationSettings;
-	Params.ConcurrencySet.Add(ConcurrencySettings);
+
+	if (ConcurrencySettings)
+	{
+		Params.ConcurrencySet.Add(ConcurrencySettings);
+	}
 
 	UAudioComponent* AudioComponent = FAudioDevice::CreateComponent(Sound, Params);
 	if (AudioComponent)
