Fix crash when toggling planar reflection preview.

diff --git a/Engine/Source/Runtime/Engine/Private/Components/SceneCaptureComponent.cpp b/Engine/Source/Runtime/Engine/Private/Components/SceneCaptureComponent.cpp
--- a/Engine/Source/Runtime/Engine/Private/Components/SceneCaptureComponent.cpp
+++ b/Engine/Source/Runtime/Engine/Private/Components/SceneCaptureComponent.cpp
@@ -874,8 +874,11 @@ void UPlanarReflectionComponent::PostEditChangeProperty(FPropertyChangedEvent& P
 		ViewStates[ViewIndex].Allocate();
 	}
 
-	ProxyMeshComponent->bVisible = bShowPreviewPlane;
-	ProxyMeshComponent->MarkRenderStateDirty();
+	if (ProxyMeshComponent)
+	{
+		ProxyMeshComponent->bVisible = bShowPreviewPlane;
+		ProxyMeshComponent->MarkRenderStateDirty();
+	}
 }
 
 #endif
