Fixed random memory stomp during DDC cooking (most likely introduced in CL 3767205) by waiting for CacheShaders render commands before destroying FMaterialResources.

Race condition flow:
* Main thread calls FMaterial::CacheShaders (UDerivedDataCacheCommandlet::Main->BeginCacheForCookedPlatformData->FMaterial::CacheShaders)
* FMaterial::CacheShaders enqueues UpdateRenderThreadShaderMap render thread command
* Main thread then deletes FMaterial by calling ClearAllCachedCookedPlatformData from UDerivedDataCacheCommandlet::Main
* Depending on the garbage collector and render thread timings, UpdateRenderThreadShaderMap may be executed after the delete and it will stomp random memory

This can be reproduced 100% by:
* Insert FWindowsPlatformProcess::Sleep(1.0f) into UpdateRenderThreadShaderMap in FMaterial::CacheShaders
* Force GC to be run every frame with gc.StressTestGC=1 and gc.ForceCollectGarbageEveryFrame=1
* Run d:\ue4_release_4_22\Samples\Showcases\VirtualCamera\VirtualCamera.uproject -run=DerivedDataCache -TargetPlatform=Windows+WindowsNoEditor -fill -DDC=CreateInstalledProjectPak -StompMalloc -stdout -CrashForUAT -unattended -NoLogTimes

diff --git a/Engine/Source/Runtime/Engine/Private/Materials/Material.cpp b/Engine/Source/Runtime/Engine/Private/Materials/Material.cpp
--- a/Engine/Source/Runtime/Engine/Private/Materials/Material.cpp
+++ b/Engine/Source/Runtime/Engine/Private/Materials/Material.cpp
@@ -4128,6 +4128,9 @@ bool UMaterial::IsCachedCookedPlatformDataLoaded( const ITargetPlatform* TargetP
 
 void UMaterial::ClearCachedCookedPlatformData( const ITargetPlatform *TargetPlatform )
 {
+	// Make sure that all CacheShaders render thead commands are finished before we destroy FMaterialResources.
+	FlushRenderingCommands();
+
 	TArray<FMaterialResource*> *CachedMaterialResourcesForPlatform = CachedMaterialResourcesForCooking.Find( TargetPlatform );
 	if ( CachedMaterialResourcesForPlatform != NULL )
 	{
@@ -4141,6 +4144,9 @@ void UMaterial::ClearCachedCookedPlatformData( const ITargetPlatform *TargetPlat
 
 void UMaterial::ClearAllCachedCookedPlatformData()
 {
+	// Make sure that all CacheShaders render thead commands are finished before we destroy FMaterialResources.
+	FlushRenderingCommands();
+
 	for ( auto It : CachedMaterialResourcesForCooking )
 	{
 		TArray<FMaterialResource*> &CachedMaterialResourcesForPlatform = It.Value;
diff --git a/Engine/Source/Runtime/Engine/Private/Materials/MaterialInstance.cpp b/Engine/Source/Runtime/Engine/Private/Materials/MaterialInstance.cpp
--- a/Engine/Source/Runtime/Engine/Private/Materials/MaterialInstance.cpp
+++ b/Engine/Source/Runtime/Engine/Private/Materials/MaterialInstance.cpp
@@ -2946,6 +2946,9 @@ bool UMaterialInstance::IsCachedCookedPlatformDataLoaded( const ITargetPlatform*
 
 void UMaterialInstance::ClearCachedCookedPlatformData( const ITargetPlatform *TargetPlatform )
 {
+	// Make sure that all CacheShaders render thead commands are finished before we destroy FMaterialResources.
+	FlushRenderingCommands();
+
 	TArray<FMaterialResource*> *CachedMaterialResourcesForPlatform = CachedMaterialResourcesForCooking.Find( TargetPlatform );
 	if ( CachedMaterialResourcesForPlatform != NULL )
 	{
@@ -2960,6 +2963,9 @@ void UMaterialInstance::ClearCachedCookedPlatformData( const ITargetPlatform *Ta
 
 void UMaterialInstance::ClearAllCachedCookedPlatformData()
 {
+	// Make sure that all CacheShaders render thead commands are finished before we destroy FMaterialResources.
+	FlushRenderingCommands();
+
 	for ( auto It : CachedMaterialResourcesForCooking )
 	{
 		TArray<FMaterialResource*> &CachedMaterialResourcesForPlatform = It.Value;
