Fix for deleted net startup actors being missing from server replay checkpoints
Fix for deleted net startup actors in streaming levels fast forwarded from a checkpoint

diff --git a/Engine/Source/Runtime/Engine/Classes/Engine/DemoNetDriver.h b/Engine/Source/Runtime/Engine/Classes/Engine/DemoNetDriver.h
--- a/Engine/Source/Runtime/Engine/Classes/Engine/DemoNetDriver.h
+++ b/Engine/Source/Runtime/Engine/Classes/Engine/DemoNetDriver.h
@@ -608,6 +608,8 @@ public:
 	virtual void ProcessLocalServerPackets() override {}
 	virtual void ProcessLocalClientPackets() override {}
 
+	virtual void InitDestroyedStartupActors() override;
+
 protected:
 	virtual UChannel* InternalCreateChannelByName(const FName& ChName) override;
 
@@ -930,7 +932,7 @@ private:
 
 	// Levels that are currently pending for fast forward.
 	// Using raw pointers, because we manually keep when levels are added and removed.
-	TMap<class ULevel*, TSet<TWeakObjectPtr<class AActor>>> LevelsPendingFastForward;
+	TSet<class ULevel*> LevelsPendingFastForward;
 
 	// Only used during recording.
 	uint32 NumLevelsAddedThisFrame;
diff --git a/Engine/Source/Runtime/Engine/Classes/Engine/Level.h b/Engine/Source/Runtime/Engine/Classes/Engine/Level.h
--- a/Engine/Source/Runtime/Engine/Classes/Engine/Level.h
+++ b/Engine/Source/Runtime/Engine/Classes/Engine/Level.h
@@ -342,6 +342,7 @@ struct FReplicatedStaticActorDestructionInfo
 	GENERATED_BODY()
 
 	FName PathName;
+	FString FullName;
 	FVector	DestroyedPosition;
 	TWeakObjectPtr<UObject> ObjOuter;
 	UPROPERTY()
diff --git a/Engine/Source/Runtime/Engine/Private/DemoNetDriver.cpp b/Engine/Source/Runtime/Engine/Private/DemoNetDriver.cpp
--- a/Engine/Source/Runtime/Engine/Private/DemoNetDriver.cpp
+++ b/Engine/Source/Runtime/Engine/Private/DemoNetDriver.cpp
@@ -4252,44 +4252,20 @@ void UDemoNetDriver::PrepFastForwardLevels()
 				continue;
 			}
 
-			TSet<TWeakObjectPtr<AActor>> LevelActors;
+			bool bStartupActors = false;
+
 			for (AActor* Actor : Level->Actors)
 			{
-				if (Actor == nullptr || !Actor->IsNetStartupActor())
-				{
-					continue;
-				}
-				else if (DeletedNetStartupActors.Contains(Actor->GetFullName()))
-				{
-					// Put this actor on the rollback list so we can undelete it during future scrubbing,
-					// then delete it.
-					QueueNetStartupActorForRollbackViaDeletion(Actor);
-					LocalWorld->DestroyActor(Actor, true);
-				}
-				else
+				if (Actor != nullptr && Actor->IsNetStartupActor())
 				{
-					if (RollbackNetStartupActors.Contains(Actor->GetFullName()))
-					{
-						LocalWorld->DestroyActor(Actor, true);
-					}
-					else
-					{
-					LevelActors.Add(Actor);
+					bStartupActors = true;
+					break;
 				}
 			}
-			}
 
-			TArray<AActor*> SpawnedActors;
-			RespawnNecessaryNetStartupActors(SpawnedActors, Level);
-
-			for(AActor* Actor : SpawnedActors)
-			{
-				LevelActors.Add(Actor);
-			}
-
-			if (LevelActors.Num() > 0)
+			if (bStartupActors)
 			{
-				LevelsPendingFastForward.Emplace(Level, MoveTemp(LevelActors));
+				LevelsPendingFastForward.Add(Level);
 			}
 		}
 	}
@@ -4339,32 +4315,6 @@ bool UDemoNetDriver::FastForwardLevels(const FGotoResult& GotoResult)
 	LevelIndices.Reserve(LevelsPendingFastForward.Num());
 	StartupActors.Reserve(LevelsPendingFastForward.Num() * 4);
 
-	for (auto It = LevelsPendingFastForward.CreateIterator(); It; ++It)
-	{
-		// Track the appropriate level, and mark it as ready.
-		FLevelStatus& LevelStatus = GetLevelStatus(GetLevelPackageName(*It.Key()));
-		LevelIndices.Add(LevelStatus.LevelIndex);
-		LevelStatus.bIsReady = true;
-
-		// Quick sanity check to make sure the actors are still valid
-		// NOTE: The only way any of these should not be valid is if the level was unloaded,
-		//			or something in the demo caused the actor to be destroyed *before*
-		//			the level was ready. Either case seems bad if we've made it this far.
-		TSet<TWeakObjectPtr<AActor>>& LevelStartupActors = It.Value();
-		for (auto ActorIt = LevelStartupActors.CreateIterator(); ActorIt; ++ActorIt)
-		{
-			if (!ensure((*ActorIt).IsValid()))
-			{
-				ActorIt.RemoveCurrent();
-			}
-		}
-
-		LocalLevels.Add(It.Key());
-		StartupActors.Append(MoveTemp(LevelStartupActors));
-	}
-
-	LevelsPendingFastForward.Reset();
-
 	struct FLocalReadPacketsHelper
 	{
 		FLocalReadPacketsHelper(UDemoNetDriver& InDriver, const float InLastPacketTime):
@@ -4391,9 +4341,9 @@ bool UDemoNetDriver::FastForwardLevels(const FGotoResult& GotoResult)
 			{
 				Ar.Seek(PreFramePos);
 				if (ensure(NumPackets != 0))
-	{
+				{
 					Packets.RemoveAt(NumPackets, Packets.Num() - NumPackets);
-	}
+				}
 				return false;
 			}
 
@@ -4401,9 +4351,9 @@ bool UDemoNetDriver::FastForwardLevels(const FGotoResult& GotoResult)
 		}
 
 		bool IsError() const
-	{
+		{
 			return bErrorOccurred;
-	}
+		}
 
 		TArray<FPlaybackPacket> Packets;
 
@@ -4412,16 +4362,19 @@ bool UDemoNetDriver::FastForwardLevels(const FGotoResult& GotoResult)
 		UDemoNetDriver& Driver;
 		const float LastPacketTime;
 
-	// We only want to process packets that are before anything we've currently processed.
-	// Further, we want to make sure that we leave the archive in a good state for later use.
-	int32 NumPackets = 0;
-	float LastReadTime = 0;
-	FArchivePos PreFramePos = 0;
+		// We only want to process packets that are before anything we've currently processed.
+		// Further, we want to make sure that we leave the archive in a good state for later use.
+		int32 NumPackets = 0;
+		float LastReadTime = 0;
+		FArchivePos PreFramePos = 0;
 
 		bool bErrorOccurred = false;
 
 	} ReadPacketsHelper(*this, LastProcessedPacketTime);
 
+	DeletedNetStartupActors.Empty();
+	DeletedActorGuids.Empty();
+
 	{
 		auto IgnoreReceivedExportGUIDs = ((UPackageMapClient*)ServerConnection->PackageMap)->ScopedIgnoreReceivedExportGUIDs();
 
@@ -4451,7 +4404,39 @@ bool UDemoNetDriver::FastForwardLevels(const FGotoResult& GotoResult)
 
 				FArchivePos PacketOffset = 0;
 				*CheckpointArchive << PacketOffset;
-				CheckpointArchive->Seek(PacketOffset + CheckpointArchive->Tell());
+
+				PacketOffset += CheckpointArchive->Tell();
+
+				if (PlaybackDemoHeader.Version >= HISTORY_MULTIPLE_LEVELS)
+				{
+					int32 LevelIndex;
+					*CheckpointArchive << LevelIndex;
+				}
+
+				if (PlaybackDemoHeader.Version >= HISTORY_DELETED_STARTUP_ACTORS)
+				{
+					if (bDeltaCheckpoint)
+					{
+						TSet<FString> DeltaActors;
+						*CheckpointArchive << DeltaActors;
+
+						DeletedNetStartupActors.Append(DeltaActors);
+
+						TSet<FNetworkGUID> DeltaGuids;
+						*CheckpointArchive << DeltaGuids;
+
+						DeletedActorGuids.Append(DeltaGuids);
+					}
+					else
+					{
+						DeletedNetStartupActors.Empty();
+						DeletedActorGuids.Empty();
+
+						*CheckpointArchive << DeletedNetStartupActors;
+					}
+				}
+
+				CheckpointArchive->Seek(PacketOffset);
 
 				if (!ReadPacketsHelper.ReadPackets(*CheckpointArchive) && ReadPacketsHelper.IsError())
 				{
@@ -4482,6 +4467,54 @@ bool UDemoNetDriver::FastForwardLevels(const FGotoResult& GotoResult)
 	// If we've gotten this far, it means we should have something to process.
 	check(ReadPacketsHelper.Packets.Num() > 0);
 
+	UWorld* LocalWorld = GetWorld();
+	for (ULevel* Level : LevelsPendingFastForward)
+	{
+		// Track the appropriate level, and mark it as ready.
+		FLevelStatus& LevelStatus = GetLevelStatus(GetLevelPackageName(*Level));
+		LevelIndices.Add(LevelStatus.LevelIndex);
+		LevelStatus.bIsReady = true;
+
+		TSet<TWeakObjectPtr<AActor>> LevelActors;
+		for (AActor* Actor : Level->Actors)
+		{
+			if (Actor == nullptr || !Actor->IsNetStartupActor())
+			{
+				continue;
+			}
+			else if (DeletedNetStartupActors.Contains(Actor->GetFullName()))
+			{
+				// Put this actor on the rollback list so we can undelete it during future scrubbing,
+				// then delete it.
+				QueueNetStartupActorForRollbackViaDeletion(Actor);
+				LocalWorld->DestroyActor(Actor, true);
+			}
+			else
+			{
+				if (RollbackNetStartupActors.Contains(Actor->GetFullName()))
+				{
+					LocalWorld->DestroyActor(Actor, true);
+				}
+				else
+				{
+					StartupActors.Add(Actor);
+				}
+			}
+		}
+
+		TArray<AActor*> SpawnedActors;
+		RespawnNecessaryNetStartupActors(SpawnedActors, Level);
+
+		for (AActor* Actor : SpawnedActors)
+		{
+			StartupActors.Add(Actor);
+		}
+
+		LocalLevels.Add(Level);
+	}
+
+	LevelsPendingFastForward.Reset();
+
 	// It's possible that the level we're streaming in may spawn Dynamic Actors.
 	// In that case, we want to make sure we track them so we can process them below.
 	// We only care about the actors if they're outered to the Level.
@@ -5572,6 +5605,32 @@ void UDemoNetDriver::OnSeamlessTravelStartDuringRecording(const FString& LevelNa
 	ReplayStreamer->RefreshHeader();
 }
 
+void UDemoNetDriver::InitDestroyedStartupActors()
+{
+	Super::InitDestroyedStartupActors();
+
+	if (World)
+	{
+		check(DeletedNetStartupActors.Num() == 0);
+		check(DeltaDeletedNetStartupActors.Num() == 0);
+
+		// add startup actors destroyed before the creation of this net driver
+		for (auto LevelIt(World->GetLevelIterator()); LevelIt; ++LevelIt)
+		{
+			ULevel* Level = *LevelIt;
+			if (Level)
+			{
+				const TArray<FReplicatedStaticActorDestructionInfo>& DestroyedReplicatedStaticActors = Level->GetDestroyedReplicatedStaticActors();
+				for(const FReplicatedStaticActorDestructionInfo& Info : DestroyedReplicatedStaticActors)
+				{
+					DeletedNetStartupActors.Add(Info.FullName);
+					DeltaDeletedNetStartupActors.Add(Info.FullName);
+				}
+			}
+		}
+	}
+}
+
 void UDemoNetDriver::NotifyActorDestroyed( AActor* Actor, bool IsSeamlessTravel )
 {
 	check( Actor != nullptr );
diff --git a/Engine/Source/Runtime/Engine/Private/Level.cpp b/Engine/Source/Runtime/Engine/Private/Level.cpp
--- a/Engine/Source/Runtime/Engine/Private/Level.cpp
+++ b/Engine/Source/Runtime/Engine/Private/Level.cpp
@@ -431,6 +431,7 @@ void ULevel::CreateReplicatedDestructionInfo(AActor* const Actor)
 	{
 		FReplicatedStaticActorDestructionInfo NewInfo;
 		NewInfo.PathName = Actor->GetFName();
+		NewInfo.FullName = Actor->GetFullName();
 		NewInfo.DestroyedPosition = Actor->GetActorLocation();
 		NewInfo.ObjOuter = Actor->GetOuter();
 		NewInfo.ObjClass = Actor->GetClass();
