If the ICH is holding an object of a given name that has been removed and an object of a different class wants to be created with that name, move the removed object out of the way and create the new override object

diff --git a/Engine/Source/Runtime/Engine/Private/InheritableComponentHandler.cpp b/Engine/Source/Runtime/Engine/Private/InheritableComponentHandler.cpp
--- a/Engine/Source/Runtime/Engine/Private/InheritableComponentHandler.cpp
+++ b/Engine/Source/Runtime/Engine/Private/InheritableComponentHandler.cpp
@@ -113,7 +113,24 @@ UActorComponent* UInheritableComponentHandler::CreateOverridenComponentTemplate(
 	}
 
 	ensure(Cast<UBlueprintGeneratedClass>(GetOuter()));
-	auto NewComponentTemplate = NewObject<UActorComponent>(
+	
+	// If we find an existing object with our name that the object recycling system won't allow for we need to deal with it 
+	// or else the NewObject call below will fatally assert
+	UObject* ExistingObj = FindObjectFast<UObject>(GetOuter(), NewComponentTemplateName);
+	if (ExistingObj && !ExistingObj->GetClass()->IsChildOf(BestArchetype->GetClass()))
+	{
+		// If this isn't an unnecessary component there is something else we need to investigate
+		// but if it is, just consign it to oblivion as its purpose is no longer required with the allocation
+		// of an object of the same name
+		UActorComponent* ExistingComp = Cast<UActorComponent>(ExistingObj);
+		if (ensure(ExistingComp) && ensure(UnnecessaryComponents.RemoveSwap(ExistingComp) > 0))
+		{
+			ExistingObj->Rename(nullptr, GetTransientPackage(), REN_DontCreateRedirectors | REN_ForceNoResetLoaders);
+			ExistingObj->MarkPendingKill();
+		}
+	}
+
+	UActorComponent* NewComponentTemplate = NewObject<UActorComponent>(
 		GetOuter(), BestArchetype->GetClass(), NewComponentTemplateName, RF_ArchetypeObject | RF_Public | RF_InheritableComponentTemplate, BestArchetype);
 
 	// HACK: NewObject can return a pre-existing object which will not have been initialized to the archetype.  When we remove the old handlers, we mark them pending
@@ -128,14 +145,14 @@ UActorComponent* UInheritableComponentHandler::CreateOverridenComponentTemplate(
 
 	// Clear transient flag if it was transient before and re copy off archetype
 	if (NewComponentTemplate->HasAnyFlags(RF_Transient) && UnnecessaryComponents.Contains(NewComponentTemplate))
-	{
-		NewComponentTemplate->ClearFlags(RF_Transient);
+		{
+			NewComponentTemplate->ClearFlags(RF_Transient);
 		UnnecessaryComponents.Remove(NewComponentTemplate);
 
-		UEngine::FCopyPropertiesForUnrelatedObjectsParams CopyParams;
-		CopyParams.bDoDelta = false;
-		UEngine::CopyPropertiesForUnrelatedObjects(BestArchetype, NewComponentTemplate, CopyParams);
-	}
+			UEngine::FCopyPropertiesForUnrelatedObjectsParams CopyParams;
+			CopyParams.bDoDelta = false;
+			UEngine::CopyPropertiesForUnrelatedObjects(BestArchetype, NewComponentTemplate, CopyParams);
+		}
 
 	FComponentOverrideRecord NewRecord;
 	NewRecord.ComponentKey = Key;
