Notify all net drivers when an actor is destroyed

diff --git a/Engine/Source/Runtime/Engine/Private/LevelActor.cpp b/Engine/Source/Runtime/Engine/Private/LevelActor.cpp
--- a/Engine/Source/Runtime/Engine/Private/LevelActor.cpp
+++ b/Engine/Source/Runtime/Engine/Private/LevelActor.cpp
@@ -30,11 +30,13 @@
 #include "ContentStreaming.h"
 #include "EditorSupportDelegates.h"
 #include "GameFramework/GameModeBase.h"
-#include "Engine/DemoNetDriver.h"
 #include "AudioDeviceManager.h"
 #include "Logging/TokenizedMessage.h"
 #include "Logging/MessageLog.h"
 #include "Misc/MapErrors.h"
+#include "GameFramework/WorldSettings.h"
+#include "Engine/NetDriver.h"
+#include "Engine/Player.h"
 
 #include "Components/BoxComponent.h"
 #include "GameFramework/MovementComponent.h"
@@ -674,12 +676,14 @@ bool UWorld::DestroyActor( AActor* ThisActor, bool bNetForce, bool bShouldModify
 	}
 
 	// Notify net drivers that this guy has been destroyed.
-	if (GEngine->GetWorldContextFromWorld(this))
+	if (FWorldContext* Context = GEngine->GetWorldContextFromWorld(this))
 	{
-		UNetDriver* ActorNetDriver = GEngine->FindNamedNetDriver(this,ThisActor->GetNetDriverName());
-		if (ActorNetDriver)
+		for (FNamedNetDriver& Driver : Context->ActiveNetDrivers)
 		{
-			ActorNetDriver->NotifyActorDestroyed(ThisActor);
+			if (Driver.NetDriver != nullptr && Driver.NetDriver->ShouldReplicateActor(ThisActor))
+			{
+				Driver.NetDriver->NotifyActorDestroyed(ThisActor);
+			}
 		}
 	}
 	else if (WorldType != EWorldType::Inactive && !IsRunningCommandlet())
@@ -689,11 +693,6 @@ bool UWorld::DestroyActor( AActor* ThisActor, bool bNetForce, bool bShouldModify
 		UE_LOG(LogSpawn, Warning, TEXT("UWorld::DestroyActor: World has no context! World: %s, Actor: %s"), *GetName(), *ThisActor->GetPathName());
 	}
 
-	if ( DemoNetDriver )
-	{
-		DemoNetDriver->NotifyActorDestroyed( ThisActor );
-	}
-
 	// Remove the actor from the actor list.
 	RemoveActor( ThisActor, bShouldModifyLevel );
 
