Adding splitscreen demo viewing.

diff --git a/Engine/Source/Runtime/Engine/Classes/Engine/DemoNetDriver.h b/Engine/Source/Runtime/Engine/Classes/Engine/DemoNetDriver.h
--- a/Engine/Source/Runtime/Engine/Classes/Engine/DemoNetDriver.h
+++ b/Engine/Source/Runtime/Engine/Classes/Engine/DemoNetDriver.h
@@ -390,7 +390,7 @@ class ENGINE_API UDemoNetDriver : public UNetDriver
 	/** Index of LevelNames that is currently loaded */
 	int32 CurrentLevelIndex;
 
-	/** This is our spectator controller that is used to view the demo world from */
+	/** This is the main spectator controller that is used to view the demo world from */
 	APlayerController* SpectatorController;
 
 	/** Our network replay streamer */
@@ -727,6 +727,69 @@ public:
 	void FinalizeFastForward( const double StartTime );
 	
 	void SpawnDemoRecSpectator( UNetConnection* Connection, const FURL& ListenURL );
+
+	/**
+	 * Restores the given player controller so that it properly points to the given NetConnection
+	 * after scrubbing when viewing a replay.
+	 *
+	 * @param PC			The PlayerController to set up the given NetConnection for
+	 * @param NetConnection	The NetConnection to be assigned to the player controller.
+	 */
+	void RestoreConnectionPostScrub(APlayerController* PC, UNetConnection* NetConnection);
+
+	/**
+	 * Sets the main spectator controller to be used and adds them to the spectator control array
+	 *
+	 * @param PC			The PlayerController to set the main controller param to.
+	 */
+	void SetSpectatorController(APlayerController* PC);
+	
+	// Splitscreen demo handling
+
+	/**
+	 * Creates a new splitscreen replay viewer.
+	 *
+	 * @param NewPlayer		The LocalPlayer in control of this new viewer
+	 * @param InWorld		The world to spawn the new viewer in.
+	 *
+	 * @return If the viewer was able to be created or not.
+	 */
+	bool SpawnSplitscreenViewer(ULocalPlayer* NewPlayer, UWorld* InWorld);
+
+	/**
+	 * Removes a splitscreen demo viewer and cleans up its connection.
+	 *
+	 * @param RemovePlayer		The PlayerController to remove from the replay system
+	 * @param bMarkOwnerForDeletion		If this function should handle deleting the given player as well.
+	 *
+	 * @return If the player was successfully removed from the replay.
+	 */
+	bool RemoveSplitscreenViewer(APlayerController* RemovePlayer, bool bMarkOwnerForDeletion=false);
+
+private:
+
+	// Internal player spawning
+	APlayerController* CreateDemoPlayerController(UNetConnection* Connection, const FURL& ListenURL);
+
+	// Internal splitscreen management
+
+	/** An array of all the spectator controllers (the main one and all splitscreen ones) that currently exist */
+	UPROPERTY(transient)
+	TArray<APlayerController*> SpectatorControllers;
+
+	/**
+	 * Removes all child connections for splitscreen viewers.
+	 * This should be done before the ClientConnections or ServerConnection
+	 * variables change or during most travel scenarios.
+	 *
+	 * @param bDeleteOwner	If the connections should delete the owning actor to the connection
+	 *
+	 * @return The number of splitscreen connections cleaned up.
+	 */
+	int32 CleanUpSplitscreenConnections(bool bDeleteOwner);
+
+public:
+
 	void ResetDemoState();
 	void JumpToEndOfLiveReplay();
 	void AddEvent(const FString& Group, const FString& Meta, const TArray<uint8>& Data);
diff --git a/Engine/Source/Runtime/Engine/Private/DemoNetDriver.cpp b/Engine/Source/Runtime/Engine/Private/DemoNetDriver.cpp
--- a/Engine/Source/Runtime/Engine/Private/DemoNetDriver.cpp
+++ b/Engine/Source/Runtime/Engine/Private/DemoNetDriver.cpp
@@ -37,6 +37,7 @@
 #include "ProfilingDebugging/CsvProfiler.h"
 #include "Misc/EngineVersion.h"
 #include "Stats/Stats2.h"
+#include "Engine/ChildConnection.h"
 
 DEFINE_LOG_CATEGORY( LogDemo );
 
@@ -749,6 +750,7 @@ void UDemoNetDriver::FinishDestroy()
 		}
 	}
 
+	CleanUpSplitscreenConnections(true);
 	FCoreUObjectDelegates::PostLoadMapWithWorld.RemoveAll(this);
 	Super::FinishDestroy();
 }
@@ -1205,6 +1207,7 @@ bool UDemoNetDriver::ContinueListen(FURL& ListenURL)
 		{
 			SpectatorController->Player = nullptr;		// Force APlayerController::DestroyNetworkActorHandled to return false
 			World->DestroyActor(SpectatorController, true);
+			SpectatorControllers.Empty();
 			SpectatorController = nullptr;
 		}
 
@@ -1466,24 +1469,29 @@ void UDemoNetDriver::TickDispatch( float DeltaSeconds )
 
 	// Update time dilation on spectator pawn to compensate for any demo dilation 
 	//	(we want to continue to fly around in real-time)
-	if ( SpectatorController != nullptr )
+	for (APlayerController* CurSpectatorController : SpectatorControllers)
 	{
+		if (CurSpectatorController == nullptr)
+		{
+			continue;
+		}
+
 		if ( World->GetWorldSettings()->DemoPlayTimeDilation > KINDA_SMALL_NUMBER )
 		{
-			SpectatorController->CustomTimeDilation = 1.0f / World->GetWorldSettings()->DemoPlayTimeDilation;
+			CurSpectatorController->CustomTimeDilation = 1.0f / World->GetWorldSettings()->DemoPlayTimeDilation;
 		}
 		else
 		{
-			SpectatorController->CustomTimeDilation = 1.0f;
+			CurSpectatorController->CustomTimeDilation = 1.0f;
 		}
 
-		if ( SpectatorController->GetSpectatorPawn() != nullptr )
+		if (CurSpectatorController->GetSpectatorPawn() != nullptr)
 		{
-			SpectatorController->GetSpectatorPawn()->CustomTimeDilation = SpectatorController->CustomTimeDilation;
-					
-			SpectatorController->GetSpectatorPawn()->PrimaryActorTick.bTickEvenWhenPaused = true;
+			CurSpectatorController->GetSpectatorPawn()->CustomTimeDilation = CurSpectatorController->CustomTimeDilation;
+
+			CurSpectatorController->GetSpectatorPawn()->PrimaryActorTick.bTickEvenWhenPaused = true;
 
-			USpectatorPawnMovement* SpectatorMovement = Cast<USpectatorPawnMovement>(SpectatorController->GetSpectatorPawn()->GetMovementComponent());
+			USpectatorPawnMovement* SpectatorMovement = Cast<USpectatorPawnMovement>(CurSpectatorController->GetSpectatorPawn()->GetMovementComponent());
 
 			if ( SpectatorMovement )
 			{
@@ -2927,12 +2935,7 @@ void UDemoNetDriver::PauseChannels( const bool bPause )
 
 		ActorChannel->CustomTimeDilation = bPause ? 0.0f : 1.0f;
 
-		if ( ActorChannel->GetActor() == SpectatorController )
-		{
-			continue;
-		}
-
-		if ( ActorChannel->GetActor() == nullptr )
+		if (ActorChannel->GetActor() == nullptr || SpectatorControllers.Contains(ActorChannel->GetActor()) )
 		{
 			continue;
 		}
@@ -3189,6 +3192,13 @@ void UDemoNetDriver::ProcessSeamlessTravel(int32 LevelIndex)
 		Controllers.Add(Iterator->Get());
 	}
 
+	// Clean up any splitscreen spectators if we have them.
+	// Let the destroy below handle deletion of the objects.
+	if (SpectatorControllers.Num() > 1)
+	{
+		CleanUpSplitscreenConnections(false);
+	}
+
 	for (int i = 0; i < Controllers.Num(); i++)
 	{
 		if (Controllers[i])
@@ -3196,9 +3206,14 @@ void UDemoNetDriver::ProcessSeamlessTravel(int32 LevelIndex)
 			// bNetForce is true so that the replicated spectator player controller will
 			// be destroyed as well.
 			Controllers[i]->Destroy(true);
+
+			// If we can, remove the spectator here as well.
+			SpectatorControllers.Remove(Cast<APlayerController>(Controllers[i]));
 		}
 	}
 
+	SpectatorControllers.Empty();
+
 	// Set this to nullptr since we just destroyed it.
 	SpectatorController = nullptr;
 
@@ -3914,21 +3929,15 @@ void UDemoNetDriver::FinalizeFastForward( const double StartTime )
 	UE_LOG( LogDemo, Log, TEXT( "Fast forward took %.2f seconds." ), FastForwardTotalSeconds );
 }
 
-void UDemoNetDriver::SpawnDemoRecSpectator( UNetConnection* Connection, const FURL& ListenURL )
+APlayerController* UDemoNetDriver::CreateDemoPlayerController(UNetConnection* Connection, const FURL& ListenURL)
 {
-	// Optionally skip spawning the demo spectator if requested via the URL option
-	if (ListenURL.HasOption(TEXT("SkipSpawnSpectatorController")))
-	{
-		return;
-	}
-
-	check( Connection != nullptr );
+	check(Connection != nullptr);
 
 	// Get the replay spectator controller class from the default game mode object,
 	// since the game mode instance isn't replicated to clients of live games.
 	AGameStateBase* GameState = GetWorld() != nullptr ? GetWorld()->GetGameState() : nullptr;
 	TSubclassOf<AGameModeBase> DefaultGameModeClass = GameState != nullptr ? GameState->GameModeClass : nullptr;
-	
+
 	// If we don't have a game mode class from the world, try to get it from the URL option.
 	// This may be true on clients who are recording a replay before the game mode class was replicated to them.
 	if (DefaultGameModeClass == nullptr)
@@ -3946,53 +3955,162 @@ void UDemoNetDriver::SpawnDemoRecSpectator( UNetConnection* Connection, const FU
 
 	if ( C == nullptr )
 	{
-		UE_LOG( LogDemo, Error, TEXT( "UDemoNetDriver::SpawnDemoRecSpectator: Failed to load demo spectator class." ) );
-		return;
+		UE_LOG(LogDemo, Error, TEXT("UDemoNetDriver::CreateDemoPlayerController: Failed to load demo spectator class."));
+		return nullptr;
 	}
 
 	FActorSpawnParameters SpawnInfo;
 	SpawnInfo.ObjectFlags |= RF_Transient;	// We never want these to save into a map
-	SpectatorController = World->SpawnActor<APlayerController>( C, SpawnInfo );
+	APlayerController* NewDemoController = World->SpawnActor<APlayerController>(C, SpawnInfo);
 
-	if ( SpectatorController == nullptr )
+	if (NewDemoController == nullptr)
 	{
-		UE_LOG( LogDemo, Error, TEXT( "UDemoNetDriver::SpawnDemoRecSpectator: Failed to spawn demo spectator." ) );
-		return;
+		UE_LOG(LogDemo, Error, TEXT("UDemoNetDriver::CreateDemoPlayerController: Failed to spawn demo spectator."));
+		return nullptr;
 	}
 
 	// Streaming volumes logic must not be affected by replay spectator camera
-	SpectatorController->bIsUsingStreamingVolumes = false;
+	NewDemoController->bIsUsingStreamingVolumes = false;
 
-	// Make sure SpectatorController->GetNetDriver returns this driver. Ensures functions that depend on it,
+	// Make sure the player controller GetNetDriver returns this driver. Ensures functions that depend on it,
 	// such as IsLocalController, work as expected.
-	SpectatorController->SetNetDriverName(NetDriverName);
+	NewDemoController->SetNetDriverName(NetDriverName);
 
 	// If the controller doesn't have a player state, we are probably recording on a client.
 	// Spawn one manually.
-	if ( SpectatorController->PlayerState == nullptr && GetWorld() != nullptr && GetWorld()->IsRecordingClientReplay())
+	if (NewDemoController->PlayerState == nullptr && GetWorld() != nullptr && GetWorld()->IsRecordingClientReplay())
 	{
-		SpectatorController->InitPlayerState();
+		NewDemoController->InitPlayerState();
 	}
 
 	// Tell the game that we're spectator and not a normal player
-	if (SpectatorController->PlayerState)
+	if (NewDemoController->PlayerState)
 	{
-		SpectatorController->PlayerState->bOnlySpectator = true;
+		NewDemoController->PlayerState->bOnlySpectator = true;
 	}
 
-	for ( FActorIterator It( World ); It; ++It)
+	for (FActorIterator It(World); It; ++It)
 	{
-		if ( It->IsA( APlayerStart::StaticClass() ) )
+		if (It->IsA(APlayerStart::StaticClass()))
 		{
-			SpectatorController->SetInitialLocationAndRotation( It->GetActorLocation(), It->GetActorRotation() );
+			NewDemoController->SetInitialLocationAndRotation(It->GetActorLocation(), It->GetActorRotation());
 			break;
 		}
 	}
+
+	NewDemoController->SetReplicates(true);
+	NewDemoController->SetAutonomousProxy(true);
+	NewDemoController->SetPlayer(Connection);
+
+	return NewDemoController;
+}
+
+void UDemoNetDriver::SpawnDemoRecSpectator( UNetConnection* Connection, const FURL& ListenURL )
+{
+	// Optionally skip spawning the demo spectator if requested via the URL option
+	if (ListenURL.HasOption(TEXT("SkipSpawnSpectatorController")))
+	{
+		return;
+	}
+
+	SpectatorController = CreateDemoPlayerController(Connection, ListenURL);
+	SpectatorControllers.Add(SpectatorController);
+}
+
+bool UDemoNetDriver::SpawnSplitscreenViewer(ULocalPlayer* NewPlayer, UWorld* InWorld)
+{
+	if (NewPlayer == nullptr || InWorld == nullptr)
+	{
+		UE_LOG(LogDemo, Warning, TEXT("UDemoNetDriver::SpawnSplitscreenViewer: Local Player or World is invalid!"));
+		return false;
+	}
+
+	if (ClientConnections.Num() == 0 && ServerConnection == nullptr)
+	{
+		UE_LOG(LogDemo, Error, TEXT("UDemoNetDriver::SpawnSplitscreenViewer: This netdriver has no demo connection data"));
+		return false;
+	}
+
+	UNetConnection* ChildConnection = CreateChild((ClientConnections.Num() > 0) ? ClientConnections[0] : ServerConnection);
+
+	APlayerController* NewSplitscreenController = CreateDemoPlayerController(ChildConnection, TEXT(""));
+	if (NewSplitscreenController == nullptr)
+	{
+		UE_LOG(LogDemo, Warning, TEXT("UDemoNetDriver::SpawnSplitscreenViewer: Unable to create new splitscreen controller"));
+		return false;
+	}
+
+	// Link this spectator to the given local player, as this will facilitate spectator pawn creation 
+	// (spectator pawns only create if the controller is linked to a local player)
+	NewSplitscreenController->Player = NewPlayer;
+	NewSplitscreenController->NetPlayerIndex = GEngine->GetGamePlayers(InWorld).Find(NewPlayer);
+
+	// Create the Pawn
+	NewSplitscreenController->ChangeState(NAME_Spectating);
+	NewPlayer->CurrentNetSpeed = 0;
+
+	// Link the local player to the player controller as the localplayer has been marked as active
+	// but without a playercontroller, the player will never be considered "ready" by other systems.
+	NewPlayer->PlayerController = NewSplitscreenController;
+
+	// This would typically be set via SetPlayer, but we need to call SetPlayer with the LocalPlayer
+	// and not with the child connection, otherwise we never create the input controls we need.
+	ChildConnection->PlayerController = NewSplitscreenController;
+	ChildConnection->OwningActor = NewSplitscreenController;
 	
-	SpectatorController->SetReplicates( true );
-	SpectatorController->SetAutonomousProxy( true );
+	// Create input control
+	NewSplitscreenController->SetPlayer(NewPlayer);
+
+	// Add to the list
+	SpectatorControllers.Add(NewSplitscreenController);
+
+	return true;
+}
+
+bool UDemoNetDriver::RemoveSplitscreenViewer(APlayerController* RemovePlayer, bool bMarkOwnerForDeletion)
+{
+	UE_LOG(LogDemo, Log, TEXT("Attempting to remove splitscreen viewer!"));
+	if (RemovePlayer && SpectatorControllers.Contains(RemovePlayer) && RemovePlayer != SpectatorController)
+	{
+		SpectatorControllers.Remove(RemovePlayer);
+		UNetConnection* RemovedNetConnection = RemovePlayer->NetConnection;
+		if (!bMarkOwnerForDeletion)
+		{
+			RemovedNetConnection->OwningActor = nullptr;
+		}
+		RemovedNetConnection->Close();
+		RemovedNetConnection->CleanUp();
+		RemovePlayer->NetConnection = nullptr;
+		return true;
+	}
+
+	return false;
+}
+
+int32 UDemoNetDriver::CleanUpSplitscreenConnections(bool bDeleteOwner)
+{
+	int32 NumSplitscreenConnectionsCleaned = 0;
+
+	for (APlayerController* CurController : SpectatorControllers)
+	{
+		UNetConnection* ControllerNetConnection = CurController->NetConnection;
+		if (ControllerNetConnection && ControllerNetConnection->IsA(UChildConnection::StaticClass()))
+		{
+			++NumSplitscreenConnectionsCleaned;
+			// With this toggled, this prevents actor deletion (which we don't want to do when scrubbing)
+			if (!bDeleteOwner)
+			{
+				ControllerNetConnection->OwningActor = nullptr;
+			}
+			ControllerNetConnection->Close();
+			ControllerNetConnection->CleanUp();
+			CurController->NetConnection = nullptr;
+		}
+	}
 
-	SpectatorController->SetPlayer( Connection );
+	FString OwnerDeletionStr(bDeleteOwner ? TEXT("with") : TEXT("without"));
+	UE_LOG(LogDemo, Log, TEXT("Cleaned up %d splitscreen connections %s owner deletion"), NumSplitscreenConnectionsCleaned, *OwnerDeletionStr);
+	return NumSplitscreenConnectionsCleaned;
 }
 
 void UDemoNetDriver::ReplayStreamingReady( const FStartStreamingResult& Result )
@@ -4735,6 +4853,9 @@ bool UDemoNetDriver::LoadCheckpoint(const FGotoResult& GotoResult)
 		// Clean package map to prepare to restore it to the checkpoint state
 		GuidCache->ResetCacheForDemo();
 
+		// Since we only count the number of sub-spectators, add one more slot for main spectator
+		// Very small optimization. We do want to clear this so that we don't end up doing during ProcessSeamlessTravel
+		SpectatorControllers.Empty(CleanUpSplitscreenConnections(true) + 1);
 		SpectatorController = nullptr;
 
 		ServerConnection->Close();
@@ -4809,6 +4930,21 @@ bool UDemoNetDriver::LoadCheckpoint(const FGotoResult& GotoResult)
 #if 1
 	TSet<const AActor*> KeepAliveActors;
 
+	// Determine if an Actor has a reference to a spectator in some way.
+	// This prevents garbage collection on splitscreen playercontrollers
+	auto HasPlayerSpectatorRef = [this](const FActorIterator& InActorIterator) {
+		for (const APlayerController* CurSpectator : SpectatorControllers)
+		{
+			if (*InActorIterator == CurSpectator || *InActorIterator == CurSpectator->GetSpectatorPawn()
+				|| InActorIterator->GetOwner() == CurSpectator)
+			{
+				return true;
+			}
+		}
+		return false;
+	};
+
+
 	// Destroy all non startup actors. They will get restored with the checkpoint
 	for ( FActorIterator It( GetWorld() ); It; ++It )
 	{
@@ -4822,7 +4958,7 @@ bool UDemoNetDriver::LoadCheckpoint(const FGotoResult& GotoResult)
 			AddNonQueuedActorForScrubbing(*It);
 		}
 		
-		const bool bShouldPreserveForPlayerController = (SpectatorController != nullptr && (*It == SpectatorController || *It == SpectatorController->GetSpectatorPawn() || It->GetOwner() == SpectatorController));
+		const bool bShouldPreserveForPlayerController = HasPlayerSpectatorRef(It);
 		const bool bShouldPreserveForRewindability = (It->bReplayRewindable && !It->IsNetStartupActor());										
 
 		if (bShouldPreserveForPlayerController || bShouldPreserveForRewindability)
@@ -4901,6 +5037,8 @@ bool UDemoNetDriver::LoadCheckpoint(const FGotoResult& GotoResult)
 	ExternalDataToObjectMap.Empty();
 	PlaybackPackets.Empty();
 
+	// Going to be recreating the splitscreen connections, but keep around the player controller.
+	CleanUpSplitscreenConnections(false);
 	ServerConnection->Close();
 	ServerConnection->CleanUp();
 
@@ -4932,6 +5070,15 @@ bool UDemoNetDriver::LoadCheckpoint(const FGotoResult& GotoResult)
 	// Create fake control channel
 	CreateInitialClientChannels();
 
+	// Respawn child connections as the parent connection has been recreated.
+	for (APlayerController* CurController : SpectatorControllers)
+	{
+		if (CurController != SpectatorController)
+		{
+			RestoreConnectionPostScrub(CurController, CreateChild(ServerConnection));
+		}
+	}
+
 	// Catch a rare case where the spectator controller is null, but a valid GUID is
 	// found on the GuidCache. The weak pointers in the NetGUIDLookup map are probably
 	// going null, and we want catch these cases and investigate further.
@@ -5354,6 +5501,30 @@ bool UDemoNetDriver::ComparePropertyState(const FDemoSavedPropertyState& State)
 	return bWasDifferent;
 }
 
+void UDemoNetDriver::RestoreConnectionPostScrub(APlayerController* PC, UNetConnection* NetConnection)
+{
+	check(NetConnection != nullptr);
+	check(PC != nullptr);
+
+	PC->Role = ROLE_AutonomousProxy;
+	PC->NetConnection = NetConnection;
+	NetConnection->LastReceiveTime = Time;
+	NetConnection->LastReceiveRealtime = FPlatformTime::Seconds();
+	NetConnection->LastGoodPacketRealtime = FPlatformTime::Seconds();
+	NetConnection->State = USOCK_Open;
+	NetConnection->PlayerController = PC;
+	NetConnection->OwningActor = PC;
+}
+
+void UDemoNetDriver::SetSpectatorController(APlayerController* PC)
+{
+	SpectatorController = PC;
+	if (PC != nullptr)
+	{
+		SpectatorControllers.AddUnique(PC);
+	}
+}
+
 /*-----------------------------------------------------------------------------
 	UDemoNetConnection.
 -----------------------------------------------------------------------------*/
@@ -5444,39 +5615,39 @@ void UDemoNetConnection::FlushNet( bool bIgnoreSimulation )
 
 void UDemoNetConnection::HandleClientPlayer( APlayerController* PC, UNetConnection* NetConnection )
 {
+	UDemoNetDriver* DemoDriver = GetDriver();
 	// If the spectator is the same, assume this is for scrubbing, and we are keeping the old one
 	// (so don't set the position, since we want to persist all that)
-	if ( GetDriver()->SpectatorController == PC )
-	{
-		PC->Role			= ROLE_AutonomousProxy;
-		PC->NetConnection	= NetConnection;
-		LastReceiveTime		= Driver->Time;
-		LastReceiveRealtime = FPlatformTime::Seconds();
-		LastGoodPacketRealtime = FPlatformTime::Seconds();
-		State				= USOCK_Open;
-		PlayerController	= PC;
-		OwningActor			= PC;
+	if (DemoDriver->SpectatorController == PC )
+	{
+		DemoDriver->RestoreConnectionPostScrub(PC, NetConnection);
+		DemoDriver->SetSpectatorController(PC);
 		return;
 	}
 
 	ULocalPlayer* LocalPlayer = nullptr;
-	for (FLocalPlayerIterator It(GEngine, Driver->GetWorld()); It; ++It)
+	uint8 PlayerIndex = 0;
+	// Attempt to find the player that doesn't already have a connection.
+	for (FLocalPlayerIterator It(GEngine, Driver->GetWorld()); It; ++It, PlayerIndex++)
 	{
-		LocalPlayer = *It;
-		break;
+		if (PC->NetPlayerIndex == PlayerIndex)
+		{
+			LocalPlayer = *It;
+			break;
+		}
 	}
 	int32 SavedNetSpeed = LocalPlayer ? LocalPlayer->CurrentNetSpeed : 0;
 
 	Super::HandleClientPlayer( PC, NetConnection );
 	
 	// Restore the netspeed if we're a local replay
-	if (GetDriver()->bIsLocalReplay && LocalPlayer)
+	if (DemoDriver->bIsLocalReplay && LocalPlayer)
 	{
 		LocalPlayer->CurrentNetSpeed = SavedNetSpeed;
 	}
 
-	// Assume this is our special spectator controller
-	GetDriver()->SpectatorController = PC;
+	// This is very likely our main demo controller.
+	DemoDriver->SetSpectatorController(PC);
 
 	for ( FActorIterator It( Driver->World ); It; ++It)
 	{
diff --git a/Engine/Source/Runtime/Engine/Private/GameInstance.cpp b/Engine/Source/Runtime/Engine/Private/GameInstance.cpp
--- a/Engine/Source/Runtime/Engine/Private/GameInstance.cpp
+++ b/Engine/Source/Runtime/Engine/Private/GameInstance.cpp
@@ -641,15 +641,25 @@ ULocalPlayer* UGameInstance::CreateLocalPlayer(int32 ControllerId, FString& OutE
 
 		NewPlayer = NewObject<ULocalPlayer>(GetEngine(), GetEngine()->LocalPlayerClass);
 		InsertIndex = AddLocalPlayer(NewPlayer, ControllerId);
-		if (bSpawnPlayerController && InsertIndex != INDEX_NONE && GetWorld() != NULL)
+		UWorld* CurrentWorld = GetWorld();
+		if (bSpawnPlayerController && InsertIndex != INDEX_NONE && CurrentWorld != nullptr)
 		{
-			if (GetWorld()->GetNetMode() != NM_Client)
+			if (CurrentWorld->GetNetMode() != NM_Client)
 			{
 				// server; spawn a new PlayerController immediately
-				if (!NewPlayer->SpawnPlayActor("", OutError, GetWorld()))
+				if (!NewPlayer->SpawnPlayActor("", OutError, CurrentWorld))
 				{
 					RemoveLocalPlayer(NewPlayer);
-					NewPlayer = NULL;
+					NewPlayer = nullptr;
+				}
+			}
+			else if (CurrentWorld->IsPlayingReplay())
+			{
+				// demo playback; ask the replay driver to spawn a splitscreen client
+				if (!CurrentWorld->DemoNetDriver->SpawnSplitscreenViewer(NewPlayer, CurrentWorld))
+				{
+					RemoveLocalPlayer(NewPlayer);
+					NewPlayer = nullptr;
 				}
 			}
 			else
@@ -699,15 +709,29 @@ int32 UGameInstance::AddLocalPlayer(ULocalPlayer* NewLocalPlayer, int32 Controll
 bool UGameInstance::RemoveLocalPlayer(ULocalPlayer* ExistingPlayer)
 {
 	// FIXME: Notify server we want to leave the game if this is an online game
-	if (ExistingPlayer->PlayerController != NULL)
+	if (ExistingPlayer->PlayerController != nullptr)
 	{
+		bool bShouldRemovePlayer = (ExistingPlayer->PlayerController->Role == ROLE_Authority);
+
 		// FIXME: Do this all inside PlayerRemoved?
 		ExistingPlayer->PlayerController->CleanupGameViewport();
 
+		UWorld* CurrentWorld = GetWorld();
+		if (CurrentWorld && CurrentWorld->DemoNetDriver && CurrentWorld->IsPlayingReplay())
+		{
+			if (!CurrentWorld->DemoNetDriver->RemoveSplitscreenViewer(ExistingPlayer->PlayerController))
+			{
+				UE_LOG(LogPlayerManagement, Warning, TEXT("UGameInstance::RemovePlayer: Did not remove player %s with ControllerId %i as it was unable to be removed from the demo"), 
+					*ExistingPlayer->GetName(), ExistingPlayer->GetControllerId());
+			}
+			bShouldRemovePlayer = true;
+		}
+
 		// Destroy the player's actors.
-		if ( ExistingPlayer->PlayerController->Role == ROLE_Authority )
+		if (bShouldRemovePlayer)
 		{
-			ExistingPlayer->PlayerController->Destroy();
+			// This is fine to forceremove as we have to be in a special case for demos or the authority
+			ExistingPlayer->PlayerController->Destroy(true);
 		}
 	}
 
@@ -720,7 +744,7 @@ bool UGameInstance::RemoveLocalPlayer(ULocalPlayer* ExistingPlayer)
 		LocalPlayers.RemoveAt(OldIndex);
 
 		// Notify the viewport so the viewport can do the fixups, resize, etc
-		if (GetGameViewportClient() != NULL)
+		if (GetGameViewportClient() != nullptr)
 		{
 			GetGameViewportClient()->NotifyPlayerRemoved(OldIndex, ExistingPlayer);
 		}
@@ -728,7 +752,7 @@ bool UGameInstance::RemoveLocalPlayer(ULocalPlayer* ExistingPlayer)
 
 	// Disassociate this viewport client from the player.
 	// Do this after notifications, as some of them require the ViewportClient.
-	ExistingPlayer->ViewportClient = NULL;
+	ExistingPlayer->ViewportClient = nullptr;
 
 	UE_LOG(LogPlayerManagement, Log, TEXT("UGameInstance::RemovePlayer: Removed player %s with ControllerId %i at index %i (%i remaining players)"), *ExistingPlayer->GetName(), ExistingPlayer->GetControllerId(), OldIndex, LocalPlayers.Num());
 
