Workaround for early data deletion issue in D3D11 with RHI thread. UTexture2D::UpdateTextureRegions would enqueue subresource region updates to the RHI thread then delete the source data immediately on the render thread, this change pushes the deletion to the RHI thread also.

diff --git a/Engine/Source/Runtime/Engine/Private/Texture2D.cpp b/Engine/Source/Runtime/Engine/Private/Texture2D.cpp
--- a/Engine/Source/Runtime/Engine/Private/Texture2D.cpp
+++ b/Engine/Source/Runtime/Engine/Private/Texture2D.cpp
@@ -1309,8 +1309,13 @@ void UTexture2D::UpdateTextureRegions(int32 MipIndex, uint32 NumRegions, const F
 							);
 					}
 				}
-				DataCleanupFunc(RegionData->SrcData, RegionData->Regions);
-				delete RegionData;
+
+				// The deletion of source data may need to be deferred to the RHI thread after the updates occur
+				RHICmdList.EnqueueLambda([RegionData, DataCleanupFunc](FRHICommandList&)
+				{
+					DataCleanupFunc(RegionData->SrcData, RegionData->Regions);
+					delete RegionData;
+				});
 			});
 	}
 }
