Fix for incorrect deleted actors in amortized delta checkpoints

diff --git a/Engine/Source/Runtime/Engine/Classes/Engine/DemoNetDriver.h b/Engine/Source/Runtime/Engine/Classes/Engine/DemoNetDriver.h
--- a/Engine/Source/Runtime/Engine/Classes/Engine/DemoNetDriver.h
+++ b/Engine/Source/Runtime/Engine/Classes/Engine/DemoNetDriver.h
@@ -1041,6 +1041,9 @@ private:
 		FArchivePos			CheckpointOffset;
 		uint32				GuidCacheSize;
 
+		TSet<FString>		DeltaDeletedNetStartupActors;
+		TSet<FNetworkGUID>	DeltaDeletedActorGuids;
+
 		void CountBytes(FArchive& Ar) const
 		{
 			CheckpointAckState.CountBytes(Ar);
diff --git a/Engine/Source/Runtime/Engine/Private/DemoNetDriver.cpp b/Engine/Source/Runtime/Engine/Private/DemoNetDriver.cpp
--- a/Engine/Source/Runtime/Engine/Private/DemoNetDriver.cpp
+++ b/Engine/Source/Runtime/Engine/Private/DemoNetDriver.cpp
@@ -1856,6 +1856,12 @@ void UDemoNetDriver::SaveCheckpoint()
 	CheckpointSaveContext.TotalCheckpointSaveFrames = 0;
 	LastCheckpointTime = DemoCurrentTime;
 
+	if (bDeltaCheckpoint)
+	{
+		CheckpointSaveContext.DeltaDeletedActorGuids = MoveTemp(DeltaDeletedActorGuids);
+		CheckpointSaveContext.DeltaDeletedNetStartupActors = MoveTemp(DeltaDeletedNetStartupActors);
+	}
+
 	UE_LOG( LogDemo, Log, TEXT( "Starting checkpoint. Actors: %i" ), GetNetworkObjectList().GetActiveObjects().Num() );
 
 	// Do the first checkpoint tick now if we're not amortizing
@@ -2023,11 +2029,11 @@ void UDemoNetDriver::TickCheckpoint()
 						// Save deleted startup actors
 						if (bDeltaCheckpoint)
 						{
-							*CheckpointArchive << DeltaDeletedNetStartupActors;
-							DeltaDeletedNetStartupActors.Empty();
+							*CheckpointArchive << CheckpointSaveContext.DeltaDeletedNetStartupActors;
+							CheckpointSaveContext.DeltaDeletedNetStartupActors.Empty();
 
-							*CheckpointArchive << DeltaDeletedActorGuids;
-							DeltaDeletedActorGuids.Empty();
+							*CheckpointArchive << CheckpointSaveContext.DeltaDeletedActorGuids;
+							CheckpointSaveContext.DeltaDeletedActorGuids.Empty();
 						}
 						else
 						{
