#ue4 - Turn back on "p.ForceJumpPeakSubstep" and fix the num apex jumps per simulation tracking to not reset during simulation in rare cases.

diff --git a/Engine/Source/Runtime/Engine/Private/Components/CharacterMovementComponent.cpp b/Engine/Source/Runtime/Engine/Private/Components/CharacterMovementComponent.cpp
--- a/Engine/Source/Runtime/Engine/Private/Components/CharacterMovementComponent.cpp
+++ b/Engine/Source/Runtime/Engine/Private/Components/CharacterMovementComponent.cpp
@@ -196,7 +196,7 @@ namespace CharacterMovementCVars
 		TEXT( "If 1, remove invalid replay samples that can occur due to oversampling (sampling at higher rate than physics is being ticked)" ),
 		ECVF_Default);
 
-	static int32 ForceJumpPeakSubstep = 0;
+	static int32 ForceJumpPeakSubstep = 1;
 	FAutoConsoleVariableRef CVarForceJumpPeakSubstep(
 		TEXT("p.ForceJumpPeakSubstep"),
 		ForceJumpPeakSubstep,
@@ -1638,6 +1638,7 @@ void UCharacterMovementComponent::SimulateRootMotion(float DeltaSeconds, const F
 			bNetworkMovementModeChanged = false;
 		}
 
+		NumJumpApexAttempts = 0;
 		StartNewPhysics(DeltaSeconds, 0);
 		// fixme laurent - simulate movement seems to have step up issues? investigate as that would be cheaper to use.
 		// 		SimulateMovement(DeltaSeconds);
@@ -2358,6 +2359,7 @@ void UCharacterMovementComponent::PerformMovement(float DeltaSeconds)
 
 		// Clear jump input now, to allow movement events to trigger it for next update.
 		CharacterOwner->ClearJumpInput(DeltaSeconds);
+		NumJumpApexAttempts = 0;
 
 		// change position
 		StartNewPhysics(DeltaSeconds, 0);
@@ -2827,10 +2829,6 @@ void UCharacterMovementComponent::StartNewPhysics(float deltaTime, int32 Iterati
 
 	const bool bSavedMovementInProgress = bMovementInProgress;
 	bMovementInProgress = true;
-	if (Iterations == 0)
-	{
-		NumJumpApexAttempts = 0;
-	}
 
 	switch ( MovementMode )
 	{
