[LLM] more tidying for EngineMisc:
 - add some missing LLM tags to particle system creation moving about 18MB to Particles.
 - move Particle, Landscape, Static & SkeletalMesh scene primitives into their respective tag (about 2MB total).
 - tag UAnimaInstance InitializeAnimation & UpdateAnimation, moving about 15MB to Animation.

diff --git a/Engine/Source/Runtime/Engine/Private/Components/SkinnedMeshComponent.cpp b/Engine/Source/Runtime/Engine/Private/Components/SkinnedMeshComponent.cpp
--- a/Engine/Source/Runtime/Engine/Private/Components/SkinnedMeshComponent.cpp
+++ b/Engine/Source/Runtime/Engine/Private/Components/SkinnedMeshComponent.cpp
@@ -505,6 +505,8 @@ void USkinnedMeshComponent::OnUnregister()
 
 void USkinnedMeshComponent::CreateRenderState_Concurrent()
 {
+	LLM_SCOPE(ELLMTag::SkeletalMesh);
+
 	if( SkeletalMesh )
 	{
 		// Attempting to track down UE-45505, where it looks as if somehow a skeletal mesh component's mesh has only been partially loaded, causing a mismatch in the LOD arrays
diff --git a/Engine/Source/Runtime/Engine/Private/Particles/ParticleComponents.cpp b/Engine/Source/Runtime/Engine/Private/Particles/ParticleComponents.cpp
--- a/Engine/Source/Runtime/Engine/Private/Particles/ParticleComponents.cpp
+++ b/Engine/Source/Runtime/Engine/Private/Particles/ParticleComponents.cpp
@@ -3840,6 +3840,7 @@ void UParticleSystemComponent::OnUnregister()
 
 void UParticleSystemComponent::CreateRenderState_Concurrent()
 {
+	LLM_SCOPE(ELLMTag::Particles);
 	SCOPE_CYCLE_COUNTER(STAT_ParticleSystemComponent_CreateRenderState_Concurrent);
 	SCOPE_CYCLE_COUNTER(STAT_ParticlesOverview_GT_CNC);
 
diff --git a/Engine/Source/Runtime/Engine/Private/Particles/ParticleGpuSimulation.cpp b/Engine/Source/Runtime/Engine/Private/Particles/ParticleGpuSimulation.cpp
--- a/Engine/Source/Runtime/Engine/Private/Particles/ParticleGpuSimulation.cpp
+++ b/Engine/Source/Runtime/Engine/Private/Particles/ParticleGpuSimulation.cpp
@@ -10,6 +10,7 @@
 #include "Stats/Stats.h"
 #include "Misc/MemStack.h"
 #include "HAL/IConsoleManager.h"
+#include "HAL/LowLevelMemTracker.h"
 #include "RHIDefinitions.h"
 #include "RHI.h"
 #include "RenderingThread.h"
@@ -4467,6 +4468,8 @@ FAutoConsoleCommand DumpTileAllocsCommand(
 
 void FFXSystem::InitGPUSimulation()
 {
+	LLM_SCOPE(ELLMTag::Particles);
+
 	check(ParticleSimulationResources == NULL);
 	ensure(GParticleSimulationTextureSizeX > 0 && GParticleSimulationTextureSizeY > 0);
 	/** How many tiles are in the simulation textures. */
@@ -4537,6 +4540,8 @@ void FFXSystem::AddGPUSimulation(FParticleSimulationGPU* Simulation)
 {
 	if (!IsPendingKill())
 	{
+		LLM_SCOPE(ELLMTag::Particles);
+
 		FFXSystem* FXSystem = this;
 		ENQUEUE_RENDER_COMMAND(FAddGPUSimulationCommand)(
 			[FXSystem, Simulation](FRHICommandListImmediate& RHICmdList)
@@ -4572,6 +4577,8 @@ void FFXSystem::RemoveGPUSimulation(FParticleSimulationGPU* Simulation)
 
 int32 FFXSystem::AddSortedGPUSimulation(FParticleSimulationGPU* Simulation, const FVector& ViewOrigin)
 {
+	LLM_SCOPE(ELLMTag::Particles);
+
 	check(FeatureLevel == ERHIFeatureLevel::SM5);
 	const int32 BufferOffset = ParticleSimulationResources->SortedParticleCount;
 	ParticleSimulationResources->SortedParticleCount += Simulation->VertexBuffer.ParticleCount;
@@ -4701,6 +4708,7 @@ void FFXSystem::SimulateGPUParticles(
 	FUniformBufferRHIParamRef SceneTexturesUniformBuffer
 	)
 {
+	LLM_SCOPE(ELLMTag::Particles);
 	check(IsInRenderingThread());
 	SCOPE_CYCLE_COUNTER(STAT_GPUParticleTickTime);
 
@@ -5276,6 +5284,8 @@ FGPUSpriteResources* BeginCreateGPUSpriteResources( const FGPUSpriteResourceData
 	FGPUSpriteResources* Resources = NULL;
 	if (RHISupportsGPUParticles())
 	{
+		LLM_SCOPE(ELLMTag::Particles);
+
 		Resources = new FGPUSpriteResources;
 		//@TODO Ideally FGPUSpriteEmitterInfo::Resources would be a TRefCountPtr<FGPUSpriteResources>, but
 		//since that class is defined in this file, we can't do that, so we just addref here instead
diff --git a/Engine/Source/Runtime/Engine/Private/Particles/ParticleModules.cpp b/Engine/Source/Runtime/Engine/Private/Particles/ParticleModules.cpp
--- a/Engine/Source/Runtime/Engine/Private/Particles/ParticleModules.cpp
+++ b/Engine/Source/Runtime/Engine/Private/Particles/ParticleModules.cpp
@@ -10,6 +10,7 @@
 #include "Serialization/ArchiveUObject.h"
 #include "UObject/UnrealType.h"
 #include "HAL/IConsoleManager.h"
+#include "HAL/LowLevelMemTracker.h"
 #include "UObject/RenderingObjectVersion.h"
 #include "UObject/UObjectHash.h"
 #include "UObject/Package.h"
@@ -4953,6 +4954,8 @@ FParticleEmitterInstance* UParticleModuleTypeDataGpu::CreateInstance(UParticleEm
 	FParticleEmitterInstance* Instance = NULL;
 	if (World->Scene && RHISupportsGPUParticles())
 	{
+		LLM_SCOPE(ELLMTag::Particles);
+
 		check( InComponent && InComponent->FXSystem );
 		Instance = InComponent->FXSystem->CreateGPUSpriteEmitterInstance( EmitterInfo );
 		Instance->InitParameters( InEmitterParent, InComponent );
diff --git a/Engine/Source/Runtime/Landscape/Private/Landscape.cpp b/Engine/Source/Runtime/Landscape/Private/Landscape.cpp
--- a/Engine/Source/Runtime/Landscape/Private/Landscape.cpp
+++ b/Engine/Source/Runtime/Landscape/Private/Landscape.cpp
@@ -2965,6 +2965,7 @@ void ALandscape::TickActor(float DeltaTime, ELevelTick TickType, FActorTickFunct
 void ALandscapeProxy::TickActor(float DeltaTime, ELevelTick TickType, FActorTickFunction& ThisTickFunction)
 {
 	CSV_SCOPED_TIMING_STAT_EXCLUSIVE(Landscape);
+	LLM_SCOPE(ELLMTag::Landscape);
 #if WITH_EDITOR
 	// editor-only
 	UWorld* World = GetWorld();
diff --git a/Engine/Source/Runtime/Renderer/Private/PrimitiveSceneInfo.cpp b/Engine/Source/Runtime/Renderer/Private/PrimitiveSceneInfo.cpp
--- a/Engine/Source/Runtime/Renderer/Private/PrimitiveSceneInfo.cpp
+++ b/Engine/Source/Runtime/Renderer/Private/PrimitiveSceneInfo.cpp
@@ -11,6 +11,7 @@
 #include "SceneCore.h"
 #include "VelocityRendering.h"
 #include "ScenePrivate.h"
+#include "HAL/LowLevelMemTracker.h"
 #include "RayTracing/RayTracingMaterialHitShaders.h"
 
 /** An implementation of FStaticPrimitiveDrawInterface that stores the drawn elements for the rendering thread to use. */
@@ -390,6 +391,8 @@ void FPrimitiveSceneInfo::RemoveCachedMeshDrawCommands()
 
 void FPrimitiveSceneInfo::AddStaticMeshes(FRHICommandListImmediate& RHICmdList, bool bAddToStaticDrawLists)
 {
+	LLM_SCOPE(ELLMTag::StaticMesh);
+
 	// Cache the primitive's static mesh elements.
 	FBatchingSPDI BatchingSPDI(this);
 	BatchingSPDI.SetHitProxy(DefaultDynamicHitProxy);
