[LLM] tag UAnimaInstance InitializeAnimation & UpdateAnimation, moving about 15MB to Animation out of EngineMisc

diff --git a/Engine/Source/Runtime/Engine/Private/Animation/AnimInstance.cpp b/Engine/Source/Runtime/Engine/Private/Animation/AnimInstance.cpp
--- a/Engine/Source/Runtime/Engine/Private/Animation/AnimInstance.cpp
+++ b/Engine/Source/Runtime/Engine/Private/Animation/AnimInstance.cpp
@@ -25,6 +25,7 @@
 #include "SkeletalRenderPublic.h"
 #include "Rendering/SkeletalMeshRenderData.h"
 #include "Kismet/KismetSystemLibrary.h"
+#include "HAL/LowLevelMemTracker.h"
 
 /** Anim stats */
 
@@ -168,6 +169,7 @@ void UAnimInstance::InitializeAnimation()
 {
 	FScopeCycleCounterUObject ContextScope(this);
 	SCOPE_CYCLE_COUNTER(STAT_AnimInitTime);
+	LLM_SCOPE(ELLMTag::Animation);
 
 	UninitializeAnimation();
 
@@ -332,6 +334,8 @@ void UAnimInstance::UpdateMontageSyncGroup()
 
 void UAnimInstance::UpdateAnimation(float DeltaSeconds, bool bNeedsValidRootMotion)
 {
+	LLM_SCOPE(ELLMTag::Animation);
+
 #if DO_CHECK
 	checkf(!bUpdatingAnimation, TEXT("UpdateAnimation already in progress, circular detected for SkeletalMeshComponent [%s], AnimInstance [%s]"), *GetNameSafe(GetOwningComponent()),  *GetName());
 	TGuardValue<bool> CircularGuard(bUpdatingAnimation, true);
