fixing bug where paths that were close to the same would not be properly scanned in -Game, now we only dont scane the path if it exactly matches a path that has already been scanned or is in the file path is a child of one that has already been scanned

diff --git a/Engine/Source/Runtime/Engine/Private/AssetManager.cpp b/Engine/Source/Runtime/Engine/Private/AssetManager.cpp
--- a/Engine/Source/Runtime/Engine/Private/AssetManager.cpp
+++ b/Engine/Source/Runtime/Engine/Private/AssetManager.cpp
@@ -2779,7 +2779,7 @@ void UAssetManager::ScanPathsSynchronous(const TArray<FString>& PathsToScan) con
 
 			for (const FString& AlreadyScanned : AlreadyScannedDirectories)
 			{
-				if (PackageName.Contains(AlreadyScanned))
+				if (PackageName == AlreadyScanned || PackageName.StartsWith(AlreadyScanned + TEXT("/")))
 				{
 					bAlreadyScanned = true;
 					break;
@@ -2805,7 +2805,7 @@ void UAssetManager::ScanPathsSynchronous(const TArray<FString>& PathsToScan) con
 		{
 			for (const FString& AlreadyScanned : AlreadyScannedDirectories)
 			{
-				if (Path.Contains(AlreadyScanned))
+				if (Path == AlreadyScanned || Path.StartsWith(AlreadyScanned + TEXT("/")))
 				{
 					bAlreadyScanned = true;
 					break;
