Fix for distance performance optimization not using position of active sound but instead that of closest listener

diff --git a/Engine/Source/Runtime/Engine/Private/ActiveSound.cpp b/Engine/Source/Runtime/Engine/Private/ActiveSound.cpp
--- a/Engine/Source/Runtime/Engine/Private/ActiveSound.cpp
+++ b/Engine/Source/Runtime/Engine/Private/ActiveSound.cpp
@@ -417,7 +417,7 @@ void FActiveSound::UpdateWaveInstances( TArray<FWaveInstance*> &InWaveInstances,
 		float ApparentMaxDistance = MaxDistance * FocusDistanceScale;
 
 		// Check if we're out of range of being audible, and early return if there's no chance of making sounds
-		if (!Sound->IsVirtualizeWhenSilent() && !AudioDevice->LocationIsAudible(ClosestListenerPtr->Transform.GetLocation(), ApparentMaxDistance))
+		if (!Sound->IsVirtualizeWhenSilent() && !AudioDevice->LocationIsAudible(Transform.GetLocation(), ApparentMaxDistance))
 		{
 			return;
 		}
