Don't generate the average scale for grass, as we will render using the mesh scale only.

diff --git a/Engine/Source/Runtime/Engine/Classes/Components/HierarchicalInstancedStaticMeshComponent.h b/Engine/Source/Runtime/Engine/Classes/Components/HierarchicalInstancedStaticMeshComponent.h
--- a/Engine/Source/Runtime/Engine/Classes/Components/HierarchicalInstancedStaticMeshComponent.h
+++ b/Engine/Source/Runtime/Engine/Classes/Components/HierarchicalInstancedStaticMeshComponent.h
@@ -229,7 +229,7 @@ public:
 	virtual bool ShouldCreatePhysicsState() const override;
 
 	bool BuildTreeIfOutdated(bool Async, bool ForceUpdate);
-	static void BuildTreeAnyThread(TArray<FMatrix>& InstanceTransforms, const FBox& MeshBox, TArray<FClusterNode>& OutClusterTree, TArray<int32>& OutSortedInstances, TArray<int32>& OutInstanceReorderTable, int32& OutOcclusionLayerNum, int32 MaxInstancesPerLeaf );
+	static void BuildTreeAnyThread(TArray<FMatrix>& InstanceTransforms, const FBox& MeshBox, TArray<FClusterNode>& OutClusterTree, TArray<int32>& OutSortedInstances, TArray<int32>& OutInstanceReorderTable, int32& OutOcclusionLayerNum, int32 MaxInstancesPerLeaf, bool InGenerateInstanceScalingRange);
 	void AcceptPrebuiltTree(TArray<FClusterNode>& InClusterTree, int32 InOcclusionLayerNumNodes, int32 InNumBuiltRenderInstances);
 	bool IsAsyncBuilding() const { return bIsAsyncBuilding; }
 	bool IsTreeFullyBuilt() const { return NumBuiltInstances == PerInstanceSMData.Num(); }
diff --git a/Engine/Source/Runtime/Engine/Private/HierarchicalInstancedStaticMesh.cpp b/Engine/Source/Runtime/Engine/Private/HierarchicalInstancedStaticMesh.cpp
--- a/Engine/Source/Runtime/Engine/Private/HierarchicalInstancedStaticMesh.cpp
+++ b/Engine/Source/Runtime/Engine/Private/HierarchicalInstancedStaticMesh.cpp
@@ -187,6 +187,7 @@ protected:
 	
 	int32 InstancingRandomSeed;
 	float DensityScaling;
+	bool GenerateInstanceScalingRange;
 
 	TArray<int32> SortIndex;
 	TArray<FVector> SortPoints;
@@ -379,12 +380,13 @@ public:
 	TUniquePtr<FClusterTree> Result;
 	TUniquePtr<FStaticMeshInstanceData> BuiltInstanceData;
 
-	FClusterBuilder(TArray<FMatrix> InTransforms, const FBox& InInstBox, int32 InMaxInstancesPerLeaf, float InDensityScaling, int32 InInstancingRandomSeed)
+	FClusterBuilder(TArray<FMatrix> InTransforms, const FBox& InInstBox, int32 InMaxInstancesPerLeaf, float InDensityScaling, int32 InInstancingRandomSeed, bool InGenerateInstanceScalingRange)
 		: OriginalNum(InTransforms.Num())
 		, InstBox(InInstBox)
 		, MaxInstancesPerLeaf(InMaxInstancesPerLeaf)
 		, InstancingRandomSeed(InInstancingRandomSeed)
 		, DensityScaling(InDensityScaling)
+		, GenerateInstanceScalingRange(InGenerateInstanceScalingRange)
 		, Transforms(MoveTemp(InTransforms))
 		, Result(nullptr)
 	{
@@ -454,10 +456,13 @@ public:
 				FBox ThisInstBox = InstBox.TransformBy(ThisInstTrans);
 				NodeBox += ThisInstBox;
 
-				FVector CurrentScale = ThisInstTrans.GetScaleVector();
+				if (GenerateInstanceScalingRange)
+				{
+					FVector CurrentScale = ThisInstTrans.GetScaleVector();
 
-				Node.MinInstanceScale = Node.MinInstanceScale.ComponentMin(CurrentScale);
-				Node.MaxInstanceScale = Node.MaxInstanceScale.ComponentMax(CurrentScale);
+					Node.MinInstanceScale = Node.MinInstanceScale.ComponentMin(CurrentScale);
+					Node.MaxInstanceScale = Node.MaxInstanceScale.ComponentMax(CurrentScale);
+				}
 			}
 			Node.BoundMin = NodeBox.Min;
 			Node.BoundMax = NodeBox.Max;
@@ -634,8 +639,11 @@ public:
 						NodeBox += ChildNode.BoundMin;
 						NodeBox += ChildNode.BoundMax;
 
-						Node.MinInstanceScale = Node.MinInstanceScale.ComponentMin(ChildNode.MinInstanceScale);
-						Node.MaxInstanceScale = Node.MaxInstanceScale.ComponentMax(ChildNode.MaxInstanceScale);
+						if (GenerateInstanceScalingRange)
+						{
+							Node.MinInstanceScale = Node.MinInstanceScale.ComponentMin(ChildNode.MinInstanceScale);
+							Node.MaxInstanceScale = Node.MaxInstanceScale.ComponentMax(ChildNode.MaxInstanceScale);
+						}
 					}
 					Node.BoundMin = NodeBox.Min;
 					Node.BoundMax = NodeBox.Max;
@@ -651,6 +659,13 @@ public:
 		{
 			Result->InstanceReorderTable[SortedInstances[Index]] = Index;
 		}
+
+		// Output a general scale of 1 if we dont want the scaling range
+		if (!GenerateInstanceScalingRange)
+		{
+			Result->Nodes[0].MinInstanceScale = FVector::OneVector;
+			Result->Nodes[0].MaxInstanceScale = FVector::OneVector;
+		}
 	}
 };
 
@@ -712,7 +727,7 @@ static void TestFoliage(const TArray<FString>& Args)
 	{
 		InstanceTransforms[Index] = Instances[Index].Transform;
 	}
-	FClusterBuilder Builder(InstanceTransforms, TempBox, 16, 1.0f, 1);
+	FClusterBuilder Builder(InstanceTransforms, TempBox, 16, 1.0f, 1, 0);
 	Builder.BuildTree();
 
 	int32 Level = 0;
@@ -1595,7 +1610,7 @@ void FHierarchicalStaticMeshSceneProxy::GetDynamicMeshElements(const TArray<cons
 				float LODScale = CVarFoliageLODDistanceScale.GetValueOnRenderThread();
 				float LODRandom = CVarRandomLODRange.GetValueOnRenderThread();
 				float MaxDrawDistanceScale = GetCachedScalabilityCVars().ViewDistanceScale;
-
+				
 				FVector AverageScale = (InstanceParams.Tree[0].MinInstanceScale + (InstanceParams.Tree[0].MaxInstanceScale - InstanceParams.Tree[0].MinInstanceScale) / 2.0f);
 				FBoxSphereBounds ScaledBounds = RenderData->Bounds.TransformBy(FTransform(FRotator::ZeroRotator, FVector::ZeroVector, AverageScale));
 				float SphereRadius = ScaledBounds.SphereRadius;
@@ -2446,7 +2461,7 @@ void UHierarchicalInstancedStaticMeshComponent::BuildTree()
 			InstanceTransforms[Index] = PerInstanceSMData[Index].Transform;
 		}
 
-		FClusterBuilder Builder(InstanceTransforms, GetStaticMesh()->GetBounds().GetBox(), DesiredInstancesPerLeaf(), CurrentDensityScaling, InstancingRandomSeed);
+		FClusterBuilder Builder(InstanceTransforms, GetStaticMesh()->GetBounds().GetBox(), DesiredInstancesPerLeaf(), CurrentDensityScaling, InstancingRandomSeed, PerInstanceSMData.Num() > 0);
 		Builder.BuildTreeAndBuffer();
 
 		NumBuiltInstances = Builder.Result->InstanceReorderTable.Num();
@@ -2513,7 +2528,8 @@ void UHierarchicalInstancedStaticMeshComponent::BuildTreeAnyThread(
 	TArray<int32>& OutSortedInstances,
 	TArray<int32>& OutInstanceReorderTable,
 	int32& OutOcclusionLayerNum,
-	int32 MaxInstancesPerLeaf
+	int32 MaxInstancesPerLeaf,
+	bool InGenerateInstanceScalingRange
 	)
 {
 	check(MaxInstancesPerLeaf > 0);
@@ -2522,7 +2538,7 @@ void UHierarchicalInstancedStaticMeshComponent::BuildTreeAnyThread(
 	float DensityScaling = 1.0f;
 	int32 InstancingRandomSeed = 1;
 
-	FClusterBuilder Builder(InstanceTransforms, MeshBox, MaxInstancesPerLeaf, DensityScaling, InstancingRandomSeed);
+	FClusterBuilder Builder(InstanceTransforms, MeshBox, MaxInstancesPerLeaf, DensityScaling, InstancingRandomSeed, InGenerateInstanceScalingRange);
 	Builder.BuildTree();
 	OutOcclusionLayerNum = Builder.Result->OutOcclusionLayerNum;
 
@@ -2752,7 +2768,7 @@ void UHierarchicalInstancedStaticMeshComponent::BuildTreeAsync()
 
 		UE_LOG(LogStaticMesh, Verbose, TEXT("Copied %d transforms in %.3fs."), Num, float(FPlatformTime::Seconds() - StartTime));
 
-		TSharedRef<FClusterBuilder, ESPMode::ThreadSafe> Builder(new FClusterBuilder(InstanceTransforms, GetStaticMesh()->GetBounds().GetBox(), DesiredInstancesPerLeaf(), CurrentDensityScaling, InstancingRandomSeed));
+		TSharedRef<FClusterBuilder, ESPMode::ThreadSafe> Builder(new FClusterBuilder(InstanceTransforms, GetStaticMesh()->GetBounds().GetBox(), DesiredInstancesPerLeaf(), CurrentDensityScaling, InstancingRandomSeed, PerInstanceSMData.Num() > 0));
 
 		bIsAsyncBuilding = true;
 
diff --git a/Engine/Source/Runtime/Landscape/Private/LandscapeGrass.cpp b/Engine/Source/Runtime/Landscape/Private/LandscapeGrass.cpp
--- a/Engine/Source/Runtime/Landscape/Private/LandscapeGrass.cpp
+++ b/Engine/Source/Runtime/Landscape/Private/LandscapeGrass.cpp
@@ -1857,7 +1857,7 @@ struct FAsyncGrassBuilder : public FGrassBuilderBase
 		{
 			TArray<int32> SortedInstances;
 			TArray<int32> InstanceReorderTable;
-			UHierarchicalInstancedStaticMeshComponent::BuildTreeAnyThread(InstanceTransforms, MeshBox, ClusterTree, SortedInstances, InstanceReorderTable, OutOcclusionLayerNum, DesiredInstancesPerLeaf);
+			UHierarchicalInstancedStaticMeshComponent::BuildTreeAnyThread(InstanceTransforms, MeshBox, ClusterTree, SortedInstances, InstanceReorderTable, OutOcclusionLayerNum, DesiredInstancesPerLeaf, false);
 
 			// in-place sort the instances
 			
