fix crash on loading - due to memory overflow

diff --git a/Engine/Source/Runtime/Engine/Private/SkeletalMesh.cpp b/Engine/Source/Runtime/Engine/Private/SkeletalMesh.cpp
--- a/Engine/Source/Runtime/Engine/Private/SkeletalMesh.cpp
+++ b/Engine/Source/Runtime/Engine/Private/SkeletalMesh.cpp
@@ -1708,7 +1708,7 @@ void USkeletalMesh::PostLoad()
 					const TArray<FTransform>& RawBonePose = RefSkeleton.GetRefBonePose();
 					check(RetargetBasePose.GetTypeSize() == RawBonePose.GetTypeSize());
 					const int32 ElementSize = RetargetBasePose.GetTypeSize();
-					FMemory::Memcpy(RetargetBasePose.GetData() + ElementSize*RawBoneNum, RawBonePose.GetData() + ElementSize*RawBoneNum, ElementSize*VBNum);
+					FMemory::Memcpy(RetargetBasePose.GetData() + RawBoneNum, RawBonePose.GetData() + RawBoneNum, ElementSize*VBNum);
 				}
 				else
 				{
