Fix for crash due to rare case where RemoveActiveSound was called more than once

diff --git a/Engine/Source/Runtime/Engine/Private/AudioDevice.cpp b/Engine/Source/Runtime/Engine/Private/AudioDevice.cpp
--- a/Engine/Source/Runtime/Engine/Private/AudioDevice.cpp
+++ b/Engine/Source/Runtime/Engine/Private/AudioDevice.cpp
@@ -2086,14 +2086,14 @@ void FAudioDevice::StopQuietSoundsDueToMaxConcurrency(TArray<FWaveInstance*>& Wa
 		}
 	}
 
-	for (int32 i = 0; i < ActiveSoundsCopy.Num(); ++i)
+	for (int32 i = ActiveSoundsCopy.Num() - 1; i >= 0; --i)
 	{
 		if (FActiveSound* ActiveSound = ActiveSoundsCopy[i])
 		{
 			if (ActiveSound->bShouldStopDueToMaxConcurrency)
 			{
-				ActiveSound->Stop(false);
-				ConcurrencyManager.StopActiveSound(ActiveSound);
+				AddSoundToStop(ActiveSound);
+				ActiveSoundsCopy.RemoveAtSwap(i, 1, false);
 			}
 		}
 	}
@@ -4390,12 +4390,18 @@ void FAudioDevice::RemoveActiveSound(FActiveSound* ActiveSound)
 	check(IsInAudioThread());
 
 	// Perform the notification
-	if (ActiveSound->GetAudioComponentID() > 0)
+	const int32 ComponentID = ActiveSound->GetAudioComponentID();
+	if (ComponentID > 0)
 	{
-		UAudioComponent::PlaybackCompleted(ActiveSound->GetAudioComponentID(), false);
+		UAudioComponent::PlaybackCompleted(ComponentID, false);
 	}
 
 	const int32 NumRemoved = ActiveSounds.RemoveSwap(ActiveSound);
+	if (!ensureMsgf(NumRemoved > 0, TEXT("Attempting to remove an already removed ActiveSound '%s'"), ActiveSound->Sound ? *ActiveSound->Sound->GetName() : TEXT("N/A")))
+	{
+		return;
+	}
+
 	check(NumRemoved == 1);
 }
 
