Log out replay header flags along with the recording version

#rnx

diff --git a/Engine/Source/Runtime/Engine/Private/DemoNetDriver.cpp b/Engine/Source/Runtime/Engine/Private/DemoNetDriver.cpp
--- a/Engine/Source/Runtime/Engine/Private/DemoNetDriver.cpp
+++ b/Engine/Source/Runtime/Engine/Private/DemoNetDriver.cpp
@@ -4187,8 +4187,8 @@ void UDemoNetDriver::ReplayStreamingReady( const FStartStreamingResult& Result )
 			}
 		}
 
-		UE_LOG(LogDemo, Log, TEXT("ReplayStreamingReady: playing back replay [%s] %s, which was recorded on engine version %s"),
-			*PlaybackDemoHeader.Guid.ToString(EGuidFormats::Digits), *DemoURL.Map, *PlaybackDemoHeader.EngineVersion.ToString());
+		UE_LOG(LogDemo, Log, TEXT("ReplayStreamingReady: playing back replay [%s] %s, which was recorded on engine version %s with flags 0x%08x"),
+			*PlaybackDemoHeader.Guid.ToString(EGuidFormats::Digits), *DemoURL.Map, *PlaybackDemoHeader.EngineVersion.ToString(), PlaybackDemoHeader.HeaderFlags);
 
 		// Notify all listeners that a demo is starting
 		OnDemoStarted.Broadcast(this);
