Added PIE orbit functionality and hotkey for ToggleDebugCamera.

diff --git a/Engine/Config/BaseInput.ini b/Engine/Config/BaseInput.ini
--- a/Engine/Config/BaseInput.ini
+++ b/Engine/Config/BaseInput.ini
@@ -45,6 +45,8 @@ TriggerKeyIsConsideredReleased=0.25
 ; Please remove EnableGDT command binding after removing deprecated GameplayDebugger module (we are going to use Gameplay Debugger plugin soon)
 +DebugExecBindings=(Key=Apostrophe,Command="EnableGDT")
 +DebugExecBindings=(Key=Quote,Command="EnableGDT")
+; --- PIE only binding
++DebugExecBindings=(Key=Backslash,Command="ToggleDebugCamera")
 
 [/Script/Engine.Console]
 HistoryBot=-1
diff --git a/Engine/Source/Runtime/Engine/Classes/Engine/DebugCameraController.h b/Engine/Source/Runtime/Engine/Classes/Engine/DebugCameraController.h
--- a/Engine/Source/Runtime/Engine/Classes/Engine/DebugCameraController.h
+++ b/Engine/Source/Runtime/Engine/Classes/Engine/DebugCameraController.h
@@ -32,6 +32,10 @@ class ENGINE_API ADebugCameraController
 	UPROPERTY()
 	uint32 bIsFrozenRendering:1;
 
+	/** Whether to orbit selected actor. */
+	UPROPERTY()
+	uint32 bIsOrbitingSelectedActor : 1;
+
 	/** Visualizes the frustum of the camera */
 	UPROPERTY()
 	class UDrawFrustumComponent* DrawFrustum;
@@ -68,6 +72,18 @@ class ENGINE_API ADebugCameraController
 	 */
 	virtual void ToggleFreezeRendering();
 
+	/**
+	 * Updates the rotation of player, based on ControlRotation after RotationInput has been applied.
+	 * This may then be modified by the PlayerCamera, and is passed to Pawn->FaceRotation().
+	 */
+	virtual void UpdateRotation(float DeltaTime) override;
+
+	/** Gets pivot to use when orbiting */
+	bool GetPivotForOrbit(FVector& PivotLocation) const;
+
+	/** Toggles camera orbit using current hit point as pivot */
+	virtual void ToggleOrbit();
+
 public:
 
 	/** Currently selected actor, may be invalid */
@@ -183,6 +199,12 @@ private:
 	/** The normalized screen location when a drag starts */
 	FVector2D LastTouchDragLocation;
 
+	/** Last position for orbit */
+	FVector LastOrbitPawnLocation;
+
+	/** Current orbit pivot, if pivot is enabled. */
+	FVector OrbitPivot;
+
 	void OnTouchBegin(ETouchIndex::Type FingerIndex, FVector Location);
 	void OnTouchEnd(ETouchIndex::Type FingerIndex, FVector Location);
 	void OnFingerMove(ETouchIndex::Type FingerIndex, FVector Location);
diff --git a/Engine/Source/Runtime/Engine/Private/DebugCameraController.cpp b/Engine/Source/Runtime/Engine/Private/DebugCameraController.cpp
--- a/Engine/Source/Runtime/Engine/Private/DebugCameraController.cpp
+++ b/Engine/Source/Runtime/Engine/Private/DebugCameraController.cpp
@@ -46,6 +46,10 @@ ADebugCameraController::ADebugCameraController(const FObjectInitializer& ObjectI
 	PrimaryActorTick.bTickEvenWhenPaused = true;
 	bShouldPerformFullTickWhenPaused = true;
 	SetAsLocalPlayerController();
+
+	bIsOrbitingSelectedActor = false;
+	LastOrbitPawnLocation = FVector::ZeroVector;
+	OrbitPivot = FVector::ZeroVector;
 }
 
 void InitializeDebugCameraInputBindings()
@@ -65,6 +69,7 @@ void InitializeDebugCameraInputBindings()
 		UPlayerInput::AddEngineDefinedActionMapping(FInputActionKeyMapping("DebugCamera_DecreaseFOV", EKeys::Period));
 		UPlayerInput::AddEngineDefinedActionMapping(FInputActionKeyMapping("DebugCamera_ToggleDisplay", EKeys::BackSpace));
 		UPlayerInput::AddEngineDefinedActionMapping(FInputActionKeyMapping("DebugCamera_FreezeRendering", EKeys::F));
+		UPlayerInput::AddEngineDefinedActionMapping(FInputActionKeyMapping("DebugCamera_Orbit", EKeys::O));
 
 		UPlayerInput::AddEngineDefinedActionMapping(FInputActionKeyMapping("DebugCamera_Select", EKeys::Gamepad_RightTrigger));
 		UPlayerInput::AddEngineDefinedActionMapping(FInputActionKeyMapping("DebugCamera_IncreaseSpeed", EKeys::Gamepad_RightShoulder));
@@ -92,6 +97,7 @@ void ADebugCameraController::SetupInputComponent()
 
 	InputComponent->BindAction("DebugCamera_ToggleDisplay", IE_Pressed, this, &ADebugCameraController::ToggleDisplay);
 	InputComponent->BindAction("DebugCamera_FreezeRendering", IE_Pressed, this, &ADebugCameraController::ToggleFreezeRendering);
+	InputComponent->BindAction("DebugCamera_Orbit", IE_Pressed, this, &ADebugCameraController::ToggleOrbit);
 
 	InputComponent->BindTouch(IE_Pressed, this, &ADebugCameraController::OnTouchBegin);
 	InputComponent->BindTouch(IE_Released, this, &ADebugCameraController::OnTouchEnd);
@@ -358,6 +364,8 @@ void ADebugCameraController::OnDeactivate( APlayerController* RestoredPC )
 		bIsFrozenRendering = false;
 	}
 
+	bIsOrbitingSelectedActor = false;
+
 	DrawFrustum->SetVisibility(false);
 	ConsoleCommand(TEXT("show camfrustums"));
 	DrawFrustum->UnregisterComponent();
@@ -380,6 +388,140 @@ void ADebugCameraController::ToggleFreezeRendering()
 	bIsFrozenRendering = !bIsFrozenRendering;
 }
 
+void ADebugCameraController::UpdateRotation(float DeltaTime)
+{
+	if (bIsOrbitingSelectedActor)
+	{
+		APawn* const CurrentPawn = GetPawnOrSpectator();
+		if (CurrentPawn)
+		{
+			bool bUpdatePawn = false;
+			FRotator ViewRotation;
+			float ViewLength = 0.0f;
+
+			if (!RotationInput.IsZero())
+			{
+				FVector ViewVector = OrbitPivot - CurrentPawn->GetActorLocation();
+				ViewRotation = ViewVector.ToOrientationRotator() + RotationInput;
+				ViewLength = ViewVector.Size();
+
+				FVector Axis;
+				float Angle;
+				FVector OppositeViewVector = -1.0 * ViewRotation.Vector();
+				FQuat::FindBetween(FVector::UpVector, OppositeViewVector).ToAxisAndAngle(Axis, Angle);
+
+				// Clamp rotation to 10 degrees from Up vector
+				const float MinAngle = PI / 18.f;
+				const float MaxAngle = PI - MinAngle;
+				if (Angle < MinAngle || Angle > MaxAngle)
+				{
+					float AdjustedAngle = FMath::Clamp(Angle, MinAngle, MaxAngle);
+					OppositeViewVector = FQuat(Axis, AdjustedAngle).RotateVector(FVector::UpVector);
+					ViewRotation = (-1.0 * OppositeViewVector).ToOrientationRotator();
+				}
+
+				bUpdatePawn = true;
+			} 
+			else
+			{
+				FVector OrbitPawnLocation = CurrentPawn->GetActorLocation();
+				FVector MoveDelta(OrbitPawnLocation - LastOrbitPawnLocation);
+
+				if (!MoveDelta.IsZero())
+				{
+					FRotationMatrix ObjectToWorld(GetControlRotation());
+					FVector MoveDeltaObj = ObjectToWorld.GetTransposed().TransformVector(MoveDelta);
+					FVector ViewVector = OrbitPivot - LastOrbitPawnLocation;
+
+					// Handle either forward or lateral motion but not both, because small 
+					// forward motion deltas while moving laterally cause the distance from pivot
+					// to drift.
+					if (FMath::Abs(MoveDeltaObj.X) > FMath::Abs(MoveDeltaObj.Y))
+					{
+						// Apply forward movement component
+						FVector ForwardDelta = ObjectToWorld.TransformVector(FVector(MoveDeltaObj.X, 0.0f, 0.0f));
+						ViewVector -= ForwardDelta;
+						ViewLength = ViewVector.Size();
+					}
+					else
+					{
+						// Apply lateral movement component, constraining distance from orbit pivot
+						FVector LateralDelta = ObjectToWorld.TransformVector(FVector(0.0f, MoveDeltaObj.Y, 0.0f));
+						ViewLength = ViewVector.Size();
+						ViewVector -= LateralDelta;
+					}
+					ViewRotation = ViewVector.Rotation();
+
+					bUpdatePawn = true;
+				}
+			}
+
+			if (bUpdatePawn)
+			{
+				// Update pawn location
+				FVector OrbitPawnLocation = OrbitPivot - ViewRotation.Vector() * ViewLength;
+				CurrentPawn->SetActorLocation(OrbitPawnLocation);
+
+				// Update pawn rotation
+				SetControlRotation(ViewRotation);
+				CurrentPawn->FaceRotation(ViewRotation, DeltaTime);
+
+				LastOrbitPawnLocation = OrbitPawnLocation;
+
+				return;
+			}
+		}
+	}
+
+	Super::UpdateRotation(DeltaTime);
+}
+
+bool ADebugCameraController::GetPivotForOrbit(FVector& PivotLocation) const
+{
+	if (SelectedActor)
+	{
+		FBox BoundingBox(ForceInit);
+		int32 NumValidComponents = 0;
+
+		// Use the center of the bounding box of the current selected actor as the pivot point for orbiting the camera
+		int32 NumSelectedActors = 0;
+
+		TInlineComponentArray<UMeshComponent*> MeshComponents(SelectedActor);
+
+		for (int32 ComponentIndex = 0; ComponentIndex < MeshComponents.Num(); ++ComponentIndex)
+		{
+			UMeshComponent* MeshComponent = MeshComponents[ComponentIndex];
+
+			if (MeshComponent->IsRegistered())
+			{
+				BoundingBox += MeshComponent->Bounds.GetBox();
+				++NumValidComponents;
+			}
+		}
+
+		if (NumValidComponents > 0)
+		{
+			PivotLocation = BoundingBox.GetCenter();
+			return true;
+		}
+	}
+
+	return false;
+}
+
+void ADebugCameraController::ToggleOrbit()
+{
+	bIsOrbitingSelectedActor = (!bIsOrbitingSelectedActor && GetPawnOrSpectator() && GetPivotForOrbit(OrbitPivot));
+
+	if (bIsOrbitingSelectedActor)
+	{
+		LastOrbitPawnLocation = GetPawnOrSpectator()->GetActorLocation();
+		FRotator ViewRotation = (OrbitPivot - LastOrbitPawnLocation).ToOrientationRotator();
+		SetControlRotation(ViewRotation);
+		GetPawnOrSpectator()->FaceRotation(ViewRotation);
+	}
+}
+
 void ADebugCameraController::SelectTargetedObject()
 {
 	FVector CamLoc;
diff --git a/Engine/Source/Runtime/Engine/Private/DebugCameraHUD.cpp b/Engine/Source/Runtime/Engine/Private/DebugCameraHUD.cpp
--- a/Engine/Source/Runtime/Engine/Private/DebugCameraHUD.cpp
+++ b/Engine/Source/Runtime/Engine/Private/DebugCameraHUD.cpp
@@ -163,11 +163,20 @@ void ADebugCameraHUD::PostRender()
 				Canvas->DrawText( RenderFont, TEXT("No trace Hit"), X, yl, 1.f, 1.f, FontRenderInfo);
 			}
 
-			if ( DCC->bShowSelectedInfo && DCC->SelectedActor != NULL )
+			if (DCC->bShowSelectedInfo && DCC->SelectedActor != NULL)
 			{
-				yl += Y;
-				Canvas->DrawText(RenderFont,  FString::Printf(TEXT("Selected actor: '%s'"), *DCC->SelectedActor->GetFName().ToString()), X, yl, 1.f, 1.f, FontRenderInfo);
-				DisplayMaterials( X, yl, Y, Cast<UMeshComponent>(DCC->SelectedComponent) );
+				if (DCC->bIsOrbitingSelectedActor)
+				{
+					yl += Y;
+					Canvas->DrawText(RenderFont, FString::Printf(TEXT("Orbiting actor: '%s'"), *DCC->SelectedActor->GetFName().ToString()), X, yl, 1.f, 1.f, FontRenderInfo);
+					DisplayMaterials(X, yl, Y, Cast<UMeshComponent>(DCC->SelectedComponent));
+				}
+				else
+				{
+					yl += Y;
+					Canvas->DrawText(RenderFont, FString::Printf(TEXT("Selected actor: '%s'"), *DCC->SelectedActor->GetFName().ToString()), X, yl, 1.f, 1.f, FontRenderInfo);
+					DisplayMaterials(X, yl, Y, Cast<UMeshComponent>(DCC->SelectedComponent));
+				}
 			}
 
 
@@ -190,6 +199,10 @@ void ADebugCameraHUD::PostRender()
 			
 			Canvas->DrawText(RenderFont, TEXT("Toggle Display: BackSpace or XButton"), X, yl, 1.f, 1.f, FontRenderInfo);
 			yl += Y;
+
+			Canvas->DrawText(RenderFont, TEXT("Orbit Selected Actor: O"), X, yl, 1.f, 1.f, FontRenderInfo);
+			yl += Y;
+
 		}
 	}
 #endif
