RL - Change particle VertexFactory to RefCounted instead of Pointer that may go out of scope. This might fix UE-72055

diff --git a/Engine/Source/Runtime/Engine/Public/ParticleVertexFactory.h b/Engine/Source/Runtime/Engine/Public/ParticleVertexFactory.h
--- a/Engine/Source/Runtime/Engine/Public/ParticleVertexFactory.h
+++ b/Engine/Source/Runtime/Engine/Public/ParticleVertexFactory.h
@@ -247,7 +247,7 @@ private:
 	int32 NumVertsInInstanceBuffer;
 
 	/** Uniform buffer with sprite paramters. */
-	FUniformBufferRHIParamRef SpriteUniformBuffer;
+	FUniformBufferRHIRef SpriteUniformBuffer;
 
 	int32 NumCutoutVerticesPerFrame;
 	FShaderResourceViewRHIParamRef CutoutGeometrySRV;
