Fixed issue where certain materials would evaluate to default colors during a asset actors operation.

diff --git a/Engine/Source/Runtime/Engine/Public/MaterialCompiler.h b/Engine/Source/Runtime/Engine/Public/MaterialCompiler.h
--- a/Engine/Source/Runtime/Engine/Public/MaterialCompiler.h
+++ b/Engine/Source/Runtime/Engine/Public/MaterialCompiler.h
@@ -426,6 +426,9 @@ public:
 	virtual int32 ExternalTextureCoordinateOffset(int32 TextureReferenceIndex, TOptional<FName> ParameterName) override { return Compiler->ExternalTextureCoordinateOffset(TextureReferenceIndex, ParameterName); }
 	virtual int32 ExternalTextureCoordinateOffset(const FGuid& ExternalTextureGuid) override { return Compiler->ExternalTextureCoordinateOffset(ExternalTextureGuid); }
 
+	virtual int32 GetTextureReferenceIndex(UTexture* Texture) override { return Compiler->GetTextureReferenceIndex(Texture); }
+	virtual UTexture* GetReferencedTexture(int32 Index) override { return Compiler->GetReferencedTexture(Index); }
+
 	virtual	int32 PixelDepth() override { return Compiler->PixelDepth();	}
 	virtual int32 SceneDepth(int32 Offset, int32 ViewportUV, bool bUseOffset) override { return Compiler->SceneDepth(Offset, ViewportUV, bUseOffset); }
 	virtual int32 SceneColor(int32 Offset, int32 ViewportUV, bool bUseOffset) override { return Compiler->SceneColor(Offset, ViewportUV, bUseOffset); }
