Hide hit line when orbiting in PIE debug camera.

diff --git a/Engine/Source/Runtime/Engine/Private/DebugCameraHUD.cpp b/Engine/Source/Runtime/Engine/Private/DebugCameraHUD.cpp
--- a/Engine/Source/Runtime/Engine/Private/DebugCameraHUD.cpp
+++ b/Engine/Source/Runtime/Engine/Private/DebugCameraHUD.cpp
@@ -155,7 +155,11 @@ void ADebugCameraHUD::PostRender()
 					yl += Y;
 					Canvas->DrawText(RenderFont, "Material: NULL", X + Y, yl, 1.f, 1.f, FontRenderInfo );
 				}
-				DrawDebugLine( GetWorld(), Hit.Location, Hit.Location+Hit.Normal*30.f, FColor::White );
+
+				if (!DCC->bIsOrbitingSelectedActor)
+				{
+					DrawDebugLine(GetWorld(), Hit.Location, Hit.Location + Hit.Normal*30.f, FColor::White);
+				}
 			}
 			else
 			{
