Make sure we log whenever we receive a network failure

diff --git a/Engine/Source/Runtime/Engine/Classes/Engine/Engine.h b/Engine/Source/Runtime/Engine/Classes/Engine/Engine.h
--- a/Engine/Source/Runtime/Engine/Classes/Engine/Engine.h
+++ b/Engine/Source/Runtime/Engine/Classes/Engine/Engine.h
@@ -1911,10 +1911,7 @@ public:
 	/** Event triggered after a network failure of any kind has occurred */
 	FOnNetworkFailure& OnNetworkFailure() { return NetworkFailureEvent; }
 	/** Called by internal engine systems after a network failure has occurred */
-	void BroadcastNetworkFailure(UWorld * World, UNetDriver *NetDriver, ENetworkFailure::Type FailureType, const FString& ErrorString = TEXT(""))
-	{
-		NetworkFailureEvent.Broadcast(World, NetDriver, FailureType, ErrorString);
-	}
+	void BroadcastNetworkFailure(UWorld * World, UNetDriver *NetDriver, ENetworkFailure::Type FailureType, const FString& ErrorString = TEXT(""));
 
 	/** Event triggered after network lag is being experienced or lag has ended */
 	FOnNetworkLagStateChanged& OnNetworkLagStateChanged() { return NetworkLagStateChangedEvent; }
diff --git a/Engine/Source/Runtime/Engine/Private/UnrealEngine.cpp b/Engine/Source/Runtime/Engine/Private/UnrealEngine.cpp
--- a/Engine/Source/Runtime/Engine/Private/UnrealEngine.cpp
+++ b/Engine/Source/Runtime/Engine/Private/UnrealEngine.cpp
@@ -10577,6 +10577,14 @@ void UEngine::WorldDestroyed( UWorld* InWorld )
 	WorldDestroyedEvent.Broadcast( InWorld );
 }
 
+void UEngine::BroadcastNetworkFailure(UWorld * World, UNetDriver *NetDriver, ENetworkFailure::Type FailureType, const FString& ErrorString /* = TEXT("") */)
+{
+	UE_LOG(LogNet, Error, TEXT("UEngine::BroadcastNetworkFailure: FailureType = %s, ErrorString = %s, Driver = %s"),
+		ENetworkFailure::ToString(FailureType), *ErrorString, (NetDriver ? *NetDriver->GetDescription() : TEXT("NONE")));
+
+	NetworkFailureEvent.Broadcast(World, NetDriver, FailureType, ErrorString);
+}
+
 UWorld* UEngine::GetWorldFromContextObject(const UObject* Object, EGetWorldErrorMode ErrorMode) const
 {
 	if (Object == nullptr)
