Move GetPlayerStateFromUniqueNetId from VoiceConfig to GameStateBase so it can also be used by game systems rather than just the VoiceConfig.

diff --git a/Engine/Source/Runtime/Engine/Classes/GameFramework/GameStateBase.h b/Engine/Source/Runtime/Engine/Classes/GameFramework/GameStateBase.h
--- a/Engine/Source/Runtime/Engine/Classes/GameFramework/GameStateBase.h
+++ b/Engine/Source/Runtime/Engine/Classes/GameFramework/GameStateBase.h
@@ -9,6 +9,7 @@
 #include "GameFramework/Actor.h"
 #include "GameFramework/Info.h"
 #include "GameFramework/GameModeBase.h"
+#include "Stats/Stats.h"
 #include "GameStateBase.generated.h"
 
 class APlayerState;
@@ -16,6 +17,8 @@ class ASpectatorPawn;
 
 ENGINE_API DECLARE_LOG_CATEGORY_EXTERN(LogGameState, Log, All);
 
+DECLARE_CYCLE_STAT_EXTERN( TEXT( "GetPlayerStateFromUniqueId" ), STAT_GetPlayerStateFromUniqueId, STATGROUP_Game , ENGINE_API);
+
 class AGameModeBase;
 class ASpectatorPawn;
 class APlayerState;
@@ -86,6 +89,9 @@ public:
 
 	virtual void GetLifetimeReplicatedProps(TArray< FLifetimeProperty > &OutLifetimeProps) const;
 
+	/** Returns the player state for a specified unique player Id */
+	APlayerState* GetPlayerStateFromUniqueNetId(const FUniqueNetIdWrapper& InPlayerId) const;
+
 	//~=============================================================================
 	// Interaction with GameModeBase
 
diff --git a/Engine/Source/Runtime/Engine/Private/GameStateBase.cpp b/Engine/Source/Runtime/Engine/Private/GameStateBase.cpp
--- a/Engine/Source/Runtime/Engine/Private/GameStateBase.cpp
+++ b/Engine/Source/Runtime/Engine/Private/GameStateBase.cpp
@@ -12,6 +12,8 @@
 
 DEFINE_LOG_CATEGORY(LogGameState);
 
+DEFINE_STAT(STAT_GetPlayerStateFromUniqueId);
+
 AGameStateBase::AGameStateBase(const FObjectInitializer& ObjectInitializer)
 	: Super(ObjectInitializer
 		.DoNotCreateDefaultSubobject(TEXT("Sprite")) )
@@ -220,3 +222,18 @@ void AGameStateBase::GetLifetimeReplicatedProps(TArray<FLifetimeProperty>& OutLi
 	DOREPLIFETIME( AGameStateBase, ReplicatedWorldTimeSeconds );
 	DOREPLIFETIME( AGameStateBase, bReplicatedHasBegunPlay );
 }
+
+APlayerState* AGameStateBase::GetPlayerStateFromUniqueNetId(const FUniqueNetIdWrapper& InPlayerId) const
+{
+	SCOPE_CYCLE_COUNTER(STAT_GetPlayerStateFromUniqueId);
+	const TArray<APlayerState*>& Players = PlayerArray;
+	for (APlayerState* Player : Players)
+	{
+		if (Player->UniqueId == InPlayerId)
+		{
+			return Player;
+		}
+	}
+
+	return nullptr;
+}
diff --git a/Engine/Source/Runtime/Engine/Private/VoiceConfig.cpp b/Engine/Source/Runtime/Engine/Private/VoiceConfig.cpp
--- a/Engine/Source/Runtime/Engine/Private/VoiceConfig.cpp
+++ b/Engine/Source/Runtime/Engine/Private/VoiceConfig.cpp
@@ -154,13 +154,9 @@ int32 UVOIPStatics::GetNumBufferedPackets()
 APlayerState* UVOIPStatics::GetPlayerStateFromUniqueNetId(UWorld* InWorld, const FUniqueNetIdWrapper& InPlayerId)
 {
 	AGameStateBase* InBase = InWorld->GetGameState();
-	TArray<APlayerState*>& Players = InBase->PlayerArray;
-	for (APlayerState* Player : Players)
+	if (InBase)
 	{
-		if (Player->UniqueId == InPlayerId)
-		{
-			return Player;
-		}
+		return InBase->GetPlayerStateFromUniqueNetId(InPlayerId);
 	}
 
 	return nullptr;
