updated camera access in PlayWorldCameraShake()
#rnx

diff --git a/Engine/Source/Runtime/Engine/Private/PlayerCameraManager.cpp b/Engine/Source/Runtime/Engine/Private/PlayerCameraManager.cpp
--- a/Engine/Source/Runtime/Engine/Private/PlayerCameraManager.cpp
+++ b/Engine/Source/Runtime/Engine/Private/PlayerCameraManager.cpp
@@ -1333,9 +1333,7 @@ void APlayerCameraManager::PlayWorldCameraShake(UWorld* InWorld, TSubclassOf<cla
 
 			if (bOrientShakeTowardsEpicenter && PlayerController->GetPawn() != NULL)
 			{
-				FVector CamLoc;
-				FRotator CamRot;
-				PlayerController->PlayerCameraManager->GetCameraViewPoint(CamLoc, CamRot);
+				const FVector CamLoc = PlayerController->PlayerCameraManager->GetCameraLocation();
 				PlayerController->ClientPlayCameraShake(Shake, ShakeScale, ECameraAnimPlaySpace::UserDefined, (Epicenter - CamLoc).Rotation());
 			}
 			else
