Fix issues with ActiveSounds not optimizing for all correct cases

diff --git a/Engine/Source/Runtime/Engine/Private/ActiveSound.cpp b/Engine/Source/Runtime/Engine/Private/ActiveSound.cpp
--- a/Engine/Source/Runtime/Engine/Private/ActiveSound.cpp
+++ b/Engine/Source/Runtime/Engine/Private/ActiveSound.cpp
@@ -396,7 +396,8 @@ void FActiveSound::UpdateWaveInstances( TArray<FWaveInstance*> &InWaveInstances,
 	// If we have an attenuation node, we can't know until we evaluate the sound cue if it's audio output going to be audible via a distance check
 	if (Sound->HasAttenuationNode() || 
 		(AudioDevice->VirtualSoundsEnabled() && (Sound->IsAllowedVirtual() || (bHandleSubtitles && bHasExternalSubtitles))) ||
-		(bHasAttenuationSettings && (AttenuationSettings.FocusDistanceScale != 1.0f || AttenuationSettings.NonFocusDistanceScale != 1.0f)))
+		!bHasAttenuationSettings ||
+		(bHasAttenuationSettings && (!AttenuationSettings.bAttenuate || AttenuationSettings.FocusDistanceScale != 1.0f || AttenuationSettings.NonFocusDistanceScale != 1.0f)))
 	{
 		bPerformDistanceCheckOptimization = false;
 	}
