Fix static analysis warning

diff --git a/Engine/Source/Runtime/Engine/Private/ContentStreaming.cpp b/Engine/Source/Runtime/Engine/Private/ContentStreaming.cpp
--- a/Engine/Source/Runtime/Engine/Private/ContentStreaming.cpp
+++ b/Engine/Source/Runtime/Engine/Private/ContentStreaming.cpp
@@ -1135,7 +1135,8 @@ void FStreamingManagerCollection::AddOrRemoveTextureStreamingManagerIfNeeded(boo
 		// some code relies on r.TextureStreaming so we're going to disable it here to reflect the hardware capabilities and system needs
 		CVarSetTextureStreaming.AsVariable()->Set(0, ECVF_SetByCode);
 	}
-	else if (!bUseTextureStreaming && TEXTURE2DMIPMAP_USE_COMPACT_BULKDATA)
+#if TEXTURE2DMIPMAP_USE_COMPACT_BULKDATA
+	else if (!bUseTextureStreaming)
 	{
 		static bool bWarned;
 		if (!bWarned)
@@ -1146,6 +1147,7 @@ void FStreamingManagerCollection::AddOrRemoveTextureStreamingManagerIfNeeded(boo
 		CVarSetTextureStreaming.AsVariable()->Set(1, ECVF_SetByCode);
 		bUseTextureStreaming = true;
 	}
+#endif
 #endif
 
 	if ( bUseTextureStreaming )
