- Add sound loop management at AudioComponent level for better audio thread perf/concurrency control
- Add au.DisableConcurrencyStopSilentForLoops cvar for disablement of StopQuietest concurrency for loops
- Fix bug with PlayInternal checking spatialization boolean to gather attenuation data for ActiveSound creation

diff --git a/Engine/Source/Runtime/Engine/Classes/Components/AudioComponent.h b/Engine/Source/Runtime/Engine/Classes/Components/AudioComponent.h
--- a/Engine/Source/Runtime/Engine/Classes/Components/AudioComponent.h
+++ b/Engine/Source/Runtime/Engine/Classes/Components/AudioComponent.h
@@ -189,6 +189,26 @@ class ENGINE_API UAudioComponent : public USceneComponent
 	/** Whether or not this audio component has been paused */
 	uint8 bIsPaused:1;
 
+	/** How AudioComponent is managing playback based on whether a listener is within 
+	  * proximity (sound's MaxDistance) or not. */
+	enum class EReplayWhenInAudibleRange : uint8
+	{
+		/** Sound lifetime is not managed by listener proximity. */
+		Disabled,
+
+		/** Sound lifetime is managed by listener proximity (in case of sound being looping and non-virtualized) */
+		Enabled,
+
+		/** Sound lifetime was managed by listener proximity (in case of sound being looping and non-virtualized),
+		  * but is fading out.*/
+		EnabledFadingOut,
+
+		Count
+	};
+
+	/** Whether or not the audio component can restart when listener is in proximity */
+	EReplayWhenInAudibleRange ReplayWhenInAudibleRange;
+
 	/**
 	* True if we should automatically attach to AutoAttachParent when Played, and detach from our parent when playback is completed.
 	* This overrides any current attachment that may be present at the time of activation (deferring initial attachment until activation, if AutoAttachParent is null).
@@ -479,6 +499,9 @@ private:
 	/** Called by the ActiveSound to inform the component that playback is finished */
 	void PlaybackCompleted(bool bFailedToStart);
 
+	/** Whether or not the sound is audible */
+	bool IsInAudibleRange() const;
+
 public:
 
 	/** Sets the sound instance parameter. */
@@ -505,6 +528,7 @@ public:
 	virtual void OnUnregister() override;
 	virtual const UObject* AdditionalStatObject() const override;
 	virtual bool IsReadyForOwnerToAutoDestroy() const override;
+	virtual void TickComponent(float DeltaTime, ELevelTick TickType, FActorComponentTickFunction* TickFunction) override;
 	//~ End ActorComponent Interface.
 
 	/** Returns a pointer to the attenuation settings to be used (if any) for this audio component dependent on the SoundAttenuation asset or overrides set. */
@@ -589,7 +613,8 @@ protected:
 	void PlayInternal(const float StartTime = 0.f, const float FadeInDuration = 0.f, const float FadeVolumeLevel = 1.f);
 
 private:
-	
+	void StopInternal();
+
 #if WITH_EDITORONLY_DATA
 	/** Utility function that updates which texture is displayed on the sprite dependent on the properties of the Audio Component. */
 	void UpdateSpriteTexture();
diff --git a/Engine/Source/Runtime/Engine/Private/ActiveSound.cpp b/Engine/Source/Runtime/Engine/Private/ActiveSound.cpp
--- a/Engine/Source/Runtime/Engine/Private/ActiveSound.cpp
+++ b/Engine/Source/Runtime/Engine/Private/ActiveSound.cpp
@@ -18,6 +18,14 @@ FAutoConsoleVariableRef CVarAudioOcclusionEnabled(
 	TEXT("Disables (1) or enables (0) audio occlusion.\n"),
 	ECVF_Default);
 
+static int32 AudioDisableConcurrencyStopSilentForLoopsCvar = 0;
+FAutoConsoleVariableRef CVarAudioDisableConcurrencyStopSilentForLoops(
+	TEXT("au.DisableConcurrencyStopSilentForLoops"),
+	AudioDisableConcurrencyStopSilentForLoopsCvar,
+	TEXT("Disables (1) or enables (0) audio concurrency for loops subscribing to stop silent.\n"),
+	ECVF_Default);
+
+
 FTraceDelegate FActiveSound::ActiveSoundTraceDelegate;
 TMap<FTraceHandle, FActiveSound::FAsyncTraceDetails> FActiveSound::TraceToActiveSoundMap;
 
@@ -346,14 +354,29 @@ void FActiveSound::GetConcurrencyHandles(TArray<FConcurrencyHandle>& OutConcurre
 	if (!ConcurrencySet.Num() && Sound)
 	{
 		Sound->GetConcurrencyHandles(OutConcurrencyHandles);
-		return;
+	}
+	else
+	{
+		for (const USoundConcurrency* Concurrency : ConcurrencySet)
+		{
+			if (Concurrency)
+			{
+				OutConcurrencyHandles.Emplace(FConcurrencyHandle(*Concurrency));
+			}
+		}
 	}
 
-	for (const USoundConcurrency* Concurrency : ConcurrencySet)
+	if (IsLooping() && AudioDisableConcurrencyStopSilentForLoopsCvar)
 	{
-		if (Concurrency)
+		for (int32 i = OutConcurrencyHandles.Num() - 1; i >= 0; --i)
 		{
-			OutConcurrencyHandles.Emplace(FConcurrencyHandle(*Concurrency));
+			if (OutConcurrencyHandles[i].Settings.ResolutionRule == EMaxConcurrentResolutionRule::StopQuietest)
+			{
+				UE_LOG(LogAudio, Warning,
+					TEXT("Concurrency not observed for '%s': StopQuietest concurrency disabled for looping sounds"),
+					Sound ? *Sound->GetName() : TEXT("N/A"));
+				OutConcurrencyHandles.RemoveAtSwap(i, 1, false);
+			}
 		}
 	}
 }
diff --git a/Engine/Source/Runtime/Engine/Private/Components/AudioComponent.cpp b/Engine/Source/Runtime/Engine/Private/Components/AudioComponent.cpp
--- a/Engine/Source/Runtime/Engine/Private/Components/AudioComponent.cpp
+++ b/Engine/Source/Runtime/Engine/Private/Components/AudioComponent.cpp
@@ -49,6 +49,7 @@ UAudioComponent::UAudioComponent(const FObjectInitializer& ObjectInitializer)
 	bOverrideSubtitlePriority = false;
 	bIsPreviewSound = false;
 	bIsPaused = false;
+	ReplayWhenInAudibleRange = EReplayWhenInAudibleRange::Disabled;
 	Priority = 1.f;
 	SubtitlePriority = DEFAULT_SUBTITLE_PRIORITY;
 	PitchMultiplier = 1.f;
@@ -70,6 +71,12 @@ UAudioComponent::UAudioComponent(const FObjectInitializer& ObjectInitializer)
 	// TODO: Consider only putting played/active components in to the map
 	FScopeLock Lock(&AudioIDToComponentMapLock);
 	AudioIDToComponentMap.Add(AudioComponentID, this);
+
+	PrimaryComponentTick.bCanEverTick = true;
+	PrimaryComponentTick.TickGroup = TG_DuringPhysics;
+	PrimaryComponentTick.bStartWithTickEnabled = false;
+	PrimaryComponentTick.SetTickFunctionEnable(false);
+	bTickInEditor = true;
 }
 
 UAudioComponent* UAudioComponent::GetAudioComponentFromID(uint64 AudioComponentID)
@@ -230,9 +237,9 @@ const UObject* UAudioComponent::AdditionalStatObject() const
 	return Sound;
 }
 
-void UAudioComponent::SetSound( USoundBase* NewSound )
+void UAudioComponent::SetSound(USoundBase* NewSound)
 {
-	const bool bPlay = IsPlaying();
+	const bool bPlay = IsPlaying() || ReplayWhenInAudibleRange != EReplayWhenInAudibleRange::Disabled;
 
 	// If this is an auto destroy component we need to prevent it from being auto-destroyed since we're really just restarting it
 	const bool bWasAutoDestroy = bAutoDestroy;
@@ -253,13 +260,43 @@ bool UAudioComponent::IsReadyForOwnerToAutoDestroy() const
 	return !IsPlaying();
 }
 
+void UAudioComponent::TickComponent(float DeltaTime, ELevelTick TickType, FActorComponentTickFunction* TickFunction)
+{
+	Super::TickComponent(DeltaTime, TickType, TickFunction);
+
+	if (bIsActive)
+	{
+		if (ReplayWhenInAudibleRange != EReplayWhenInAudibleRange::Disabled && !IsInAudibleRange())
+		{
+			// If this is an auto destroy component we need to prevent it from being auto-destroyed since we're really just restarting it
+			bool bCurrentAutoDestroy = bAutoDestroy;
+			bAutoDestroy = false;
+			StopInternal();
+			bAutoDestroy = bCurrentAutoDestroy;
+		}
+	}
+	else 
+	{
+		if (ReplayWhenInAudibleRange != EReplayWhenInAudibleRange::Disabled && IsInAudibleRange())
+		{
+			Play(); // Play refreshes ReplayWhenInAudibleRange internally
+		}
+	}
+}
+
+
 void UAudioComponent::OnUpdateTransform(EUpdateTransformFlags UpdateTransformFlags, ETeleportType Teleport)
 {
 	Super::OnUpdateTransform(UpdateTransformFlags, Teleport);
 
-	if (bIsActive && !bPreviewComponent)
+	if (bPreviewComponent)
 	{
-		if (FAudioDevice* AudioDevice = GetAudioDevice())
+		return;
+	}
+
+	if (FAudioDevice* AudioDevice = GetAudioDevice())
+	{
+		if (bIsActive)
 		{
 			DECLARE_CYCLE_STAT(TEXT("FAudioThreadTask.UpdateAudioComponentTransform"), STAT_AudioUpdateComponentTransform, STATGROUP_AudioThreadCommands);
 
@@ -299,6 +336,22 @@ void UAudioComponent::CancelAutoAttachment(bool bDetachFromParent)
 	}
 }
 
+bool UAudioComponent::IsInAudibleRange() const
+{
+	FAudioDevice* AudioDevice = GetAudioDevice();
+	if (!AudioDevice)
+	{
+		return false;
+	}
+
+	float MaxDistance = 0.0f;
+	float FocusFactor = 0.0f;
+	const FVector Location = GetComponentTransform().GetLocation();
+	const FSoundAttenuationSettings* AttenuationSettingsToApply = GetAttenuationSettingsToApply();
+
+	AudioDevice->GetMaxDistanceAndFocusFactor(Sound, GetWorld(), Location, AttenuationSettingsToApply, MaxDistance, FocusFactor);
+	return AudioDevice->SoundIsAudible(Sound, GetWorld(), Location, AttenuationSettingsToApply, MaxDistance, FocusFactor);
+}
 
 void UAudioComponent::Play(float StartTime)
 {
@@ -355,7 +408,7 @@ void UAudioComponent::PlayInternal(const float StartTime, const float FadeInDura
 			}
 
 			// Create / configure new ActiveSound
-			const FSoundAttenuationSettings* AttenuationSettingsToApply = (bAllowSpatialization ? GetAttenuationSettingsToApply() : nullptr);
+			const FSoundAttenuationSettings* AttenuationSettingsToApply = GetAttenuationSettingsToApply();
 
 			float MaxDistance = 0.0f;
 			float FocusFactor = 0.0f;
@@ -461,6 +514,19 @@ void UAudioComponent::PlayInternal(const float StartTime, const float FadeInDura
 			++ActiveCount;
 
 			AudioDevice->AddNewActiveSound(NewActiveSound);
+
+			const bool bCanVirtualize = Sound && Sound->IsVirtualizeWhenSilent() && AudioDevice->VirtualSoundsEnabled();
+			if (!bIsPreviewSound && NewActiveSound.bHasAttenuationSettings && !bCanVirtualize)
+			{
+				const bool bLooping = NewActiveSound.IsLooping();
+				ReplayWhenInAudibleRange = bLooping ? EReplayWhenInAudibleRange::Enabled : EReplayWhenInAudibleRange::Disabled;
+				PrimaryComponentTick.SetTickFunctionEnable(bLooping);
+			}
+			else
+			{
+				ReplayWhenInAudibleRange = EReplayWhenInAudibleRange::Disabled;
+			}
+
 			bIsActive = true;
 		}
 	}
@@ -500,6 +566,12 @@ void UAudioComponent::FadeOut( float FadeOutDuration, float FadeVolumeLevel )
 	{
 		if (FadeOutDuration > 0.0f)
 		{
+			if (FadeVolumeLevel == 0.0f)
+			{
+				ReplayWhenInAudibleRange = EReplayWhenInAudibleRange::EnabledFadingOut;
+				PrimaryComponentTick.SetTickFunctionEnable(false);
+			}
+
 			if (FAudioDevice* AudioDevice = GetAudioDevice())
 			{
 				DECLARE_CYCLE_STAT(TEXT("FAudioThreadTask.FadeOut"), STAT_AudioFadeOut, STATGROUP_AudioThreadCommands);
@@ -530,6 +602,12 @@ void UAudioComponent::AdjustVolume( float AdjustVolumeDuration, float AdjustVolu
 	{
 		if (FAudioDevice* AudioDevice = GetAudioDevice())
 		{
+			if (ReplayWhenInAudibleRange == EReplayWhenInAudibleRange::EnabledFadingOut)
+			{
+				ReplayWhenInAudibleRange= EReplayWhenInAudibleRange::Enabled;
+				PrimaryComponentTick.SetTickFunctionEnable(true);
+			}
+
 			DECLARE_CYCLE_STAT(TEXT("FAudioThreadTask.AdjustVolume"), STAT_AudioAdjustVolume, STATGROUP_AudioThreadCommands);
 
 			const uint64 MyAudioComponentID = AudioComponentID;
@@ -557,6 +635,15 @@ void UAudioComponent::AdjustVolume( float AdjustVolumeDuration, float AdjustVolu
 }
 
 void UAudioComponent::Stop()
+{
+	PrimaryComponentTick.SetTickFunctionEnable(false);
+
+	ReplayWhenInAudibleRange = EReplayWhenInAudibleRange::Disabled;
+
+	StopInternal();
+}
+
+void UAudioComponent::StopInternal()
 {
 	if (bIsActive)
 	{
@@ -621,6 +708,12 @@ void UAudioComponent::PlaybackCompleted(bool bFailedToStart)
 
 	if (!bIsActive)
 	{
+		if (bFailedToStart)
+		{
+			ReplayWhenInAudibleRange = EReplayWhenInAudibleRange::Disabled;
+			PrimaryComponentTick.SetTickFunctionEnable(false);
+		}
+
 		if (!bFailedToStart && GetWorld() != nullptr && (OnAudioFinished.IsBound() || OnAudioFinishedNative.IsBound()))
 		{
 			INC_DWORD_STAT(STAT_AudioFinishedDelegatesCalled);
