CacheAggregateValues when changing SoundCue values

diff --git a/Engine/Source/Runtime/Engine/Private/SoundCue.cpp b/Engine/Source/Runtime/Engine/Private/SoundCue.cpp
--- a/Engine/Source/Runtime/Engine/Private/SoundCue.cpp
+++ b/Engine/Source/Runtime/Engine/Private/SoundCue.cpp
@@ -254,6 +254,8 @@ void USoundCue::PostEditChangeProperty(struct FPropertyChangedEvent& PropertyCha
 		// Propagate branch exclusion to child nodes which care (sound node random)
 		RecursivelySetExcludeBranchCulling(FirstNode);
 	}
+
+	CacheAggregateValues();
 }
 #endif
 
