Use an SRV for the BoneMatrices on all Metal platforms - this ensures that out-of-bounds access is handled properly.

diff --git a/Engine/Source/Runtime/Engine/Private/GPUSkinVertexFactory.cpp b/Engine/Source/Runtime/Engine/Private/GPUSkinVertexFactory.cpp
--- a/Engine/Source/Runtime/Engine/Private/GPUSkinVertexFactory.cpp
+++ b/Engine/Source/Runtime/Engine/Private/GPUSkinVertexFactory.cpp
@@ -76,14 +76,12 @@ static TAutoConsoleVariable<int32> CVarVelocityTest(
 // These should match USE_BONES_SRV_BUFFER
 static inline bool SupportsBonesBufferSRV(EShaderPlatform Platform)
 {
-	//#todo-rco: Add Metal support
-	return IsFeatureLevelSupported(Platform, ERHIFeatureLevel::SM4) || IsVulkanPlatform(Platform);
+	return IsFeatureLevelSupported(Platform, ERHIFeatureLevel::SM4) || IsVulkanPlatform(Platform) || IsMetalPlatform(Platform);
 }
 
 static inline bool SupportsBonesBufferSRV(ERHIFeatureLevel::Type InFeatureLevel)
 {
-	//#todo-rco: Add Metal support
-	return InFeatureLevel > ERHIFeatureLevel::ES3_1 || IsVulkanPlatform(GMaxRHIShaderPlatform);
+	return InFeatureLevel > ERHIFeatureLevel::ES3_1 || IsVulkanPlatform(GMaxRHIShaderPlatform) || IsMetalPlatform(GMaxRHIShaderPlatform);
 }
 // ---
 
