Reverting part of change that broke 2D loops again.

diff --git a/Engine/Source/Runtime/Engine/Private/Components/AudioComponent.cpp b/Engine/Source/Runtime/Engine/Private/Components/AudioComponent.cpp
--- a/Engine/Source/Runtime/Engine/Private/Components/AudioComponent.cpp
+++ b/Engine/Source/Runtime/Engine/Private/Components/AudioComponent.cpp
@@ -347,7 +347,7 @@ bool UAudioComponent::IsInAudibleRange() const
 	float MaxDistance = 0.0f;
 	float FocusFactor = 0.0f;
 	const FVector Location = GetComponentTransform().GetLocation();
-	const FSoundAttenuationSettings* AttenuationSettingsToApply = GetAttenuationSettingsToApply();
+	const FSoundAttenuationSettings* AttenuationSettingsToApply = bAllowSpatialization ? GetAttenuationSettingsToApply() : nullptr;
 
 	AudioDevice->GetMaxDistanceAndFocusFactor(Sound, GetWorld(), Location, AttenuationSettingsToApply, MaxDistance, FocusFactor);
 	return AudioDevice->SoundIsAudible(Sound, GetWorld(), Location, AttenuationSettingsToApply, MaxDistance, FocusFactor);
@@ -408,7 +408,7 @@ void UAudioComponent::PlayInternal(const float StartTime, const float FadeInDura
 			}
 
 			// Create / configure new ActiveSound
-			const FSoundAttenuationSettings* AttenuationSettingsToApply = GetAttenuationSettingsToApply();
+			const FSoundAttenuationSettings* AttenuationSettingsToApply = bAllowSpatialization ? GetAttenuationSettingsToApply() : nullptr;
 
 			float MaxDistance = 0.0f;
 			float FocusFactor = 0.0f;
