Fix corrupted compressed curve data (expecting curve tangents that were not written out).

diff --git a/Engine/Source/Runtime/Engine/Private/Animation/AnimCompressionDerivedData.h b/Engine/Source/Runtime/Engine/Private/Animation/AnimCompressionDerivedData.h
--- a/Engine/Source/Runtime/Engine/Private/Animation/AnimCompressionDerivedData.h
+++ b/Engine/Source/Runtime/Engine/Private/Animation/AnimCompressionDerivedData.h
@@ -50,7 +50,7 @@ public:
 		// This is a version string that mimics the old versioning scheme. If you
 		// want to bump this version, generate a new guid using VS->Tools->Create GUID and
 		// return it here. Ex.
-		return TEXT("AA0E2605005B4B2EBAACD66B5578C97E");
+		return TEXT("267426BCDCE54E129C9D584B35330967");
 	}
 
 	virtual FString GetPluginSpecificCacheKeySuffix() const override;
diff --git a/Engine/Source/Runtime/Engine/Private/Curves/RichCurve.cpp b/Engine/Source/Runtime/Engine/Private/Curves/RichCurve.cpp
--- a/Engine/Source/Runtime/Engine/Private/Curves/RichCurve.cpp
+++ b/Engine/Source/Runtime/Engine/Private/Curves/RichCurve.cpp
@@ -5,6 +5,8 @@
 
 DECLARE_CYCLE_STAT(TEXT("RichCurve Eval"), STAT_RichCurve_Eval, STATGROUP_Engine);
 
+// Broken - do not turn on! 
+#define MIXEDKEY_STRIPS_TANGENTS 0
 
 /* FRichCurveKey interface
  *****************************************************************************/
@@ -1328,6 +1330,7 @@ void FRichCurve::CompressCurve(FCompressedRichCurve& OutCurve, float ErrorThresh
 		}
 		else if (CompressionFormat == RCCF_Mixed)
 		{
+#if MIXEDKEY_STRIPS_TANGENTS
 			for (const FRichCurveKey& Key : Keys)
 			{
 				if (Key.InterpMode == RCIM_Cubic)
@@ -1335,6 +1338,9 @@ void FRichCurve::CompressCurve(FCompressedRichCurve& OutCurve, float ErrorThresh
 					PackedDataSize += 2 * sizeof(float);
 				}
 			}
+#else
+			PackedDataSize += Keys.Num() * 2 * sizeof(float); // Always have tangents
+#endif
 		}
 
 		OutCurve.CompressedKeys.Empty(PackedDataSize);
@@ -1416,7 +1422,11 @@ void FRichCurve::CompressCurve(FCompressedRichCurve& OutCurve, float ErrorThresh
 		{
 			*KeyData++ = Key.Value;
 
+#if MIXEDKEY_STRIPS_TANGENTS
 			if (Key.InterpMode == RCIM_Cubic)
+#else
+			if (CompressionFormat == RCCF_Mixed || Key.InterpMode == RCIM_Cubic)
+#endif
 			{
 				check(CompressionFormat == RCCF_Cubic || CompressionFormat == RCCF_Mixed);
 				*KeyData++ = Key.ArriveTangent;
@@ -1534,6 +1544,7 @@ struct MixedKeyDataAdapter
 
 	KeyDataHandle GetKeyDataHandle(int32 KeyIndexToQuery) const
 	{
+#if MIXEDKEY_STRIPS_TANGENTS
 		int32 Offset = 0;
 		for (int32 KeyIndex = 0; KeyIndex < KeyIndexToQuery; ++KeyIndex)
 		{
@@ -1541,6 +1552,9 @@ struct MixedKeyDataAdapter
 		}
 
 		return Offset;
+#else
+		return KeyIndexToQuery * 3;
+#endif
 	};
 
 	constexpr float GetKeyValue(KeyDataHandle Handle) const
