Deterministically invalidate GUID when light component transform has changed enough to invalidate lighting.

diff --git a/Engine/Source/Runtime/Engine/Private/Components/LightComponent.cpp b/Engine/Source/Runtime/Engine/Private/Components/LightComponent.cpp
--- a/Engine/Source/Runtime/Engine/Private/Components/LightComponent.cpp
+++ b/Engine/Source/Runtime/Engine/Private/Components/LightComponent.cpp
@@ -1121,12 +1121,14 @@ void ULightComponent::ApplyComponentInstanceData(FPrecomputedLightInstanceData*
 {
 	check(LightMapData);
 
+	LightGuid = LightMapData->LightGuid;
+
 	if (!LightMapData->Transform.Equals(GetComponentTransform()))
 	{
+		++LightGuid.D;
 		return;
 	}
 
-	LightGuid = LightMapData->LightGuid;
 	PreviewShadowMapChannel = LightMapData->PreviewShadowMapChannel;
 
 	MarkRenderStateDirty();
