Fix client crash with disconnecting from warmup.

diff --git a/Engine/Source/Runtime/Engine/Private/DemoNetDriver.cpp b/Engine/Source/Runtime/Engine/Private/DemoNetDriver.cpp
--- a/Engine/Source/Runtime/Engine/Private/DemoNetDriver.cpp
+++ b/Engine/Source/Runtime/Engine/Private/DemoNetDriver.cpp
@@ -4099,8 +4099,8 @@ int32 UDemoNetDriver::CleanUpSplitscreenConnections(bool bDeleteOwner)
 
 	for (APlayerController* CurController : SpectatorControllers)
 	{
-		UNetConnection* ControllerNetConnection = CurController->NetConnection;
-		if (ControllerNetConnection && ControllerNetConnection->IsA(UChildConnection::StaticClass()))
+		UNetConnection* ControllerNetConnection = (CurController != nullptr) ? CurController->NetConnection : nullptr;
+		if (ControllerNetConnection != nullptr && ControllerNetConnection->IsA(UChildConnection::StaticClass()))
 		{
 			++NumSplitscreenConnectionsCleaned;
 			// With this toggled, this prevents actor deletion (which we don't want to do when scrubbing)
