- Fix potential crash due to requesting to evict the same sound in the same frame from multiple concurrency groups.
- Add more logging and safeguarding around creatings/revoking concurrency membership

diff --git a/Engine/Source/Runtime/Engine/Private/SoundConcurrency.cpp b/Engine/Source/Runtime/Engine/Private/SoundConcurrency.cpp
--- a/Engine/Source/Runtime/Engine/Private/SoundConcurrency.cpp
+++ b/Engine/Source/Runtime/Engine/Private/SoundConcurrency.cpp
@@ -70,11 +70,18 @@ void FConcurrencyGroup::AddActiveSound(FActiveSound* ActiveSound)
 {
 	check(GroupID != 0);
 
-	const int32 LastIndex = ActiveSound->ConcurrencyGroupIDs.Num();
+	int32 LastIndex = ActiveSound->ConcurrencyGroupIDs.Num();
 	if (ActiveSound->ConcurrencyGroupIDs.AddUnique(GroupID) == LastIndex)
 	{
-		ActiveSounds.Add(ActiveSound);
-		ActiveSound->ConcurrencyGeneration = Generation++;
+		LastIndex = ActiveSounds.Num();
+		if (ActiveSounds.AddUnique(ActiveSound) == LastIndex)
+		{
+			ActiveSound->ConcurrencyGeneration = Generation++;
+		}
+		else
+		{
+			UE_LOG(LogAudio, Fatal, TEXT("Attempting to increment active sound ConcurrencyGeneration by adding active sound '%s' to group multiple times."), *ActiveSound->GetOwnerName());
+		}
 	}
 	else
 	{
@@ -84,6 +91,8 @@ void FConcurrencyGroup::AddActiveSound(FActiveSound* ActiveSound)
 
 void FConcurrencyGroup::RemoveActiveSound(FActiveSound* ActiveSound)
 {
+	check(ActiveSound);
+
 	// Cache generation being removed
 	const int32 RemovedGeneration = ActiveSound->ConcurrencyGeneration;
 
@@ -246,6 +255,7 @@ void FSoundConcurrencyManager::CreateNewGroupsFromHandles(
 FActiveSound* FSoundConcurrencyManager::CreateNewActiveSound(const FActiveSound& NewActiveSound)
 {
 	check(NewActiveSound.GetSound());
+	check(IsInAudioThread());
 
 	// If there are no concurrency settings associated then there is no limit on this sound
 	TArray<FConcurrencyHandle> ConcurrencyHandles;
@@ -302,7 +312,7 @@ FConcurrencyGroup* FSoundConcurrencyManager::CanPlaySound(const FActiveSound& Ne
 		// If no room for new sound, early out
 		if (FActiveSound* SoundToEvict = GetEvictableSound(NewActiveSound, *ConcurrencyGroup))
 		{
-			OutSoundsToEvict.Add(SoundToEvict);
+			OutSoundsToEvict.AddUnique(SoundToEvict);
 		}
 		else
 		{
