Tell PVS not to complain because it is wrong
#rnx

diff --git a/Engine/Source/Runtime/Engine/Private/EdGraph/EdGraphPin.cpp b/Engine/Source/Runtime/Engine/Private/EdGraph/EdGraphPin.cpp
--- a/Engine/Source/Runtime/Engine/Private/EdGraph/EdGraphPin.cpp
+++ b/Engine/Source/Runtime/Engine/Private/EdGraph/EdGraphPin.cpp
@@ -687,7 +687,7 @@ bool UEdGraphPin::DoesDefaultValueMatchAutogenerated() const
 {
 	if (const UEdGraphSchema* Schema = GetSchema())
 	{
-		return Schema->DoesDefaultValueMatchAutogenerated(*this);
+		return Schema->DoesDefaultValueMatchAutogenerated(*this); //-V522
 	}
 	return false;
 }
