Memreport ordering should fallback on alpha sort if sizes match

diff --git a/Engine/Source/Runtime/Engine/Private/UnrealEngine.cpp b/Engine/Source/Runtime/Engine/Private/UnrealEngine.cpp
--- a/Engine/Source/Runtime/Engine/Private/UnrealEngine.cpp
+++ b/Engine/Source/Runtime/Engine/Private/UnrealEngine.cpp
@@ -3246,9 +3246,14 @@ struct FCompareFSortedTexture
 	FCompareFSortedTexture( bool InAlphaSort )
 		: bAlphaSort( InAlphaSort )
 	{}
-	FORCEINLINE bool operator()( const FSortedTexture& A, const FSortedTexture& B ) const
+	FORCEINLINE bool operator()(const FSortedTexture& A, const FSortedTexture& B) const
 	{
-		return bAlphaSort ? ( A.Name < B.Name ) : ( B.CurrentSize < A.CurrentSize );
+		if (bAlphaSort || A.CurrentSize == B.CurrentSize)
+		{
+			return A.Name < B.Name;
+		}
+
+		return B.CurrentSize < A.CurrentSize;
 	}
 };
 
@@ -3310,7 +3315,13 @@ struct FCompareFSortedStaticMesh
 		{
 			return ((B.VertexCountTotalMobile) < (A.VertexCountTotalMobile));
 		}
-		return bAlphaSort ? (A.Name < B.Name) : ((B.NumKB + B.ResKBExc) < (A.NumKB + A.ResKBExc));
+
+		if (bAlphaSort || (B.NumKB + B.ResKBExc) == (A.NumKB + A.ResKBExc))
+		{
+			return A.Name < B.Name;
+		}
+
+		return (B.NumKB + B.ResKBExc) < (A.NumKB + A.ResKBExc);
 	}
 };
 
@@ -3369,7 +3380,13 @@ struct FCompareFSortedSkeletalMesh
 		{
 			return ((B.VertexCountTotalMobile) < (A.VertexCountTotalMobile));
 		}
-		return bAlphaSort ? (A.Name < B.Name) : ((B.NumKB + B.ResKBExc) < (A.NumKB + A.ResKBExc));
+
+		if (bAlphaSort || (B.NumKB + B.ResKBExc) == (A.NumKB + A.ResKBExc))
+		{
+			return A.Name < B.Name;
+		}
+
+		return (B.NumKB + B.ResKBExc) < (A.NumKB + A.ResKBExc);
 	}
 };
 
@@ -3416,7 +3433,12 @@ struct FCompareFSortedAnimAsset
 	{}
 	FORCEINLINE bool operator()(const FSortedAnimAsset& A, const FSortedAnimAsset& B) const
 	{
-		return bAlphaSort ? (A.Name < B.Name) : ((B.NumKB + B.ResKBExc) < (A.NumKB + A.ResKBExc));
+		if (bAlphaSort || (B.NumKB + B.ResKBExc) == (A.NumKB + A.ResKBExc))
+		{
+			return A.Name < B.Name;
+		}
+
+		return (B.NumKB + B.ResKBExc) < (A.NumKB + A.ResKBExc);
 	}
 };
 
@@ -3440,7 +3462,12 @@ struct FCompareFSortedSet
 	{}
 	FORCEINLINE bool operator()( const FSortedSet& A, const FSortedSet& B ) const
 	{
-		return bAlphaSort ? ( A.Name < B.Name ) : ( B.Size < A.Size );
+		if (bAlphaSort || A.Size == B.Size)
+		{
+			return A.Name < B.Name;
+		}
+
+		return B.Size < A.Size;
 	}
 };
 
@@ -3507,7 +3534,12 @@ struct FCompareFSortedParticleSet
 	{}
 	FORCEINLINE bool operator()( const FSortedParticleSet& A, const FSortedParticleSet& B ) const
 	{
-		return bAlphaSort ? ( A.Name < B.Name ) : ( B.Size < A.Size );
+		if (bAlphaSort || A.Size == B.Size)
+		{
+			return A.Name < B.Name;
+		}
+
+		return B.Size < A.Size;
 	}
 };
 
@@ -7238,7 +7270,12 @@ bool UEngine::HandleObjCommand( const TCHAR* Cmd, FOutputDevice& Ar )
 
 				FORCEINLINE bool operator()( const FSubItem& A, const FSubItem& B ) const
 				{
-					return bAlphaSort ? (A.Object->GetPathName() < B.Object->GetPathName()) : (B.Max < A.Max);
+					if (bAlphaSort || A.Max == B.Max)
+					{
+						return A.Object->GetPathName() < B.Object->GetPathName();
+					}
+
+					return B.Max < A.Max;
 				}
 			};
 			Objects.Sort( FCompareFSubItem( bAlphaSort ) );
