Framework - Adding logging to world cleanup instead of on RemoveFromWorld, may add some additional logging to CleanupWorld later to assist with level transition/shutdown debugging.

Marc.Audy

diff --git a/Engine/Source/Runtime/Engine/Private/World.cpp b/Engine/Source/Runtime/Engine/Private/World.cpp
--- a/Engine/Source/Runtime/Engine/Private/World.cpp
+++ b/Engine/Source/Runtime/Engine/Private/World.cpp
@@ -2428,8 +2428,6 @@ void UWorld::RemoveFromWorld( ULevel* Level, bool bAllowIncrementalRemoval )
 	// If the level may be removed incrementally then there must also be no level pending visibility
 	if ( ((CurrentLevelPendingVisibility == nullptr) || (!bAllowIncrementalRemoval && (CurrentLevelPendingVisibility != Level))) && Level->bIsVisible )
 	{
-		UE_LOG(LogWorld, Log, TEXT("UWorld::RemoveFromWorld for %s"), *Level->GetOutermost()->GetName());
-
 		// Keep track of timing.
 		double StartTime = FPlatformTime::Seconds();	
 
@@ -3973,6 +3971,8 @@ bool UWorld::IsNavigationRebuilt() const
 
 void UWorld::CleanupWorld(bool bSessionEnded, bool bCleanupResources, UWorld* NewWorld)
 {
+	UE_LOG(LogWorld, Log, TEXT("UWorld::CleanupWorld for %s, bSessionEnded=%s, bCleanupResources=%s"), *GetName(), bSessionEnded ? TEXT("true") : TEXT("false"), bCleanupResources ? TEXT("true") : TEXT("false"));
+
 	check(IsVisibilityRequestPending() == false);
 	bCleanedUpWorld = true;
 
