Add comment explaining GUID offsetting
#rnx

diff --git a/Engine/Source/Runtime/Engine/Private/Components/LightComponent.cpp b/Engine/Source/Runtime/Engine/Private/Components/LightComponent.cpp
--- a/Engine/Source/Runtime/Engine/Private/Components/LightComponent.cpp
+++ b/Engine/Source/Runtime/Engine/Private/Components/LightComponent.cpp
@@ -1125,6 +1125,8 @@ void ULightComponent::ApplyComponentInstanceData(FPrecomputedLightInstanceData*
 
 	if (!LightMapData->Transform.Equals(GetComponentTransform()))
 	{
+		// Rather than resetting the guid and reallocating one randomly, we'll offset it so that consecutive loads 
+		// behave deterministically  while still disassociating it from any now invalidated baked lighting
 		++LightGuid.D;
 		return;
 	}
