Added support to ke/kismetevent for broadcasting to all instances of a class by passing in the class name as the target
Improved the success/failure logging for all variants of ke
- ke classname command [args] will now find all instances of that class in the current world and call it on them (printing out the number of instances & successes, as well as the fully qualified class name found)
- ke instance_name_or_path command [args] will now print out the fully qualified name of the object found and whether or not the call was successful (not just whether or not the object was found).  Note that the name displayed in the world outliner for an actor is not always the same as the underlying actor name
- ke * command [args] will now print out the number of successes

diff --git a/Engine/Source/Runtime/Engine/Private/UnrealEngine.cpp b/Engine/Source/Runtime/Engine/Private/UnrealEngine.cpp
--- a/Engine/Source/Runtime/Engine/Private/UnrealEngine.cpp
+++ b/Engine/Source/Runtime/Engine/Private/UnrealEngine.cpp
@@ -4775,27 +4775,55 @@ bool UEngine::HandleKismetEventCommand(UWorld* InWorld, const TCHAR* Cmd, FOutpu
 	if (ObjectName == TEXT("*"))
 	{
 		// Send the command to everything in the world we're dealing with...
+		int32 NumInstanceCallsSucceeded = 0;
 		for (TObjectIterator<UObject> It; It; ++It)
 		{
 			UObject* const Obj = *It;
 			UWorld const* const ObjWorld = Obj->GetWorld();
 			if (ObjWorld == InWorld)
 			{
-				Obj->CallFunctionByNameWithArguments(Cmd, Ar, NULL, true);
+				const bool bSucceeded = Obj->CallFunctionByNameWithArguments(Cmd, Ar, nullptr, true);
+				NumInstanceCallsSucceeded += bSucceeded ? 1 : 0;
 			}
 		}
+
+		Ar.Logf(TEXT("Called '%s' on everything in the world and %d instances succeeded"), Cmd, NumInstanceCallsSucceeded);
 	}
 	else
 	{
 		UObject* ObjectToMatch = FindObject<UObject>(ANY_PACKAGE, *ObjectName);
 
-		if (ObjectToMatch == NULL)
+		if (ObjectToMatch == nullptr)
 		{
-			Ar.Logf(TEXT("Failed to find object named '%s'.  Specify a valid name or *"), *ObjectName);
+			Ar.Logf(TEXT("Failed to find an object named '%s'.  Specify a valid object name, class name, or * (if using a class name for a BP, don't forget the _C)"), *ObjectName);
 		}
 		else
 		{
-			ObjectToMatch->CallFunctionByNameWithArguments(Cmd, Ar, NULL, true);
+			if (UClass* ClassToMatch = Cast<UClass>(ObjectToMatch))
+			{
+				// Call it on all instances of the class
+				int32 NumInstancesFound = 0;
+				int32 NumInstanceCallsSucceeded = 0;
+				for (FObjectIterator It(ClassToMatch); It; ++It)
+				{
+					UObject* const Obj = *It;
+					UWorld const* const ObjWorld = Obj->GetWorld();
+					if (ObjWorld == InWorld)
+					{
+						const bool bSucceeded = Obj->CallFunctionByNameWithArguments(Cmd, Ar, nullptr, true);
+						++NumInstancesFound;
+						NumInstanceCallsSucceeded += bSucceeded ? 1 : 0;
+					}
+				}
+
+				Ar.Logf(TEXT("Called '%s' on %d instance(s) of class '%s' (%d succeeded)"), Cmd, NumInstancesFound, *ClassToMatch->GetPathName(), NumInstanceCallsSucceeded);
+			}
+			else
+			{
+				// Call it on the specific instance specified
+				const bool bSucceeded = ObjectToMatch->CallFunctionByNameWithArguments(Cmd, Ar, nullptr, true);
+				Ar.Logf(TEXT("Called '%s' on instance '%s' (call %s)"), Cmd, *ObjectToMatch->GetPathName(), bSucceeded ? TEXT("succeeded") : TEXT("failed"));
+			}
 		}
 	}
 
