Fix issue with sounds with autodestroy set being culled too early if their loops are throttled by new distance-culling AudioComponet functionality.

diff --git a/Engine/Source/Runtime/Engine/Private/Components/AudioComponent.cpp b/Engine/Source/Runtime/Engine/Private/Components/AudioComponent.cpp
--- a/Engine/Source/Runtime/Engine/Private/Components/AudioComponent.cpp
+++ b/Engine/Source/Runtime/Engine/Private/Components/AudioComponent.cpp
@@ -713,6 +713,12 @@ void UAudioComponent::PlaybackCompleted(bool bFailedToStart)
 			ReplayWhenInAudibleRange = EReplayWhenInAudibleRange::Disabled;
 			PrimaryComponentTick.SetTickFunctionEnable(false);
 		}
+		// Don't attempt to auto destroy or report that anything is finished if ReplayWhenInAudibleRange
+		// is still active, as this call for playback being complete is just from stopping due to range throttling
+		else if (bAutoDestroy && ReplayWhenInAudibleRange != EReplayWhenInAudibleRange::Disabled)
+		{
+			return;
+		}
 
 		if (!bFailedToStart && GetWorld() != nullptr && (OnAudioFinished.IsBound() || OnAudioFinishedNative.IsBound()))
 		{
