Adding Getting Player Controller via the Controller ID as a blueprint widget

diff --git a/Engine/Source/Runtime/Engine/Classes/Kismet/GameplayStatics.h b/Engine/Source/Runtime/Engine/Classes/Kismet/GameplayStatics.h
--- a/Engine/Source/Runtime/Engine/Classes/Kismet/GameplayStatics.h
+++ b/Engine/Source/Runtime/Engine/Classes/Kismet/GameplayStatics.h
@@ -110,6 +110,10 @@ class ENGINE_API UGameplayStatics : public UBlueprintFunctionLibrary
 	UFUNCTION(BlueprintPure, Category="Game", meta=(WorldContext="WorldContextObject", UnsafeDuringActorConstruction="true"))
 	static class APlayerController* GetPlayerController(const UObject* WorldContextObject, int32 PlayerIndex);
 
+	/** Returns the player controller that has the given controller ID */
+	UFUNCTION(BlueprintPure, Category = "Game", meta = (WorldContext = "WorldContextObject", UnsafeDuringActorConstruction = "true"))
+	static class APlayerController* GetPlayerControllerFromID(const UObject* WorldContextObject, int32 ControllerID);
+
 	/** Returns the player pawn at the specified player index */
 	UFUNCTION(BlueprintPure, Category="Game", meta=(WorldContext="WorldContextObject", UnsafeDuringActorConstruction="true"))
 	static class APawn* GetPlayerPawn(const UObject* WorldContextObject, int32 PlayerIndex);
diff --git a/Engine/Source/Runtime/Engine/Private/GameplayStatics.cpp b/Engine/Source/Runtime/Engine/Private/GameplayStatics.cpp
--- a/Engine/Source/Runtime/Engine/Private/GameplayStatics.cpp
+++ b/Engine/Source/Runtime/Engine/Private/GameplayStatics.cpp
@@ -217,6 +217,23 @@ class APlayerController* UGameplayStatics::GetPlayerController(const UObject* Wo
 	return nullptr;
 }
 
+class APlayerController* UGameplayStatics::GetPlayerControllerFromID(const UObject* WorldContextObject, int32 ControllerID)
+{
+	if (UWorld* World = GEngine->GetWorldFromContextObject(WorldContextObject, EGetWorldErrorMode::LogAndReturnNull))
+	{
+		for (FConstPlayerControllerIterator Iterator = World->GetPlayerControllerIterator(); Iterator; ++Iterator)
+		{
+			APlayerController* PlayerController = Iterator->Get();
+			int32 PlayerControllerID = GetPlayerControllerID(PlayerController);
+			if (PlayerControllerID != INDEX_NONE && PlayerControllerID == ControllerID)
+			{
+				return PlayerController;
+			}
+		}
+	}
+	return nullptr;
+}
+
 ACharacter* UGameplayStatics::GetPlayerCharacter(const UObject* WorldContextObject, int32 PlayerIndex)
 {
 	APlayerController* PC = GetPlayerController(WorldContextObject, PlayerIndex);
