Fixed selected static mesh collision color material override.

diff --git a/Engine/Source/Runtime/Engine/Private/StaticMeshRender.cpp b/Engine/Source/Runtime/Engine/Private/StaticMeshRender.cpp
--- a/Engine/Source/Runtime/Engine/Private/StaticMeshRender.cpp
+++ b/Engine/Source/Runtime/Engine/Private/StaticMeshRender.cpp
@@ -1153,7 +1153,10 @@ void FStaticMeshSceneProxy::GetDynamicMeshElements(const TArray<const FSceneView
 										if (bSectionIsSelected)
 										{
 											// Override the mesh's material with our material that draws the collision color
-											MeshElement.MaterialRenderProxy = GEngine->ShadedLevelColorationUnlitMaterial->GetRenderProxy();
+											MeshElement.MaterialRenderProxy = new FOverrideSelectionColorMaterialRenderProxy(
+												GEngine->ShadedLevelColorationUnlitMaterial->GetRenderProxy(),
+												GetSelectionColor(GEngine->GetSelectedMaterialColor(), bSectionIsSelected, IsHovered())
+											);
 										}
 	#endif
 										if (MeshElement.bDitheredLODTransition && LODMask.IsDithered())
