Temporarily removing shadermapid.isvalid() check to unblock development builds.
#rnx

diff --git a/Engine/Source/Runtime/Engine/Private/Materials/MaterialShared.cpp b/Engine/Source/Runtime/Engine/Private/Materials/MaterialShared.cpp
--- a/Engine/Source/Runtime/Engine/Private/Materials/MaterialShared.cpp
+++ b/Engine/Source/Runtime/Engine/Private/Materials/MaterialShared.cpp
@@ -1636,7 +1636,7 @@ bool FMaterial::CacheShaders(EShaderPlatform Platform, bool bApplyCompletedShade
 bool FMaterial::CacheShaders(const FMaterialShaderMapId& ShaderMapId, EShaderPlatform Platform, bool bApplyCompletedShaderMapForRendering)
 {
 	bool bSucceeded = false;
-	check(ShaderMapId.IsValid());
+	//check(ShaderMapId.IsValid());
 
 	// If we loaded this material with inline shaders, use what was loaded (GameThreadShaderMap) instead of looking in the DDC
 	if (bContainsInlineShaders)
