Add additional check for valid UTexture when binding to material uniform buffer

diff --git a/Engine/Source/Runtime/Engine/Private/Materials/MaterialUniformExpressions.cpp b/Engine/Source/Runtime/Engine/Private/Materials/MaterialUniformExpressions.cpp
--- a/Engine/Source/Runtime/Engine/Private/Materials/MaterialUniformExpressions.cpp
+++ b/Engine/Source/Runtime/Engine/Private/Materials/MaterialUniformExpressions.cpp
@@ -415,17 +415,15 @@ void FUniformExpressionSet::FillUniformBuffer(const FMaterialRenderContext& Mate
 			void** ResourceTableSamplerPtr = (void**)((uint8*)BufferCursor + 1 * SHADER_PARAMETER_POINTER_ALIGNMENT);
 			BufferCursor = ((uint8*)BufferCursor) + (SHADER_PARAMETER_POINTER_ALIGNMENT * 2);
 
-			if (Value && Value->Resource)
+			// TextureReference.TextureReferenceRHI is cleared from a render command issued by UTexture::BeginDestroy
+			// It's possible for this command to trigger before a given material is cleaned up and removed from deferred update list
+			// Technically I don't think it's necesary to check 'Resource' for nullptr here, as if TextureReferenceRHI has been initialized, that should be enough
+			// Going to leave the check for now though, to hopefully avoid any unexpected problems
+			if (Value && Value->Resource && Value->TextureReference.TextureReferenceRHI)
 			{
 				//@todo-rco: Help track down a invalid values
 				checkf(Value->IsA(UTexture::StaticClass()), TEXT("Expecting a UTexture! Value='%s' class='%s'"), *Value->GetName(), *Value->GetClass()->GetName());
 
-				// UMaterial / UMaterialInstance should have caused all dependent textures to be PostLoaded, which initializes their rendering resource
-				checkf(Value->TextureReference.TextureReferenceRHI, TEXT("Texture %s of class %s had invalid texture reference. Material %s with texture expression in slot %i"),
-					*Value->GetName(), *Value->GetClass()->GetName(),
-					*MaterialRenderContext.Material.GetFriendlyName(), ExpressionIndex);
-
-				
 				*ResourceTableTexturePtr = Value->TextureReference.TextureReferenceRHI;
 				FSamplerStateRHIRef* SamplerSource = &Value->Resource->SamplerStateRHI;
 
