ifdef changes to audio component to tick looping sounds for potential revival due to perf
- Begin work fixing issue of components not auto-destroying with replay sound feature enabled.

diff --git a/Engine/Source/Runtime/Engine/Classes/Components/AudioComponent.h b/Engine/Source/Runtime/Engine/Classes/Components/AudioComponent.h
--- a/Engine/Source/Runtime/Engine/Classes/Components/AudioComponent.h
+++ b/Engine/Source/Runtime/Engine/Classes/Components/AudioComponent.h
@@ -196,13 +196,12 @@ class ENGINE_API UAudioComponent : public USceneComponent
 		/** Sound lifetime is not managed by listener proximity. */
 		Disabled,
 
+		/** Loop replay is enabled and disable has been requested. */
+		DisableRequested,
+
 		/** Sound lifetime is managed by listener proximity (in case of sound being looping and non-virtualized) */
 		Enabled,
 
-		/** Sound lifetime was managed by listener proximity (in case of sound being looping and non-virtualized),
-		  * but is fading out.*/
-		EnabledFadingOut,
-
 		Count
 	};
 
diff --git a/Engine/Source/Runtime/Engine/Private/Components/AudioComponent.cpp b/Engine/Source/Runtime/Engine/Private/Components/AudioComponent.cpp
--- a/Engine/Source/Runtime/Engine/Private/Components/AudioComponent.cpp
+++ b/Engine/Source/Runtime/Engine/Private/Components/AudioComponent.cpp
@@ -21,6 +21,7 @@ FAutoConsoleVariableRef CVarBackedAnalysisTimeShift(
 	TEXT("Value: The time in seconds to shift the timeline."),
 	ECVF_Default);
 
+#define REPLAY_LOOPS_ENABLED 0
 
 /*-----------------------------------------------------------------------------
 UAudioComponent implementation.
@@ -239,7 +240,7 @@ const UObject* UAudioComponent::AdditionalStatObject() const
 
 void UAudioComponent::SetSound(USoundBase* NewSound)
 {
-	const bool bPlay = IsPlaying() || ReplayWhenInAudibleRange != EReplayWhenInAudibleRange::Disabled;
+	const bool bPlay = IsPlaying();
 
 	// If this is an auto destroy component we need to prevent it from being auto-destroyed since we're really just restarting it
 	const bool bWasAutoDestroy = bAutoDestroy;
@@ -264,22 +265,26 @@ void UAudioComponent::TickComponent(float DeltaTime, ELevelTick TickType, FActor
 {
 	Super::TickComponent(DeltaTime, TickType, TickFunction);
 
-	if (bIsActive)
+	const bool bIsReplayEnabled = ReplayWhenInAudibleRange == EReplayWhenInAudibleRange::Enabled;
+	if (bIsReplayEnabled)
 	{
-		if (ReplayWhenInAudibleRange != EReplayWhenInAudibleRange::Disabled && !IsInAudibleRange())
+		if (ActiveCount > 0)
 		{
-			// If this is an auto destroy component we need to prevent it from being auto-destroyed since we're really just restarting it
-			bool bCurrentAutoDestroy = bAutoDestroy;
-			bAutoDestroy = false;
-			StopInternal();
-			bAutoDestroy = bCurrentAutoDestroy;
+			if (!IsInAudibleRange())
+			{
+				// If this is an auto destroy component we need to prevent it from being auto-destroyed since we're really just restarting it
+				bool bCurrentAutoDestroy = bAutoDestroy;
+				bAutoDestroy = false;
+				StopInternal();
+				bAutoDestroy = bCurrentAutoDestroy;
+			}
 		}
-	}
-	else 
-	{
-		if (ReplayWhenInAudibleRange != EReplayWhenInAudibleRange::Disabled && IsInAudibleRange())
+		else
 		{
-			Play(); // Play refreshes ReplayWhenInAudibleRange internally
+			if (IsInAudibleRange())
+			{
+				Play(); // Play refreshes ReplayWhenInAudibleRange internally
+			}
 		}
 	}
 }
@@ -515,18 +520,16 @@ void UAudioComponent::PlayInternal(const float StartTime, const float FadeInDura
 
 			AudioDevice->AddNewActiveSound(NewActiveSound);
 
+#if REPLAY_LOOPS_ENABLED
+			// Only allow replay if this is the first sound to play
 			const bool bCanVirtualize = Sound && Sound->IsVirtualizeWhenSilent() && AudioDevice->VirtualSoundsEnabled();
-			if (!bIsPreviewSound && NewActiveSound.bHasAttenuationSettings && !bCanVirtualize)
+			if (!bIsPreviewSound && NewActiveSound.bHasAttenuationSettings && !bCanVirtualize && !bAutoDestroy)
 			{
 				const bool bLooping = NewActiveSound.IsLooping();
 				ReplayWhenInAudibleRange = bLooping ? EReplayWhenInAudibleRange::Enabled : EReplayWhenInAudibleRange::Disabled;
 				PrimaryComponentTick.SetTickFunctionEnable(bLooping);
 			}
-			else
-			{
-				ReplayWhenInAudibleRange = EReplayWhenInAudibleRange::Disabled;
-			}
-
+#endif // REPLAY_LOOPS_ENABLED
 			bIsActive = true;
 		}
 	}
@@ -566,11 +569,13 @@ void UAudioComponent::FadeOut( float FadeOutDuration, float FadeVolumeLevel )
 	{
 		if (FadeOutDuration > 0.0f)
 		{
+#if REPLAY_LOOPS_ENABLED
 			if (FadeVolumeLevel == 0.0f)
 			{
-				ReplayWhenInAudibleRange = EReplayWhenInAudibleRange::EnabledFadingOut;
+				ReplayWhenInAudibleRange = EReplayWhenInAudibleRange::DisableRequested;
 				PrimaryComponentTick.SetTickFunctionEnable(false);
 			}
+#endif // REPLAY_LOOPS_ENABLED
 
 			if (FAudioDevice* AudioDevice = GetAudioDevice())
 			{
@@ -602,11 +607,13 @@ void UAudioComponent::AdjustVolume( float AdjustVolumeDuration, float AdjustVolu
 	{
 		if (FAudioDevice* AudioDevice = GetAudioDevice())
 		{
-			if (ReplayWhenInAudibleRange == EReplayWhenInAudibleRange::EnabledFadingOut)
+#if REPLAY_LOOPS_ENABLED
+			if (AdjustVolumeLevel > 0.0f && ReplayWhenInAudibleRange == EReplayWhenInAudibleRange::DisableRequested)
 			{
-				ReplayWhenInAudibleRange= EReplayWhenInAudibleRange::Enabled;
+				ReplayWhenInAudibleRange = EReplayWhenInAudibleRange::Enabled;
 				PrimaryComponentTick.SetTickFunctionEnable(true);
 			}
+#endif // REPLAY_LOOPS_ENABLED
 
 			DECLARE_CYCLE_STAT(TEXT("FAudioThreadTask.AdjustVolume"), STAT_AudioAdjustVolume, STATGROUP_AudioThreadCommands);
 
@@ -636,9 +643,10 @@ void UAudioComponent::AdjustVolume( float AdjustVolumeDuration, float AdjustVolu
 
 void UAudioComponent::Stop()
 {
+#if REPLAY_LOOPS_ENABLED
+	ReplayWhenInAudibleRange = EReplayWhenInAudibleRange::DisableRequested;
 	PrimaryComponentTick.SetTickFunctionEnable(false);
-
-	ReplayWhenInAudibleRange = EReplayWhenInAudibleRange::Disabled;
+#endif // REPLAY_LOOPS_ENABLED
 
 	StopInternal();
 }
@@ -703,42 +711,41 @@ void UAudioComponent::PlaybackCompleted(bool bFailedToStart)
 	check(ActiveCount > 0);
 	--ActiveCount;
 
+	if (ActiveCount > 0)
+	{
+		return;
+	}
+
 	// Mark inactive before calling destroy to avoid recursion
-	bIsActive = (ActiveCount > 0);
+	if (bIsActive)
+	{
+		bIsActive = false;
+	}
 
-	if (!bIsActive)
+	if (bFailedToStart)
 	{
-		if (bFailedToStart)
-		{
-			ReplayWhenInAudibleRange = EReplayWhenInAudibleRange::Disabled;
-			PrimaryComponentTick.SetTickFunctionEnable(false);
-		}
-		// Don't attempt to auto destroy or report that anything is finished if ReplayWhenInAudibleRange
-		// is still active, as this call for playback being complete is just from stopping due to range throttling
-		else if (bAutoDestroy && ReplayWhenInAudibleRange != EReplayWhenInAudibleRange::Disabled)
-		{
-			return;
-		}
+		ReplayWhenInAudibleRange = EReplayWhenInAudibleRange::Disabled;
+		PrimaryComponentTick.SetTickFunctionEnable(false);
+	}
 
-		if (!bFailedToStart && GetWorld() != nullptr && (OnAudioFinished.IsBound() || OnAudioFinishedNative.IsBound()))
-		{
-			INC_DWORD_STAT(STAT_AudioFinishedDelegatesCalled);
-			SCOPE_CYCLE_COUNTER(STAT_AudioFinishedDelegates);
+	if (!bFailedToStart && GetWorld() != nullptr && (OnAudioFinished.IsBound() || OnAudioFinishedNative.IsBound()))
+	{
+		INC_DWORD_STAT(STAT_AudioFinishedDelegatesCalled);
+		SCOPE_CYCLE_COUNTER(STAT_AudioFinishedDelegates);
 
-			OnAudioFinished.Broadcast();
-			OnAudioFinishedNative.Broadcast(this);
-		}
+		OnAudioFinished.Broadcast();
+		OnAudioFinishedNative.Broadcast(this);
+	}
 
-		// Auto destruction is handled via marking object for deletion.
-		if (bAutoDestroy)
-		{
-			DestroyComponent();
-		}
-		// Otherwise see if we should detach ourself and wait until we're needed again
-		else if (bAutoManageAttachment)
-		{
-			CancelAutoAttachment(true);
-		}
+	// Auto destruction is handled via marking object for deletion.
+	if (bAutoDestroy)
+	{
+		DestroyComponent();
+	}
+	// Otherwise see if we should detach ourself and wait until we're needed again
+	else if (bAutoManageAttachment)
+	{
+		CancelAutoAttachment(true);
 	}
 }
 
