Swapping a fatal assert to a conditionally logged error as the issue should not block the client. Materials with invalid shader IDs can gracefully fall back to the default material as-is.

diff --git a/Engine/Source/Runtime/Engine/Private/Materials/MaterialShared.cpp b/Engine/Source/Runtime/Engine/Private/Materials/MaterialShared.cpp
--- a/Engine/Source/Runtime/Engine/Private/Materials/MaterialShared.cpp
+++ b/Engine/Source/Runtime/Engine/Private/Materials/MaterialShared.cpp
@@ -510,7 +510,7 @@ void FMaterial::GetShaderMapId(EShaderPlatform Platform, FMaterialShaderMapId& O
 		OutId.QualityLevel = GetQualityLevelForShaderMapId();
 		OutId.FeatureLevel = GetFeatureLevel();
 
-		UE_LOG(LogMaterial, Warning, TEXT("Tried to access an uncooked shader map ID in a cooked application"));
+		UE_LOG(LogMaterial, Log, TEXT("Tried to access an uncooked shader map ID in a cooked application"));
 #endif
 	}
 }
@@ -1636,7 +1636,7 @@ bool FMaterial::CacheShaders(EShaderPlatform Platform, bool bApplyCompletedShade
 bool FMaterial::CacheShaders(const FMaterialShaderMapId& ShaderMapId, EShaderPlatform Platform, bool bApplyCompletedShaderMapForRendering)
 {
 	bool bSucceeded = false;
-	//check(ShaderMapId.IsValid());
+	UE_CLOG(!ShaderMapId.IsValid(), LogMaterial, Warning, TEXT("Invalid shader map ID caching shaders for '%s', will use default material."), *GetFriendlyName());
 
 	// If we loaded this material with inline shaders, use what was loaded (GameThreadShaderMap) instead of looking in the DDC
 	if (bContainsInlineShaders)
