Add a cvar for the option of overlapping render asset streaming ticks

diff --git a/Engine/Source/Runtime/Engine/Private/Streaming/StreamingManagerTexture.cpp b/Engine/Source/Runtime/Engine/Private/Streaming/StreamingManagerTexture.cpp
--- a/Engine/Source/Runtime/Engine/Private/Streaming/StreamingManagerTexture.cpp
+++ b/Engine/Source/Runtime/Engine/Private/Streaming/StreamingManagerTexture.cpp
@@ -27,6 +27,12 @@
 
 CSV_DECLARE_CATEGORY_MODULE_EXTERN(CORE_API, Basic);
 
+static TAutoConsoleVariable<int32> CVarStreamingOverlapAssetAndLevelTicks(
+	TEXT("r.Streaming.OverlapAssetAndLevelTicks"),
+	PLATFORM_PS4 || PLATFORM_XBOXONE,
+	TEXT("Ticks render asset streaming info on a high priority task thread while ticking levels on GT"),
+	ECVF_Default);
+
 bool TrackRenderAsset( const FString& AssetName );
 bool UntrackRenderAsset( const FString& AssetName );
 void ListTrackedRenderAssets( FOutputDevice& Ar, int32 NumTextures );
@@ -1395,12 +1401,13 @@ void FRenderAssetStreamingManager::UpdateResourceStreaming( float DeltaTime, boo
 			SetLastUpdateTime();
 		}
 
-		FEvent* SyncEvent;
+		FEvent* SyncEvent = nullptr;
 		// Optimization: overlapping UpdateStreamingRenderAssets() and IncrementalUpdate();
 		// Restrict this optimization to platforms tested to have a win;
 		// Platforms tested and results (ave exec time of UpdateResourceStreaming):
 		//   PS4 Pro - from ~0.55 ms/frame to ~0.15 ms/frame
-		const bool bOverlappedExecution = !!(PLATFORM_PS4) && !(WITH_EDITOR) && bUseThreadingForPerf;
+		//   XB1 X - from ~0.45 ms/frame to ~0.17 ms/frame
+		const bool bOverlappedExecution = !(WITH_EDITOR) && bUseThreadingForPerf && CVarStreamingOverlapAssetAndLevelTicks.GetValueOnGameThread();
 		if (bOverlappedExecution)
 		{
 			SyncEvent = FPlatformProcess::GetSynchEventFromPool(false);
