Prevent backwards compatible replay warnings when receiving an RPC with no params

#rnx

diff --git a/Engine/Source/Runtime/Engine/Private/RepLayout.cpp b/Engine/Source/Runtime/Engine/Private/RepLayout.cpp
--- a/Engine/Source/Runtime/Engine/Private/RepLayout.cpp
+++ b/Engine/Source/Runtime/Engine/Private/RepLayout.cpp
@@ -5591,48 +5591,48 @@ void FRepLayout::ReceivePropertiesForRPC(
 {
 	check(Function == Owner);
 
-	for (int32 i = 0; i < Parents.Num(); i++)
+	if (ERepLayoutState::Normal == LayoutState)
 	{
-		if (Parents[i].ArrayIndex == 0 && !EnumHasAnyFlags(Parents[i].Flags, ERepParentFlags::IsZeroConstructible))
+		for (int32 i = 0; i < Parents.Num(); i++)
 		{
-			// If this property needs to be constructed, make sure we do that
-			Parents[i].Property->InitializeValue(Data + Parents[i]);
+			if (Parents[i].ArrayIndex == 0 && !EnumHasAnyFlags(Parents[i].Flags, ERepParentFlags::IsZeroConstructible))
+			{
+				// If this property needs to be constructed, make sure we do that
+				Parents[i].Property->InitializeValue(Data + Parents[i]);
+			}
 		}
-	}
 
-	if (Channel->Connection->InternalAck)
-	{
-		bool bHasUnmapped = false;
-		bool bGuidsChanged = false;
+		if (Channel->Connection->InternalAck)
+		{
+			bool bHasUnmapped = false;
+			bool bGuidsChanged = false;
 
-		// Let package map know we want to track and know about any guids that are unmapped during the serialize call
-		// We have to do this manually since we aren't passing in any unmapped info
-		Reader.PackageMap->ResetTrackedGuids(true);
+			// Let package map know we want to track and know about any guids that are unmapped during the serialize call
+			// We have to do this manually since we aren't passing in any unmapped info
+			Reader.PackageMap->ResetTrackedGuids(true);
 
-		ReceiveProperties_BackwardsCompatible(Channel->Connection, nullptr, Data, Reader, bHasUnmapped, false, bGuidsChanged);
+			ReceiveProperties_BackwardsCompatible(Channel->Connection, nullptr, Data, Reader, bHasUnmapped, false, bGuidsChanged);
 
-		if (Reader.PackageMap->GetTrackedUnmappedGuids().Num() > 0)
-		{
-			bHasUnmapped = true;
-			UnmappedGuids = Reader.PackageMap->GetTrackedUnmappedGuids();
-		}
+			if (Reader.PackageMap->GetTrackedUnmappedGuids().Num() > 0)
+			{
+				bHasUnmapped = true;
+				UnmappedGuids = Reader.PackageMap->GetTrackedUnmappedGuids();
+			}
 
-		Reader.PackageMap->ResetTrackedGuids(false);
+			Reader.PackageMap->ResetTrackedGuids(false);
 
-		if (bHasUnmapped)
-		{
-			UE_LOG(LogRepTraffic, Log, TEXT("Unable to resolve RPC parameter to do being unmapped. Object[%d] %s. Function %s."),
+			if (bHasUnmapped)
+			{
+				UE_LOG(LogRepTraffic, Log, TEXT("Unable to resolve RPC parameter to do being unmapped. Object[%d] %s. Function %s."),
 					Channel->ChIndex, *Object->GetName(), *Function->GetName());
+			}
 		}
-	}
-	else
-	{
-		Reader.PackageMap->ResetTrackedGuids(true);
+		else
+		{
+			Reader.PackageMap->ResetTrackedGuids(true);
 
-		static FRepSerializationSharedInfo Empty;
+			static FRepSerializationSharedInfo Empty;
 
-		if (ERepLayoutState::Normal == LayoutState)
-		{
 			for (int32 i = 0; i < Parents.Num(); i++)
 			{
 				if (Cast<UBoolProperty>(Parents[i].Property) || Reader.ReadBit())
