- Move r-value assignment to new line as could help debug/may be an issue with Android-specific optimization for FindOrAdd. Hoping to get better stack to sus-out the issue.
- Add more safeguard checks to ensure concurrency manager is only accessed on the audio thread

diff --git a/Engine/Source/Runtime/Engine/Classes/Sound/SoundConcurrency.h b/Engine/Source/Runtime/Engine/Classes/Sound/SoundConcurrency.h
--- a/Engine/Source/Runtime/Engine/Classes/Sound/SoundConcurrency.h
+++ b/Engine/Source/Runtime/Engine/Classes/Sound/SoundConcurrency.h
@@ -212,12 +212,6 @@ public:
 	FSoundConcurrencyManager(class FAudioDevice* InAudioDevice);
 	ENGINE_API ~FSoundConcurrencyManager();
 
-	void CreateNewGroupsFromHandles(
-		const FActiveSound& NewActiveSound,
-		const TArray<FConcurrencyHandle>& ConcurrencyHandles,
-		TArray<FConcurrencyGroup*>& OutGroupsToApply
-	);
-
 	/** Returns a newly allocated active sound given the input active sound struct. Will return nullptr if the active sound concurrency evaluation doesn't allow for it. */
 	FActiveSound* CreateNewActiveSound(const FActiveSound& NewActiveSound);
 
@@ -235,6 +229,13 @@ private: // Methods
 	/** Creates a new concurrency group and returns pointer to said group */
 	FConcurrencyGroup& CreateNewConcurrencyGroup(const FConcurrencyHandle& ConcurrencyHandle);
 
+	/** Creates new concurrency groups from handle array */
+	void CreateNewGroupsFromHandles(
+		const FActiveSound& NewActiveSound,
+		const TArray<FConcurrencyHandle>& ConcurrencyHandles,
+		TArray<FConcurrencyGroup*>& OutGroupsToApply
+	);
+
 	/**  Creates an active sound to play, assigning it to the provided concurrency groups, and evicting required sounds */
 	FActiveSound* CreateAndEvictActiveSounds(const FActiveSound& NewActiveSound, const TArray<FConcurrencyGroup*>& GroupsToApply, const TArray<FActiveSound*>& SoundsToEvict);
 
diff --git a/Engine/Source/Runtime/Engine/Private/SoundConcurrency.cpp b/Engine/Source/Runtime/Engine/Private/SoundConcurrency.cpp
--- a/Engine/Source/Runtime/Engine/Private/SoundConcurrency.cpp
+++ b/Engine/Source/Runtime/Engine/Private/SoundConcurrency.cpp
@@ -542,7 +542,8 @@ FActiveSound* FSoundConcurrencyManager::CreateAndEvictActiveSounds(const FActive
 
 				const float ActiveSoundGeneration = static_cast<float>(CurActiveSound->ConcurrencyGeneration);
 				const float GenerationDelta = NextGeneration - ActiveSoundGeneration;
-				CurActiveSound->ConcurrencyGroupVolumeScales.FindOrAdd(ConcurrencyGroup->GetGroupID()) = FMath::Pow(Volume, GenerationDelta);
+				float& VolumeScale = CurActiveSound->ConcurrencyGroupVolumeScales.FindOrAdd(ConcurrencyGroup->GetGroupID());
+				VolumeScale = FMath::Pow(Volume, GenerationDelta);
 			}
 		}
 
@@ -581,6 +582,8 @@ FActiveSound* FSoundConcurrencyManager::CreateAndEvictActiveSounds(const FActive
 
 void FSoundConcurrencyManager::StopActiveSound(FActiveSound* ActiveSound)
 {
+	check(IsInAudioThread());
+
 	// Remove this sound from it's concurrency list
 	for (const FConcurrencyGroupID ConcurrencyGroupID : ActiveSound->ConcurrencyGroupIDs)
 	{
@@ -648,6 +651,8 @@ void FSoundConcurrencyManager::StopActiveSound(FActiveSound* ActiveSound)
 
 void FSoundConcurrencyManager::UpdateQuietSoundsToStop()
 {
+	check(IsInAudioThread());
+
 	for (auto& ConcurrenyGroupEntry : ConcurrencyGroups)
 	{
 		ConcurrenyGroupEntry.Value.StopQuietSoundsDueToMaxConcurrency();
