TMap is a hash-map so the underlying table can be realloc'd and that doesn't guarantee that pointers will always be resized, it may simply memcpy into a fresh block, thus if you mean to reference entries by pointer then you must allocate them from the heap individually or you'll eventually stomp the heap.

diff --git a/Engine/Source/Runtime/Engine/Classes/Sound/SoundConcurrency.h b/Engine/Source/Runtime/Engine/Classes/Sound/SoundConcurrency.h
--- a/Engine/Source/Runtime/Engine/Classes/Sound/SoundConcurrency.h
+++ b/Engine/Source/Runtime/Engine/Classes/Sound/SoundConcurrency.h
@@ -171,7 +171,7 @@ public:
 	void StopQuietSoundsDueToMaxConcurrency();
 };
 
-typedef TMap<FConcurrencyGroupID, FConcurrencyGroup> FConcurrencyGroups;
+typedef TMap<FConcurrencyGroupID, FConcurrencyGroup*> FConcurrencyGroups;
 
 struct FSoundInstanceEntry
 {
diff --git a/Engine/Source/Runtime/Engine/Private/SoundConcurrency.cpp b/Engine/Source/Runtime/Engine/Private/SoundConcurrency.cpp
--- a/Engine/Source/Runtime/Engine/Private/SoundConcurrency.cpp
+++ b/Engine/Source/Runtime/Engine/Private/SoundConcurrency.cpp
@@ -283,15 +283,15 @@ FConcurrencyGroup& FSoundConcurrencyManager::CreateNewConcurrencyGroup(const FCo
 {
 	//Create & add new concurrency group to the map
 	FConcurrencyGroupID GroupID = FConcurrencyGroup::GenerateNewID();
-	ConcurrencyGroups.Emplace(GroupID, FConcurrencyGroup(GroupID, ConcurrencyHandle));
+	ConcurrencyGroups.Emplace(GroupID, new FConcurrencyGroup(GroupID, ConcurrencyHandle));
 
-	return *ConcurrencyGroups.Find(GroupID);
+	return *ConcurrencyGroups.FindRef(GroupID);
 }
 
 FConcurrencyGroup* FSoundConcurrencyManager::CanPlaySound(const FActiveSound& NewActiveSound, const FConcurrencyGroupID GroupID, TArray<FActiveSound*>& OutSoundsToEvict)
 {
 	check(GroupID != 0);
-	FConcurrencyGroup* ConcurrencyGroup = ConcurrencyGroups.Find(GroupID);
+	FConcurrencyGroup* ConcurrencyGroup = ConcurrencyGroups.FindRef(GroupID);
 	if (!ConcurrencyGroup)
 	{
 		UE_LOG(LogAudio, Warning, TEXT("Attempting to add active sound '%s' (owner '%s') to invalid concurrency group."),
@@ -587,7 +587,7 @@ void FSoundConcurrencyManager::StopActiveSound(FActiveSound* ActiveSound)
 	// Remove this sound from it's concurrency list
 	for (const FConcurrencyGroupID ConcurrencyGroupID : ActiveSound->ConcurrencyGroupIDs)
 	{
-		FConcurrencyGroup* ConcurrencyGroup = ConcurrencyGroups.Find(ConcurrencyGroupID);
+		FConcurrencyGroup* ConcurrencyGroup = ConcurrencyGroups.FindRef(ConcurrencyGroupID);
 		if (!ConcurrencyGroup)
 		{
 			UE_LOG(LogAudio, Error, TEXT("Attempting to remove stopped sound '%s' from inactive concurrency group."),
@@ -644,6 +644,8 @@ void FSoundConcurrencyManager::StopActiveSound(FActiveSound* ActiveSound)
 					}
 				}
 			}
+			
+			delete ConcurrencyGroup;
 		}
 	}
 	ActiveSound->ConcurrencyGroupIDs.Reset();
@@ -655,6 +657,6 @@ void FSoundConcurrencyManager::UpdateQuietSoundsToStop()
 
 	for (auto& ConcurrenyGroupEntry : ConcurrencyGroups)
 	{
-		ConcurrenyGroupEntry.Value.StopQuietSoundsDueToMaxConcurrency();
+		ConcurrenyGroupEntry.Value->StopQuietSoundsDueToMaxConcurrency();
 	}
 }
