Fix static analysis warning

diff --git a/Engine/Source/Runtime/Engine/Private/Streaming/StreamingManagerTexture.cpp b/Engine/Source/Runtime/Engine/Private/Streaming/StreamingManagerTexture.cpp
--- a/Engine/Source/Runtime/Engine/Private/Streaming/StreamingManagerTexture.cpp
+++ b/Engine/Source/Runtime/Engine/Private/Streaming/StreamingManagerTexture.cpp
@@ -29,7 +29,7 @@ CSV_DECLARE_CATEGORY_MODULE_EXTERN(CORE_API, Basic);
 
 static TAutoConsoleVariable<int32> CVarStreamingOverlapAssetAndLevelTicks(
 	TEXT("r.Streaming.OverlapAssetAndLevelTicks"),
-	PLATFORM_PS4 || PLATFORM_XBOXONE,
+	!WITH_EDITOR && (PLATFORM_PS4 || PLATFORM_XBOXONE),
 	TEXT("Ticks render asset streaming info on a high priority task thread while ticking levels on GT"),
 	ECVF_Default);
 
@@ -1407,7 +1407,7 @@ void FRenderAssetStreamingManager::UpdateResourceStreaming( float DeltaTime, boo
 		// Platforms tested and results (ave exec time of UpdateResourceStreaming):
 		//   PS4 Pro - from ~0.55 ms/frame to ~0.15 ms/frame
 		//   XB1 X - from ~0.45 ms/frame to ~0.17 ms/frame
-		const bool bOverlappedExecution = !(WITH_EDITOR) && bUseThreadingForPerf && CVarStreamingOverlapAssetAndLevelTicks.GetValueOnGameThread();
+		const bool bOverlappedExecution = bUseThreadingForPerf && CVarStreamingOverlapAssetAndLevelTicks.GetValueOnGameThread();
 		if (bOverlappedExecution)
 		{
 			SyncEvent = FPlatformProcess::GetSynchEventFromPool(false);
