UMG - Updating the code responsible for constraining the UI to just be over the players screen when camera view constraints are in play.  This was only partially working for safezones.

diff --git a/Engine/Source/Runtime/Engine/Private/GameViewportClient.cpp b/Engine/Source/Runtime/Engine/Private/GameViewportClient.cpp
--- a/Engine/Source/Runtime/Engine/Private/GameViewportClient.cpp
+++ b/Engine/Source/Runtime/Engine/Private/GameViewportClient.cpp
@@ -2515,11 +2515,15 @@ void UGameViewportClient::RemoveViewportWidgetContent( TSharedRef<SWidget> Viewp
 
 void UGameViewportClient::AddViewportWidgetForPlayer(ULocalPlayer* Player, TSharedRef<SWidget> ViewportContent, const int32 ZOrder)
 {
-	TSharedPtr< IGameLayerManager > GameLayerManager(GameLayerManagerPtr.Pin());
-	if ( GameLayerManager.IsValid() )
+	if (ensure(Player))
 	{
-		GameLayerManager->AddWidgetForPlayer(Player, ViewportContent, ZOrder);
+		TSharedPtr< IGameLayerManager > GameLayerManager(GameLayerManagerPtr.Pin());
+		if (GameLayerManager.IsValid())
+		{
+			GameLayerManager->AddWidgetForPlayer(Player, ViewportContent, ZOrder);
+		}
 	}
+	//TODO - If this fails what should we do?
 }
 
 void UGameViewportClient::RemoveViewportWidgetForPlayer(ULocalPlayer* Player, TSharedRef<SWidget> ViewportContent)
diff --git a/Engine/Source/Runtime/Engine/Private/Slate/SGameLayerManager.cpp b/Engine/Source/Runtime/Engine/Private/Slate/SGameLayerManager.cpp
--- a/Engine/Source/Runtime/Engine/Private/Slate/SGameLayerManager.cpp
+++ b/Engine/Source/Runtime/Engine/Private/Slate/SGameLayerManager.cpp
@@ -169,12 +169,22 @@ void SGameLayerManager::AddWidgetForPlayer(ULocalPlayer* Player, TSharedRef<SWid
 
 void SGameLayerManager::RemoveWidgetForPlayer(ULocalPlayer* Player, TSharedRef<SWidget> ViewportContent)
 {
-	TSharedPtr<FPlayerLayer>* PlayerLayerPtr = PlayerLayers.Find(Player);
-	if ( PlayerLayerPtr )
+	if (TSharedPtr<FPlayerLayer> PlayerLayer = PlayerLayers.FindRef(Player) )
 	{
-		TSharedPtr<FPlayerLayer> PlayerLayer = *PlayerLayerPtr;
 		PlayerLayer->Widget->RemoveSlot(ViewportContent);
 	}
+	else
+	{
+		// If no local user is specified, we need to find the widget and purge it.
+		for (auto& PlayerLayerEntry : PlayerLayers)
+		{
+			const TSharedPtr<FPlayerLayer>& PlayerLayerRef = PlayerLayerEntry.Value;
+			if (PlayerLayerRef->Widget->RemoveSlot(ViewportContent))
+			{
+				return;
+			}
+		}
+	}
 }
 
 void SGameLayerManager::ClearWidgetsForPlayer(ULocalPlayer* Player)
@@ -386,9 +396,9 @@ void SGameLayerManager::RemoveMissingPlayerLayers(const TArray<ULocalPlayer*>& G
 	TArray<ULocalPlayer*> ToRemove;
 
 	// Find the player layers for players that no longer exist
-	for ( TMap< ULocalPlayer*, TSharedPtr<FPlayerLayer> >::TIterator It(PlayerLayers); It; ++It )
+	for (auto& PlayerLayerEntry : PlayerLayers)
 	{
-		ULocalPlayer* Key = ( *It ).Key;
+		ULocalPlayer* Key = Cast<ULocalPlayer>(PlayerLayerEntry.Key.ResolveObjectPtr());
 		if ( !GamePlayers.Contains(Key) )
 		{
 			ToRemove.Add(Key);
@@ -404,10 +414,12 @@ void SGameLayerManager::RemoveMissingPlayerLayers(const TArray<ULocalPlayer*>& G
 
 void SGameLayerManager::RemovePlayerWidgets(ULocalPlayer* LocalPlayer)
 {
-	TSharedPtr<FPlayerLayer> Layer = PlayerLayers.FindRef(LocalPlayer);
+	FObjectKey LocalPlayerKey(LocalPlayer);
+
+	TSharedPtr<FPlayerLayer> Layer = PlayerLayers.FindRef(LocalPlayerKey);
 	PlayerCanvas->RemoveSlot(Layer->Widget.ToSharedRef());
 
-	PlayerLayers.Remove(LocalPlayer);
+	PlayerLayers.Remove(LocalPlayerKey);
 }
 
 void SGameLayerManager::AddOrUpdatePlayerLayers(const FGeometry& AllottedGeometry, UGameViewportClient* ViewportClient, const TArray<ULocalPlayer*>& GamePlayers)
@@ -433,16 +445,9 @@ void SGameLayerManager::AddOrUpdatePlayerLayers(const FGeometry& AllottedGeometr
 			FPerPlayerSplitscreenData& SplitData = SplitInfo[SplitType].PlayerData[PlayerIndex];
 
 			// Viewport Sizes
-			FVector2D Size, Position;
-			Size.X = SplitData.SizeX;
-			Size.Y = SplitData.SizeY;
-			Position.X = SplitData.OriginX;
-			Position.Y = SplitData.OriginY;
-
-			FVector2D AspectRatioInset = GetAspectRatioInset(Player);
-
-			Position += AspectRatioInset;
-			Size -= (AspectRatioInset * 2.0f);
+			FVector2D Position(0, 0);
+			FVector2D Size(SplitData.SizeX, SplitData.SizeY);
+			GetNormalizeRect(Player, Position, Size);
 
 			Size = Size * AllottedGeometry.GetLocalSize() * InverseDPIScale;
 			Position = Position * AllottedGeometry.GetLocalSize() * InverseDPIScale;
@@ -458,10 +463,9 @@ void SGameLayerManager::AddOrUpdatePlayerLayers(const FGeometry& AllottedGeometr
 	}
 }
 
-FVector2D SGameLayerManager::GetAspectRatioInset(ULocalPlayer* LocalPlayer) const
+bool SGameLayerManager::GetNormalizeRect(ULocalPlayer* LocalPlayer, FVector2D& OutPosition, FVector2D& OutSize) const
 {
 	QUICK_SCOPE_CYCLE_COUNTER(STAT_SGameLayerManager_GetAspectRatioInset);
-	FVector2D Offset(0.f, 0.f);
 	if ( LocalPlayer )
 	{
 		FSceneViewProjectionData ProjectionData;
@@ -470,13 +474,20 @@ FVector2D SGameLayerManager::GetAspectRatioInset(ULocalPlayer* LocalPlayer) cons
 			const FIntRect ViewRect = ProjectionData.GetViewRect();
 			const FIntRect ConstrainedViewRect = ProjectionData.GetConstrainedViewRect();
 
+			const FIntPoint ViewportSize = LocalPlayer->ViewportClient->Viewport->GetSizeXY();
+
 			// Return normalized coordinates.
-			Offset.X = ( ConstrainedViewRect.Min.X - ViewRect.Min.X ) / (float)ViewRect.Width();
-			Offset.Y = ( ConstrainedViewRect.Min.Y - ViewRect.Min.Y ) / (float)ViewRect.Height();
+			OutPosition.X = ConstrainedViewRect.Min.X / (float)ViewportSize.X;
+			OutPosition.Y = ConstrainedViewRect.Min.Y / (float)ViewportSize.Y;
+
+			OutSize.X = ConstrainedViewRect.Width() / (float)ViewportSize.X;
+			OutSize.Y = ConstrainedViewRect.Height() / (float)ViewportSize.Y;
+
+			return true;
 		}
 	}
 
-	return Offset;
+	return false;
 }
 
 void SGameLayerManager::SetDefaultWindowTitleBarHeight(float Height)
diff --git a/Engine/Source/Runtime/Engine/Public/Slate/SGameLayerManager.h b/Engine/Source/Runtime/Engine/Public/Slate/SGameLayerManager.h
--- a/Engine/Source/Runtime/Engine/Public/Slate/SGameLayerManager.h
+++ b/Engine/Source/Runtime/Engine/Public/Slate/SGameLayerManager.h
@@ -12,6 +12,7 @@
 #include "Widgets/SCanvas.h"
 #include "Widgets/Layout/SBox.h"
 #include "Styling/SlateTypes.h"
+#include "UObject/ObjectKey.h"
 
 class FPaintArgs;
 class FSceneViewport;
@@ -169,7 +170,7 @@ private:
 	void RemoveMissingPlayerLayers(const TArray<ULocalPlayer*>& GamePlayers);
 	void RemovePlayerWidgets(ULocalPlayer* LocalPlayer);
 	void AddOrUpdatePlayerLayers(const FGeometry& AllottedGeometry, UGameViewportClient* ViewportClient, const TArray<ULocalPlayer*>& GamePlayers);
-	FVector2D GetAspectRatioInset(ULocalPlayer* LocalPlayer) const;
+	bool GetNormalizeRect(ULocalPlayer* LocalPlayer, FVector2D& OutPosition, FVector2D& OutSize) const;
 
 	void UpdateWindowTitleBar();
 	void UpdateWindowTitleBarVisibility();
@@ -178,7 +179,7 @@ private:
 private:
 	FGeometry CachedGeometry;
 
-	TMap < ULocalPlayer*, TSharedPtr<FPlayerLayer> > PlayerLayers;
+	TMap<FObjectKey, TSharedPtr<FPlayerLayer>> PlayerLayers;
 
 	TAttribute<FSceneViewport*> SceneViewport;
 	TSharedPtr<SVerticalBox> WidgetHost;
