- Re-enable AudioComponents loop revival but require opt in boolean to be set in content and move off of TickComponent.
- Rename enum to 'retrigger' to match new boolean and avoid semantic collision (replay collides with notion of game 'replays')
- Set tick rate to 100ms and scale to 3s over distance by factor of ratio Distance/MaxDistance.  Helps amortize ticks better (assuming a fair distance distribution of objects) and throttle perf

diff --git a/Engine/Source/Runtime/Engine/Classes/Components/AudioComponent.h b/Engine/Source/Runtime/Engine/Classes/Components/AudioComponent.h
--- a/Engine/Source/Runtime/Engine/Classes/Components/AudioComponent.h
+++ b/Engine/Source/Runtime/Engine/Classes/Components/AudioComponent.h
@@ -162,6 +162,11 @@ class ENGINE_API UAudioComponent : public USceneComponent
 	UPROPERTY(EditAnywhere, BlueprintReadWrite, Category = Sound)
 	uint8 bSuppressSubtitles:1;
 
+	/** If true, re-triggers a loop if a listener comes back within range.
+	  * Only works for attenuated, non-virtualized sounds. */
+	UPROPERTY(EditAnywhere, BlueprintReadWrite, Category = Attenuation)
+	uint8 bRetriggerLoopOnProximity : 1;
+
 	/** Whether this audio component is previewing a sound */
 	uint8 bPreviewComponent:1;
 
@@ -189,24 +194,22 @@ class ENGINE_API UAudioComponent : public USceneComponent
 	/** Whether or not this audio component has been paused */
 	uint8 bIsPaused:1;
 
-	/** How AudioComponent is managing playback based on whether a listener is within 
-	  * proximity (sound's MaxDistance) or not. */
-	enum class EReplayWhenInAudibleRange : uint8
+	/** How playback management is based on whether a listener is within
+	  * proximity (sound's MaxDistance). */
+	enum class ERetriggerWhenInAudibleRange : uint8
 	{
 		/** Sound lifetime is not managed by listener proximity. */
 		Disabled,
 
-		/** Loop replay is enabled and disable has been requested. */
+		/** Loop re-trigger is enabled and disable has been requested. */
 		DisableRequested,
 
 		/** Sound lifetime is managed by listener proximity (in case of sound being looping and non-virtualized) */
 		Enabled,
-
-		Count
 	};
 
 	/** Whether or not the audio component can restart when listener is in proximity */
-	EReplayWhenInAudibleRange ReplayWhenInAudibleRange;
+	ERetriggerWhenInAudibleRange RetriggerWhenInAudibleRange;
 
 	/**
 	* True if we should automatically attach to AutoAttachParent when Played, and detach from our parent when playback is completed.
@@ -498,8 +501,11 @@ private:
 	/** Called by the ActiveSound to inform the component that playback is finished */
 	void PlaybackCompleted(bool bFailedToStart);
 
-	/** Whether or not the sound is audible */
-	bool IsInAudibleRange() const;
+	/** Whether or not the sound is audible. */
+	bool IsInAudibleRange(float* OutMaxDistance) const;
+
+	/** Returns true if in re-trigger range, false if not. */
+	bool GetRetriggerRate(float* OutRetriggerRate) const;
 
 public:
 
@@ -527,7 +533,6 @@ public:
 	virtual void OnUnregister() override;
 	virtual const UObject* AdditionalStatObject() const override;
 	virtual bool IsReadyForOwnerToAutoDestroy() const override;
-	virtual void TickComponent(float DeltaTime, ELevelTick TickType, FActorComponentTickFunction* TickFunction) override;
 	//~ End ActorComponent Interface.
 
 	/** Returns a pointer to the attenuation settings to be used (if any) for this audio component dependent on the SoundAttenuation asset or overrides set. */
@@ -539,6 +544,8 @@ public:
 	/** Collects the various attenuation shapes that may be applied to the sound played by the audio component for visualization in the editor or via the in game debug visualization. */
 	void CollectAttenuationShapesForVisualization(TMultiMap<EAttenuationShape::Type, FBaseAttenuationSettings::AttenuationShapeDetails>& ShapeDetailsMap) const;
 
+	void OnUpdateProximityRetrigger(float DeltaTime);
+
 	/** Returns the active audio device to use for this component based on whether or not the component is playing in a world. */
 	FAudioDevice* GetAudioDevice() const;
 
@@ -573,6 +580,9 @@ private:
 
 	uint64 AudioComponentID;
 
+	float RetriggerTimeSinceLastUpdate;
+	float RetriggerUpdateInterval;
+
 	/** Saved relative transform before auto attachement. Used during detachment to restore the transform if we had automatically attached. */
 	FVector SavedAutoAttachRelativeLocation;
 	FRotator SavedAutoAttachRelativeRotation;
diff --git a/Engine/Source/Runtime/Engine/Private/AudioDevice.cpp b/Engine/Source/Runtime/Engine/Private/AudioDevice.cpp
--- a/Engine/Source/Runtime/Engine/Private/AudioDevice.cpp
+++ b/Engine/Source/Runtime/Engine/Private/AudioDevice.cpp
@@ -617,6 +617,8 @@ void FAudioDevice::Teardown()
 	}
 
 	PluginListeners.Reset();
+
+	ProximityRetriggerComponents.Reset();
 }
 
 void FAudioDevice::Suspend(bool bGameTicking)
@@ -689,6 +691,24 @@ void FAudioDevice::AddReferencedObjects(FReferenceCollector& Collector)
 	Collector.AddReferencedObjects(PluginSettingsObjects);
 }
 
+void FAudioDevice::RegisterProximityRetriggeringAudioComponent(UAudioComponent& Component)
+{
+	check(IsInGameThread());
+	ProximityRetriggerComponents.AddUnique(&Component);
+}
+
+const TArray<FAudioComponentPtr>& FAudioDevice::GetProximityRetriggerComponents() const
+{
+	check(IsInGameThread());
+	return ProximityRetriggerComponents;
+}
+
+void FAudioDevice::UnregisterProximityRetriggeringAudioComponent(UAudioComponent& Component)
+{
+	check(IsInGameThread());
+	ProximityRetriggerComponents.RemoveSwap(&Component);
+}
+
 void FAudioDevice::ResetInterpolation()
 {
 	check(IsInAudioThread());
@@ -3682,6 +3702,12 @@ void FAudioDevice::Update(bool bGameTicking)
 {
 	LLM_SCOPE(ELLMTag::AudioMisc);
 
+	// Must be separate from subsequent call as editor is considered game thread
+	if (IsInGameThread())
+	{
+		UpdateProximityRetriggerComponents();
+	}
+
 	if (!IsInAudioThread())
 	{
 		check(IsInGameThread());
@@ -5568,6 +5594,24 @@ float FAudioDevice::GetGameDeltaTime() const
 	return FMath::Min(DeltaTime, 0.5f);
 }
 
+void FAudioDevice::UpdateProximityRetriggerComponents()
+{
+	for (int32 i = ProximityRetriggerComponents.Num() - 1; i >= 0; --i)
+	{
+		FAudioComponentPtr& ComponentPtr = ProximityRetriggerComponents[i];
+		if (ComponentPtr.IsStale())
+		{
+			UE_LOG(LogAudio, Warning, TEXT("AudioComponent destroyed but not removed from retrigger proximity."));
+			ProximityRetriggerComponents.RemoveAtSwap(i, 1, false);
+		}
+		else if (ComponentPtr.IsValid())
+		{
+			ComponentPtr->OnUpdateProximityRetrigger(GetDeviceDeltaTime());
+		}
+	}
+}
+
+
 #if WITH_EDITOR
 void FAudioDevice::OnBeginPIE(const bool bIsSimulating)
 {
diff --git a/Engine/Source/Runtime/Engine/Private/AudioDeviceManager.cpp b/Engine/Source/Runtime/Engine/Private/AudioDeviceManager.cpp
--- a/Engine/Source/Runtime/Engine/Private/AudioDeviceManager.cpp
+++ b/Engine/Source/Runtime/Engine/Private/AudioDeviceManager.cpp
@@ -166,6 +166,8 @@ void FAudioDeviceManager::ToggleAudioMixer()
 						}
 					}
 
+					TArray<FAudioComponentPtr> RetriggerComponents = AudioDevice->GetProximityRetriggerComponents();
+
 					// Tear it down and delete the old audio device. This does a bunch of cleanup.
 					AudioDevice->Teardown();
 					delete AudioDevice;
@@ -182,6 +184,12 @@ void FAudioDeviceManager::ToggleAudioMixer()
 						AudioDevice->SetMaxChannels(QualityLevelMaxChannels);
 					}
 
+					// Re-register the new audio device's re-trigger components
+					for (FAudioComponentPtr Component : RetriggerComponents)
+					{
+						AudioDevice->RegisterProximityRetriggeringAudioComponent(*Component);
+					}
+
 					// Transfer the sound mix modifiers to the new audio engine
 					AudioDevice->SetSoundMixModifiers(MixModifiers, PrevPassiveSoundMixModifiers, BaseSoundMix);
 					// Setup the mute state of the audio device to be the same that it was
diff --git a/Engine/Source/Runtime/Engine/Private/Components/AudioComponent.cpp b/Engine/Source/Runtime/Engine/Private/Components/AudioComponent.cpp
--- a/Engine/Source/Runtime/Engine/Private/Components/AudioComponent.cpp
+++ b/Engine/Source/Runtime/Engine/Private/Components/AudioComponent.cpp
@@ -21,7 +21,28 @@ FAutoConsoleVariableRef CVarBackedAnalysisTimeShift(
 	TEXT("Value: The time in seconds to shift the timeline."),
 	ECVF_Default);
 
-#define REPLAY_LOOPS_ENABLED 0
+static int32 bRetriggerLoopOnProximityEnabledCVar = 1;
+FAutoConsoleVariableRef CVarRetriggerLoopOnProximityEnabled(
+	TEXT("au.RetriggerLoopOnProximity"),
+	bRetriggerLoopOnProximityEnabledCVar,
+	TEXT("Enables or disables whether retriggering is supported on AudioComponent loops.\n")
+	TEXT("Value: The time in seconds to shift the timeline."),
+	ECVF_Default);
+
+static float RetriggerLoopOnProximityDistanceCVar = 15000.0f;
+FAutoConsoleVariableRef CVarRetriggerLoopOnProximityDistance(
+	TEXT("au.RetriggerLoopOnProximityDistance"),
+	RetriggerLoopOnProximityDistanceCVar,
+	TEXT("Sets retrigger loop proximity distance.\n")
+	TEXT("Value: The distance to scale update rate based on distance to listener."),
+	ECVF_Default);
+
+namespace
+{
+	float AudioComponentRetriggerUpdateRateMin = 0.1f;
+	float AudioComponentRetriggerUpdateRateMax = 3.0f;
+} // namespace <>
+
 
 /*-----------------------------------------------------------------------------
 UAudioComponent implementation.
@@ -50,7 +71,12 @@ UAudioComponent::UAudioComponent(const FObjectInitializer& ObjectInitializer)
 	bOverrideSubtitlePriority = false;
 	bIsPreviewSound = false;
 	bIsPaused = false;
-	ReplayWhenInAudibleRange = EReplayWhenInAudibleRange::Disabled;
+
+	bRetriggerLoopOnProximity = false;
+	RetriggerWhenInAudibleRange = ERetriggerWhenInAudibleRange::Disabled;
+	RetriggerTimeSinceLastUpdate = 0.0f;
+	RetriggerUpdateInterval = 0.1f;
+
 	Priority = 1.f;
 	SubtitlePriority = DEFAULT_SUBTITLE_PRIORITY;
 	PitchMultiplier = 1.f;
@@ -69,15 +95,11 @@ UAudioComponent::UAudioComponent(const FObjectInitializer& ObjectInitializer)
 	AudioDeviceHandle = INDEX_NONE;
 	AudioComponentID = FPlatformAtomics::InterlockedIncrement(reinterpret_cast<volatile int64*>(&AudioComponentIDCounter));
 
-	// TODO: Consider only putting played/active components in to the map
-	FScopeLock Lock(&AudioIDToComponentMapLock);
-	AudioIDToComponentMap.Add(AudioComponentID, this);
-
-	PrimaryComponentTick.bCanEverTick = true;
-	PrimaryComponentTick.TickGroup = TG_DuringPhysics;
-	PrimaryComponentTick.bStartWithTickEnabled = false;
-	PrimaryComponentTick.SetTickFunctionEnable(false);
-	bTickInEditor = true;
+	{
+		// TODO: Consider only putting played/active components in to the map
+		FScopeLock Lock(&AudioIDToComponentMapLock);
+		AudioIDToComponentMap.Add(AudioComponentID, this);
+	}
 }
 
 UAudioComponent* UAudioComponent::GetAudioComponentFromID(uint64 AudioComponentID)
@@ -99,6 +121,14 @@ void UAudioComponent::BeginDestroy()
 		Stop();
 	}
 
+	if (RetriggerWhenInAudibleRange == ERetriggerWhenInAudibleRange::Enabled)
+	{
+		if (FAudioDevice* AudioDevice = GetAudioDevice())
+		{
+			AudioDevice->UnregisterProximityRetriggeringAudioComponent(*this);
+		}
+	}
+
 	FScopeLock Lock(&AudioIDToComponentMapLock);
 	AudioIDToComponentMap.Remove(AudioComponentID);
 }
@@ -261,35 +291,35 @@ bool UAudioComponent::IsReadyForOwnerToAutoDestroy() const
 	return !IsPlaying();
 }
 
-void UAudioComponent::TickComponent(float DeltaTime, ELevelTick TickType, FActorComponentTickFunction* TickFunction)
+void UAudioComponent::OnUpdateProximityRetrigger(float DeltaTime)
 {
-	Super::TickComponent(DeltaTime, TickType, TickFunction);
+	if (RetriggerUpdateInterval > 0.0f)
+	{
+		RetriggerTimeSinceLastUpdate += DeltaTime;
+		if (RetriggerUpdateInterval > RetriggerTimeSinceLastUpdate)
+		{
+			return;
+		}
+		RetriggerTimeSinceLastUpdate = 0.0f;
+	}
 
-	const bool bIsReplayEnabled = ReplayWhenInAudibleRange == EReplayWhenInAudibleRange::Enabled;
-	if (bIsReplayEnabled)
+	if (RetriggerWhenInAudibleRange == ERetriggerWhenInAudibleRange::Enabled)
 	{
-		if (ActiveCount > 0)
+		bool bIsInAudibleRange = GetRetriggerRate(&RetriggerUpdateInterval);
+		if (bIsActive)
 		{
-			if (!IsInAudibleRange())
+			if (!bIsInAudibleRange)
 			{
-				// If this is an auto destroy component we need to prevent it from being auto-destroyed since we're really just restarting it
-				bool bCurrentAutoDestroy = bAutoDestroy;
-				bAutoDestroy = false;
 				StopInternal();
-				bAutoDestroy = bCurrentAutoDestroy;
 			}
 		}
-		else
+		else if (bIsInAudibleRange)
 		{
-			if (IsInAudibleRange())
-			{
-				Play(); // Play refreshes ReplayWhenInAudibleRange internally
-			}
+			Play(); // Play refreshes RetriggerWhenInAudibleRange internally
 		}
 	}
 }
 
-
 void UAudioComponent::OnUpdateTransform(EUpdateTransformFlags UpdateTransformFlags, ETeleportType Teleport)
 {
 	Super::OnUpdateTransform(UpdateTransformFlags, Teleport);
@@ -341,7 +371,7 @@ void UAudioComponent::CancelAutoAttachment(bool bDetachFromParent)
 	}
 }
 
-bool UAudioComponent::IsInAudibleRange() const
+bool UAudioComponent::IsInAudibleRange(float* OutMaxDistance) const
 {
 	FAudioDevice* AudioDevice = GetAudioDevice();
 	if (!AudioDevice)
@@ -353,11 +383,41 @@ bool UAudioComponent::IsInAudibleRange() const
 	float FocusFactor = 0.0f;
 	const FVector Location = GetComponentTransform().GetLocation();
 	const FSoundAttenuationSettings* AttenuationSettingsToApply = bAllowSpatialization ? GetAttenuationSettingsToApply() : nullptr;
-
 	AudioDevice->GetMaxDistanceAndFocusFactor(Sound, GetWorld(), Location, AttenuationSettingsToApply, MaxDistance, FocusFactor);
+	
+	if (OutMaxDistance)
+	{
+		*OutMaxDistance = MaxDistance;
+	}
+
 	return AudioDevice->SoundIsAudible(Sound, GetWorld(), Location, AttenuationSettingsToApply, MaxDistance, FocusFactor);
 }
 
+bool UAudioComponent::GetRetriggerRate(float* OutRetriggerRate) const
+{
+	check(OutRetriggerRate);
+
+	FAudioDevice* AudioDevice = GetAudioDevice();
+	if (!AudioDevice)
+	{
+		*OutRetriggerRate = AudioComponentRetriggerUpdateRateMin;
+		return false;
+	}
+
+	float MaxDistance = 0.0f;
+	if (IsInAudibleRange(&MaxDistance))
+	{
+		*OutRetriggerRate = AudioComponentRetriggerUpdateRateMin;
+		return true;
+	}
+
+	const FVector Location = GetComponentTransform().GetLocation();
+	float DistanceToListener = AudioDevice->GetDistanceToNearestListener(Location);
+	const float DistanceRatio = (DistanceToListener - MaxDistance) / FMath::Max(RetriggerLoopOnProximityDistanceCVar, 1.0f);
+	*OutRetriggerRate = FMath::Lerp(AudioComponentRetriggerUpdateRateMin, AudioComponentRetriggerUpdateRateMax, FMath::Min(DistanceRatio, 1.0f));
+	return false;
+}
+
 void UAudioComponent::Play(float StartTime)
 {
 	PlayInternal(StartTime);
@@ -520,16 +580,36 @@ void UAudioComponent::PlayInternal(const float StartTime, const float FadeInDura
 
 			AudioDevice->AddNewActiveSound(NewActiveSound);
 
-#if REPLAY_LOOPS_ENABLED
-			// Only allow replay if this is the first sound to play
-			const bool bCanVirtualize = Sound && Sound->IsVirtualizeWhenSilent() && AudioDevice->VirtualSoundsEnabled();
-			if (!bIsPreviewSound && NewActiveSound.bHasAttenuationSettings && !bCanVirtualize && !bAutoDestroy)
+			if (World && World->bAllowAudioPlayback)
 			{
-				const bool bLooping = NewActiveSound.IsLooping();
-				ReplayWhenInAudibleRange = bLooping ? EReplayWhenInAudibleRange::Enabled : EReplayWhenInAudibleRange::Disabled;
-				PrimaryComponentTick.SetTickFunctionEnable(bLooping);
+				if (bRetriggerLoopOnProximity && bRetriggerLoopOnProximityEnabledCVar)
+				{
+					static FString RetriggerFailedStr = TEXT("Re-trigger loop in listener proximity failed: Sound ");
+					if (bAutoDestroy)
+					{
+						UE_LOG(LogAudio, Warning, TEXT("%s'%s' set to AutoDestroy."), *RetriggerFailedStr, Sound ? *Sound->GetName() : TEXT("N/A"));
+					}
+					else if (Sound && Sound->IsVirtualizeWhenSilent() && AudioDevice->VirtualSoundsEnabled())
+					{
+						UE_LOG(LogAudio, Warning, TEXT("%s'%s' set to 'PlayWhenSilent'."), *RetriggerFailedStr, Sound ? *Sound->GetName() : TEXT("N/A"));
+					}
+					else if (!NewActiveSound.bHasAttenuationSettings)
+					{
+						UE_LOG(LogAudio, Warning, TEXT("%s'%s' not using attenuation."), *RetriggerFailedStr, Sound ? *Sound->GetName() : TEXT("N/A"));
+					}
+					else if (!bIsPreviewSound)
+					{
+						const bool bLooping = NewActiveSound.IsLooping();
+						RetriggerWhenInAudibleRange = bLooping ? ERetriggerWhenInAudibleRange::Enabled : ERetriggerWhenInAudibleRange::Disabled;
+
+						if (bLooping)
+						{
+							AudioDevice->RegisterProximityRetriggeringAudioComponent(*this);
+						}
+					}
+				}
 			}
-#endif // REPLAY_LOOPS_ENABLED
+
 			bIsActive = true;
 		}
 	}
@@ -569,18 +649,16 @@ void UAudioComponent::FadeOut( float FadeOutDuration, float FadeVolumeLevel )
 	{
 		if (FadeOutDuration > 0.0f)
 		{
-#if REPLAY_LOOPS_ENABLED
-			if (FadeVolumeLevel == 0.0f)
-			{
-				ReplayWhenInAudibleRange = EReplayWhenInAudibleRange::DisableRequested;
-				PrimaryComponentTick.SetTickFunctionEnable(false);
-			}
-#endif // REPLAY_LOOPS_ENABLED
-
 			if (FAudioDevice* AudioDevice = GetAudioDevice())
 			{
 				DECLARE_CYCLE_STAT(TEXT("FAudioThreadTask.FadeOut"), STAT_AudioFadeOut, STATGROUP_AudioThreadCommands);
 
+				if (RetriggerWhenInAudibleRange == ERetriggerWhenInAudibleRange::Enabled && FadeVolumeLevel == 0.0f)
+				{
+					RetriggerWhenInAudibleRange = ERetriggerWhenInAudibleRange::DisableRequested;
+					AudioDevice->UnregisterProximityRetriggeringAudioComponent(*this);
+				}
+
 				const uint64 MyAudioComponentID = AudioComponentID;
 				FAudioThread::RunCommandOnAudioThread([AudioDevice, MyAudioComponentID, FadeOutDuration, FadeVolumeLevel]()
 				{
@@ -599,6 +677,15 @@ void UAudioComponent::FadeOut( float FadeOutDuration, float FadeVolumeLevel )
 			Stop();
 		}
 	}
+	// Calling FadeOut to 0.0f is effectively the same as calling stop if sound is out-of-bounds and is set to retrigger, so kill loop
+	else if (RetriggerWhenInAudibleRange == ERetriggerWhenInAudibleRange::Enabled && FadeVolumeLevel == 0.0f)
+	{
+		if (FAudioDevice* AudioDevice = GetAudioDevice())
+		{
+			RetriggerWhenInAudibleRange = ERetriggerWhenInAudibleRange::Disabled;
+			AudioDevice->UnregisterProximityRetriggeringAudioComponent(*this);
+		}
+	}
 }
 
 void UAudioComponent::AdjustVolume( float AdjustVolumeDuration, float AdjustVolumeLevel )
@@ -607,13 +694,11 @@ void UAudioComponent::AdjustVolume( float AdjustVolumeDuration, float AdjustVolu
 	{
 		if (FAudioDevice* AudioDevice = GetAudioDevice())
 		{
-#if REPLAY_LOOPS_ENABLED
-			if (AdjustVolumeLevel > 0.0f && ReplayWhenInAudibleRange == EReplayWhenInAudibleRange::DisableRequested)
+			if (AdjustVolumeLevel > 0.0f && RetriggerWhenInAudibleRange == ERetriggerWhenInAudibleRange::DisableRequested)
 			{
-				ReplayWhenInAudibleRange = EReplayWhenInAudibleRange::Enabled;
-				PrimaryComponentTick.SetTickFunctionEnable(true);
+				RetriggerWhenInAudibleRange = ERetriggerWhenInAudibleRange::Enabled;
+				AudioDevice->RegisterProximityRetriggeringAudioComponent(*this);
 			}
-#endif // REPLAY_LOOPS_ENABLED
 
 			DECLARE_CYCLE_STAT(TEXT("FAudioThreadTask.AdjustVolume"), STAT_AudioAdjustVolume, STATGROUP_AudioThreadCommands);
 
@@ -643,10 +728,14 @@ void UAudioComponent::AdjustVolume( float AdjustVolumeDuration, float AdjustVolu
 
 void UAudioComponent::Stop()
 {
-#if REPLAY_LOOPS_ENABLED
-	ReplayWhenInAudibleRange = EReplayWhenInAudibleRange::DisableRequested;
-	PrimaryComponentTick.SetTickFunctionEnable(false);
-#endif // REPLAY_LOOPS_ENABLED
+	if (FAudioDevice* AudioDevice = GetAudioDevice())
+	{
+		if (RetriggerWhenInAudibleRange == ERetriggerWhenInAudibleRange::Enabled)
+		{
+			RetriggerWhenInAudibleRange = ERetriggerWhenInAudibleRange::Disabled;
+			AudioDevice->UnregisterProximityRetriggeringAudioComponent(*this);
+		}
+	}
 
 	StopInternal();
 }
@@ -724,8 +813,20 @@ void UAudioComponent::PlaybackCompleted(bool bFailedToStart)
 
 	if (bFailedToStart)
 	{
-		ReplayWhenInAudibleRange = EReplayWhenInAudibleRange::Disabled;
-		PrimaryComponentTick.SetTickFunctionEnable(false);
+		RetriggerWhenInAudibleRange = ERetriggerWhenInAudibleRange::Disabled;
+		if (FAudioDevice* AudioDevice = GetAudioDevice())
+		{
+			AudioDevice->UnregisterProximityRetriggeringAudioComponent(*this);
+		}
+	}
+
+	if (RetriggerWhenInAudibleRange == ERetriggerWhenInAudibleRange::DisableRequested)
+	{
+		RetriggerWhenInAudibleRange = ERetriggerWhenInAudibleRange::Disabled;
+		if (FAudioDevice* AudioDevice = GetAudioDevice())
+		{
+			AudioDevice->UnregisterProximityRetriggeringAudioComponent(*this);
+		}
 	}
 
 	if (!bFailedToStart && GetWorld() != nullptr && (OnAudioFinished.IsBound() || OnAudioFinishedNative.IsBound()))
diff --git a/Engine/Source/Runtime/Engine/Public/AudioDevice.h b/Engine/Source/Runtime/Engine/Public/AudioDevice.h
--- a/Engine/Source/Runtime/Engine/Public/AudioDevice.h
+++ b/Engine/Source/Runtime/Engine/Public/AudioDevice.h
@@ -299,6 +299,7 @@ struct FSoundMixClassOverride
 };
 
 typedef TMap<USoundClass*, FSoundMixClassOverride> FSoundMixClassOverrideMap;
+typedef TWeakObjectPtr<UAudioComponent> FAudioComponentPtr;
 
 struct FActivatedReverb
 {
@@ -568,6 +569,21 @@ public:
 	 */
 	void AddReferencedObjects(FReferenceCollector& Collector);
 
+	/**
+	 * Register an audio component to re-trigger based on re-entering proximity (MaxDistance as provided by owned sound instance).
+	 */
+	void RegisterProximityRetriggeringAudioComponent(UAudioComponent& Component);
+
+	/**
+	 * Returns set of AudioComponents currently set to re-trigger on listener proximity
+	 */
+	const TArray<FAudioComponentPtr>& GetProximityRetriggerComponents() const;
+
+	/**
+	 * Unregister an audio component to re-trigger based on re-entering proximity (MaxDistance as provided by owned sound instance).
+	 */
+	void UnregisterProximityRetriggeringAudioComponent(UAudioComponent& Component);
+
 	/**
 	 * Iterate over the active AudioComponents for wave instances that could be playing.
 	 *
@@ -706,6 +722,7 @@ public:
 		FVector Location;
 
 		void CommonInit();
+		void UpdateProximityRetriggerComponents(float DeltaTime);
 
 		friend class FAudioDevice;
 	};
@@ -1482,6 +1499,8 @@ public:
 	/** Returns the game's delta time */
 	float GetGameDeltaTime() const;
 
+	void UpdateProximityRetriggerComponents();
+
 	/** Sets the update delta time for the audio frame */
 	virtual void UpdateDeviceDeltaTime()
 	{
@@ -1802,6 +1821,8 @@ private:
 	TArray<USoundWave*> PrecachingSoundWaves;
 	TArray<FWaveInstance*> ActiveWaveInstances;
 
+	TArray<FAudioComponentPtr> ProximityRetriggerComponents;
+
 	/** Cached copy of sound class adjusters array. Cached to avoid allocating every frame. */
 	TArray<FSoundClassAdjuster> SoundClassAdjustersCopy;
 
