- Handle null Template inside CanSkipTickDUeToVisibility

diff --git a/Engine/Source/Runtime/Engine/Private/Particles/ParticleComponents.cpp b/Engine/Source/Runtime/Engine/Private/Particles/ParticleComponents.cpp
--- a/Engine/Source/Runtime/Engine/Private/Particles/ParticleComponents.cpp
+++ b/Engine/Source/Runtime/Engine/Private/Particles/ParticleComponents.cpp
@@ -3508,7 +3508,7 @@ bool UParticleSystemComponent::CanBeOccluded()const
 
 bool UParticleSystemComponent::CanSkipTickDueToVisibility()
 {
-	if (Template->IsLooping() && CanConsiderInvisible() && !bWasDeactivated)
+	if (Template && Template->IsLooping() && CanConsiderInvisible() && !bWasDeactivated)
 	{
 		SCOPE_CYCLE_COUNTER(STAT_UParticleSystemComponent_LOD_Inactive);
 		bForcedInActive = true;
