Do not create dynamic point/spot lights render proxies if they are not supported on mobile

diff --git a/Engine/Source/Runtime/Core/Private/HAL/ConsoleManager.cpp b/Engine/Source/Runtime/Core/Private/HAL/ConsoleManager.cpp
--- a/Engine/Source/Runtime/Core/Private/HAL/ConsoleManager.cpp
+++ b/Engine/Source/Runtime/Core/Private/HAL/ConsoleManager.cpp
@@ -1918,7 +1918,7 @@ static TAutoConsoleVariable<int32> CVarMobileNumDynamicPointLights(
 	TEXT("r.MobileNumDynamicPointLights"),
 	4,
 	TEXT("The number of dynamic point lights to support on mobile devices. Setting this to 0 for games which do not require dynamic point lights will reduce the number of shaders generated."), 
-	ECVF_RenderThreadSafe);
+	ECVF_RenderThreadSafe | ECVF_ReadOnly);
 
 static TAutoConsoleVariable<int32> CVarMobileDynamicPointLightsUseStaticBranch(
 	TEXT("r.MobileDynamicPointLightsUseStaticBranch"),
diff --git a/Engine/Source/Runtime/Engine/Private/Components/PointLightComponent.cpp b/Engine/Source/Runtime/Engine/Private/Components/PointLightComponent.cpp
--- a/Engine/Source/Runtime/Engine/Private/Components/PointLightComponent.cpp
+++ b/Engine/Source/Runtime/Engine/Private/Components/PointLightComponent.cpp
@@ -95,9 +95,25 @@ UPointLightComponent::UPointLightComponent(const FObjectInitializer& ObjectIniti
 	bUseInverseSquaredFalloff = true;
 }
 
+static bool IsPointLightSupported(const UPointLightComponent* InLight)
+{
+	if (GMaxRHIFeatureLevel <= ERHIFeatureLevel::ES3_1 && InLight->IsMovable())
+	{
+		// if project does not support dynamic point lights on mobile do not add them to the renderer 
+		static auto* CVarPointLights = IConsoleManager::Get().FindTConsoleVariableDataInt(TEXT("r.MobileNumDynamicPointLights"));
+		const bool bPointLights = CVarPointLights->GetValueOnAnyThread() > 0;
+		return bPointLights;
+	}
+	return true;
+}
+
 FLightSceneProxy* UPointLightComponent::CreateSceneProxy() const
 {
-	return new FPointLightSceneProxy(this);
+	if (IsPointLightSupported(this))
+	{
+		return new FPointLightSceneProxy(this);
+	}
+	return nullptr;
 }
 
 void UPointLightComponent::SetLightFalloffExponent(float NewLightFalloffExponent)
diff --git a/Engine/Source/Runtime/Engine/Private/Components/SpotLightComponent.cpp b/Engine/Source/Runtime/Engine/Private/Components/SpotLightComponent.cpp
--- a/Engine/Source/Runtime/Engine/Private/Components/SpotLightComponent.cpp
+++ b/Engine/Source/Runtime/Engine/Private/Components/SpotLightComponent.cpp
@@ -268,9 +268,25 @@ void USpotLightComponent::SetLightBrightness(float InBrightness)
 //	}
 //}
 
+static bool IsSpotLightSupported(const USpotLightComponent* InLight)
+{
+	if (GMaxRHIFeatureLevel <= ERHIFeatureLevel::ES3_1 && InLight->IsMovable())
+	{
+		// if project does not support dynamic point/spot lights on mobile do not add them to the renderer 
+		static auto* CVarPointLights = IConsoleManager::Get().FindTConsoleVariableDataInt(TEXT("r.MobileNumDynamicPointLights"));
+		const bool bPointLights = CVarPointLights->GetValueOnAnyThread() > 0;
+		return bPointLights;
+	}
+	return true;
+}
+
 FLightSceneProxy* USpotLightComponent::CreateSceneProxy() const
 {
-	return new FSpotLightSceneProxy(this);
+	if (IsSpotLightSupported(this))
+	{
+		return new FSpotLightSceneProxy(this);
+	}
+	return nullptr;
 }
 
 FSphere USpotLightComponent::GetBoundingSphere() const
