Added PIE orbit for specific hitpoint, fixed orbit issue for actors with non-visible components, orbit forward motion now stops at pivot rather than passing through.

diff --git a/Engine/Source/Runtime/Engine/Classes/Engine/DebugCameraController.h b/Engine/Source/Runtime/Engine/Classes/Engine/DebugCameraController.h
--- a/Engine/Source/Runtime/Engine/Classes/Engine/DebugCameraController.h
+++ b/Engine/Source/Runtime/Engine/Classes/Engine/DebugCameraController.h
@@ -36,6 +36,10 @@ class ENGINE_API ADebugCameraController
 	UPROPERTY()
 	uint32 bIsOrbitingSelectedActor : 1;
 
+	/** When orbiting, true if using actor center as pivot, false if using last selected hitpoint */
+	UPROPERTY()
+	uint32 bOrbitPivotUseCenter:1;
+
 	/** Visualizes the frustum of the camera */
 	UPROPERTY()
 	class UDrawFrustumComponent* DrawFrustum;
@@ -78,11 +82,20 @@ class ENGINE_API ADebugCameraController
 	 */
 	virtual void UpdateRotation(float DeltaTime) override;
 
+	/** Updates the rotation and location of player when orbiting */
+	void UpdateRotationForOrbit(float DeltaTime);
+
 	/** Gets pivot to use when orbiting */
 	bool GetPivotForOrbit(FVector& PivotLocation) const;
 
-	/** Toggles camera orbit using current hit point as pivot */
-	virtual void ToggleOrbit();
+	/** Toggles camera orbit */
+	void ToggleOrbit(bool bOrbitCenter);
+
+	/** Toggles camera orbit center */
+	void ToggleOrbitCenter();
+
+	/** Toggles camera orbit hitpoint */
+	void ToggleOrbitHitPoint();
 
 public:
 
@@ -98,6 +111,10 @@ public:
 	UPROPERTY()
 	class UPrimitiveComponent* SelectedComponent;
 
+	/** Selected hit point */
+	UPROPERTY()
+	FHitResult SelectedHitPoint;
+
 	/** Controller that was active before this was spawned */
 	UPROPERTY()	
 	class APlayerController* OriginalControllerRef;
@@ -202,9 +219,12 @@ private:
 	/** Last position for orbit */
 	FVector LastOrbitPawnLocation;
 
-	/** Current orbit pivot, if pivot is enabled. */
+	/** Current orbit pivot, if orbit is enabled */
 	FVector OrbitPivot;
 
+	/** Current orbit radius, if orbit is enabled */
+	float OrbitRadius;
+
 	void OnTouchBegin(ETouchIndex::Type FingerIndex, FVector Location);
 	void OnTouchEnd(ETouchIndex::Type FingerIndex, FVector Location);
 	void OnFingerMove(ETouchIndex::Type FingerIndex, FVector Location);
diff --git a/Engine/Source/Runtime/Engine/Private/DebugCameraController.cpp b/Engine/Source/Runtime/Engine/Private/DebugCameraController.cpp
--- a/Engine/Source/Runtime/Engine/Private/DebugCameraController.cpp
+++ b/Engine/Source/Runtime/Engine/Private/DebugCameraController.cpp
@@ -23,6 +23,7 @@
 #include "GameFramework/GameStateBase.h"
 
 static const float SPEED_SCALE_ADJUSTMENT = 0.5f;
+static const float MIN_ORBIT_RADIUS = 30.0f;
 
 ADebugCameraController::ADebugCameraController(const FObjectInitializer& ObjectInitializer)
 	: Super(ObjectInitializer)
@@ -48,8 +49,10 @@ ADebugCameraController::ADebugCameraController(const FObjectInitializer& ObjectI
 	SetAsLocalPlayerController();
 
 	bIsOrbitingSelectedActor = false;
+	bOrbitPivotUseCenter = false;
 	LastOrbitPawnLocation = FVector::ZeroVector;
 	OrbitPivot = FVector::ZeroVector;
+	OrbitRadius = MIN_ORBIT_RADIUS;
 }
 
 void InitializeDebugCameraInputBindings()
@@ -69,7 +72,8 @@ void InitializeDebugCameraInputBindings()
 		UPlayerInput::AddEngineDefinedActionMapping(FInputActionKeyMapping("DebugCamera_DecreaseFOV", EKeys::Period));
 		UPlayerInput::AddEngineDefinedActionMapping(FInputActionKeyMapping("DebugCamera_ToggleDisplay", EKeys::BackSpace));
 		UPlayerInput::AddEngineDefinedActionMapping(FInputActionKeyMapping("DebugCamera_FreezeRendering", EKeys::F));
-		UPlayerInput::AddEngineDefinedActionMapping(FInputActionKeyMapping("DebugCamera_Orbit", EKeys::O));
+		UPlayerInput::AddEngineDefinedActionMapping(FInputActionKeyMapping("DebugCamera_OrbitCenter", EKeys::O));
+		UPlayerInput::AddEngineDefinedActionMapping(FInputActionKeyMapping("DebugCamera_OrbitHitPoint", EKeys::O, true));
 
 		UPlayerInput::AddEngineDefinedActionMapping(FInputActionKeyMapping("DebugCamera_Select", EKeys::Gamepad_RightTrigger));
 		UPlayerInput::AddEngineDefinedActionMapping(FInputActionKeyMapping("DebugCamera_IncreaseSpeed", EKeys::Gamepad_RightShoulder));
@@ -97,7 +101,8 @@ void ADebugCameraController::SetupInputComponent()
 
 	InputComponent->BindAction("DebugCamera_ToggleDisplay", IE_Pressed, this, &ADebugCameraController::ToggleDisplay);
 	InputComponent->BindAction("DebugCamera_FreezeRendering", IE_Pressed, this, &ADebugCameraController::ToggleFreezeRendering);
-	InputComponent->BindAction("DebugCamera_Orbit", IE_Pressed, this, &ADebugCameraController::ToggleOrbit);
+	InputComponent->BindAction("DebugCamera_OrbitCenter", IE_Pressed, this, &ADebugCameraController::ToggleOrbitCenter);
+	InputComponent->BindAction("DebugCamera_OrbitHitPoint", IE_Pressed, this, &ADebugCameraController::ToggleOrbitHitPoint);
 
 	InputComponent->BindTouch(IE_Pressed, this, &ADebugCameraController::OnTouchBegin);
 	InputComponent->BindTouch(IE_Released, this, &ADebugCameraController::OnTouchEnd);
@@ -146,6 +151,7 @@ void ADebugCameraController::Select( FHitResult const& Hit )
 	// store selection
 	SelectedActor = Hit.GetActor();
 	SelectedComponent = Hit.Component.Get();
+	SelectedHitPoint = Hit;
 
 	//BP Event
 	ReceiveOnActorSelected(SelectedActor, Hit.ImpactPoint, Hit.ImpactNormal, Hit);
@@ -392,116 +398,123 @@ void ADebugCameraController::UpdateRotation(float DeltaTime)
 {
 	if (bIsOrbitingSelectedActor)
 	{
-		APawn* const CurrentPawn = GetPawnOrSpectator();
-		if (CurrentPawn)
+		UpdateRotationForOrbit(DeltaTime);
+	}
+	else
+	{
+		Super::UpdateRotation(DeltaTime);
+	}
+}
+
+void ADebugCameraController::UpdateRotationForOrbit(float DeltaTime)
+{
+	APawn* const CurrentPawn = GetPawnOrSpectator();
+
+	if (bIsOrbitingSelectedActor && CurrentPawn)
+	{
+		bool bUpdatePawn = false;
+		FRotator ViewRotation = GetControlRotation();
+
+		// Handle rotation input
+		if (!RotationInput.IsZero())
 		{
-			bool bUpdatePawn = false;
-			FRotator ViewRotation;
-			float ViewLength = 0.0f;
+			ViewRotation += RotationInput;
 
-			if (!RotationInput.IsZero())
+			FVector Axis;
+			float Angle;
+			FVector OppositeViewVector = -1.0 * ViewRotation.Vector();
+			FQuat::FindBetween(FVector::UpVector, OppositeViewVector).ToAxisAndAngle(Axis, Angle);
+
+			// Clamp rotation to 10 degrees from Up vector
+			const float MinAngle = PI / 18.f;
+			const float MaxAngle = PI - MinAngle;
+			if (Angle < MinAngle || Angle > MaxAngle)
 			{
-				FVector ViewVector = OrbitPivot - CurrentPawn->GetActorLocation();
-				ViewRotation = ViewVector.ToOrientationRotator() + RotationInput;
-				ViewLength = ViewVector.Size();
-
-				FVector Axis;
-				float Angle;
-				FVector OppositeViewVector = -1.0 * ViewRotation.Vector();
-				FQuat::FindBetween(FVector::UpVector, OppositeViewVector).ToAxisAndAngle(Axis, Angle);
-
-				// Clamp rotation to 10 degrees from Up vector
-				const float MinAngle = PI / 18.f;
-				const float MaxAngle = PI - MinAngle;
-				if (Angle < MinAngle || Angle > MaxAngle)
-				{
-					float AdjustedAngle = FMath::Clamp(Angle, MinAngle, MaxAngle);
-					OppositeViewVector = FQuat(Axis, AdjustedAngle).RotateVector(FVector::UpVector);
-					ViewRotation = (-1.0 * OppositeViewVector).ToOrientationRotator();
-				}
+				float AdjustedAngle = FMath::Clamp(Angle, MinAngle, MaxAngle);
+				OppositeViewVector = FQuat(Axis, AdjustedAngle).RotateVector(FVector::UpVector);
+				ViewRotation = (-1.0 * OppositeViewVector).ToOrientationRotator();
+			}
 
-				bUpdatePawn = true;
-			} 
-			else
+			bUpdatePawn = true;
+		}
+		else // Handle movement input
+		{
+			FVector OrbitPawnLocation = CurrentPawn->GetActorLocation();
+			FVector MoveDelta(OrbitPawnLocation - LastOrbitPawnLocation);
+
+			if (!MoveDelta.IsZero())
 			{
-				FVector OrbitPawnLocation = CurrentPawn->GetActorLocation();
-				FVector MoveDelta(OrbitPawnLocation - LastOrbitPawnLocation);
+				FRotationMatrix ObjectToWorld(ViewRotation);
+				FVector MoveDeltaObj = ObjectToWorld.GetTransposed().TransformVector(MoveDelta);
 
-				if (!MoveDelta.IsZero())
+				// Handle either forward or lateral motion but not both, because small forward
+				// motion deltas while moving laterally cause the distance from pivot to drift
+				if (FMath::IsNearlyZero(MoveDeltaObj.Y, 0.01f))
 				{
-					FRotationMatrix ObjectToWorld(GetControlRotation());
-					FVector MoveDeltaObj = ObjectToWorld.GetTransposed().TransformVector(MoveDelta);
-					FVector ViewVector = OrbitPivot - LastOrbitPawnLocation;
-
-					// Handle either forward or lateral motion but not both, because small 
-					// forward motion deltas while moving laterally cause the distance from pivot
-					// to drift.
-					if (FMath::Abs(MoveDeltaObj.X) > FMath::Abs(MoveDeltaObj.Y))
-					{
-						// Apply forward movement component
-						FVector ForwardDelta = ObjectToWorld.TransformVector(FVector(MoveDeltaObj.X, 0.0f, 0.0f));
-						ViewVector -= ForwardDelta;
-						ViewLength = ViewVector.Size();
-					}
-					else
+					// Clamp delta to avoid flipping to opposite view
+					float MaxDelta = OrbitRadius - MIN_ORBIT_RADIUS;
+					if (MoveDeltaObj.X > MaxDelta)
 					{
-						// Apply lateral movement component, constraining distance from orbit pivot
-						FVector LateralDelta = ObjectToWorld.TransformVector(FVector(0.0f, MoveDeltaObj.Y, 0.0f));
-						ViewLength = ViewVector.Size();
-						ViewVector -= LateralDelta;
+						MoveDeltaObj.X = MaxDelta;
 					}
-					ViewRotation = ViewVector.Rotation();
 
-					bUpdatePawn = true;
+					OrbitRadius -= MoveDeltaObj.X;
+				}
+				else
+				{
+					// Apply lateral movement component, constraining distance from orbit pivot
+					FVector LateralDelta = ObjectToWorld.TransformVector(FVector(0.0f, MoveDeltaObj.Y, 0.0f));
+					ViewRotation = (OrbitPivot - LastOrbitPawnLocation - LateralDelta).ToOrientationRotator();
 				}
+				bUpdatePawn = true;
 			}
+		}
 
-			if (bUpdatePawn)
-			{
-				// Update pawn location
-				FVector OrbitPawnLocation = OrbitPivot - ViewRotation.Vector() * ViewLength;
-				CurrentPawn->SetActorLocation(OrbitPawnLocation);
-
-				// Update pawn rotation
-				SetControlRotation(ViewRotation);
-				CurrentPawn->FaceRotation(ViewRotation, DeltaTime);
-
-				LastOrbitPawnLocation = OrbitPawnLocation;
+		if (bUpdatePawn)
+		{
+			LastOrbitPawnLocation = OrbitPivot - ViewRotation.Vector() * OrbitRadius;
+			CurrentPawn->SetActorLocation(LastOrbitPawnLocation);
 
-				return;
-			}
+			SetControlRotation(ViewRotation);
+			CurrentPawn->FaceRotation(ViewRotation, DeltaTime);
 		}
 	}
-
-	Super::UpdateRotation(DeltaTime);
 }
 
 bool ADebugCameraController::GetPivotForOrbit(FVector& PivotLocation) const
 {
 	if (SelectedActor)
 	{
-		FBox BoundingBox(ForceInit);
-		int32 NumValidComponents = 0;
+		if (bOrbitPivotUseCenter)
+		{
+			FBox BoundingBox(ForceInit);
+			int32 NumValidComponents = 0;
 
-		// Use the center of the bounding box of the current selected actor as the pivot point for orbiting the camera
-		int32 NumSelectedActors = 0;
+			// Use the center of the bounding box of the current selected actor as the pivot point for orbiting the camera
+			int32 NumSelectedActors = 0;
 
-		TInlineComponentArray<UMeshComponent*> MeshComponents(SelectedActor);
+			TInlineComponentArray<UMeshComponent*> MeshComponents(SelectedActor);
 
-		for (int32 ComponentIndex = 0; ComponentIndex < MeshComponents.Num(); ++ComponentIndex)
-		{
-			UMeshComponent* MeshComponent = MeshComponents[ComponentIndex];
+			for (int32 ComponentIndex = 0; ComponentIndex < MeshComponents.Num(); ++ComponentIndex)
+			{
+				UMeshComponent* MeshComponent = MeshComponents[ComponentIndex];
+
+				if (MeshComponent->IsRegistered() && MeshComponent->IsVisible())
+				{
+					BoundingBox += MeshComponent->Bounds.GetBox();
+					++NumValidComponents;
+				}
+			}
 
-			if (MeshComponent->IsRegistered())
+			if (NumValidComponents > 0)
 			{
-				BoundingBox += MeshComponent->Bounds.GetBox();
-				++NumValidComponents;
+				PivotLocation = BoundingBox.GetCenter();
+				return true;
 			}
 		}
-
-		if (NumValidComponents > 0)
+		else
 		{
-			PivotLocation = BoundingBox.GetCenter();
+			PivotLocation = SelectedHitPoint.Location;
 			return true;
 		}
 	}
@@ -509,19 +522,57 @@ bool ADebugCameraController::GetPivotForOrbit(FVector& PivotLocation) const
 	return false;
 }
 
-void ADebugCameraController::ToggleOrbit()
+void ADebugCameraController::ToggleOrbit(bool bOrbitCenter)
 {
-	bIsOrbitingSelectedActor = (!bIsOrbitingSelectedActor && GetPawnOrSpectator() && GetPivotForOrbit(OrbitPivot));
-
 	if (bIsOrbitingSelectedActor)
 	{
-		LastOrbitPawnLocation = GetPawnOrSpectator()->GetActorLocation();
-		FRotator ViewRotation = (OrbitPivot - LastOrbitPawnLocation).ToOrientationRotator();
-		SetControlRotation(ViewRotation);
-		GetPawnOrSpectator()->FaceRotation(ViewRotation);
+		bIsOrbitingSelectedActor = false;
+	}
+	else
+	{
+		APawn* const CurrentPawn = GetPawnOrSpectator();
+		bOrbitPivotUseCenter = bOrbitCenter;
+		bIsOrbitingSelectedActor = (CurrentPawn && GetPivotForOrbit(OrbitPivot));
+
+		if (bIsOrbitingSelectedActor)
+		{
+			LastOrbitPawnLocation = CurrentPawn->GetActorLocation();
+			FVector ViewVector = OrbitPivot - LastOrbitPawnLocation;
+			float ViewLength = ViewVector.Size();
+
+			if (ViewLength == 0.0f)
+			{
+				bIsOrbitingSelectedActor = false;
+				return;
+			}
+			else if (ViewLength >= MIN_ORBIT_RADIUS)
+			{
+				OrbitRadius = ViewLength;
+			}
+			else
+			{
+				LastOrbitPawnLocation = OrbitPivot - ViewVector.GetSafeNormal() * MIN_ORBIT_RADIUS;
+				CurrentPawn->SetActorLocation(LastOrbitPawnLocation);
+				OrbitRadius = MIN_ORBIT_RADIUS;
+			}
+
+			FRotator ViewRotation = ViewVector.ToOrientationRotator();
+			SetControlRotation(ViewRotation);
+			CurrentPawn->FaceRotation(ViewRotation);
+		}
 	}
 }
 
+void ADebugCameraController::ToggleOrbitCenter()
+{
+	ToggleOrbit(true);
+}
+
+void ADebugCameraController::ToggleOrbitHitPoint()
+{
+	ToggleOrbit(false);
+}
+
 void ADebugCameraController::SelectTargetedObject()
 {
 	FVector CamLoc;
diff --git a/Engine/Source/Runtime/Engine/Private/DebugCameraHUD.cpp b/Engine/Source/Runtime/Engine/Private/DebugCameraHUD.cpp
--- a/Engine/Source/Runtime/Engine/Private/DebugCameraHUD.cpp
+++ b/Engine/Source/Runtime/Engine/Private/DebugCameraHUD.cpp
@@ -101,8 +101,6 @@ void ADebugCameraHUD::PostRender()
 			Canvas->DrawText(RenderFont, SpeedString, X, yl, 1.f, 1.f, FontRenderInfo);
 			yl += Y;
 
-			//Canvas->DrawText(FString::Printf(TEXT("CamLoc:%s CamRot:%s"), *CamLoc.ToString(), *CamRot.ToString() ));
-
 			const TCHAR* CVarComplexName = TEXT("g.DebugCameraTraceComplex");
 			bool bTraceComplex = true;
 
@@ -112,7 +110,7 @@ void ADebugCameraHUD::PostRender()
 
 			FCollisionQueryParams TraceParams(NAME_None, FCollisionQueryParams::GetUnknownStatId(), bTraceComplex, this);
 			FHitResult Hit;
-			bool bHit = GetWorld()->LineTraceSingleByChannel(Hit, CamLoc, CamRot.Vector() * 100000.f + CamLoc, ECC_Pawn, TraceParams);
+			bool bHit = GetWorld()->LineTraceSingleByChannel(Hit, CamLoc, CamRot.Vector() * 100000.f + CamLoc, ECC_Visibility, TraceParams);
 
 			yl += Y;
 			Canvas->DrawText(RenderFont, FString::Printf(TEXT("Trace info (%s = %d):"), CVarComplexName, bTraceComplex ? 1 : 0), X, yl, 1.f, 1.f, FontRenderInfo);
@@ -159,6 +157,11 @@ void ADebugCameraHUD::PostRender()
 				if (!DCC->bIsOrbitingSelectedActor)
 				{
 					DrawDebugLine(GetWorld(), Hit.Location, Hit.Location + Hit.Normal*30.f, FColor::White);
+
+					if (DCC->SelectedActor)
+					{
+						DrawDebugLine(GetWorld(), DCC->SelectedHitPoint.Location, DCC->SelectedHitPoint.Location + DCC->SelectedHitPoint.Normal*30.f, FColor::Red);
+					}
 				}
 			}
 			else
@@ -171,21 +174,28 @@ void ADebugCameraHUD::PostRender()
 			{
 				if (DCC->bIsOrbitingSelectedActor)
 				{
-					yl += Y;
-					Canvas->DrawText(RenderFont, FString::Printf(TEXT("Orbiting actor: '%s'"), *DCC->SelectedActor->GetFName().ToString()), X, yl, 1.f, 1.f, FontRenderInfo);
-					DisplayMaterials(X, yl, Y, Cast<UMeshComponent>(DCC->SelectedComponent));
+					if (DCC->bOrbitPivotUseCenter)
+					{
+						yl += Y;
+						Canvas->DrawText(RenderFont, FString::Printf(TEXT("Orbiting actor (center): '%s'"), *DCC->SelectedActor->GetFName().ToString()), X, yl, 1.f, 1.f, FontRenderInfo);
+					}
+					else
+					{
+						yl += Y;
+						Canvas->DrawText(RenderFont, FString::Printf(TEXT("Orbiting actor (hitpoint): '%s'"), *DCC->SelectedActor->GetFName().ToString()), X, yl, 1.f, 1.f, FontRenderInfo);
+					}
 				}
 				else
 				{
 					yl += Y;
 					Canvas->DrawText(RenderFont, FString::Printf(TEXT("Selected actor: '%s'"), *DCC->SelectedActor->GetFName().ToString()), X, yl, 1.f, 1.f, FontRenderInfo);
-					DisplayMaterials(X, yl, Y, Cast<UMeshComponent>(DCC->SelectedComponent));
 				}
+				DisplayMaterials(X, yl, Y, Cast<UMeshComponent>(DCC->SelectedComponent));
 			}
 
 
 			// controls display
-			yl += Y*15;
+			yl += Y*14;
 			
 			Canvas->SetDrawColor(64, 64, 255, 255);
 			Canvas->DrawText(RenderFont, TEXT("Controls"), X, yl, 1.f, 1.f, FontRenderInfo);
@@ -207,6 +217,8 @@ void ADebugCameraHUD::PostRender()
 			Canvas->DrawText(RenderFont, TEXT("Orbit Selected Actor: O"), X, yl, 1.f, 1.f, FontRenderInfo);
 			yl += Y;
 
+			Canvas->DrawText(RenderFont, TEXT("Orbit Selected Hitpoint: Shift-O"), X, yl, 1.f, 1.f, FontRenderInfo);
+			yl += Y;
 		}
 	}
 #endif
