Fix shape materials not being assigned correctly in immediate mode.

diff --git a/Engine/Source/Runtime/Engine/Private/PhysicsEngine/ImmediatePhysics/ImmediatePhysicsActor.cpp b/Engine/Source/Runtime/Engine/Private/PhysicsEngine/ImmediatePhysics/ImmediatePhysicsActor.cpp
--- a/Engine/Source/Runtime/Engine/Private/PhysicsEngine/ImmediatePhysics/ImmediatePhysicsActor.cpp
+++ b/Engine/Source/Runtime/Engine/Private/PhysicsEngine/ImmediatePhysics/ImmediatePhysicsActor.cpp
@@ -54,6 +54,7 @@ void FActor::CreateGeometry(PxRigidActor* RigidActor, const PxTransform& ActorTo
 		Materials.SetNumUninitialized(NumMaterials);
 		Shape->getMaterials(Materials.GetData(), sizeof(Materials[0]) * NumMaterials);
 
+#if PHYSICS_INTERFACE_LLIMMEDIATE
 		FMaterial* NewMaterial = nullptr;
 		if(NumMaterials > 0)
 		{
@@ -78,6 +79,24 @@ void FActor::CreateGeometry(PxRigidActor* RigidActor, const PxTransform& ActorTo
 			case PxGeometryType::eTRIANGLEMESH: Shapes.Emplace(BodyLocalShape, BoundsCenter, BoundsMag, new PxTriangleMeshGeometry(GeomHolder.triangleMesh().triangleMesh, GeomHolder.triangleMesh().scale, GeomHolder.triangleMesh().meshFlags), NewMaterial); break;
 			default: continue;	//we don't bother with other types
 		}
+#else
+		FMaterial NewMaterial;
+		if (NumMaterials > 0)
+		{
+			NewMaterial = FMaterial(Materials[0]);
+		}
+
+		switch (GeomHolder.getType())
+		{
+		case PxGeometryType::eSPHERE:		Shapes.Emplace(BodyLocalShape, BoundsCenter, BoundsMag, new PxSphereGeometry(GeomHolder.sphere().radius), NewMaterial); break;
+		case PxGeometryType::eCAPSULE:		Shapes.Emplace(BodyLocalShape, BoundsCenter, BoundsMag, new PxCapsuleGeometry(GeomHolder.capsule().radius, GeomHolder.capsule().halfHeight), NewMaterial); break;
+		case PxGeometryType::eBOX:			Shapes.Emplace(BodyLocalShape, BoundsCenter, BoundsMag, new PxBoxGeometry(GeomHolder.box().halfExtents), NewMaterial); break;
+		case PxGeometryType::eCONVEXMESH:	Shapes.Emplace(BodyLocalShape, BoundsCenter, BoundsMag, new PxConvexMeshGeometry(GeomHolder.convexMesh().convexMesh, GeomHolder.convexMesh().scale, GeomHolder.convexMesh().meshFlags), NewMaterial); break;
+		case PxGeometryType::eHEIGHTFIELD:	Shapes.Emplace(BodyLocalShape, BoundsCenter, BoundsMag, new PxHeightFieldGeometry(GeomHolder.heightField().heightField, GeomHolder.heightField().heightFieldFlags, GeomHolder.heightField().heightScale, GeomHolder.heightField().rowScale, GeomHolder.heightField().columnScale), NewMaterial); break;
+		case PxGeometryType::eTRIANGLEMESH: Shapes.Emplace(BodyLocalShape, BoundsCenter, BoundsMag, new PxTriangleMeshGeometry(GeomHolder.triangleMesh().triangleMesh, GeomHolder.triangleMesh().scale, GeomHolder.triangleMesh().meshFlags), NewMaterial); break;
+		default: continue;	//we don't bother with other types
+		}
+#endif
 	}
 }
 
@@ -100,6 +119,7 @@ bool FActor::AddShape(PxShape* InShape)
 	Materials.SetNumUninitialized(NumMaterials);
 	InShape->getMaterials(Materials.GetData(), sizeof(Materials[0]) * NumMaterials);
 
+#if PHYSICS_INTERFACE_LLIMMEDIATE
 	FMaterial* NewMaterial = nullptr;
 	if(NumMaterials > 0)
 	{
@@ -124,7 +144,24 @@ bool FActor::AddShape(PxShape* InShape)
 		case PxGeometryType::eTRIANGLEMESH: Shapes.Emplace(LocalPose, BoundsCenter, BoundsMagnitude, new PxTriangleMeshGeometry(GeomHolder.triangleMesh().triangleMesh, GeomHolder.triangleMesh().scale, GeomHolder.triangleMesh().meshFlags), NewMaterial); break;
 		default: return false;	//we don't bother with other types
 	}
+#else
+	FMaterial NewMaterial;
+	if (NumMaterials > 0)
+	{
+		NewMaterial = FMaterial(Materials[0]);
+	}
 
+	switch (GeomHolder.getType())
+	{
+	case PxGeometryType::eSPHERE:		Shapes.Emplace(LocalPose, BoundsCenter, BoundsMagnitude, new PxSphereGeometry(GeomHolder.sphere().radius), NewMaterial); break;
+	case PxGeometryType::eCAPSULE:		Shapes.Emplace(LocalPose, BoundsCenter, BoundsMagnitude, new PxCapsuleGeometry(GeomHolder.capsule().radius, GeomHolder.capsule().halfHeight), NewMaterial); break;
+	case PxGeometryType::eBOX:			Shapes.Emplace(LocalPose, BoundsCenter, BoundsMagnitude, new PxBoxGeometry(GeomHolder.box().halfExtents), NewMaterial); break;
+	case PxGeometryType::eCONVEXMESH:	Shapes.Emplace(LocalPose, BoundsCenter, BoundsMagnitude, new PxConvexMeshGeometry(GeomHolder.convexMesh().convexMesh, GeomHolder.convexMesh().scale, GeomHolder.convexMesh().meshFlags), NewMaterial); break;
+	case PxGeometryType::eHEIGHTFIELD:	Shapes.Emplace(LocalPose, BoundsCenter, BoundsMagnitude, new PxHeightFieldGeometry(GeomHolder.heightField().heightField, GeomHolder.heightField().heightFieldFlags, GeomHolder.heightField().heightScale, GeomHolder.heightField().rowScale, GeomHolder.heightField().columnScale), NewMaterial); break;
+	case PxGeometryType::eTRIANGLEMESH: Shapes.Emplace(LocalPose, BoundsCenter, BoundsMagnitude, new PxTriangleMeshGeometry(GeomHolder.triangleMesh().triangleMesh, GeomHolder.triangleMesh().scale, GeomHolder.triangleMesh().meshFlags), NewMaterial); break;
+	default: return false;	//we don't bother with other types
+	}
+#endif
 	// If we get here then we succeeded in creating a shape.
 	Shapes.Last().UserData = InShape->userData;
 
diff --git a/Engine/Source/Runtime/Engine/Public/Physics/ImmediatePhysics/ImmediatePhysicsShape.h b/Engine/Source/Runtime/Engine/Public/Physics/ImmediatePhysics/ImmediatePhysicsShape.h
--- a/Engine/Source/Runtime/Engine/Public/Physics/ImmediatePhysics/ImmediatePhysicsShape.h
+++ b/Engine/Source/Runtime/Engine/Public/Physics/ImmediatePhysics/ImmediatePhysicsShape.h
@@ -21,6 +21,9 @@ struct FShape
 #if WITH_PHYSX
 	PxTransform LocalTM;
 	FMaterial* Material;
+#if !PHYSICS_INTERFACE_LLIMMEDIATE
+	FMaterial InternalMaterial;
+#endif
 	PxGeometry* Geometry;
 	PxVec3 BoundsOffset;
 	float BoundsMagnitude;
@@ -44,7 +47,18 @@ struct FShape
 		, UserData(nullptr)
 	{
 	}
-
+#if !PHYSICS_INTERFACE_LLIMMEDIATE
+	FShape(const PxTransform& InLocalTM, const PxVec3& InBoundsOffset, const float InBoundsMagnitude, PxGeometry* InGeometry, FMaterial InMaterial)
+		: LocalTM(InLocalTM)
+		, InternalMaterial(InMaterial)
+		, Geometry(InGeometry)
+		, BoundsOffset(InBoundsOffset)
+		, BoundsMagnitude(InBoundsMagnitude)
+		, UserData(nullptr)
+	{
+		Material = &InternalMaterial;
+	}
+#endif
 #endif
 };
 
