Only log first 50 timers that go over threshold to prevent log spam when server is loading a creative island

diff --git a/Engine/Source/Runtime/Engine/Private/TimerManager.cpp b/Engine/Source/Runtime/Engine/Private/TimerManager.cpp
--- a/Engine/Source/Runtime/Engine/Private/TimerManager.cpp
+++ b/Engine/Source/Runtime/Engine/Private/TimerManager.cpp
@@ -29,6 +29,12 @@ static FAutoConsoleVariableRef CVarDumpTimerLogsThreshold(
 	TEXT("Threshold (in milliseconds) after which we log timer info to try and help track down spikes in the timer code. Disabled when set to 0"),
 	ECVF_Default);
 
+static int32 MaxExpiredTimersToLog = 50;
+static FAutoConsoleVariableRef CVarMaxExpiredTimersToLog(
+	TEXT("TimerManager.MaxExpiredTimersToLog"), 
+	MaxExpiredTimersToLog,
+	TEXT("Maximum number of TimerData exceeding the threshold to log in a single frame."));
+
 namespace
 {
 	void DescribeFTimerDataSafely(FOutputDevice& Ar, const FTimerData& Data)
@@ -555,6 +561,7 @@ void FTimerManager::Tick(float DeltaTime)
 
 	const double StartTime = FPlatformTime::Seconds();
 	bool bDumpTimerLogsThresholdExceeded = false;
+	int32 NbExpiredTimers = 0;
 
 	InternalTime += DeltaTime;
 
@@ -582,7 +589,11 @@ void FTimerManager::Tick(float DeltaTime)
 
 			if (bDumpTimerLogsThresholdExceeded)
 			{
-				DescribeFTimerDataSafely(*GLog, *Top);
+				++NbExpiredTimers;
+				if (NbExpiredTimers <= MaxExpiredTimersToLog)
+				{
+					DescribeFTimerDataSafely(*GLog, *Top);
+				}
 			}
 
 			// Set the relevant level context for this timer
@@ -631,6 +642,7 @@ void FTimerManager::Tick(float DeltaTime)
 				if (DeltaT >= DumpTimerLogsThreshold)
 				{
 					bDumpTimerLogsThresholdExceeded = true;
+                    ++NbExpiredTimers;
 					UE_LOG(LogEngine, Log, TEXT("TimerManager's time threshold of %.2fms exceeded with a deltaT of %.4f, dumping current timer data."), DumpTimerLogsThreshold, DeltaT);
 
 					if (Top)
@@ -670,6 +682,11 @@ void FTimerManager::Tick(float DeltaTime)
 		}
 	}
 
+	if (NbExpiredTimers > MaxExpiredTimersToLog)
+	{
+		UE_LOG(LogEngine, Log, TEXT("TimerManager's caught %d Timers exceeding the time threshold. Only the first %d were logged."), NbExpiredTimers, MaxExpiredTimersToLog);
+	}
+
 	// Timer has been ticked.
 	LastTickedFrame = GFrameCounter;
 
